Updating badges to use the new rest api

diff --git a/Engine/Source/Programs/UnrealGameSync/PostBadgeStatus/Program.cs b/Engine/Source/Programs/UnrealGameSync/PostBadgeStatus/Program.cs
--- a/Engine/Source/Programs/UnrealGameSync/PostBadgeStatus/Program.cs
+++ b/Engine/Source/Programs/UnrealGameSync/PostBadgeStatus/Program.cs
@@ -29,7 +29,7 @@ namespace WriteBadgeStatus
 			if (Arguments.Count > 0 || Name == null || Change == null || Project == null || RestUrl == null || Status == null || Url == null)
 			{
 				Console.WriteLine("Syntax:");
-				Console.WriteLine("  PostBadgeStatus -Name=<Name> -Change=<CL> -Project=<DepotPath> -RestUrls=<Url> -Status=<Status> -Url=<Url>");
+				Console.WriteLine("  PostBadgeStatus -Name=<Name> -Change=<CL> -Project=<DepotPath> -RestUrl=<Url> -Status=<Status> -Url=<Url>");
 				return 1;
 			}
 			if (RestUrl != null)
