Merging latest fixes for UnrealGameSync from //UE4/Dev-Core.

#rnx

diff --git a/Engine/Source/Programs/UnrealGameSync/UnrealGameSync/Controls/WorkspaceControl.cs b/Engine/Source/Programs/UnrealGameSync/UnrealGameSync/Controls/WorkspaceControl.cs
--- a/Engine/Source/Programs/UnrealGameSync/UnrealGameSync/Controls/WorkspaceControl.cs
+++ b/Engine/Source/Programs/UnrealGameSync/UnrealGameSync/Controls/WorkspaceControl.cs
@@ -4105,5 +4105,13 @@ namespace UnrealGameSync
 		{
 			SetBisectStateForSelection(BisectState.Include);
 		}
+
+		private void WorkspaceControl_VisibleChanged(object sender, EventArgs e)
+		{
+			if(PerforceMonitor != null && PerforceMonitor.IsActive != Visible)
+			{
+				PerforceMonitor.IsActive = Visible;
+			}
+		}
 	}
 }
diff --git a/Engine/Source/Programs/UnrealGameSync/UnrealGameSync/DeleteFilesTask.cs b/Engine/Source/Programs/UnrealGameSync/UnrealGameSync/DeleteFilesTask.cs
--- a/Engine/Source/Programs/UnrealGameSync/UnrealGameSync/DeleteFilesTask.cs
+++ b/Engine/Source/Programs/UnrealGameSync/UnrealGameSync/DeleteFilesTask.cs
@@ -28,12 +28,18 @@ namespace UnrealGameSync
 		{
 			StringBuilder FailMessage = new StringBuilder();
 
-			foreach(FileInfo FileToSync in FilesToSync)
+			if(FilesToSync.Count > 0)
 			{
+				List<string> RevisionsToSync = new List<string>();
+				foreach(FileInfo FileToSync in FilesToSync)
+				{
+					RevisionsToSync.Add(String.Format("{0}#have", PerforceUtils.EscapePath(FileToSync.FullName)));
+				}
+
 				StringWriter Log = new StringWriter();
-				if(!Perforce.ForceSync(String.Format("{0}#have", FileToSync.FullName), Log))
+				if(!Perforce.Sync(RevisionsToSync, x => { }, new List<string>(), true, null, Log))
 				{
-					FailMessage.AppendFormat("{0} ({1})\r\n", FileToSync.FullName, Log.ToString().Trim().Replace('\n', ' '));
+					FailMessage.Append(Log.ToString());
 				}
 			}
 
diff --git a/Engine/Source/Programs/UnrealGameSync/UnrealGameSync/DetectProjectSettingsTask.cs b/Engine/Source/Programs/UnrealGameSync/UnrealGameSync/DetectProjectSettingsTask.cs
--- a/Engine/Source/Programs/UnrealGameSync/UnrealGameSync/DetectProjectSettingsTask.cs
+++ b/Engine/Source/Programs/UnrealGameSync/UnrealGameSync/DetectProjectSettingsTask.cs
@@ -310,7 +310,7 @@ namespace UnrealGameSync
 
 			// Figure out if it's an enterprise project
 			List<string> ProjectLines;
-			if(Perforce.Print(NewSelectedClientFileName, out ProjectLines, Log))
+			if(NewSelectedClientFileName.EndsWith(".uproject", StringComparison.InvariantCultureIgnoreCase) && PerforceClient.Print(NewSelectedClientFileName, out ProjectLines, Log))
 			{
 				string Text = String.Join("\n", ProjectLines);
 				bIsEnterpriseProject = Utility.IsEnterpriseProjectFromText(Text);
diff --git a/Engine/Source/Programs/UnrealGameSync/UnrealGameSync/FindFoldersToCleanTask.cs b/Engine/Source/Programs/UnrealGameSync/UnrealGameSync/FindFoldersToCleanTask.cs
--- a/Engine/Source/Programs/UnrealGameSync/UnrealGameSync/FindFoldersToCleanTask.cs
+++ b/Engine/Source/Programs/UnrealGameSync/UnrealGameSync/FindFoldersToCleanTask.cs
@@ -249,7 +249,7 @@ namespace UnrealGameSync
 						}
 						FileFolder = NextFileFolder;
 					}
-					FileFolder.NameToFile.Add(Tokens[Tokens.Length - 1], FileRecord);
+					FileFolder.NameToFile[Tokens[Tokens.Length - 1]] = FileRecord;
 				}
 			}
 
diff --git a/Engine/Source/Programs/UnrealGameSync/UnrealGameSync/Perforce.cs b/Engine/Source/Programs/UnrealGameSync/UnrealGameSync/Perforce.cs
--- a/Engine/Source/Programs/UnrealGameSync/UnrealGameSync/Perforce.cs
+++ b/Engine/Source/Programs/UnrealGameSync/UnrealGameSync/Perforce.cs
@@ -5,8 +5,10 @@ using System.Collections.Generic;
 using System.Diagnostics;
 using System.IO;
 using System.Linq;
+using System.Reflection;
 using System.Text;
 using System.Threading.Tasks;
+using System.Windows.Forms;
 
 namespace UnrealGameSync
 {
@@ -1053,7 +1055,7 @@ namespace UnrealGameSync
 
 		public bool Stat(string Filter, out List<PerforceFileRecord> FileRecords, TextWriter Log)
 		{
-			return RunCommand(String.Format("fstat \"{0}\"", Filter), out FileRecords, CommandOptions.None, Log);
+			return RunCommand(String.Format("fstat \"{0}\"", Filter), out FileRecords, CommandOptions.IgnoreFilesNotOnClientError, Log);
 		}
 
 		public bool Sync(string Filter, TextWriter Log)
@@ -1061,7 +1063,7 @@ namespace UnrealGameSync
 			return RunCommand("sync " + Filter, CommandOptions.IgnoreFilesUpToDateError, Log);
 		}
 
-		public bool Sync(List<string> FileRevisions, Action<PerforceFileRecord> SyncOutput, List<string> TamperedFiles, PerforceSyncOptions Options, TextWriter Log)
+		public bool Sync(List<string> FileRevisions, Action<PerforceFileRecord> SyncOutput, List<string> TamperedFiles, bool bForce, PerforceSyncOptions Options, TextWriter Log)
 		{
 			// Write all the files we want to sync to a temp file
 			string TempFileName = Path.GetTempFileName();
@@ -1085,6 +1087,10 @@ namespace UnrealGameSync
 						CommandLine.AppendFormat(" -v net.tcpsize={0}", Options.TcpBufferSize);
 					}
 					CommandLine.Append(" sync");
+					if(bForce)
+					{
+						CommandLine.AppendFormat(" -f");
+					}
 					if(Options != null && Options.NumThreads > 1)
 					{
 						CommandLine.AppendFormat(" --parallel=threads={0}", Options.NumThreads);
@@ -1137,7 +1143,7 @@ namespace UnrealGameSync
 
 		public bool SyncPreview(string Filter, int ChangeNumber, bool bOnlyFilesInThisChange, out List<PerforceFileRecord> FileRecords, TextWriter Log)
 		{
-			return RunCommand(String.Format("sync -n {0}@{1}{2}", Filter, bOnlyFilesInThisChange? "=" : "", ChangeNumber), out FileRecords, CommandOptions.IgnoreFilesUpToDateError | CommandOptions.IgnoreNoSuchFilesError, Log);
+			return RunCommand(String.Format("sync -n {0}@{1}{2}", Filter, bOnlyFilesInThisChange? "=" : "", ChangeNumber), out FileRecords, CommandOptions.IgnoreFilesUpToDateError | CommandOptions.IgnoreNoSuchFilesError | CommandOptions.IgnoreFilesNotInClientViewError, Log);
 		}
 
 		public bool ForceSync(string Filter, TextWriter Log)
@@ -1304,6 +1310,7 @@ namespace UnrealGameSync
 			{
 				FullCommandLine.Append("-s ");
 			}
+			FullCommandLine.AppendFormat("-zprog=UGS -zversion={0} ", Assembly.GetExecutingAssembly().GetName().Version.ToString());
 			FullCommandLine.Append(CommandLine);
 
 			return FullCommandLine.ToString();
@@ -1541,6 +1548,16 @@ namespace UnrealGameSync
 			}
 		}
 
+		static public string EscapePath(string Path)
+		{
+			string NewPath = Path;
+			NewPath = NewPath.Replace("%", "%25");
+			NewPath = NewPath.Replace("*", "%2A");
+			NewPath = NewPath.Replace("#", "%23");
+			NewPath = NewPath.Replace("@", "%40");
+			return NewPath;
+		}
+
 		static public string UnescapePath(string Path)
 		{
 			string NewPath = Path;
diff --git a/Engine/Source/Programs/UnrealGameSync/UnrealGameSync/PerforceModalTask.cs b/Engine/Source/Programs/UnrealGameSync/UnrealGameSync/PerforceModalTask.cs
--- a/Engine/Source/Programs/UnrealGameSync/UnrealGameSync/PerforceModalTask.cs
+++ b/Engine/Source/Programs/UnrealGameSync/UnrealGameSync/PerforceModalTask.cs
@@ -45,7 +45,14 @@ namespace UnrealGameSync
 					PerforceConnection Perforce = new PerforceConnection(UserName, null, null);
 					if(ProjectFileName != null)
 					{
-						Directory.SetCurrentDirectory(Path.GetDirectoryName(ProjectFileName));
+						try
+						{
+							Directory.SetCurrentDirectory(Path.GetDirectoryName(ProjectFileName));
+						}
+						catch
+						{
+							// Just ignore
+						}
 					}
 
 					string NewServerAndPort;
diff --git a/Engine/Source/Programs/UnrealGameSync/UnrealGameSync/PerforceMonitor.cs b/Engine/Source/Programs/UnrealGameSync/UnrealGameSync/PerforceMonitor.cs
--- a/Engine/Source/Programs/UnrealGameSync/UnrealGameSync/PerforceMonitor.cs
+++ b/Engine/Source/Programs/UnrealGameSync/UnrealGameSync/PerforceMonitor.cs
@@ -92,6 +92,12 @@ namespace UnrealGameSync
 			}
 		}
 
+		public bool IsActive
+		{
+			get;
+			set;
+		}
+
 		public string LastStatusMessage
 		{
 			get;
@@ -177,7 +183,7 @@ namespace UnrealGameSync
 				}
 
 				// Wait for another request, or scan for new builds after a timeout
-				RefreshEvent.WaitOne(60 * 1000);
+				RefreshEvent.WaitOne((IsActive? 1 : 10) * 60 * 1000);
 			}
 		}
 
diff --git a/Engine/Source/Programs/UnrealGameSync/UnrealGameSync/ProgramApplicationContext.cs b/Engine/Source/Programs/UnrealGameSync/UnrealGameSync/ProgramApplicationContext.cs
--- a/Engine/Source/Programs/UnrealGameSync/UnrealGameSync/ProgramApplicationContext.cs
+++ b/Engine/Source/Programs/UnrealGameSync/UnrealGameSync/ProgramApplicationContext.cs
@@ -24,6 +24,7 @@ namespace UnrealGameSync
 		bool bRestoreState;
 		string UpdateSpawn;
 		bool bUnstable;
+		bool bIsClosing;
 
 		BoundedLogWriter Log;
 		UserSettings Settings;
@@ -217,8 +218,9 @@ namespace UnrealGameSync
 
 		private void OnUpdateAvailable()
 		{
-			if(MainWindowInstance != null && MainWindowInstance.CanPerformUpdate())
+			if(MainWindowInstance != null && !bIsClosing && MainWindowInstance.CanPerformUpdate())
 			{
+				bIsClosing = true;
 				MainWindowInstance.ForceClose();
 			}
 		}
diff --git a/Engine/Source/Programs/UnrealGameSync/UnrealGameSync/Properties/AssemblyInfo.cs b/Engine/Source/Programs/UnrealGameSync/UnrealGameSync/Properties/AssemblyInfo.cs
--- a/Engine/Source/Programs/UnrealGameSync/UnrealGameSync/Properties/AssemblyInfo.cs
+++ b/Engine/Source/Programs/UnrealGameSync/UnrealGameSync/Properties/AssemblyInfo.cs
@@ -34,5 +34,5 @@ using System.Runtime.InteropServices;
 // You can specify all the values or you can default the Build and Revision Numbers 
 // by using the '*' as shown below:
 // [assembly: AssemblyVersion("1.0.*")]
-[assembly: AssemblyVersion("1.133.*")]
-[assembly: AssemblyFileVersion("1.133.0.0")]
+[assembly: AssemblyVersion("1.139.*")]
+[assembly: AssemblyFileVersion("1.139.0.0")]
diff --git a/Engine/Source/Programs/UnrealGameSync/UnrealGameSync/Workspace.cs b/Engine/Source/Programs/UnrealGameSync/UnrealGameSync/Workspace.cs
--- a/Engine/Source/Programs/UnrealGameSync/UnrealGameSync/Workspace.cs
+++ b/Engine/Source/Programs/UnrealGameSync/UnrealGameSync/Workspace.cs
@@ -451,7 +451,7 @@ namespace UnrealGameSync
 					}
 					SyncDepotPaths.RemoveAll(x => !Filter.Matches(x));
 
-					// Findd all the depot paths that will be synced
+					// Find all the depot paths that will be synced
 					HashSet<string> RemainingDepotPaths = new HashSet<string>(StringComparer.InvariantCultureIgnoreCase);
 					RemainingDepotPaths.UnionWith(RemoveDepotPaths);
 					RemainingDepotPaths.UnionWith(SyncDepotPaths);
@@ -463,7 +463,7 @@ namespace UnrealGameSync
 
 					// Sync them all
 					List<string> TamperedFiles = new List<string>();
-					if(!Perforce.Sync(SyncRevisions, Record => UpdateSyncProgress(Record, RemainingDepotPaths, SyncRevisions.Count), TamperedFiles, Context.PerforceSyncOptions, Log))
+					if(!Perforce.Sync(SyncRevisions, Record => UpdateSyncProgress(Record, RemainingDepotPaths, SyncRevisions.Count), TamperedFiles, false, Context.PerforceSyncOptions, Log))
 					{
 						StatusMessage = "Aborted sync due to errors.";
 						return WorkspaceUpdateResult.FailedToSync;
@@ -984,6 +984,7 @@ namespace UnrealGameSync
 
 		public List<string> GetSyncPaths(bool bSyncAllProjects, string[] SyncFilter)
 		{
+			// Add the default project paths
 			List<string> SyncPaths = new List<string>();
 			if(bSyncAllProjects || !SelectedClientFileName.EndsWith(".uproject", StringComparison.InvariantCultureIgnoreCase))
 			{
@@ -999,6 +1000,8 @@ namespace UnrealGameSync
 				}
 				SyncPaths.Add(PerforceUtils.GetClientOrDepotDirectoryName(SelectedClientFileName) + "/...");
 			}
+
+			// Apply the sync filter to that list. We only want inclusive rules in the output list, but we can manually apply exclusions to previous entries.
 			if(SyncFilter != null)
 			{
 				foreach(string SyncPath in SyncFilter)
@@ -1008,8 +1011,31 @@ namespace UnrealGameSync
 					{
 						SyncPaths.Add(ClientRootPath + TrimSyncPath);
 					}
+					else if(TrimSyncPath.StartsWith("-/") && TrimSyncPath.EndsWith("..."))
+					{
+						SyncPaths.RemoveAll(x => x.StartsWith(ClientRootPath + TrimSyncPath.Substring(1, TrimSyncPath.Length - 4)));
+					}
 				}
 			}
+
+			// Sort the remaining paths by length, and remove any paths which are included twice
+			SyncPaths = SyncPaths.OrderBy(x => x.Length).ToList();
+			for(int Idx = 0; Idx < SyncPaths.Count; Idx++)
+			{
+				string SyncPath = SyncPaths[Idx];
+				if(SyncPath.EndsWith("..."))
+				{
+					string SyncPathPrefix = SyncPath.Substring(0, SyncPath.Length - 3);
+					for(int OtherIdx = SyncPaths.Count - 1; OtherIdx > Idx; OtherIdx--)
+					{
+						if(SyncPaths[OtherIdx].StartsWith(SyncPathPrefix))
+						{
+							SyncPaths.RemoveAt(OtherIdx);
+						}
+					}
+				}
+			}
+
 			return SyncPaths;
 		}
 		
