Merge latest UGS fixes (up to 1.140) to 4.20 branch.

diff --git a/Engine/Source/Programs/UnrealGameSync/UnrealGameSync/Controls/WorkspaceControl.cs b/Engine/Source/Programs/UnrealGameSync/UnrealGameSync/Controls/WorkspaceControl.cs
--- a/Engine/Source/Programs/UnrealGameSync/UnrealGameSync/Controls/WorkspaceControl.cs
+++ b/Engine/Source/Programs/UnrealGameSync/UnrealGameSync/Controls/WorkspaceControl.cs
@@ -126,7 +126,7 @@ namespace UnrealGameSync
 		IWorkspaceControlOwner Owner;
 		string ApiUrl;
 		string DataFolder;
-		BoundedLogWriter Log;
+		LineBasedTextWriter Log;
 
 		UserSettings Settings;
 		UserWorkspaceSettings WorkspaceSettings;
@@ -216,7 +216,7 @@ namespace UnrealGameSync
 		bool bUpdateBuildMetadataPosted;
 		bool bUpdateReviewsPosted;
 
-		public WorkspaceControl(IWorkspaceControlOwner InOwner, string InApiUrl, string InOriginalExecutableFileName, DetectProjectSettingsTask DetectSettings, BoundedLogWriter InLog, UserSettings InSettings)
+		public WorkspaceControl(IWorkspaceControlOwner InOwner, string InApiUrl, string InOriginalExecutableFileName, DetectProjectSettingsTask DetectSettings, LineBasedTextWriter InLog, UserSettings InSettings)
 		{
 			InitializeComponent();
 
@@ -648,7 +648,6 @@ namespace UnrealGameSync
 			Context.StartTime = DateTime.UtcNow;
 			Context.PerforceSyncOptions = (PerforceSyncOptions)Settings.SyncOptions.Clone();
 
-			Log.WriteLine();
 			Log.WriteLine("Updating workspace at {0}...", Context.StartTime.ToLocalTime().ToString());
 			Log.WriteLine("  ChangeNumber={0}", Context.ChangeNumber);
 			Log.WriteLine("  Options={0}", Context.Options.ToString());
diff --git a/Engine/Source/Programs/UnrealGameSync/UnrealGameSync/FindFoldersToCleanTask.cs b/Engine/Source/Programs/UnrealGameSync/UnrealGameSync/FindFoldersToCleanTask.cs
--- a/Engine/Source/Programs/UnrealGameSync/UnrealGameSync/FindFoldersToCleanTask.cs
+++ b/Engine/Source/Programs/UnrealGameSync/UnrealGameSync/FindFoldersToCleanTask.cs
@@ -104,44 +104,56 @@ namespace UnrealGameSync
 
 		void MergeTrees(FolderToClean LocalFolder, PerforceHaveFolder PerforceFolder, HashSet<string> OpenClientPaths)
 		{
-			// Loop through all the local sub-folders
-			foreach(FolderToClean LocalSubFolder in LocalFolder.NameToSubFolder.Values)
+			if(PerforceFolder == null)
 			{
-				PerforceHaveFolder PerforceSubFolder;
-				if(PerforceFolder.NameToSubFolder.TryGetValue(LocalSubFolder.Name, out PerforceSubFolder))
+				// Loop through all the local sub-folders
+				foreach(FolderToClean LocalSubFolder in LocalFolder.NameToSubFolder.Values)
 				{
-					MergeTrees(LocalSubFolder, PerforceSubFolder, OpenClientPaths);
+					MergeTrees(LocalSubFolder, null, OpenClientPaths);
 				}
-			}
 
-			// Also merge all the Perforce folders that no longer exist
-			foreach(KeyValuePair<string, PerforceHaveFolder> PerforceSubFolderPair in PerforceFolder.NameToSubFolder)
+				// Delete everything
+				LocalFolder.FilesToDelete.AddRange(LocalFolder.NameToFile.Values);
+			}
+			else
 			{
-				FolderToClean LocalSubFolder;
-				if(!LocalFolder.NameToSubFolder.TryGetValue(PerforceSubFolderPair.Key, out LocalSubFolder))
+				// Loop through all the local sub-folders
+				foreach(FolderToClean LocalSubFolder in LocalFolder.NameToSubFolder.Values)
 				{
-					LocalSubFolder = new FolderToClean(new DirectoryInfo(Path.Combine(LocalFolder.Directory.FullName, PerforceSubFolderPair.Key)));
-					MergeTrees(LocalSubFolder, PerforceSubFolderPair.Value, OpenClientPaths);
-					LocalFolder.NameToSubFolder.Add(LocalSubFolder.Name, LocalSubFolder);
+					PerforceHaveFolder PerforceSubFolder;
+					PerforceFolder.NameToSubFolder.TryGetValue(LocalSubFolder.Name, out PerforceSubFolder);
+					MergeTrees(LocalSubFolder, PerforceSubFolder, OpenClientPaths);
 				}
-			}
 
-			// Find all the files that need to be re-synced
-			foreach(KeyValuePair<string, PerforceFileRecord> FilePair in PerforceFolder.NameToFile)
-			{
-				FileInfo LocalFile;
-				if(!LocalFolder.NameToFile.TryGetValue(FilePair.Key, out LocalFile))
+				// Also merge all the Perforce folders that no longer exist
+				foreach(KeyValuePair<string, PerforceHaveFolder> PerforceSubFolderPair in PerforceFolder.NameToSubFolder)
 				{
-					LocalFolder.FilesToSync.Add(new FileInfo(Path.Combine(LocalFolder.Directory.FullName, FilePair.Key)));
+					FolderToClean LocalSubFolder;
+					if(!LocalFolder.NameToSubFolder.TryGetValue(PerforceSubFolderPair.Key, out LocalSubFolder))
+					{
+						LocalSubFolder = new FolderToClean(new DirectoryInfo(Path.Combine(LocalFolder.Directory.FullName, PerforceSubFolderPair.Key)));
+						MergeTrees(LocalSubFolder, PerforceSubFolderPair.Value, OpenClientPaths);
+						LocalFolder.NameToSubFolder.Add(LocalSubFolder.Name, LocalSubFolder);
+					}
 				}
-				else if((FilePair.Value.Flags & PerforceFileFlags.AlwaysWritable) == 0 && (LocalFile.Attributes & FileAttributes.ReadOnly) == 0 && !OpenClientPaths.Contains(FilePair.Value.ClientPath))
+
+				// Find all the files that need to be re-synced
+				foreach(KeyValuePair<string, PerforceFileRecord> FilePair in PerforceFolder.NameToFile)
 				{
-					LocalFolder.FilesToSync.Add(LocalFile);
+					FileInfo LocalFile;
+					if(!LocalFolder.NameToFile.TryGetValue(FilePair.Key, out LocalFile))
+					{
+						LocalFolder.FilesToSync.Add(new FileInfo(Path.Combine(LocalFolder.Directory.FullName, FilePair.Key)));
+					}
+					else if((FilePair.Value.Flags & PerforceFileFlags.AlwaysWritable) == 0 && (LocalFile.Attributes & FileAttributes.ReadOnly) == 0 && !OpenClientPaths.Contains(FilePair.Value.ClientPath))
+					{
+						LocalFolder.FilesToSync.Add(LocalFile);
+					}
 				}
-			}
 
-			// Find all the files that should be deleted
-			LocalFolder.FilesToDelete.AddRange(LocalFolder.NameToFile.Values.Where(x => !PerforceFolder.NameToFile.ContainsKey(x.Name)));
+				// Find all the files that should be deleted
+				LocalFolder.FilesToDelete.AddRange(LocalFolder.NameToFile.Values.Where(x => !PerforceFolder.NameToFile.ContainsKey(x.Name)));
+			}
 
 			// Figure out if this folder is just an empty directory that needs to be removed
 			LocalFolder.bEmptyLeaf = LocalFolder.NameToFile.Count == 0 && LocalFolder.NameToSubFolder.Count == 0 && LocalFolder.FilesToSync.Count == 0;
diff --git a/Engine/Source/Programs/UnrealGameSync/UnrealGameSync/Forms/MainWindow.cs b/Engine/Source/Programs/UnrealGameSync/UnrealGameSync/Forms/MainWindow.cs
--- a/Engine/Source/Programs/UnrealGameSync/UnrealGameSync/Forms/MainWindow.cs
+++ b/Engine/Source/Programs/UnrealGameSync/UnrealGameSync/Forms/MainWindow.cs
@@ -57,7 +57,7 @@ namespace UnrealGameSync
 
 		string ApiUrl;
 		string DataFolder;
-		BoundedLogWriter Log;
+		LineBasedTextWriter Log;
 		UserSettings Settings;
 		int TabMenu_TabIdx = -1;
 		int ChangingWorkspacesRefCount;
@@ -73,7 +73,7 @@ namespace UnrealGameSync
 
 		IMainWindowTabPanel CurrentTabPanel;
 
-		public MainWindow(string InApiUrl, string InDataFolder, bool bInRestoreStateOnLoad, string InOriginalExecutableFileName, List<DetectProjectSettingsResult> StartupProjects, BoundedLogWriter InLog, UserSettings InSettings)
+		public MainWindow(string InApiUrl, string InDataFolder, bool bInRestoreStateOnLoad, string InOriginalExecutableFileName, List<DetectProjectSettingsResult> StartupProjects, LineBasedTextWriter InLog, UserSettings InSettings)
 		{
 			InitializeComponent();
 
@@ -444,6 +444,8 @@ namespace UnrealGameSync
 
 				TimeSpan IntervalToFirstTick = NextScheduleTime - CurrentTime;
 				ScheduleTimer = new System.Threading.Timer(x => MainThreadSynchronizationContext.Post((o) => { if(!IsDisposed){ ScheduleTimerElapsed(); } }, null), null, IntervalToFirstTick, TimeSpan.FromDays(1));
+
+				Log.WriteLine("Schedule: Started ScheduleTimer for {0} ({1} remaining)", NextScheduleTime, IntervalToFirstTick);
 			}
 		}
 
@@ -453,16 +455,20 @@ namespace UnrealGameSync
 			{
 				ScheduleTimer.Dispose();
 				ScheduleTimer = null;
+				Log.WriteLine("Schedule: Stopped ScheduleTimer");
 			}
 			StopScheduleSettledTimer();
 		}
 
 		private void ScheduleTimerElapsed()
 		{
+			Log.WriteLine("Schedule: Timer Elapsed");
+
 			// Try to open any missing tabs. 
 			int NumInitialTabs = TabControl.GetTabCount();
 			foreach (UserSelectedProjectSettings ScheduledProject in Settings.ScheduleProjects)
 			{
+				Log.WriteLine("Schedule: Attempting to open {0}", ScheduledProject);
 				TryOpenProject(ScheduledProject, -1, OpenProjectOptions.Quiet);
 			}
 
@@ -481,6 +487,7 @@ namespace UnrealGameSync
 		{
 			StopScheduleSettledTimer();
 			ScheduleSettledTimer = new System.Threading.Timer(x => MainThreadSynchronizationContext.Post((o) => { if(!IsDisposed){ ScheduleSettledTimerElapsed(); } }, null), null, TimeSpan.FromSeconds(20.0), TimeSpan.FromMilliseconds(-1.0));
+			Log.WriteLine("Schedule: Started ScheduleSettledTimer");
 		}
 
 		private void StopScheduleSettledTimer()
@@ -489,17 +496,25 @@ namespace UnrealGameSync
 			{
 				ScheduleSettledTimer.Dispose();
 				ScheduleSettledTimer = null;
+
+				Log.WriteLine("Schedule: Stopped ScheduleSettledTimer");
 			}
 		}
 
 		private void ScheduleSettledTimerElapsed()
 		{
+			Log.WriteLine("Schedule: Starting Sync");
 			for(int Idx = 0; Idx < TabControl.GetTabCount(); Idx++)
 			{
 				WorkspaceControl Workspace = TabControl.GetTabData(Idx) as WorkspaceControl;
-				if(Workspace != null && (Settings.ScheduleAnyOpenProject || Settings.ScheduleProjects.Contains(Workspace.SelectedProject)))
+				if(Workspace != null)
 				{
-					Workspace.ScheduleTimerElapsed();
+					Log.WriteLine("Schedule: Considering {0}", Workspace.SelectedFileName);
+					if(Settings.ScheduleAnyOpenProject || Settings.ScheduleProjects.Contains(Workspace.SelectedProject))
+					{
+						Log.WriteLine("Schedule: Starting Sync");
+						Workspace.ScheduleTimerElapsed();
+					}
 				}
 			}
 		}
@@ -631,7 +646,8 @@ namespace UnrealGameSync
 
 		int TryOpenProject(UserSelectedProjectSettings Project, int ReplaceTabIdx, OpenProjectOptions Options = OpenProjectOptions.None)
 		{
-			using(DetectProjectSettingsTask DetectProjectSettings = new DetectProjectSettingsTask(Project, DataFolder, Log))
+			Log.WriteLine("Detecting settings for {0}", Project);
+			using(DetectProjectSettingsTask DetectProjectSettings = new DetectProjectSettingsTask(Project, DataFolder, new PrefixedTextWriter("  ", Log)))
 			{
 				string ErrorMessage;
 				if(ModalTask.Execute(this, DetectProjectSettings, "Opening Project", "Opening project, please wait...", out ErrorMessage) != ModalTaskResult.Succeeded)
@@ -660,6 +676,7 @@ namespace UnrealGameSync
 					{
 						if(Workspace.SelectedFileName.Equals(ProjectSettings.NewSelectedFileName, StringComparison.InvariantCultureIgnoreCase))
 						{
+							Log.WriteLine("  Already open in tab {0}", TabIdx);
 							if((Options & OpenProjectOptions.Quiet) == 0)
 							{
 								TabControl.SelectTab(TabIdx);
@@ -670,10 +687,12 @@ namespace UnrealGameSync
 						{
 							if((Options & OpenProjectOptions.Quiet) == 0 && MessageBox.Show(String.Format("{0} is already open under {1}.\n\nWould you like to close it?", Path.GetFileNameWithoutExtension(Workspace.SelectedFileName), Workspace.BranchDirectoryName, Path.GetFileNameWithoutExtension(ProjectSettings.NewSelectedFileName)), "Branch already open", MessageBoxButtons.YesNo) == System.Windows.Forms.DialogResult.Yes)
 							{
+								Log.WriteLine("  Another project already open in this workspace, tab {0}. Replacing.", TabIdx);
 								TabControl.RemoveTab(TabIdx);
 							}
 							else
 							{
+								Log.WriteLine("  Another project already open in this workspace, tab {0}. Aborting.", TabIdx);
 								return -1;
 							}
 						}
@@ -709,10 +728,12 @@ namespace UnrealGameSync
 			if(ReplaceTabIdx == -1)
 			{
 				int NewTabIdx = TabControl.InsertTab(-1, NewTabName, NewWorkspace);
+				Log.WriteLine("  Inserted tab {0}", NewTabIdx);
 				return NewTabIdx;
 			}
 			else
 			{
+				Log.WriteLine("  Replacing tab {0}", ReplaceTabIdx);
 				TabControl.InsertTab(ReplaceTabIdx + 1, NewTabName, NewWorkspace);
 				TabControl.RemoveTab(ReplaceTabIdx);
 				return ReplaceTabIdx;
diff --git a/Engine/Source/Programs/UnrealGameSync/UnrealGameSync/Forms/ScheduleWindow.cs b/Engine/Source/Programs/UnrealGameSync/UnrealGameSync/Forms/ScheduleWindow.cs
--- a/Engine/Source/Programs/UnrealGameSync/UnrealGameSync/Forms/ScheduleWindow.cs
+++ b/Engine/Source/Programs/UnrealGameSync/UnrealGameSync/Forms/ScheduleWindow.cs
@@ -33,7 +33,7 @@ namespace UnrealGameSync
 
 			foreach(UserSelectedProjectSettings Project in LocalFileToProject.Values.OrderBy(x => x.ToString()))
 			{
-				bool Enabled = ScheduledProjects.Contains(Project);
+				bool Enabled = ScheduledProjects.Any(x => x.LocalPath == Project.LocalPath);
 				ProjectListBox.Items.Add(Project, Enabled);
 			}
 
diff --git a/Engine/Source/Programs/UnrealGameSync/UnrealGameSync/OutputAdapters.cs b/Engine/Source/Programs/UnrealGameSync/UnrealGameSync/OutputAdapters.cs
--- a/Engine/Source/Programs/UnrealGameSync/UnrealGameSync/OutputAdapters.cs
+++ b/Engine/Source/Programs/UnrealGameSync/UnrealGameSync/OutputAdapters.cs
@@ -1,4 +1,4 @@
-﻿// Copyright 1998-2018 Epic Games, Inc. All Rights Reserved.
+// Copyright 1998-2018 Epic Games, Inc. All Rights Reserved.
 
 using System;
 using System.Collections.Generic;
@@ -77,6 +77,30 @@ namespace UnrealGameSync
 			Inner.WriteLine(Prefix + Line);
 		}
 	}
+
+	class TimestampLogWriter : LineBasedTextWriter
+	{
+		LineBasedTextWriter Inner;
+
+		public TimestampLogWriter(LineBasedTextWriter Inner)
+		{
+			this.Inner = Inner;
+		}
+
+		protected override void Dispose(bool bDisposing)
+		{
+			if(Inner != null)
+			{
+				Inner.Dispose();
+				Inner = null;
+			}
+		}
+
+		protected override void FlushLine(string Line)
+		{
+			Inner.WriteLine("[{0}] {1}", DateTime.Now, Line);
+		}
+	}
 	
 	class BoundedLogWriter : LineBasedTextWriter
 	{
@@ -277,11 +301,12 @@ namespace UnrealGameSync
 
 		protected override void FlushLine(string Line)
 		{
-			if(Line.StartsWith(DirectivePrefix))
+			string TrimLine = Line.Trim();
+			if(TrimLine.StartsWith(DirectivePrefix))
 			{
 				// Line that just contains a progress directive
 				bool bSkipLine = false;
-				ProcessInternal(Line.Substring(DirectivePrefix.Length), ref bSkipLine);
+				ProcessInternal(TrimLine.Substring(DirectivePrefix.Length), ref bSkipLine);
 			}
 			else
 			{
@@ -289,7 +314,6 @@ namespace UnrealGameSync
 				string RemainingLine = Line;
 
 				// Look for a progress directive at the end of a line, in square brackets
-				string TrimLine = Line.TrimEnd();
 				if(TrimLine.EndsWith("]"))
 				{
 					for(int LastIdx = TrimLine.Length - 2; LastIdx >= 0 && TrimLine[LastIdx] != ']'; LastIdx--)
diff --git a/Engine/Source/Programs/UnrealGameSync/UnrealGameSync/ProgramApplicationContext.cs b/Engine/Source/Programs/UnrealGameSync/UnrealGameSync/ProgramApplicationContext.cs
--- a/Engine/Source/Programs/UnrealGameSync/UnrealGameSync/ProgramApplicationContext.cs
+++ b/Engine/Source/Programs/UnrealGameSync/UnrealGameSync/ProgramApplicationContext.cs
@@ -26,7 +26,7 @@ namespace UnrealGameSync
 		bool bUnstable;
 		bool bIsClosing;
 
-		BoundedLogWriter Log;
+		TimestampLogWriter Log;
 		UserSettings Settings;
 		ActivationListener ActivationListener;
 
@@ -64,7 +64,7 @@ namespace UnrealGameSync
 			MainThreadSynchronizationContext = WindowsFormsSynchronizationContext.Current;
 
 			// Create the log file
-			Log = new BoundedLogWriter(Path.Combine(DataFolder, "UnrealGameSync.log"));
+			Log = new TimestampLogWriter(new BoundedLogWriter(Path.Combine(DataFolder, "UnrealGameSync.log")));
 			Log.WriteLine("Application version: {0}", Assembly.GetExecutingAssembly().GetName().Version);
 			Log.WriteLine("Started at {0}", DateTime.Now.ToString());
 
@@ -147,7 +147,8 @@ namespace UnrealGameSync
 			List<DetectProjectSettingsTask> Tasks = new List<DetectProjectSettingsTask>();
 			foreach(UserSelectedProjectSettings OpenProject in Settings.OpenProjects)
 			{
-				Tasks.Add(new DetectProjectSettingsTask(OpenProject, DataFolder, Log));
+				Log.WriteLine("Detecting settings for {0}", OpenProject);
+				Tasks.Add(new DetectProjectSettingsTask(OpenProject, DataFolder, new PrefixedTextWriter("  ", Log)));
 			}
 
 			// Detect settings for the project we want to open
diff --git a/Engine/Source/Programs/UnrealGameSync/UnrealGameSync/Properties/AssemblyInfo.cs b/Engine/Source/Programs/UnrealGameSync/UnrealGameSync/Properties/AssemblyInfo.cs
--- a/Engine/Source/Programs/UnrealGameSync/UnrealGameSync/Properties/AssemblyInfo.cs
+++ b/Engine/Source/Programs/UnrealGameSync/UnrealGameSync/Properties/AssemblyInfo.cs
@@ -34,5 +34,5 @@ using System.Runtime.InteropServices;
 // You can specify all the values or you can default the Build and Revision Numbers 
 // by using the '*' as shown below:
 // [assembly: AssemblyVersion("1.0.*")]
-[assembly: AssemblyVersion("1.139.*")]
-[assembly: AssemblyFileVersion("1.139.0.0")]
+[assembly: AssemblyVersion("1.140.*")]
+[assembly: AssemblyFileVersion("1.140.0.0")]
diff --git a/Engine/Source/Programs/UnrealGameSync/UnrealGameSync/Workspace.cs b/Engine/Source/Programs/UnrealGameSync/UnrealGameSync/Workspace.cs
--- a/Engine/Source/Programs/UnrealGameSync/UnrealGameSync/Workspace.cs
+++ b/Engine/Source/Programs/UnrealGameSync/UnrealGameSync/Workspace.cs
@@ -731,6 +731,7 @@ namespace UnrealGameSync
 							Log.WriteLine("Removing {0} binaries...", ArchiveTypeAndDepotPath.Key);
 							Progress.Set(String.Format("Removing {0} binaries...", ArchiveTypeAndDepotPath.Key), 0.0f);
 							ArchiveUtils.RemoveExtractedFiles(LocalRootPath, ManifestFileName, Progress, Log);
+							File.Delete(ManifestFileName);
 							Log.WriteLine();
 						}
 
