Removing MSSQL support from UGSAPI for 4.20

diff --git a/Engine/Source/Programs/UnrealGameSync/UnrealGameSyncWebAPI/Connectors/SqlConnector.cs b/Engine/Source/Programs/UnrealGameSync/UnrealGameSyncWebAPI/Connectors/SqlConnector.cs
--- a/Engine/Source/Programs/UnrealGameSync/UnrealGameSyncWebAPI/Connectors/SqlConnector.cs
+++ b/Engine/Source/Programs/UnrealGameSync/UnrealGameSyncWebAPI/Connectors/SqlConnector.cs
@@ -12,20 +12,19 @@ namespace UnrealGameSyncWebAPI.Connectors
 {
 	public static class SqlConnector
 	{
-		private static string MsSqlConnectionString = System.Configuration.ConfigurationManager.AppSettings["MsSqlConnectionString"];
-		private static string SQLiteConnectionString = System.Configuration.ConfigurationManager.AppSettings["SQLiteConnectionString"];
+		private static string ConnectionString = System.Configuration.ConfigurationManager.AppSettings["ConnectionString"];
 		public static long[] GetLastIds()
 		{
 			// Get Last IDs
 			long LastEventId = 0;
 			long LastCommentId = 0;
 			long LastBuildID = 0;
-			using (SqlConnection Connection = new SqlConnection(MsSqlConnectionString))
+			using (SQLiteConnection Connection = new SQLiteConnection(ConnectionString))
 			{
 				Connection.Open();
-				using (SqlCommand Command = new SqlCommand("SELECT MAX(ID) FROM dbo.UserVotes;", Connection))
+				using (SQLiteCommand Command = new SQLiteCommand("SELECT MAX(ID) FROM [UserVotes];", Connection))
 				{
-					using (SqlDataReader Reader = Command.ExecuteReader())
+					using (SQLiteDataReader Reader = Command.ExecuteReader())
 					{
 						while (Reader.Read())
 						{
@@ -35,9 +34,9 @@ namespace UnrealGameSyncWebAPI.Connectors
 						}
 					}
 				}
-				using (SqlCommand Command = new SqlCommand("SELECT MAX(ID) FROM dbo.Comments;", Connection))
+				using (SQLiteCommand Command = new SQLiteCommand("SELECT MAX(ID) FROM [Comments];", Connection))
 				{
-					using (SqlDataReader Reader = Command.ExecuteReader())
+					using (SQLiteDataReader Reader = Command.ExecuteReader())
 					{
 						while (Reader.Read())
 						{
@@ -47,9 +46,9 @@ namespace UnrealGameSyncWebAPI.Connectors
 						}
 					}
 				}
-				using (SqlCommand Command = new SqlCommand("SELECT MAX(ID) FROM dbo.CIS;", Connection))
+				using (SQLiteCommand Command = new SQLiteCommand("SELECT MAX(ID) FROM [CIS];", Connection))
 				{
-					using (SqlDataReader Reader = Command.ExecuteReader())
+					using (SQLiteDataReader Reader = Command.ExecuteReader())
 					{
 						while (Reader.Read())
 						{
@@ -65,13 +64,13 @@ namespace UnrealGameSyncWebAPI.Connectors
 		public static List<EventData> GetUserVotes(string Project, long LastEventId)
 		{
 			List<EventData> ReturnedEvents = new List<EventData>();
-			using (SqlConnection Connection = new SqlConnection(MsSqlConnectionString))
+			using (SQLiteConnection Connection = new SQLiteConnection(ConnectionString))
 			{
 				Connection.Open();
-				using (SqlCommand Command = new SqlCommand("SELECT Id, Changelist, UserName, Verdict, Project FROM dbo.UserVotes WHERE Id > @param1 ORDER BY Id", Connection))
+				using (SQLiteCommand Command = new SQLiteCommand("SELECT Id, Changelist, UserName, Verdict, Project FROM [UserVotes] WHERE Id > @param1 ORDER BY Id", Connection))
 				{
 					Command.Parameters.AddWithValue("@param1", LastEventId);
-					using (SqlDataReader Reader = Command.ExecuteReader())
+					using (SQLiteDataReader Reader = Command.ExecuteReader())
 					{
 						while (Reader.Read())
 						{
@@ -96,13 +95,13 @@ namespace UnrealGameSyncWebAPI.Connectors
 		public static List<CommentData> GetComments(string Project, long LastCommentId)
 		{
 			List<CommentData> ReturnedComments = new List<CommentData>();
-			using (SqlConnection Connection = new SqlConnection(MsSqlConnectionString))
+			using (SQLiteConnection Connection = new SQLiteConnection(ConnectionString))
 			{
 				Connection.Open();
-				using (SqlCommand Command = new SqlCommand("SELECT Id, ChangeNumber, UserName, Text, Project FROM dbo.Comments WHERE Id > @param1 ORDER BY Id", Connection))
+				using (SQLiteCommand Command = new SQLiteCommand("SELECT Id, ChangeNumber, UserName, Text, Project FROM [Comments] WHERE Id > @param1 ORDER BY Id", Connection))
 				{
 					Command.Parameters.AddWithValue("@param1", LastCommentId);
-					using (SqlDataReader Reader = Command.ExecuteReader())
+					using (SQLiteDataReader Reader = Command.ExecuteReader())
 					{
 						while (Reader.Read())
 						{
@@ -126,13 +125,13 @@ namespace UnrealGameSyncWebAPI.Connectors
 		public static List<BuildData> GetCIS(string Project, long LastBuildId)
 		{
 			List<BuildData> ReturnedBuilds = new List<BuildData>();
-			using (SqlConnection Connection = new SqlConnection(MsSqlConnectionString))
+			using (SQLiteConnection Connection = new SQLiteConnection(ConnectionString))
 			{
 				Connection.Open();
-				using (SqlCommand Command = new SqlCommand("SELECT Id, ChangeNumber, BuildType, Result, Url, Project FROM dbo.CIS WHERE Id > @param1 ORDER BY Id", Connection))
+				using (SQLiteCommand Command = new SQLiteCommand("SELECT Id, ChangeNumber, BuildType, Result, Url, Project FROM [CIS] WHERE Id > @param1 ORDER BY Id", Connection))
 				{
 					Command.Parameters.AddWithValue("@param1", LastBuildId);
-					using (SqlDataReader Reader = Command.ExecuteReader())
+					using (SQLiteDataReader Reader = Command.ExecuteReader())
 					{
 						while (Reader.Read())
 						{
@@ -158,21 +157,7 @@ namespace UnrealGameSyncWebAPI.Connectors
 		}
 		public static void PostCIS(BuildData Build)
 		{
-			using (SqlConnection Connection = new SqlConnection(MsSqlConnectionString))
-			{
-				Connection.Open();
-				using (SqlCommand Command = new SqlCommand("INSERT INTO dbo.CIS (ChangeNumber, BuildType, Result, URL, Project, ArchivePath) VALUES (@ChangeNumber, @BuildType, @Result, @URL, @Project, @ArchivePath)", Connection))
-				{
-					Command.Parameters.AddWithValue("@ChangeNumber", Build.ChangeNumber);
-					Command.Parameters.AddWithValue("@BuildType", Build.BuildType);
-					Command.Parameters.AddWithValue("@Result", Build.Result);
-					Command.Parameters.AddWithValue("@URL", Build.Url);
-					Command.Parameters.AddWithValue("@Project", Build.Project);
-					Command.Parameters.AddWithValue("@ArchivePath", Build.ArchivePath);
-					Command.ExecuteNonQuery();
-				}
-			}
-			using (SQLiteConnection Connection = new SQLiteConnection(SQLiteConnectionString))
+			using (SQLiteConnection Connection = new SQLiteConnection(ConnectionString))
 			{
 				Connection.Open();
 				using (SQLiteCommand Command = new SQLiteCommand("INSERT INTO [CIS] (ChangeNumber, BuildType, Result, URL, Project, ArchivePath) VALUES (@ChangeNumber, @BuildType, @Result, @URL, @Project, @ArchivePath)", Connection))
@@ -190,19 +175,7 @@ namespace UnrealGameSyncWebAPI.Connectors
 
 		public static void PostEvent(EventData Event)
 		{
-			using (SqlConnection Connection = new SqlConnection(MsSqlConnectionString))
-			{
-				Connection.Open();
-				using (SqlCommand Command = new SqlCommand("INSERT INTO dbo.UserVotes (Changelist, UserName, Verdict, Project) VALUES (@Changelist, @UserName, @Verdict, @Project)", Connection))
-				{
-					Command.Parameters.AddWithValue("@Changelist", Event.Change);
-					Command.Parameters.AddWithValue("@UserName", Event.UserName.ToString());
-					Command.Parameters.AddWithValue("@Verdict", Event.Type.ToString());
-					Command.Parameters.AddWithValue("@Project", Event.Project);
-					Command.ExecuteNonQuery();
-				}
-			}
-			using (SQLiteConnection Connection = new SQLiteConnection(SQLiteConnectionString))
+			using (SQLiteConnection Connection = new SQLiteConnection(ConnectionString))
 			{
 				Connection.Open();
 				using (SQLiteCommand Command = new SQLiteCommand("INSERT INTO [UserVotes] (Changelist, UserName, Verdict, Project) VALUES (@Changelist, @UserName, @Verdict, @Project)", Connection))
@@ -218,19 +191,7 @@ namespace UnrealGameSyncWebAPI.Connectors
 
 		public static void PostComment(CommentData Comment)
 		{
-			using (SqlConnection Connection = new SqlConnection(MsSqlConnectionString))
-			{
-				Connection.Open();
-				using (SqlCommand Command = new SqlCommand("INSERT INTO dbo.Comments (ChangeNumber, UserName, Text, Project) VALUES (@ChangeNumber, @UserName, @Text, @Project)", Connection))
-				{
-					Command.Parameters.AddWithValue("@ChangeNumber", Comment.ChangeNumber);
-					Command.Parameters.AddWithValue("@UserName", Comment.UserName);
-					Command.Parameters.AddWithValue("@Text", Comment.Text);
-					Command.Parameters.AddWithValue("@Project", Comment.Project);
-					Command.ExecuteNonQuery();
-				}
-			}
-			using (SQLiteConnection Connection = new SQLiteConnection(SQLiteConnectionString))
+			using (SQLiteConnection Connection = new SQLiteConnection(ConnectionString))
 			{
 				Connection.Open();
 				using (SQLiteCommand Command = new SQLiteCommand("INSERT INTO [Comments] (ChangeNumber, UserName, Text, Project) VALUES (@ChangeNumber, @UserName, @Text, @Project)", Connection))
@@ -246,23 +207,7 @@ namespace UnrealGameSyncWebAPI.Connectors
 
 		public static void PostTelemetryData(TelemetryTimingData Data, string Version, string IpAddress)
 		{
-			using (SqlConnection Connection = new SqlConnection(MsSqlConnectionString))
-			{
-				Connection.Open();
-				using (SqlCommand Command = new SqlCommand("INSERT INTO dbo.[Telemetry.v2] (Action, Result, UserName, Project, Timestamp, Duration, Version, IpAddress) VALUES (@Action, @Result, @UserName, @Project, @Timestamp, @Duration, @Version, @IpAddress)", Connection))
-				{
-					Command.Parameters.AddWithValue("@Action", Data.Action);
-					Command.Parameters.AddWithValue("@Result", Data.Result);
-					Command.Parameters.AddWithValue("@UserName", Data.UserName);
-					Command.Parameters.AddWithValue("@Project", Data.Project);
-					Command.Parameters.AddWithValue("@Timestamp", Data.Timestamp);
-					Command.Parameters.AddWithValue("@Duration", Data.Duration);
-					Command.Parameters.AddWithValue("@Version", Version);
-					Command.Parameters.AddWithValue("@IPAddress", IpAddress);
-					Command.ExecuteNonQuery();
-				}
-			}
-			using (SQLiteConnection Connection = new SQLiteConnection(SQLiteConnectionString))
+			using (SQLiteConnection Connection = new SQLiteConnection(ConnectionString))
 			{
 				Connection.Open();
 				using (SQLiteCommand Command = new SQLiteCommand("INSERT INTO [Telemetry.v2] (Action, Result, UserName, Project, Timestamp, Duration, Version, IpAddress) VALUES (@Action, @Result, @UserName, @Project, @Timestamp, @Duration, @Version, @IpAddress)", Connection))
@@ -282,29 +227,7 @@ namespace UnrealGameSyncWebAPI.Connectors
 
 		public static void PostErrorData(TelemetryErrorData Data, string Version, string IpAddress)
 		{
-			using (SqlConnection Connection = new SqlConnection(MsSqlConnectionString))
-			{
-				Connection.Open();
-				using (SqlCommand Command = new SqlCommand("INSERT INTO dbo.[Errors] (Type, Text, UserName, Project, Timestamp, Version, IpAddress) VALUES (@Type, @Text, @UserName, @Project, @Timestamp, @Version, @IpAddress)", Connection))
-				{
-					Command.Parameters.AddWithValue("@Type", Data.Type.ToString());
-					Command.Parameters.AddWithValue("@Text", Data.Text);
-					Command.Parameters.AddWithValue("@UserName", Data.UserName);
-					if (Data.Project == null)
-					{
-						Command.Parameters.AddWithValue("@Project", DBNull.Value);
-					}
-					else
-					{
-						Command.Parameters.AddWithValue("@Project", Data.Project);
-					}
-					Command.Parameters.AddWithValue("@Timestamp", Data.Timestamp);
-					Command.Parameters.AddWithValue("@Version", Version);
-					Command.Parameters.AddWithValue("@IPAddress", IpAddress);
-					Command.ExecuteNonQuery();
-				}
-			}
-			using (SQLiteConnection Connection = new SQLiteConnection(SQLiteConnectionString))
+			using (SQLiteConnection Connection = new SQLiteConnection(ConnectionString))
 			{
 				Connection.Open();
 				using (SQLiteCommand Command = new SQLiteCommand("INSERT INTO [Errors] (Type, Text, UserName, Project, Timestamp, Version, IpAddress) VALUES (@Type, @Text, @UserName, @Project, @Timestamp, @Version, @IpAddress)", Connection))
diff --git a/Engine/Source/Programs/UnrealGameSync/UnrealGameSyncWebAPI/Global.asax.cs b/Engine/Source/Programs/UnrealGameSync/UnrealGameSyncWebAPI/Global.asax.cs
--- a/Engine/Source/Programs/UnrealGameSync/UnrealGameSyncWebAPI/Global.asax.cs
+++ b/Engine/Source/Programs/UnrealGameSync/UnrealGameSyncWebAPI/Global.asax.cs
@@ -14,16 +14,15 @@ namespace UnrealGameSyncWebAPI
 		{
 			try
 			{
-				string MsSqlConnectionString = System.Configuration.ConfigurationManager.AppSettings["MsSqlConnectionString"];
-				string SQLiteConnectionString = System.Configuration.ConfigurationManager.AppSettings["SQLiteConnectionString"];
-				string FileName = Regex.Match(SQLiteConnectionString, "Data Source=(.+);.+").Groups[1].Value;
+				string ConnectionString = System.Configuration.ConfigurationManager.AppSettings["ConnectionString"];
+				string FileName = Regex.Match(ConnectionString, "Data Source=(.+);.+").Groups[1].Value;
 				Directory.CreateDirectory(new FileInfo(FileName).Directory.FullName);
 				if (!File.Exists(FileName))
 				{
 					SQLiteConnection.CreateFile(FileName);
 				}
 				// initialize the db if it doesn't exist
-				using (SQLiteConnection Connection = new SQLiteConnection(SQLiteConnectionString, true))
+				using (SQLiteConnection Connection = new SQLiteConnection(ConnectionString, true))
 				{
 					Connection.Open();
 					string CreateTables = @"CREATE TABLE IF NOT EXISTS
diff --git a/Engine/Source/Programs/UnrealGameSync/UnrealGameSyncWebAPI/Web.template.config b/Engine/Source/Programs/UnrealGameSync/UnrealGameSyncWebAPI/Web.template.config
--- a/Engine/Source/Programs/UnrealGameSync/UnrealGameSyncWebAPI/Web.template.config
+++ b/Engine/Source/Programs/UnrealGameSync/UnrealGameSyncWebAPI/Web.template.config
@@ -8,10 +8,8 @@
     <!-- Leave this connection string blank! It should go in either the .debug or .release transform config
          There is a custom build step on BeforeBuild that generates a transformed Web.config on the fly. The generated Web.config should not be checked into source control, 
          and this template.config should be checked in in it's place. -->
-    <add key="MsSqlConnectionString"
+    <add key="ConnectionString"
       value=""/>
-    <add key="SQLiteConnectionString"
-     value=""/>
   </appSettings>
   <!--
     For a description of web.config changes see http://go.microsoft.com/fwlink/?LinkId=235367.
