updating name from UnrealGameSyncWebAPI to UnrealGameSyncMetadataService

diff --git a/Engine/Source/Programs/UnrealGameSync/UnrealGameSyncMetadataServer/App_Start/WebApiConfig.cs b/Engine/Source/Programs/UnrealGameSync/UnrealGameSyncMetadataServer/App_Start/WebApiConfig.cs
new file mode 100644
--- /dev/null
+++ b/Engine/Source/Programs/UnrealGameSync/UnrealGameSyncMetadataServer/App_Start/WebApiConfig.cs
@@ -0,0 +1,21 @@
+﻿// Copyright 1998-2018 Epic Games, Inc. All Rights Reserved.
+
+using System.Web.Http;
+
+namespace UnrealGameSyncMetadataServer
+{
+	public static class WebApiConfig
+	{
+		public static void Register(HttpConfiguration config)
+		{
+			// Web API routes
+			config.MapHttpAttributeRoutes();
+
+			config.Routes.MapHttpRoute(
+				name: "DefaultApi",
+				routeTemplate: "api/{controller}/{id}",
+				defaults: new { id = RouteParameter.Optional }
+			);
+		}
+	}
+}
diff --git a/Engine/Source/Programs/UnrealGameSync/UnrealGameSyncMetadataServer/Connectors/SqlConnector.cs b/Engine/Source/Programs/UnrealGameSync/UnrealGameSyncMetadataServer/Connectors/SqlConnector.cs
new file mode 100644
--- /dev/null
+++ b/Engine/Source/Programs/UnrealGameSync/UnrealGameSyncMetadataServer/Connectors/SqlConnector.cs
@@ -0,0 +1,258 @@
+﻿// Copyright 1998-2018 Epic Games, Inc. All Rights Reserved.
+
+using System;
+using System.Collections.Generic;
+using System.Data.SqlClient;
+using System.Data.SQLite;
+using System.IO;
+using System.Text.RegularExpressions;
+using UnrealGameSyncMetadataServer.Models;
+
+namespace UnrealGameSyncMetadataServer.Connectors
+{
+	public static class SqlConnector
+	{
+		private static string ConnectionString = System.Configuration.ConfigurationManager.AppSettings["ConnectionString"];
+		public static long[] GetLastIds()
+		{
+			// Get Last IDs
+			long LastEventId = 0;
+			long LastCommentId = 0;
+			long LastBuildID = 0;
+			using (SQLiteConnection Connection = new SQLiteConnection(ConnectionString))
+			{
+				Connection.Open();
+				using (SQLiteCommand Command = new SQLiteCommand("SELECT MAX(ID) FROM [UserVotes];", Connection))
+				{
+					using (SQLiteDataReader Reader = Command.ExecuteReader())
+					{
+						while (Reader.Read())
+						{
+							LastEventId = Reader.GetInt64(0);
+							LastEventId = Math.Max(LastEventId - 5000, 0);
+							break;
+						}
+					}
+				}
+				using (SQLiteCommand Command = new SQLiteCommand("SELECT MAX(ID) FROM [Comments];", Connection))
+				{
+					using (SQLiteDataReader Reader = Command.ExecuteReader())
+					{
+						while (Reader.Read())
+						{
+							LastCommentId = Reader.GetInt32(0);
+							LastCommentId = Math.Max(LastCommentId - 5000, 0);
+							break;
+						}
+					}
+				}
+				using (SQLiteCommand Command = new SQLiteCommand("SELECT MAX(ID) FROM [CIS];", Connection))
+				{
+					using (SQLiteDataReader Reader = Command.ExecuteReader())
+					{
+						while (Reader.Read())
+						{
+							LastBuildID = Reader.GetInt32(0);
+							LastBuildID = Math.Max(LastBuildID - 5000, 0);
+							break;
+						}
+					}
+				}
+			}
+			return new long[] { LastEventId, LastCommentId, LastBuildID };
+		}
+		public static List<EventData> GetUserVotes(string Project, long LastEventId)
+		{
+			List<EventData> ReturnedEvents = new List<EventData>();
+			using (SQLiteConnection Connection = new SQLiteConnection(ConnectionString))
+			{
+				Connection.Open();
+				using (SQLiteCommand Command = new SQLiteCommand("SELECT Id, Changelist, UserName, Verdict, Project FROM [UserVotes] WHERE Id > @param1 ORDER BY Id", Connection))
+				{
+					Command.Parameters.AddWithValue("@param1", LastEventId);
+					using (SQLiteDataReader Reader = Command.ExecuteReader())
+					{
+						while (Reader.Read())
+						{
+							EventData Review = new EventData();
+							Review.Id = Reader.GetInt64(0);
+							Review.Change = Reader.GetInt32(1);
+							Review.UserName = Reader.GetString(2);
+							Review.Project = Reader.IsDBNull(4) ? null : Reader.GetString(4);
+							if (Enum.TryParse(Reader.GetString(3), out Review.Type))
+							{
+								if (Review.Project == null || String.Compare(Review.Project, Project, true) == 0)
+								{
+									ReturnedEvents.Add(Review);
+								}
+							}
+						}
+					}
+				}
+			}
+			return ReturnedEvents;
+		}
+		public static List<CommentData> GetComments(string Project, long LastCommentId)
+		{
+			List<CommentData> ReturnedComments = new List<CommentData>();
+			using (SQLiteConnection Connection = new SQLiteConnection(ConnectionString))
+			{
+				Connection.Open();
+				using (SQLiteCommand Command = new SQLiteCommand("SELECT Id, ChangeNumber, UserName, Text, Project FROM [Comments] WHERE Id > @param1 ORDER BY Id", Connection))
+				{
+					Command.Parameters.AddWithValue("@param1", LastCommentId);
+					using (SQLiteDataReader Reader = Command.ExecuteReader())
+					{
+						while (Reader.Read())
+						{
+							CommentData Comment = new CommentData();
+							Comment.Id = Reader.GetInt32(0);
+							Comment.ChangeNumber = Reader.GetInt32(1);
+							Comment.UserName = Reader.GetString(2);
+							Comment.Text = Reader.GetString(3);
+							Comment.Project = Reader.GetString(4);
+							if (Comment.Project == null || String.Compare(Comment.Project, Project, true) == 0)
+							{
+								ReturnedComments.Add(Comment);
+							}
+						}
+					}
+				}
+			}
+			return ReturnedComments;
+		}
+
+		public static List<BuildData> GetCIS(string Project, long LastBuildId)
+		{
+			List<BuildData> ReturnedBuilds = new List<BuildData>();
+			using (SQLiteConnection Connection = new SQLiteConnection(ConnectionString))
+			{
+				Connection.Open();
+				using (SQLiteCommand Command = new SQLiteCommand("SELECT Id, ChangeNumber, BuildType, Result, Url, Project FROM [CIS] WHERE Id > @param1 ORDER BY Id", Connection))
+				{
+					Command.Parameters.AddWithValue("@param1", LastBuildId);
+					using (SQLiteDataReader Reader = Command.ExecuteReader())
+					{
+						while (Reader.Read())
+						{
+							BuildData Build = new BuildData();
+							Build.Id = Reader.GetInt32(0);
+							Build.ChangeNumber = Reader.GetInt32(1);
+							Build.BuildType = Reader.GetString(2).TrimEnd();
+							if (Enum.TryParse(Reader.GetString(3).TrimEnd(), true, out Build.Result))
+							{
+								Build.Url = Reader.GetString(4);
+								Build.Project = Reader.IsDBNull(5) ? null : Reader.GetString(5);
+								if (Build.Project == null || String.Compare(Build.Project, Project, true) == 0 || MatchesWildcard(Build.Project, Project))
+								{
+									ReturnedBuilds.Add(Build);
+								}
+							}
+							LastBuildId = Math.Max(LastBuildId, Build.Id);
+						}
+					}
+				}
+			}
+			return ReturnedBuilds;
+		}
+		public static void PostCIS(BuildData Build)
+		{
+			using (SQLiteConnection Connection = new SQLiteConnection(ConnectionString))
+			{
+				Connection.Open();
+				using (SQLiteCommand Command = new SQLiteCommand("INSERT INTO [CIS] (ChangeNumber, BuildType, Result, URL, Project, ArchivePath) VALUES (@ChangeNumber, @BuildType, @Result, @URL, @Project, @ArchivePath)", Connection))
+				{
+					Command.Parameters.AddWithValue("@ChangeNumber", Build.ChangeNumber);
+					Command.Parameters.AddWithValue("@BuildType", Build.BuildType);
+					Command.Parameters.AddWithValue("@Result", Build.Result);
+					Command.Parameters.AddWithValue("@URL", Build.Url);
+					Command.Parameters.AddWithValue("@Project", Build.Project);
+					Command.Parameters.AddWithValue("@ArchivePath", Build.ArchivePath);
+					Command.ExecuteNonQuery();
+				}
+			}
+		}
+
+		public static void PostEvent(EventData Event)
+		{
+			using (SQLiteConnection Connection = new SQLiteConnection(ConnectionString))
+			{
+				Connection.Open();
+				using (SQLiteCommand Command = new SQLiteCommand("INSERT INTO [UserVotes] (Changelist, UserName, Verdict, Project) VALUES (@Changelist, @UserName, @Verdict, @Project)", Connection))
+				{
+					Command.Parameters.AddWithValue("@Changelist", Event.Change);
+					Command.Parameters.AddWithValue("@UserName", Event.UserName.ToString());
+					Command.Parameters.AddWithValue("@Verdict", Event.Type.ToString());
+					Command.Parameters.AddWithValue("@Project", Event.Project);
+					Command.ExecuteNonQuery();
+				}
+			}
+		}
+
+		public static void PostComment(CommentData Comment)
+		{
+			using (SQLiteConnection Connection = new SQLiteConnection(ConnectionString))
+			{
+				Connection.Open();
+				using (SQLiteCommand Command = new SQLiteCommand("INSERT INTO [Comments] (ChangeNumber, UserName, Text, Project) VALUES (@ChangeNumber, @UserName, @Text, @Project)", Connection))
+				{
+					Command.Parameters.AddWithValue("@ChangeNumber", Comment.ChangeNumber);
+					Command.Parameters.AddWithValue("@UserName", Comment.UserName);
+					Command.Parameters.AddWithValue("@Text", Comment.Text);
+					Command.Parameters.AddWithValue("@Project", Comment.Project);
+					Command.ExecuteNonQuery();
+				}
+			}
+		}
+
+		public static void PostTelemetryData(TelemetryTimingData Data, string Version, string IpAddress)
+		{
+			using (SQLiteConnection Connection = new SQLiteConnection(ConnectionString))
+			{
+				Connection.Open();
+				using (SQLiteCommand Command = new SQLiteCommand("INSERT INTO [Telemetry.v2] (Action, Result, UserName, Project, Timestamp, Duration, Version, IpAddress) VALUES (@Action, @Result, @UserName, @Project, @Timestamp, @Duration, @Version, @IpAddress)", Connection))
+				{
+					Command.Parameters.AddWithValue("@Action", Data.Action);
+					Command.Parameters.AddWithValue("@Result", Data.Result);
+					Command.Parameters.AddWithValue("@UserName", Data.UserName);
+					Command.Parameters.AddWithValue("@Project", Data.Project);
+					Command.Parameters.AddWithValue("@Timestamp", Data.Timestamp);
+					Command.Parameters.AddWithValue("@Duration", Data.Duration);
+					Command.Parameters.AddWithValue("@Version", Version);
+					Command.Parameters.AddWithValue("@IPAddress", IpAddress);
+					Command.ExecuteNonQuery();
+				}
+			}
+		}
+
+		public static void PostErrorData(TelemetryErrorData Data, string Version, string IpAddress)
+		{
+			using (SQLiteConnection Connection = new SQLiteConnection(ConnectionString))
+			{
+				Connection.Open();
+				using (SQLiteCommand Command = new SQLiteCommand("INSERT INTO [Errors] (Type, Text, UserName, Project, Timestamp, Version, IpAddress) VALUES (@Type, @Text, @UserName, @Project, @Timestamp, @Version, @IpAddress)", Connection))
+				{
+					Command.Parameters.AddWithValue("@Type", Data.Type.ToString());
+					Command.Parameters.AddWithValue("@Text", Data.Text);
+					Command.Parameters.AddWithValue("@UserName", Data.UserName);
+					if (Data.Project == null)
+					{
+						Command.Parameters.AddWithValue("@Project", DBNull.Value);
+					}
+					else
+					{
+						Command.Parameters.AddWithValue("@Project", Data.Project);
+					}
+					Command.Parameters.AddWithValue("@Timestamp", Data.Timestamp);
+					Command.Parameters.AddWithValue("@Version", Version);
+					Command.Parameters.AddWithValue("@IPAddress", IpAddress);
+					Command.ExecuteNonQuery();
+				}
+			}
+		}
+		private static bool MatchesWildcard(string Wildcard, string Project)
+		{
+			return Wildcard.EndsWith("...") && Project.StartsWith(Wildcard.Substring(0, Wildcard.Length - 3), StringComparison.InvariantCultureIgnoreCase);
+		}
+	}
+}
\ No newline at end of file
diff --git a/Engine/Source/Programs/UnrealGameSync/UnrealGameSyncMetadataServer/Controllers/CISController.cs b/Engine/Source/Programs/UnrealGameSync/UnrealGameSyncMetadataServer/Controllers/CISController.cs
new file mode 100644
--- /dev/null
+++ b/Engine/Source/Programs/UnrealGameSync/UnrealGameSyncMetadataServer/Controllers/CISController.cs
@@ -0,0 +1,25 @@
+﻿// Copyright 1998-2018 Epic Games, Inc. All Rights Reserved.
+
+using System.Collections.Generic;
+using System.Web.Http;
+using UnrealGameSyncMetadataServer.Connectors;
+using UnrealGameSyncMetadataServer.Models;
+
+namespace UnrealGameSyncMetadataServer.Controllers
+{
+    public class CISController : ApiController
+    {
+		public long[] Get()
+		{
+			return SqlConnector.GetLastIds();
+		}
+		public List<BuildData> Get(string Project, long LastBuildId)
+		{
+			return SqlConnector.GetCIS(Project, LastBuildId);
+		}
+		public void Post([FromBody]BuildData Build)
+		{
+			SqlConnector.PostCIS(Build);
+		}
+	}
+}
diff --git a/Engine/Source/Programs/UnrealGameSync/UnrealGameSyncMetadataServer/Controllers/CommentController.cs b/Engine/Source/Programs/UnrealGameSync/UnrealGameSyncMetadataServer/Controllers/CommentController.cs
new file mode 100644
--- /dev/null
+++ b/Engine/Source/Programs/UnrealGameSync/UnrealGameSyncMetadataServer/Controllers/CommentController.cs
@@ -0,0 +1,21 @@
+﻿// Copyright 1998-2018 Epic Games, Inc. All Rights Reserved.
+
+using System.Collections.Generic;
+using System.Web.Http;
+using UnrealGameSyncMetadataServer.Connectors;
+using UnrealGameSyncMetadataServer.Models;
+
+namespace UnrealGameSyncMetadataServer.Controllers
+{
+	public class CommentController : ApiController
+	{
+		public List<CommentData> Get(string Project, long LastCommentId)
+		{
+			return SqlConnector.GetComments(Project, LastCommentId);
+		}
+		public void Post([FromBody] CommentData Comment)
+		{
+			SqlConnector.PostComment(Comment);
+		}
+	}
+}
\ No newline at end of file
diff --git a/Engine/Source/Programs/UnrealGameSync/UnrealGameSyncMetadataServer/Controllers/ErrorController.cs b/Engine/Source/Programs/UnrealGameSync/UnrealGameSyncMetadataServer/Controllers/ErrorController.cs
new file mode 100644
--- /dev/null
+++ b/Engine/Source/Programs/UnrealGameSync/UnrealGameSyncMetadataServer/Controllers/ErrorController.cs
@@ -0,0 +1,16 @@
+﻿// Copyright 1998-2018 Epic Games, Inc. All Rights Reserved.
+
+using System.Web.Http;
+using UnrealGameSyncMetadataServer.Connectors;
+using UnrealGameSyncMetadataServer.Models;
+
+namespace UnrealGameSyncMetadataServer.Controllers
+{
+    public class ErrorController : ApiController
+    {
+		public void Post([FromBody] TelemetryErrorData Data, string Version, string IpAddress)
+		{
+			SqlConnector.PostErrorData(Data, Version, IpAddress);
+		}
+    }
+}
diff --git a/Engine/Source/Programs/UnrealGameSync/UnrealGameSyncMetadataServer/Controllers/EventController.cs b/Engine/Source/Programs/UnrealGameSync/UnrealGameSyncMetadataServer/Controllers/EventController.cs
new file mode 100644
--- /dev/null
+++ b/Engine/Source/Programs/UnrealGameSync/UnrealGameSyncMetadataServer/Controllers/EventController.cs
@@ -0,0 +1,22 @@
+﻿// Copyright 1998-2018 Epic Games, Inc. All Rights Reserved.
+
+using System.Collections.Generic;
+using System.Web.Http;
+using UnrealGameSyncMetadataServer.Models;
+using UnrealGameSyncMetadataServer.Connectors;
+
+namespace UnrealGameSyncMetadataServer.Controllers
+{
+	public class EventController : ApiController
+	{
+		public List<EventData> Get(string Project, long LastEventId)
+		{
+			return SqlConnector.GetUserVotes(Project, LastEventId);
+		}
+
+		public void Post([FromBody] EventData Event)
+		{
+			SqlConnector.PostEvent(Event);
+		}
+	}
+}
\ No newline at end of file
diff --git a/Engine/Source/Programs/UnrealGameSync/UnrealGameSyncMetadataServer/Controllers/TelemetryController.cs b/Engine/Source/Programs/UnrealGameSync/UnrealGameSyncMetadataServer/Controllers/TelemetryController.cs
new file mode 100644
--- /dev/null
+++ b/Engine/Source/Programs/UnrealGameSync/UnrealGameSyncMetadataServer/Controllers/TelemetryController.cs
@@ -0,0 +1,16 @@
+﻿// Copyright 1998-2018 Epic Games, Inc. All Rights Reserved.
+
+using System.Web.Http;
+using UnrealGameSyncMetadataServer.Connectors;
+using UnrealGameSyncMetadataServer.Models;
+
+namespace UnrealGameSyncMetadataServer.Controllers
+{
+    public class TelemetryController : ApiController
+    {
+		public void Post([FromBody] TelemetryTimingData Data, string Version, string IpAddress)
+		{
+			SqlConnector.PostTelemetryData(Data, Version, IpAddress);
+		}
+    }
+}
diff --git a/Engine/Source/Programs/UnrealGameSync/UnrealGameSyncMetadataServer/Global.asax.cs b/Engine/Source/Programs/UnrealGameSync/UnrealGameSyncMetadataServer/Global.asax.cs
new file mode 100644
--- /dev/null
+++ b/Engine/Source/Programs/UnrealGameSync/UnrealGameSyncMetadataServer/Global.asax.cs
@@ -0,0 +1,96 @@
+﻿// Copyright 1998-2018 Epic Games, Inc. All Rights Reserved.
+using System;
+using System.Data.SQLite;
+using System.IO;
+using System.Text.RegularExpressions;
+using System.Web;
+using System.Web.Http;
+
+namespace UnrealGameSyncMetadataServer
+{
+	public class WebApiApplication : System.Web.HttpApplication
+	{
+		protected void Application_Start()
+		{
+			try
+			{
+				string ConnectionString = System.Configuration.ConfigurationManager.AppSettings["ConnectionString"];
+				string FileName = Regex.Match(ConnectionString, "Data Source=(.+);.+").Groups[1].Value;
+				Directory.CreateDirectory(new FileInfo(FileName).Directory.FullName);
+				if (!File.Exists(FileName))
+				{
+					SQLiteConnection.CreateFile(FileName);
+				}
+				// initialize the db if it doesn't exist
+				using (SQLiteConnection Connection = new SQLiteConnection(ConnectionString, true))
+				{
+					Connection.Open();
+					string CreateTables = @"CREATE TABLE IF NOT EXISTS
+										[CIS] (
+											[Id]           INTEGER       NOT NULL PRIMARY KEY AUTOINCREMENT,
+											[ChangeNumber] INTEGER       NOT NULL,
+											[BuildType]    NCHAR (32)    NOT NULL,
+											[Result]       NCHAR (10)    NOT NULL,
+											[Url]          VARCHAR (512) NOT NULL,
+											[Project]      VARCHAR (256) NULL,
+											[ArchivePath]  VARCHAR (512) NULL
+										);
+										
+										CREATE TABLE IF NOT EXISTS
+										[Comments] (
+											[Id]           INTEGER    	  NOT NULL PRIMARY KEY AUTOINCREMENT,
+											[ChangeNumber] INTEGER        NOT NULL,
+											[UserName]     VARCHAR (128)  NOT NULL,
+											[Text]         VARCHAR (140)  NOT NULL,
+											[Project]      NVARCHAR (256) NOT NULL
+										);
+
+										CREATE TABLE IF NOT EXISTS
+										[Errors] (
+											[Id]        INTEGER		   NOT NULL PRIMARY KEY AUTOINCREMENT,
+											[Type]      VARCHAR (50)   NOT NULL,
+											[Text]      VARCHAR (1024) NOT NULL,
+											[UserName]  NVARCHAR (128) NOT NULL,
+											[Project]   VARCHAR (128)  NULL,
+											[Timestamp] DATETIME       NOT NULL,
+											[Version]   VARCHAR (64)   NOT NULL,
+											[IpAddress] VARCHAR (64)   NOT NULL
+										);
+										
+										CREATE TABLE IF NOT EXISTS
+										[Telemetry.v2] (
+											[Id]        INTEGER		   NOT NULL PRIMARY KEY AUTOINCREMENT,
+											[Action]    VARCHAR (128)  NOT NULL,
+											[Result]    VARCHAR (128)  NOT NULL,
+											[UserName]  NVARCHAR (128) NOT NULL,
+											[Project]   VARCHAR (128)  NOT NULL,
+											[Timestamp] DATETIME       NOT NULL,
+											[Duration]  REAL           NOT NULL,
+											[Version]   VARCHAR (64)   NULL,
+											[IpAddress] VARCHAR (64)   NULL
+										);
+
+										CREATE TABLE IF NOT EXISTS
+										[UserVotes] (
+											[Id]         INTEGER		NOT NULL PRIMARY KEY AUTOINCREMENT,
+											[Changelist] INTEGER        NOT NULL,
+											[UserName]   NVARCHAR (128) NOT NULL,
+											[Verdict]    NCHAR (32)     NOT NULL,
+											[Project]    NVARCHAR (256) NULL
+										);";
+					using (SQLiteCommand Command = new SQLiteCommand(CreateTables, Connection))
+					{
+						Command.ExecuteNonQuery();
+					}
+				}
+
+			}
+			catch(Exception)
+			{
+				HttpRuntime.UnloadAppDomain();
+				return;
+			}
+			GlobalConfiguration.Configure(WebApiConfig.Register);
+		}
+	}
+}
diff --git a/Engine/Source/Programs/UnrealGameSync/UnrealGameSyncMetadataServer/Models/BuildData.cs b/Engine/Source/Programs/UnrealGameSync/UnrealGameSyncMetadataServer/Models/BuildData.cs
new file mode 100644
--- /dev/null
+++ b/Engine/Source/Programs/UnrealGameSync/UnrealGameSyncMetadataServer/Models/BuildData.cs
@@ -0,0 +1,39 @@
+﻿// Copyright 1998-2018 Epic Games, Inc. All Rights Reserved.
+
+using System;
+using System.Collections.Generic;
+using System.Linq;
+using System.Web;
+
+namespace UnrealGameSyncMetadataServer.Models
+{
+	public class BuildData
+	{
+		public enum BuildDataResult
+		{
+			Starting,
+			Failure,
+			Warning,
+			Success,
+			Skipped,
+		}
+
+		public long Id;
+		public int ChangeNumber;
+		public string BuildType;
+		public BuildDataResult Result;
+		public string Url;
+		public string Project;
+		public string ArchivePath;
+
+		public bool IsSuccess
+		{
+			get { return Result == BuildDataResult.Success || Result == BuildDataResult.Warning; }
+		}
+
+		public bool IsFailure
+		{
+			get { return Result == BuildDataResult.Failure; }
+		}
+	}
+}
\ No newline at end of file
diff --git a/Engine/Source/Programs/UnrealGameSync/UnrealGameSyncMetadataServer/Models/CommentData.cs b/Engine/Source/Programs/UnrealGameSync/UnrealGameSyncMetadataServer/Models/CommentData.cs
new file mode 100644
--- /dev/null
+++ b/Engine/Source/Programs/UnrealGameSync/UnrealGameSyncMetadataServer/Models/CommentData.cs
@@ -0,0 +1,18 @@
+﻿// Copyright 1998-2018 Epic Games, Inc. All Rights Reserved.
+
+using System;
+using System.Collections.Generic;
+using System.Linq;
+using System.Web;
+
+namespace UnrealGameSyncMetadataServer.Models
+{
+	public class CommentData
+	{
+		public long Id;
+		public int ChangeNumber;
+		public string UserName;
+		public string Text;
+		public string Project;
+	}
+}
\ No newline at end of file
diff --git a/Engine/Source/Programs/UnrealGameSync/UnrealGameSyncMetadataServer/Models/EventData.cs b/Engine/Source/Programs/UnrealGameSync/UnrealGameSyncMetadataServer/Models/EventData.cs
new file mode 100644
--- /dev/null
+++ b/Engine/Source/Programs/UnrealGameSync/UnrealGameSyncMetadataServer/Models/EventData.cs
@@ -0,0 +1,38 @@
+﻿// Copyright 1998-2018 Epic Games, Inc. All Rights Reserved.
+
+using System;
+using System.Collections.Generic;
+using System.Linq;
+using System.Web;
+
+namespace UnrealGameSyncMetadataServer.Models
+{
+	public class EventData
+	{
+		public enum EventType
+		{
+			Syncing,
+
+			// Reviews
+			Compiles,
+			DoesNotCompile,
+			Good,
+			Bad,
+			Unknown,
+
+			// Starred builds
+			Starred,
+			Unstarred,
+
+			// Investigating events
+			Investigating,
+			Resolved,
+		}
+		
+		public long Id;
+		public int Change;
+		public string UserName;
+		public EventType Type;
+		public string Project;
+	}
+}
\ No newline at end of file
diff --git a/Engine/Source/Programs/UnrealGameSync/UnrealGameSyncMetadataServer/Models/EventSummary.cs b/Engine/Source/Programs/UnrealGameSync/UnrealGameSyncMetadataServer/Models/EventSummary.cs
new file mode 100644
--- /dev/null
+++ b/Engine/Source/Programs/UnrealGameSync/UnrealGameSyncMetadataServer/Models/EventSummary.cs
@@ -0,0 +1,29 @@
+﻿// Copyright 1998-2018 Epic Games, Inc. All Rights Reserved.
+
+using System;
+using System.Collections.Generic;
+using System.Linq;
+using System.Web;
+
+namespace UnrealGameSyncMetadataServer.Models
+{
+	public class EventSummary
+	{
+		public enum ReviewVerdict
+		{
+			Unknown,
+			Good,
+			Bad,
+			Mixed,
+		}
+
+		public int ChangeNumber;
+		public ReviewVerdict Verdict;
+		public List<EventData> SyncEvents = new List<EventData>();
+		public List<EventData> Reviews = new List<EventData>();
+		public List<string> CurrentUsers = new List<string>();
+		public EventData LastStarReview;
+		public List<BuildData> Builds = new List<BuildData>();
+		public List<CommentData> Comments = new List<CommentData>();
+	}
+}
\ No newline at end of file
diff --git a/Engine/Source/Programs/UnrealGameSync/UnrealGameSyncMetadataServer/Models/TelemetryErrorData.cs b/Engine/Source/Programs/UnrealGameSync/UnrealGameSyncMetadataServer/Models/TelemetryErrorData.cs
new file mode 100644
--- /dev/null
+++ b/Engine/Source/Programs/UnrealGameSync/UnrealGameSyncMetadataServer/Models/TelemetryErrorData.cs
@@ -0,0 +1,22 @@
+﻿// Copyright 1998-2018 Epic Games, Inc. All Rights Reserved.
+
+using System;
+using System.Collections.Generic;
+using System.Linq;
+using System.Web;
+
+namespace UnrealGameSyncMetadataServer.Models
+{
+	public class TelemetryErrorData
+	{
+		public enum TelemetryErrorType
+		{
+			Crash,
+		}
+		public TelemetryErrorType Type;
+		public string Text;
+		public string UserName;
+		public string Project;
+		public DateTime Timestamp;
+	}
+}
\ No newline at end of file
diff --git a/Engine/Source/Programs/UnrealGameSync/UnrealGameSyncMetadataServer/Models/TelemetryTimingData.cs b/Engine/Source/Programs/UnrealGameSync/UnrealGameSyncMetadataServer/Models/TelemetryTimingData.cs
new file mode 100644
--- /dev/null
+++ b/Engine/Source/Programs/UnrealGameSync/UnrealGameSyncMetadataServer/Models/TelemetryTimingData.cs
@@ -0,0 +1,19 @@
+﻿// Copyright 1998-2018 Epic Games, Inc. All Rights Reserved.
+
+using System;
+using System.Collections.Generic;
+using System.Linq;
+using System.Web;
+
+namespace UnrealGameSyncMetadataServer.Models
+{
+	public class TelemetryTimingData
+	{
+		public string Action;
+		public string Result;
+		public string UserName;
+		public string Project;
+		public DateTime Timestamp;
+		public float Duration;
+	}
+}
\ No newline at end of file
diff --git a/Engine/Source/Programs/UnrealGameSync/UnrealGameSyncMetadataServer/Properties/AssemblyInfo.cs b/Engine/Source/Programs/UnrealGameSync/UnrealGameSyncMetadataServer/Properties/AssemblyInfo.cs
new file mode 100644
--- /dev/null
+++ b/Engine/Source/Programs/UnrealGameSync/UnrealGameSyncMetadataServer/Properties/AssemblyInfo.cs
@@ -0,0 +1,37 @@
+﻿// Copyright 1998-2018 Epic Games, Inc. All Rights Reserved.
+
+using System.Reflection;
+using System.Runtime.CompilerServices;
+using System.Runtime.InteropServices;
+
+// General Information about an assembly is controlled through the following
+// set of attributes. Change these attribute values to modify the information
+// associated with an assembly.
+[assembly: AssemblyTitle("UnrealGameSyncWebAPI")]
+[assembly: AssemblyDescription("")]
+[assembly: AssemblyConfiguration("")]
+[assembly: AssemblyCompany("")]
+[assembly: AssemblyProduct("UnrealGameSync")]
+[assembly: AssemblyCopyright("Copyright 1998-2018 Epic Games, Inc. All Rights Reserved.")]
+[assembly: AssemblyTrademark("")]
+[assembly: AssemblyCulture("")]
+
+// Setting ComVisible to false makes the types in this assembly not visible
+// to COM components.  If you need to access a type in this assembly from
+// COM, set the ComVisible attribute to true on that type.
+[assembly: ComVisible(false)]
+
+// The following GUID is for the ID of the typelib if this project is exposed to COM
+[assembly: Guid("f6044fb1-c977-4705-8cad-f17962adfb41")]
+
+// Version information for an assembly consists of the following four values:
+//
+//      Major Version
+//      Minor Version
+//      Build Number
+//      Revision
+//
+// You can specify all the values or you can default the Revision and Build Numbers
+// by using the '*' as shown below:
+[assembly: AssemblyVersion("1.0.0.0")]
+[assembly: AssemblyFileVersion("1.0.0.0")]
diff --git a/Engine/Source/Programs/UnrealGameSync/UnrealGameSyncMetadataServer/UnrealGameSyncMetadataServer.csproj b/Engine/Source/Programs/UnrealGameSync/UnrealGameSyncMetadataServer/UnrealGameSyncMetadataServer.csproj
new file mode 100644
--- /dev/null
+++ b/Engine/Source/Programs/UnrealGameSync/UnrealGameSyncMetadataServer/UnrealGameSyncMetadataServer.csproj
@@ -0,0 +1,176 @@
+﻿<?xml version="1.0" encoding="utf-8"?>
+<Project ToolsVersion="14.0" DefaultTargets="Build" xmlns="http://schemas.microsoft.com/developer/msbuild/2003">
+  <Import Project="$(MSBuildExtensionsPath)\$(MSBuildToolsVersion)\Microsoft.Common.props" Condition="Exists('$(MSBuildExtensionsPath)\$(MSBuildToolsVersion)\Microsoft.Common.props')" />
+  <PropertyGroup>
+    <Configuration Condition=" '$(Configuration)' == '' ">Debug</Configuration>
+    <Platform Condition=" '$(Platform)' == '' ">AnyCPU</Platform>
+    <ProductVersion>
+    </ProductVersion>
+    <SchemaVersion>2.0</SchemaVersion>
+    <ProjectGuid>{C8D13E98-C8CD-42FB-9C79-F3EB11B0F3E8}</ProjectGuid>
+    <ProjectTypeGuids>{349c5851-65df-11da-9384-00065b846f21};{fae04ec0-301f-11d3-bf4b-00c04f79efbc}</ProjectTypeGuids>
+    <OutputType>Library</OutputType>
+    <AppDesignerFolder>Properties</AppDesignerFolder>
+    <RootNamespace>UnrealGameSyncMetadataServer</RootNamespace>
+    <AssemblyName>UnrealGameSyncMetadataServer</AssemblyName>
+    <TargetFrameworkVersion>v4.5</TargetFrameworkVersion>
+    <MvcBuildViews>false</MvcBuildViews>
+    <UseIISExpress>true</UseIISExpress>
+    <IISExpressSSLPort />
+    <IISExpressAnonymousAuthentication />
+    <IISExpressWindowsAuthentication />
+    <IISExpressUseClassicPipelineMode />
+    <UseGlobalApplicationHostFile />
+    <NuGetPackageImportStamp>
+    </NuGetPackageImportStamp>
+    <TargetFrameworkProfile />
+  </PropertyGroup>
+  <PropertyGroup Condition=" '$(Configuration)|$(Platform)' == 'Debug|AnyCPU' ">
+    <TransformWebConfigEnabled>False</TransformWebConfigEnabled>
+    <DebugSymbols>true</DebugSymbols>
+    <DebugType>full</DebugType>
+    <Optimize>false</Optimize>
+    <OutputPath>bin\</OutputPath>
+    <DefineConstants>DEBUG;TRACE</DefineConstants>
+    <ErrorReport>prompt</ErrorReport>
+    <WarningLevel>4</WarningLevel>
+  </PropertyGroup>
+  <PropertyGroup Condition=" '$(Configuration)|$(Platform)' == 'Release|AnyCPU' ">
+    <TransformWebConfigEnabled>False</TransformWebConfigEnabled>
+    <DebugType>pdbonly</DebugType>
+    <Optimize>true</Optimize>
+    <OutputPath>bin\</OutputPath>
+    <DefineConstants>TRACE</DefineConstants>
+    <ErrorReport>prompt</ErrorReport>
+    <WarningLevel>4</WarningLevel>
+  </PropertyGroup>
+  <ItemGroup>
+    <Reference Include="Microsoft.CSharp" />
+    <Reference Include="Newtonsoft.Json, Version=11.0.0.0, Culture=neutral, PublicKeyToken=30ad4fe6b2a6aeed, processorArchitecture=MSIL">
+      <HintPath>..\packages\Newtonsoft.Json.11.0.2\lib\net45\Newtonsoft.Json.dll</HintPath>
+      <Private>True</Private>
+    </Reference>
+    <Reference Include="System" />
+    <Reference Include="System.Data" />
+    <Reference Include="System.Data.DataSetExtensions" />
+    <Reference Include="System.Data.SQLite, Version=1.0.108.0, Culture=neutral, PublicKeyToken=db937bc2d44ff139, processorArchitecture=MSIL">
+      <HintPath>..\packages\System.Data.SQLite.Core.1.0.108.0\lib\net45\System.Data.SQLite.dll</HintPath>
+      <Private>True</Private>
+    </Reference>
+    <Reference Include="System.Drawing" />
+    <Reference Include="System.Web.DynamicData" />
+    <Reference Include="System.Web.Entity" />
+    <Reference Include="System.Web.ApplicationServices" />
+    <Reference Include="System.ComponentModel.DataAnnotations" />
+    <Reference Include="System.Web.Extensions" />
+    <Reference Include="System.Web" />
+    <Reference Include="System.Web.Abstractions" />
+    <Reference Include="System.Web.Routing" />
+    <Reference Include="System.Xml" />
+    <Reference Include="System.Configuration" />
+    <Reference Include="System.Runtime.Serialization" />
+    <Reference Include="System.Net.Http">
+    </Reference>
+    <Reference Include="System.Net.Http.Formatting, Version=5.2.3.0, Culture=neutral, PublicKeyToken=31bf3856ad364e35, processorArchitecture=MSIL">
+      <HintPath>..\packages\Microsoft.AspNet.WebApi.Client.5.2.3\lib\net45\System.Net.Http.Formatting.dll</HintPath>
+    </Reference>
+    <Reference Include="System.Net.Http.WebRequest">
+    </Reference>
+    <Reference Include="System.Web.Http, Version=5.2.3.0, Culture=neutral, PublicKeyToken=31bf3856ad364e35, processorArchitecture=MSIL">
+      <HintPath>..\packages\Microsoft.AspNet.WebApi.Core.5.2.3\lib\net45\System.Web.Http.dll</HintPath>
+    </Reference>
+    <Reference Include="System.Web.Http.WebHost, Version=5.2.3.0, Culture=neutral, PublicKeyToken=31bf3856ad364e35, processorArchitecture=MSIL">
+      <HintPath>..\packages\Microsoft.AspNet.WebApi.WebHost.5.2.3\lib\net45\System.Web.Http.WebHost.dll</HintPath>
+    </Reference>
+    <Reference Include="System.Xml.Linq" />
+  </ItemGroup>
+  <ItemGroup>
+    <Folder Include="App_Data\" />
+  </ItemGroup>
+  <ItemGroup>
+    <Compile Include="App_Start\WebApiConfig.cs" />
+    <Compile Include="Controllers\CISController.cs" />
+    <Compile Include="Controllers\CommentController.cs" />
+    <Compile Include="Controllers\ErrorController.cs" />
+    <Compile Include="Controllers\EventController.cs" />
+    <Compile Include="Controllers\TelemetryController.cs" />
+    <Compile Include="Global.asax.cs">
+      <DependentUpon>Global.asax</DependentUpon>
+    </Compile>
+    <Compile Include="Models\BuildData.cs" />
+    <Compile Include="Models\CommentData.cs" />
+    <Compile Include="Models\EventData.cs" />
+    <Compile Include="Models\EventSummary.cs" />
+    <Compile Include="Models\TelemetryErrorData.cs" />
+    <Compile Include="Models\TelemetryTimingData.cs" />
+    <Compile Include="Properties\AssemblyInfo.cs" />
+    <Compile Include="Connectors\SqlConnector.cs" />
+  </ItemGroup>
+  <ItemGroup>
+    <Content Include="Global.asax" />
+    <Content Include="Web.config" />
+    <None Include="Properties\PublishProfiles\File.pubxml" />
+    <None Include="Web.template.config" />
+    <None Include="Web.Debug.config">
+      <DependentUpon>Web.template.config</DependentUpon>
+    </None>
+    <None Include="Web.Release.config">
+      <DependentUpon>Web.template.config</DependentUpon>
+    </None>
+  </ItemGroup>
+  <ItemGroup>
+    <Content Include="NotForLicensees\Web.config">
+      <SubType>Designer</SubType>
+    </Content>
+    <None Include="Properties\PublishProfiles\FilePublish.pubxml" />
+    <None Include="Properties\PublishProfiles\Profile.pubxml" />
+  </ItemGroup>
+  <ItemGroup>
+    <Content Include="packages.config" />
+  </ItemGroup>
+  <PropertyGroup>
+    <VisualStudioVersion Condition="'$(VisualStudioVersion)' == ''">10.0</VisualStudioVersion>
+    <VSToolsPath Condition="'$(VSToolsPath)' == ''">$(MSBuildExtensionsPath32)\Microsoft\VisualStudio\v$(VisualStudioVersion)</VSToolsPath>
+  </PropertyGroup>
+  <Import Project="$(MSBuildBinPath)\Microsoft.CSharp.targets" />
+  <Import Project="$(VSToolsPath)\WebApplications\Microsoft.WebApplication.targets" Condition="'$(VSToolsPath)' != ''" />
+  <Import Project="$(MSBuildExtensionsPath32)\Microsoft\VisualStudio\v10.0\WebApplications\Microsoft.WebApplication.targets" Condition="false" />
+  <Target Name="MvcBuildViews" AfterTargets="AfterBuild" Condition="'$(MvcBuildViews)'=='true'">
+    <AspNetCompiler VirtualPath="temp" PhysicalPath="$(WebProjectOutputDir)" />
+  </Target>
+  <ProjectExtensions>
+    <VisualStudio>
+      <FlavorProperties GUID="{349c5851-65df-11da-9384-00065b846f21}">
+        <WebProjectProperties>
+          <UseIIS>True</UseIIS>
+          <AutoAssignPort>True</AutoAssignPort>
+          <DevelopmentServerPort>62486</DevelopmentServerPort>
+          <DevelopmentServerVPath>/</DevelopmentServerVPath>
+          <IISUrl>http://localhost:62487/</IISUrl>
+          <NTLMAuthentication>False</NTLMAuthentication>
+          <UseCustomServer>False</UseCustomServer>
+          <CustomServerUrl>
+          </CustomServerUrl>
+          <SaveServerSettingsInUserFile>False</SaveServerSettingsInUserFile>
+        </WebProjectProperties>
+      </FlavorProperties>
+    </VisualStudio>
+  </ProjectExtensions>
+  <!-- To modify your build process, add your task inside one of the targets below and uncomment it.
+       Other similar extension points exist, see Microsoft.Common.targets.
+  <Target Name="BeforeBuild">
+  </Target> -->
+  <Target Name="BeforeBuild">
+    <Message Importance="High" Text="Creating web.config from web.template.config" />
+    <TransformXml Source="Web.template.config" Transform="Web.$(Configuration).config" Destination="Web.config" />
+    <Message Importance="High" Text="Found NotForLicensees\Web.config, applying transform" Condition="Exists('NotForLicensees\Web.config')" />
+    <TransformXml Source="Web.config" Transform="NotForLicensees\Web.config" Destination="Web.config" Condition="Exists('NotForLicensees\Web.config')" />
+  </Target>
+  <Import Project="..\packages\System.Data.SQLite.Core.1.0.108.0\build\net45\System.Data.SQLite.Core.targets" Condition="Exists('..\packages\System.Data.SQLite.Core.1.0.108.0\build\net45\System.Data.SQLite.Core.targets')" />
+  <Target Name="EnsureNuGetPackageBuildImports" BeforeTargets="PrepareForBuild">
+    <PropertyGroup>
+      <ErrorText>This project references NuGet package(s) that are missing on this computer. Use NuGet Package Restore to download them.  For more information, see http://go.microsoft.com/fwlink/?LinkID=322105. The missing file is {0}.</ErrorText>
+    </PropertyGroup>
+    <Error Condition="!Exists('..\packages\System.Data.SQLite.Core.1.0.108.0\build\net45\System.Data.SQLite.Core.targets')" Text="$([System.String]::Format('$(ErrorText)', '..\packages\System.Data.SQLite.Core.1.0.108.0\build\net45\System.Data.SQLite.Core.targets'))" />
+  </Target>
+</Project>
\ No newline at end of file
diff --git a/Engine/Source/Programs/UnrealGameSync/UnrealGameSyncMetadataServer/Web.Debug.config b/Engine/Source/Programs/UnrealGameSync/UnrealGameSyncMetadataServer/Web.Debug.config
new file mode 100644
--- /dev/null
+++ b/Engine/Source/Programs/UnrealGameSync/UnrealGameSyncMetadataServer/Web.Debug.config
@@ -0,0 +1,30 @@
+﻿<?xml version="1.0"?>
+
+<!-- For more information on using Web.config transformation visit http://go.microsoft.com/fwlink/?LinkId=301874 -->
+
+<configuration xmlns:xdt="http://schemas.microsoft.com/XML-Document-Transform">
+  <!--
+    In the example below, the "SetAttributes" transform will change the value of
+    "connectionString" to use "ReleaseSQLServer" only when the "Match" locator
+    finds an attribute "name" that has a value of "MyDB".
+
+    <connectionStrings>
+      <add name="MyDB"
+        connectionString="Data Source=ReleaseSQLServer;Initial Catalog=MyReleaseDB;Integrated Security=True"
+        xdt:Transform="SetAttributes" xdt:Locator="Match(name)"/>
+    </connectionStrings>
+  -->
+  <system.web>
+    <!--
+      In the example below, the "Replace" transform will replace the entire
+      <customErrors> section of your Web.config file.
+      Note that because there is only one customErrors section under the
+      <system.web> node, there is no need to use the "xdt:Locator" attribute.
+
+      <customErrors defaultRedirect="GenericError.htm"
+        mode="RemoteOnly" xdt:Transform="Replace">
+        <error statusCode="500" redirect="InternalError.htm"/>
+      </customErrors>
+    -->
+  </system.web>
+</configuration>
diff --git a/Engine/Source/Programs/UnrealGameSync/UnrealGameSyncMetadataServer/Web.Release.config b/Engine/Source/Programs/UnrealGameSync/UnrealGameSyncMetadataServer/Web.Release.config
new file mode 100644
--- /dev/null
+++ b/Engine/Source/Programs/UnrealGameSync/UnrealGameSyncMetadataServer/Web.Release.config
@@ -0,0 +1,31 @@
+﻿<?xml version="1.0"?>
+
+<!-- For more information on using Web.config transformation visit http://go.microsoft.com/fwlink/?LinkId=301874 -->
+
+<configuration xmlns:xdt="http://schemas.microsoft.com/XML-Document-Transform">
+  <!--
+    In the example below, the "SetAttributes" transform will change the value of
+    "connectionString" to use "ReleaseSQLServer" only when the "Match" locator
+    finds an attribute "name" that has a value of "MyDB".
+
+    <connectionStrings>
+      <add name="MyDB"
+        connectionString="Data Source=ReleaseSQLServer;Initial Catalog=MyReleaseDB;Integrated Security=True"
+        xdt:Transform="SetAttributes" xdt:Locator="Match(name)"/>
+    </connectionStrings>
+  -->
+  <system.web>
+    <compilation xdt:Transform="RemoveAttributes(debug)" />
+    <!--
+      In the example below, the "Replace" transform will replace the entire
+      <customErrors> section of your Web.config file.
+      Note that because there is only one customErrors section under the
+      <system.web> node, there is no need to use the "xdt:Locator" attribute.
+
+      <customErrors defaultRedirect="GenericError.htm"
+        mode="RemoteOnly" xdt:Transform="Replace">
+        <error statusCode="500" redirect="InternalError.htm"/>
+      </customErrors>
+    -->
+  </system.web>
+</configuration>
diff --git a/Engine/Source/Programs/UnrealGameSync/UnrealGameSyncMetadataServer/Web.template.config b/Engine/Source/Programs/UnrealGameSync/UnrealGameSyncMetadataServer/Web.template.config
new file mode 100644
--- /dev/null
+++ b/Engine/Source/Programs/UnrealGameSync/UnrealGameSyncMetadataServer/Web.template.config
@@ -0,0 +1,42 @@
+﻿<?xml version="1.0" encoding="utf-8"?>
+<!--
+  For more information on how to configure your ASP.NET application, please visit
+  http://go.microsoft.com/fwlink/?LinkId=301879
+  -->
+<configuration>
+  <appSettings>
+    <!-- Leave this connection string blank! It should go in either the .debug or .release transform config
+         There is a custom build step on BeforeBuild that generates a transformed Web.config on the fly. The generated Web.config should not be checked into source control, 
+         and this template.config should be checked in in it's place. -->
+    <add key="ConnectionString"
+      value=""/>
+  </appSettings>
+  <!--
+    For a description of web.config changes see http://go.microsoft.com/fwlink/?LinkId=235367.
+
+    The following attributes can be set on the <httpRuntime> tag.
+      <system.Web>
+        <httpRuntime targetFramework="4.5" />
+      </system.Web>
+  -->
+  <system.web>
+    <compilation debug="true" targetFramework="4.5" />
+    <httpRuntime targetFramework="4.5" />
+  </system.web>
+  <system.webServer>
+    <handlers>
+      <remove name="ExtensionlessUrlHandler-Integrated-4.0" />
+      <remove name="OPTIONSVerbHandler" />
+      <remove name="TRACEVerbHandler" />
+      <add name="ExtensionlessUrlHandler-Integrated-4.0" path="*." verb="*" type="System.Web.Handlers.TransferRequestHandler" preCondition="integratedMode,runtimeVersionv4.0" />
+    </handlers>
+  </system.webServer>
+  <runtime>
+    <assemblyBinding xmlns="urn:schemas-microsoft-com:asm.v1">
+      <dependentAssembly>
+        <assemblyIdentity name="Newtonsoft.Json" publicKeyToken="30ad4fe6b2a6aeed" culture="neutral" />
+        <bindingRedirect oldVersion="0.0.0.0-11.0.0.0" newVersion="11.0.0.0" />
+      </dependentAssembly>
+    </assemblyBinding>
+  </runtime>
+</configuration>
\ No newline at end of file
diff --git a/Engine/Source/Programs/UnrealGameSync/UnrealGameSyncMetadataServer/packages.config b/Engine/Source/Programs/UnrealGameSync/UnrealGameSyncMetadataServer/packages.config
new file mode 100644
--- /dev/null
+++ b/Engine/Source/Programs/UnrealGameSync/UnrealGameSyncMetadataServer/packages.config
@@ -0,0 +1,9 @@
+﻿<?xml version="1.0" encoding="utf-8"?>
+<packages>
+  <package id="Microsoft.AspNet.WebApi" version="5.2.3" targetFramework="net452" />
+  <package id="Microsoft.AspNet.WebApi.Client" version="5.2.3" targetFramework="net452" />
+  <package id="Microsoft.AspNet.WebApi.Core" version="5.2.3" targetFramework="net452" />
+  <package id="Microsoft.AspNet.WebApi.WebHost" version="5.2.3" targetFramework="net452" />
+  <package id="Newtonsoft.Json" version="11.0.2" targetFramework="net452" />
+  <package id="System.Data.SQLite.Core" version="1.0.108.0" targetFramework="net45" />
+</packages>
\ No newline at end of file
diff --git a/Engine/Source/Programs/UnrealGameSync/UnrealGameSyncWebAPI/App_Start/WebApiConfig.cs b/Engine/Source/Programs/UnrealGameSync/UnrealGameSyncWebAPI/App_Start/WebApiConfig.cs
deleted file mode 100644
diff --git a/Engine/Source/Programs/UnrealGameSync/UnrealGameSyncWebAPI/Connectors/SqlConnector.cs b/Engine/Source/Programs/UnrealGameSync/UnrealGameSyncWebAPI/Connectors/SqlConnector.cs
deleted file mode 100644
diff --git a/Engine/Source/Programs/UnrealGameSync/UnrealGameSyncWebAPI/Controllers/CISController.cs b/Engine/Source/Programs/UnrealGameSync/UnrealGameSyncWebAPI/Controllers/CISController.cs
deleted file mode 100644
diff --git a/Engine/Source/Programs/UnrealGameSync/UnrealGameSyncWebAPI/Controllers/CommentController.cs b/Engine/Source/Programs/UnrealGameSync/UnrealGameSyncWebAPI/Controllers/CommentController.cs
deleted file mode 100644
diff --git a/Engine/Source/Programs/UnrealGameSync/UnrealGameSyncWebAPI/Controllers/ErrorController.cs b/Engine/Source/Programs/UnrealGameSync/UnrealGameSyncWebAPI/Controllers/ErrorController.cs
deleted file mode 100644
diff --git a/Engine/Source/Programs/UnrealGameSync/UnrealGameSyncWebAPI/Controllers/EventController.cs b/Engine/Source/Programs/UnrealGameSync/UnrealGameSyncWebAPI/Controllers/EventController.cs
deleted file mode 100644
diff --git a/Engine/Source/Programs/UnrealGameSync/UnrealGameSyncWebAPI/Controllers/TelemetryController.cs b/Engine/Source/Programs/UnrealGameSync/UnrealGameSyncWebAPI/Controllers/TelemetryController.cs
deleted file mode 100644
diff --git a/Engine/Source/Programs/UnrealGameSync/UnrealGameSyncWebAPI/Global.asax.cs b/Engine/Source/Programs/UnrealGameSync/UnrealGameSyncWebAPI/Global.asax.cs
deleted file mode 100644
diff --git a/Engine/Source/Programs/UnrealGameSync/UnrealGameSyncWebAPI/Models/BuildData.cs b/Engine/Source/Programs/UnrealGameSync/UnrealGameSyncWebAPI/Models/BuildData.cs
deleted file mode 100644
diff --git a/Engine/Source/Programs/UnrealGameSync/UnrealGameSyncWebAPI/Models/CommentData.cs b/Engine/Source/Programs/UnrealGameSync/UnrealGameSyncWebAPI/Models/CommentData.cs
deleted file mode 100644
diff --git a/Engine/Source/Programs/UnrealGameSync/UnrealGameSyncWebAPI/Models/EventData.cs b/Engine/Source/Programs/UnrealGameSync/UnrealGameSyncWebAPI/Models/EventData.cs
deleted file mode 100644
diff --git a/Engine/Source/Programs/UnrealGameSync/UnrealGameSyncWebAPI/Models/EventSummary.cs b/Engine/Source/Programs/UnrealGameSync/UnrealGameSyncWebAPI/Models/EventSummary.cs
deleted file mode 100644
diff --git a/Engine/Source/Programs/UnrealGameSync/UnrealGameSyncWebAPI/Models/TelemetryErrorData.cs b/Engine/Source/Programs/UnrealGameSync/UnrealGameSyncWebAPI/Models/TelemetryErrorData.cs
deleted file mode 100644
diff --git a/Engine/Source/Programs/UnrealGameSync/UnrealGameSyncWebAPI/Models/TelemetryTimingData.cs b/Engine/Source/Programs/UnrealGameSync/UnrealGameSyncWebAPI/Models/TelemetryTimingData.cs
deleted file mode 100644
diff --git a/Engine/Source/Programs/UnrealGameSync/UnrealGameSyncWebAPI/Properties/AssemblyInfo.cs b/Engine/Source/Programs/UnrealGameSync/UnrealGameSyncWebAPI/Properties/AssemblyInfo.cs
deleted file mode 100644
diff --git a/Engine/Source/Programs/UnrealGameSync/UnrealGameSyncWebAPI/UnrealGameSyncWebAPI.csproj b/Engine/Source/Programs/UnrealGameSync/UnrealGameSyncWebAPI/UnrealGameSyncWebAPI.csproj
deleted file mode 100644
diff --git a/Engine/Source/Programs/UnrealGameSync/UnrealGameSyncWebAPI/Web.Debug.config b/Engine/Source/Programs/UnrealGameSync/UnrealGameSyncWebAPI/Web.Debug.config
deleted file mode 100644
diff --git a/Engine/Source/Programs/UnrealGameSync/UnrealGameSyncWebAPI/Web.Release.config b/Engine/Source/Programs/UnrealGameSync/UnrealGameSyncWebAPI/Web.Release.config
deleted file mode 100644
diff --git a/Engine/Source/Programs/UnrealGameSync/UnrealGameSyncWebAPI/Web.template.config b/Engine/Source/Programs/UnrealGameSync/UnrealGameSyncWebAPI/Web.template.config
deleted file mode 100644
diff --git a/Engine/Source/Programs/UnrealGameSync/UnrealGameSyncWebAPI/packages.config b/Engine/Source/Programs/UnrealGameSync/UnrealGameSyncWebAPI/packages.config
deleted file mode 100644
