merging 4159132 into 4.20
- Fixed issue where namespace was still wrong from rename in global.asax and assemblyinfo
- Removed a couple of erroneous publish profiles
- Added GET to /error. takes optional records param. 10 by default

diff --git a/Engine/Source/Programs/UnrealGameSync/UnrealGameSyncMetadataServer/Connectors/SqlConnector.cs b/Engine/Source/Programs/UnrealGameSync/UnrealGameSyncMetadataServer/Connectors/SqlConnector.cs
--- a/Engine/Source/Programs/UnrealGameSync/UnrealGameSyncMetadataServer/Connectors/SqlConnector.cs
+++ b/Engine/Source/Programs/UnrealGameSync/UnrealGameSyncMetadataServer/Connectors/SqlConnector.cs
@@ -155,6 +155,38 @@ namespace UnrealGameSyncMetadataServer.Connectors
 			}
 			return ReturnedBuilds;
 		}
+
+		public static List<TelemetryErrorData> GetErrorData(int Records)
+		{
+			List<TelemetryErrorData> ReturnedErrors = new List<TelemetryErrorData>();
+			using (SQLiteConnection Connection = new SQLiteConnection(ConnectionString))
+			{
+				Connection.Open();
+				using (SQLiteCommand Command = new SQLiteCommand("SELECT Id, Type, Text, UserName, Project, Timestamp, Version, IpAddress FROM [Errors] ORDER BY Id DESC LIMIT @param1", Connection))
+				{
+					Command.Parameters.AddWithValue("@param1", Records);
+					using (SQLiteDataReader Reader = Command.ExecuteReader())
+					{
+						while (Reader.Read())
+						{
+							TelemetryErrorData Error = new TelemetryErrorData();
+							Error.Id = Reader.GetInt32(0);
+							Enum.TryParse(Reader.GetString(1), true, out Error.Type);
+							Error.Text = Reader.GetString(2);
+							Error.UserName = Reader.GetString(3);
+							Error.Project = Reader.IsDBNull(4) ? null : Reader.GetString(4);
+							Error.Timestamp = Reader.GetDateTime(5);
+							Error.Version = Reader.GetString(6);
+							Error.IpAddress = Reader.GetString(7);
+							ReturnedErrors.Add(Error);
+						}
+					}
+				}
+			}
+			return ReturnedErrors;
+		}
+
+
 		public static void PostCIS(BuildData Build)
 		{
 			using (SQLiteConnection Connection = new SQLiteConnection(ConnectionString))
diff --git a/Engine/Source/Programs/UnrealGameSync/UnrealGameSyncMetadataServer/Controllers/ErrorController.cs b/Engine/Source/Programs/UnrealGameSync/UnrealGameSyncMetadataServer/Controllers/ErrorController.cs
--- a/Engine/Source/Programs/UnrealGameSync/UnrealGameSyncMetadataServer/Controllers/ErrorController.cs
+++ b/Engine/Source/Programs/UnrealGameSync/UnrealGameSyncMetadataServer/Controllers/ErrorController.cs
@@ -1,5 +1,6 @@
 ﻿// Copyright 1998-2018 Epic Games, Inc. All Rights Reserved.
 
+using System.Collections.Generic;
 using System.Web.Http;
 using UnrealGameSyncMetadataServer.Connectors;
 using UnrealGameSyncMetadataServer.Models;
@@ -8,6 +9,10 @@ namespace UnrealGameSyncMetadataServer.Controllers
 {
     public class ErrorController : ApiController
     {
+		public List<TelemetryErrorData> Get(int Records = 10)
+		{
+			return SqlConnector.GetErrorData(Records);
+		}
 		public void Post([FromBody] TelemetryErrorData Data, string Version, string IpAddress)
 		{
 			SqlConnector.PostErrorData(Data, Version, IpAddress);
diff --git a/Engine/Source/Programs/UnrealGameSync/UnrealGameSyncMetadataServer/Models/TelemetryErrorData.cs b/Engine/Source/Programs/UnrealGameSync/UnrealGameSyncMetadataServer/Models/TelemetryErrorData.cs
--- a/Engine/Source/Programs/UnrealGameSync/UnrealGameSyncMetadataServer/Models/TelemetryErrorData.cs
+++ b/Engine/Source/Programs/UnrealGameSync/UnrealGameSyncMetadataServer/Models/TelemetryErrorData.cs
@@ -13,10 +13,13 @@ namespace UnrealGameSyncMetadataServer.Models
 		{
 			Crash,
 		}
+		public int Id;
 		public TelemetryErrorType Type;
 		public string Text;
 		public string UserName;
 		public string Project;
 		public DateTime Timestamp;
+		public string Version;
+		public string IpAddress;
 	}
 }
\ No newline at end of file
diff --git a/Engine/Source/Programs/UnrealGameSync/UnrealGameSyncMetadataServer/Properties/AssemblyInfo.cs b/Engine/Source/Programs/UnrealGameSync/UnrealGameSyncMetadataServer/Properties/AssemblyInfo.cs
--- a/Engine/Source/Programs/UnrealGameSync/UnrealGameSyncMetadataServer/Properties/AssemblyInfo.cs
+++ b/Engine/Source/Programs/UnrealGameSync/UnrealGameSyncMetadataServer/Properties/AssemblyInfo.cs
@@ -7,7 +7,7 @@ using System.Runtime.InteropServices;
 // General Information about an assembly is controlled through the following
 // set of attributes. Change these attribute values to modify the information
 // associated with an assembly.
-[assembly: AssemblyTitle("UnrealGameSyncWebAPI")]
+[assembly: AssemblyTitle("UnrealGameSyncMetadataServer")]
 [assembly: AssemblyDescription("")]
 [assembly: AssemblyConfiguration("")]
 [assembly: AssemblyCompany("")]
diff --git a/Engine/Source/Programs/UnrealGameSync/UnrealGameSyncMetadataServer/UnrealGameSyncMetadataServer.csproj b/Engine/Source/Programs/UnrealGameSync/UnrealGameSyncMetadataServer/UnrealGameSyncMetadataServer.csproj
--- a/Engine/Source/Programs/UnrealGameSync/UnrealGameSyncMetadataServer/UnrealGameSyncMetadataServer.csproj
+++ b/Engine/Source/Programs/UnrealGameSync/UnrealGameSyncMetadataServer/UnrealGameSyncMetadataServer.csproj
@@ -109,7 +109,6 @@
   <ItemGroup>
     <Content Include="Global.asax" />
     <Content Include="Web.config" />
-    <None Include="Properties\PublishProfiles\File.pubxml" />
     <None Include="Web.template.config" />
     <None Include="Web.Debug.config">
       <DependentUpon>Web.template.config</DependentUpon>
