UGS: Try to read the editor executable path from the target receipt, rather than guessing the executable path directly.

diff --git a/Engine/Source/Programs/UnrealGameSync/UnrealGameSync/Controls/WorkspaceControl.cs b/Engine/Source/Programs/UnrealGameSync/UnrealGameSync/Controls/WorkspaceControl.cs
--- a/Engine/Source/Programs/UnrealGameSync/UnrealGameSync/Controls/WorkspaceControl.cs
+++ b/Engine/Source/Programs/UnrealGameSync/UnrealGameSync/Controls/WorkspaceControl.cs
@@ -21,6 +21,7 @@ using System.Windows.Forms.VisualStyles;
 using Microsoft.Win32;
 using System.Threading;
 using System.Drawing.Imaging;
+using System.Web.Script.Serialization;
 
 namespace UnrealGameSync
 {
@@ -2286,6 +2287,40 @@ namespace UnrealGameSync
 
 		private string GetEditorExePath(BuildConfig Config)
 		{
+			// Try to read the executable path from the target receipt
+			List<string> ReceiptFileNames = GetEditorReceiptPaths(Config);
+			foreach(string ReceiptFileName in ReceiptFileNames)
+			{
+				if(File.Exists(ReceiptFileName))
+				{
+					Log.WriteLine("Reading {0}", ReceiptFileName);
+					try
+					{
+						string Text = File.ReadAllText(ReceiptFileName);
+						JavaScriptSerializer Serializer = new JavaScriptSerializer();
+						Dictionary<string, object> RawObject = Serializer.Deserialize<Dictionary<string, object>>(Text);
+
+						object LaunchFileNameObject;
+						if(RawObject.TryGetValue("Launch", out LaunchFileNameObject))
+						{
+							string LaunchFileName = LaunchFileNameObject as string;
+							if(LaunchFileName != null)
+							{
+								LaunchFileName = LaunchFileName.Replace("$(EngineDir)", Path.Combine(BranchDirectoryName, "Engine"));
+								LaunchFileName = LaunchFileName.Replace("$(ProjectDir)", Path.GetDirectoryName(SelectedFileName));
+								return Path.GetFullPath(LaunchFileName);
+							}
+						}
+					}
+					catch(Exception Ex)
+					{
+						Log.WriteLine("  Exception while parsing receipt: {0}", Ex.ToString());
+					}
+					break;
+				}
+			}
+
+			// Otherwise use the standard editor path
 			string ExeFileName = "UE4Editor.exe";
 			if((Config != BuildConfig.DebugGame || PerforceMonitor.LatestProjectConfigFile.GetValue("Options.DebugGameHasSeparateExecutable", false)) && Config != BuildConfig.Development)
 			{
