UGS: Fix exception entering custom server settings when UGS registry key does not exist.

diff --git a/Engine/Source/Programs/UnrealGameSync/UnrealGameSyncLauncher/Program.cs b/Engine/Source/Programs/UnrealGameSync/UnrealGameSyncLauncher/Program.cs
--- a/Engine/Source/Programs/UnrealGameSync/UnrealGameSyncLauncher/Program.cs
+++ b/Engine/Source/Programs/UnrealGameSync/UnrealGameSyncLauncher/Program.cs
@@ -119,35 +119,42 @@ namespace UnrealGameSyncLauncher
 
 		public static void SaveSettings(string ServerAndPort, string UserName, string DepotPath)
 		{
-			using (RegistryKey Key = Registry.CurrentUser.OpenSubKey("SOFTWARE\\Epic Games\\UnrealGameSync", true))
+			try
 			{
-				if(String.IsNullOrEmpty(ServerAndPort))
-				{
-					try { Key.DeleteValue("Server"); } catch(Exception) { }
-				}
-				else
+				using (RegistryKey Key = Registry.CurrentUser.CreateSubKey("SOFTWARE\\Epic Games\\UnrealGameSync"))
 				{
-					Key.SetValue("Server", ServerAndPort);
-				}
+					if(String.IsNullOrEmpty(ServerAndPort))
+					{
+						try { Key.DeleteValue("Server"); } catch(Exception) { }
+					}
+					else
+					{
+						Key.SetValue("Server", ServerAndPort);
+					}
 
-				if(String.IsNullOrEmpty(UserName))
-				{
-					try { Key.DeleteValue("UserName"); } catch(Exception) { }
-				}
-				else
-				{
-					Key.SetValue("UserName", UserName);
-				}
+					if(String.IsNullOrEmpty(UserName))
+					{
+						try { Key.DeleteValue("UserName"); } catch(Exception) { }
+					}
+					else
+					{
+						Key.SetValue("UserName", UserName);
+					}
 
-				if(String.IsNullOrEmpty(DepotPath) || (DefaultDepotPath != null && String.Equals(DepotPath, DefaultDepotPath, StringComparison.InvariantCultureIgnoreCase)))
-				{
-					try { Key.DeleteValue("DepotPath"); } catch(Exception) { }
-				}
-				else
-				{
-					Key.SetValue("DepotPath", DepotPath);
+					if(String.IsNullOrEmpty(DepotPath) || (DefaultDepotPath != null && String.Equals(DepotPath, DefaultDepotPath, StringComparison.InvariantCultureIgnoreCase)))
+					{
+						try { Key.DeleteValue("DepotPath"); } catch(Exception) { }
+					}
+					else
+					{
+						Key.SetValue("DepotPath", DepotPath);
+					}
 				}
 			}
+			catch(Exception Ex)
+			{
+				MessageBox.Show("Unable to save settings.\n\n" + Ex.ToString());
+			}
 		}
 
 		public static bool SyncAndRun(PerforceConnection Perforce, string BaseDepotPath, bool bUnstable, string[] Args, Mutex InstanceMutex, TextWriter LogWriter)
