UGS: Add proper support for DPI scaling.

diff --git a/Engine/Source/Programs/UnrealGameSync/UnrealGameSync/Controls/StatusPanel.cs b/Engine/Source/Programs/UnrealGameSync/UnrealGameSync/Controls/StatusPanel.cs
--- a/Engine/Source/Programs/UnrealGameSync/UnrealGameSync/Controls/StatusPanel.cs
+++ b/Engine/Source/Programs/UnrealGameSync/UnrealGameSync/Controls/StatusPanel.cs
@@ -82,7 +82,7 @@ namespace UnrealGameSync
 
 		public override Size Measure(Graphics Graphics, StatusElementResources Resources)
 		{
-			return Icon.Size;
+			return new Size((int)(Icon.Width * Graphics.DpiX / 96.0f), (int)(Icon.Height * Graphics.DpiY / 96.0f));
 		}
 
 		public override void Draw(Graphics Graphics, StatusElementResources Resources)
@@ -106,7 +106,7 @@ namespace UnrealGameSync
 
 		public override Size Measure(Graphics Graphics, StatusElementResources Resources)
 		{
-			return IconSize;
+			return new Size((int)(IconSize.Width * Graphics.DpiX / 96.0f), (int)(IconSize.Height * Graphics.DpiY / 96.0f));
 		}
 
 		public override void Draw(Graphics Graphics, StatusElementResources Resources)
@@ -543,7 +543,7 @@ namespace UnrealGameSync
 			int BodyHeight = Height;
 			if(Alert != null)
 			{
-				AlertDividerY = Height - Font.Height - 16;
+				AlertDividerY = Height - (int)(Font.Height * 2);
 				Alert.Layout(Graphics, Point.Empty, Resources);
 				Alert.Layout(Graphics, new Point((Width - Alert.Bounds.Width) / 2, (Height + AlertDividerY) / 2), Resources);
 				BodyHeight = AlertDividerY;
@@ -551,7 +551,7 @@ namespace UnrealGameSync
 
 			// Get the logo size
 			Image DrawProjectLogo = ProjectLogo ?? Properties.Resources.DefaultProjectLogo;
-			float LogoScale = Math.Min((float)BodyHeight - ((Caption != null)? Font.Height : 0) / DrawProjectLogo.Height, 1.0f);
+			float LogoScale = Math.Min((float)BodyHeight - ((Caption != null)? Font.Height : 0) / DrawProjectLogo.Height, Graphics.DpiY / 96.0f);
 			int LogoWidth = (int)(DrawProjectLogo.Width * LogoScale);
 			int LogoHeight = (int)(DrawProjectLogo.Height * LogoScale);
 
diff --git a/Engine/Source/Programs/UnrealGameSync/UnrealGameSync/Controls/TabControl.cs b/Engine/Source/Programs/UnrealGameSync/UnrealGameSync/Controls/TabControl.cs
--- a/Engine/Source/Programs/UnrealGameSync/UnrealGameSync/Controls/TabControl.cs
+++ b/Engine/Source/Programs/UnrealGameSync/UnrealGameSync/Controls/TabControl.cs
@@ -243,14 +243,17 @@ namespace UnrealGameSync
 		{
 			using(Graphics Graphics = CreateGraphics())
 			{
+				float DpiScaleX = Graphics.DpiX / 96.0f;
+				float DpiScaleY = Graphics.DpiY / 96.0f;
+
 				for(int Idx = 0; Idx < Tabs.Count; Idx++)
 				{
 					TabData Tab = Tabs[Idx];
 					Tab.TextSize = TextRenderer.MeasureText(Graphics, Tab.Name, Font, new Size(int.MaxValue, int.MaxValue), TextFormatFlags.NoPadding);
-					Tab.Width = TabPadding + Tab.TextSize.Width + TabPadding;
+					Tab.Width = TabPadding + Tab.TextSize.Width + (int)(TabPadding * DpiScaleX);
 					if(Idx == SelectedTabIdx)
 					{
-						Tab.Width += TabCloseButtonWidth;
+						Tab.Width += (int)(TabCloseButtonWidth * DpiScaleX);
 					}
 				}
 
@@ -515,17 +518,21 @@ namespace UnrealGameSync
 				SmoothingMode PrevSmoothingMode = e.Graphics.SmoothingMode;
 				e.Graphics.SmoothingMode = System.Drawing.Drawing2D.SmoothingMode.AntiAlias;
 
-				int CloseMidX = SelectedTab.MinX + SelectedTab.Width - (TabPadding + TabCloseButtonWidth) / 2;// + (13 / 2);
-				int CloseMidY = (ClientSize.Height - 4) / 2;
+				float DpiScaleX = e.Graphics.DpiX / 96.0f;
+				float DpiScaleY = e.Graphics.DpiY / 96.0f;
+
+				int CloseMidX = SelectedTab.MinX + SelectedTab.Width - (int)((TabPadding + TabCloseButtonWidth) * DpiScaleX / 2);// + (13 / 2);
+				int CloseMidY = (ClientSize.Height - (int)(4 * DpiScaleY)) / 2;
 
-				Rectangle CloseButton = new Rectangle(CloseMidX - (13 / 2), CloseMidY - (13 / 2), 13, 13);
+				Rectangle CloseButton = new Rectangle(CloseMidX - (int)((13 / 2) * DpiScaleX), CloseMidY - (int)((13 / 2) * DpiScaleY), (int)(13 * DpiScaleX), (int)(13 * DpiScaleY));
 				e.Graphics.FillEllipse(SystemBrushes.ControlDark, CloseButton);
 
-				using(Pen CrossPen = new Pen(SystemBrushes.Window, 2.0f))
+				using(Pen CrossPen = new Pen(SystemBrushes.Window, 2.0f * DpiScaleX))
 				{
-					const float Indent = 3.5f;
-					e.Graphics.DrawLine(CrossPen, CloseButton.Left + Indent, CloseButton.Top + Indent, CloseButton.Right - Indent, CloseButton.Bottom - Indent);
-					e.Graphics.DrawLine(CrossPen, CloseButton.Left + Indent, CloseButton.Bottom - Indent, CloseButton.Right - Indent, CloseButton.Top + Indent);
+					float IndentX = 3.5f * DpiScaleX;
+					float IndentY = 3.5f * DpiScaleY;
+					e.Graphics.DrawLine(CrossPen, CloseButton.Left + IndentX, CloseButton.Top + IndentY, CloseButton.Right - IndentX, CloseButton.Bottom - IndentY);
+					e.Graphics.DrawLine(CrossPen, CloseButton.Left + IndentX, CloseButton.Bottom - IndentY, CloseButton.Right - IndentX, CloseButton.Top + IndentY);
 				}
 				e.Graphics.SmoothingMode = PrevSmoothingMode;
 
@@ -541,22 +548,26 @@ namespace UnrealGameSync
 
 		void DrawBackground(Graphics Graphics, TabData Tab, Brush BackgroundBrush, Brush StripeBrush, Tuple<Color, float> Highlight)
 		{
+			float DpiScaleY = Graphics.DpiY / 96.0f;
+
 			Graphics.SmoothingMode = System.Drawing.Drawing2D.SmoothingMode.None;
 			Graphics.FillRectangle(BackgroundBrush, Tab.MinX + 1, 1, Tab.Width - 1, ClientSize.Height - 3);
-			Graphics.FillRectangle(StripeBrush, Tab.MinX + 1, ClientSize.Height - 5, Tab.Width - 1, 3);
+			Graphics.FillRectangle(StripeBrush, Tab.MinX + 1, ClientSize.Height - (int)(5 * DpiScaleY), Tab.Width - 1, (int)(5 * DpiScaleY) - 2);
 
 			if(Highlight != null && Highlight.Item2 > 0.0f)
 			{
 				using(SolidBrush Brush = new SolidBrush(Highlight.Item1))
 				{
-					Graphics.FillRectangle(Brush, Tab.MinX + 1, ClientSize.Height - 5, (int)((Tab.Width - 2) * Highlight.Item2), 3);
+					Graphics.FillRectangle(Brush, Tab.MinX + 1, ClientSize.Height - (int)(5 * DpiScaleY), (int)((Tab.Width - 2) * Highlight.Item2), (int)(5 * DpiScaleY) - 2);
 				}
 			}
 		}
 
 		void DrawText(Graphics Graphics, TabData Tab)
 		{
-			TextRenderer.DrawText(Graphics, Tab.Name, Font, new Rectangle(Tab.MinX + TabPadding, 0, Tab.TextSize.Width, ClientSize.Height - 3), Color.Black, TextFormatFlags.NoPadding | TextFormatFlags.EndEllipsis | TextFormatFlags.VerticalCenter);
+			float DpiScaleX = Graphics.DpiX / 96.0f;
+			float DpiScaleY = Graphics.DpiY / 96.0f;
+			TextRenderer.DrawText(Graphics, Tab.Name, Font, new Rectangle(Tab.MinX + (int)(TabPadding * DpiScaleX), 0, Tab.TextSize.Width, ClientSize.Height - (int)(3 * DpiScaleY)), Color.Black, TextFormatFlags.NoPadding | TextFormatFlags.EndEllipsis | TextFormatFlags.VerticalCenter);
 		}
 	}
 }
diff --git a/Engine/Source/Programs/UnrealGameSync/UnrealGameSync/Controls/WorkspaceControl.cs b/Engine/Source/Programs/UnrealGameSync/UnrealGameSync/Controls/WorkspaceControl.cs
--- a/Engine/Source/Programs/UnrealGameSync/UnrealGameSync/Controls/WorkspaceControl.cs
+++ b/Engine/Source/Programs/UnrealGameSync/UnrealGameSync/Controls/WorkspaceControl.cs
@@ -397,21 +397,26 @@ namespace UnrealGameSync
 					ColumnWidths[Idx] = BuildList.Columns[Idx].Width;
 				}
 
-				MinColumnWidths = Enumerable.Repeat(32, BuildList.Columns.Count).ToArray();
-				MinColumnWidths[IconColumn.Index] = 50;
-				MinColumnWidths[TypeColumn.Index] = 100;
-				MinColumnWidths[TimeColumn.Index] = 75;
-				MinColumnWidths[ChangeColumn.Index] = 75;
-				MinColumnWidths[CISColumn.Index] = 200;
-
-				DesiredColumnWidths = Enumerable.Repeat(65536, BuildList.Columns.Count).ToArray();
-				DesiredColumnWidths[IconColumn.Index] = MinColumnWidths[IconColumn.Index];
-				DesiredColumnWidths[TypeColumn.Index] = MinColumnWidths[TypeColumn.Index];
-				DesiredColumnWidths[TimeColumn.Index] = MinColumnWidths[TimeColumn.Index];
-				DesiredColumnWidths[ChangeColumn.Index] = MinColumnWidths[ChangeColumn.Index];
-				DesiredColumnWidths[AuthorColumn.Index] = 120;
-				DesiredColumnWidths[CISColumn.Index] = 200;
-				DesiredColumnWidths[StatusColumn.Index] = 300;
+				using(Graphics Graphics = Graphics.FromHwnd(IntPtr.Zero))
+				{
+					float DpiScaleX = Graphics.DpiX / 96.0f;
+
+					MinColumnWidths = Enumerable.Repeat(32, BuildList.Columns.Count).ToArray();
+					MinColumnWidths[IconColumn.Index] = (int)(50 * DpiScaleX);
+					MinColumnWidths[TypeColumn.Index] = (int)(100 * DpiScaleX);
+					MinColumnWidths[TimeColumn.Index] = (int)(75 * DpiScaleX);
+					MinColumnWidths[ChangeColumn.Index] = (int)(75 * DpiScaleX);
+					MinColumnWidths[CISColumn.Index] = (int)(200 * DpiScaleX);
+
+					DesiredColumnWidths = Enumerable.Repeat(65536, BuildList.Columns.Count).ToArray();
+					DesiredColumnWidths[IconColumn.Index] = MinColumnWidths[IconColumn.Index];
+					DesiredColumnWidths[TypeColumn.Index] = MinColumnWidths[TypeColumn.Index];
+					DesiredColumnWidths[TimeColumn.Index] = MinColumnWidths[TimeColumn.Index];
+					DesiredColumnWidths[ChangeColumn.Index] = MinColumnWidths[ChangeColumn.Index];
+					DesiredColumnWidths[AuthorColumn.Index] = (int)(120 * DpiScaleX);
+					DesiredColumnWidths[CISColumn.Index] = (int)(200 * DpiScaleX);
+					DesiredColumnWidths[StatusColumn.Index] = (int)(300 * DpiScaleX);
+				}
 
 				ColumnWeights = Enumerable.Repeat(1.0f, BuildList.Columns.Count).ToArray();
 				ColumnWeights[IconColumn.Index] = 3.0f;
@@ -1483,7 +1488,10 @@ namespace UnrealGameSync
 		{
 			PerforceChangeSummary Change = (PerforceChangeSummary)e.Item.Tag;
 
-			int IconY = e.Bounds.Top + (e.Bounds.Height - 16) / 2;
+			float DpiScaleX = e.Graphics.DpiX / 96.0f;
+			float DpiScaleY = e.Graphics.DpiY / 96.0f;
+
+			float IconY = e.Bounds.Top + (e.Bounds.Height - 16 * DpiScaleY) / 2;
 
 			StringFormat Format = new StringFormat();
 			Format.LineAlignment = StringAlignment.Center;
@@ -1500,20 +1508,20 @@ namespace UnrealGameSync
 			{
 				EventSummary Summary = EventMonitor.GetSummaryForChange(Change.Number);
 
-				int MinX = 4;
+				float MinX = 4 * DpiScaleX;
 				if((Summary != null && EventMonitor.WasSyncedByCurrentUser(Summary.ChangeNumber)) || (Workspace != null && Workspace.CurrentChangeNumber == Change.Number))
 				{
-					e.Graphics.DrawImage(Properties.Resources.Icons, MinX, IconY, PreviousSyncIcon, GraphicsUnit.Pixel);
+					e.Graphics.DrawImage(Properties.Resources.Icons, MinX * DpiScaleX, IconY, PreviousSyncIcon, GraphicsUnit.Pixel);
 				}
 				else if(WorkspaceSettings != null && WorkspaceSettings.AdditionalChangeNumbers.Contains(Change.Number))
 				{
-					e.Graphics.DrawImage(Properties.Resources.Icons, MinX, IconY, AdditionalSyncIcon, GraphicsUnit.Pixel);
+					e.Graphics.DrawImage(Properties.Resources.Icons, MinX * DpiScaleX, IconY, AdditionalSyncIcon, GraphicsUnit.Pixel);
 				}
 				else if(bAllowSync && ((Summary != null && Summary.LastStarReview != null && Summary.LastStarReview.Type == EventType.Starred) || PromotedChangeNumbers.Contains(Change.Number)))
 				{
-					e.Graphics.DrawImage(Properties.Resources.Icons, MinX, IconY, PromotedBuildIcon, GraphicsUnit.Pixel);
+					e.Graphics.DrawImage(Properties.Resources.Icons, MinX * DpiScaleX, IconY, PromotedBuildIcon, GraphicsUnit.Pixel);
 				}
-				MinX += PromotedBuildIcon.Width;
+				MinX += PromotedBuildIcon.Width * DpiScaleX;
 
 				if(bAllowSync)
 				{
@@ -1570,7 +1578,7 @@ namespace UnrealGameSync
 					}
 					e.Graphics.DrawImage(Properties.Resources.Icons, MinX, IconY, QualityIcon, GraphicsUnit.Pixel);
 
-					MinX += QualityIcon.Width;
+					MinX += QualityIcon.Width * DpiScaleX;
 				}
 			}
 			else if(e.ColumnIndex == ChangeColumn.Index || e.ColumnIndex == TimeColumn.Index)
@@ -1591,12 +1599,12 @@ namespace UnrealGameSync
 				    e.Graphics.IntersectClip(e.Bounds);
 				    e.Graphics.SmoothingMode = SmoothingMode.HighQuality;
 
-					RemainingBounds = new Rectangle(RemainingBounds.Left, RemainingBounds.Top, RemainingBounds.Width - 2, RemainingBounds.Height);
+					RemainingBounds = new Rectangle(RemainingBounds.Left, RemainingBounds.Top, RemainingBounds.Width - (int)(2 * DpiScaleX), RemainingBounds.Height);
 
 					Point ListLocation = GetBadgeListLocation(Layout.DescriptionBadges, RemainingBounds, HorizontalAlign.Right, VerticalAlignment.Middle);
 					DrawBadges(e.Graphics, ListLocation, Layout.DescriptionBadges, BadgeAlpha);
 
-					RemainingBounds = new Rectangle(RemainingBounds.Left, RemainingBounds.Top, ListLocation.X - RemainingBounds.Left - 2, RemainingBounds.Height);
+					RemainingBounds = new Rectangle(RemainingBounds.Left, RemainingBounds.Top, ListLocation.X - RemainingBounds.Left - (int)(2 * DpiScaleX), RemainingBounds.Height);
 				}
 
 				TextRenderer.DrawText(e.Graphics, e.SubItem.Text, CurrentFont, RemainingBounds, TextColor, TextFormatFlags.EndEllipsis | TextFormatFlags.SingleLine | TextFormatFlags.VerticalCenter | TextFormatFlags.NoPrefix);
@@ -1628,7 +1636,7 @@ namespace UnrealGameSync
 			}
 			else if(e.ColumnIndex == StatusColumn.Index)
 			{
-				int MaxX = e.SubItem.Bounds.Right;
+				float MaxX = e.SubItem.Bounds.Right;
 
 				if(Change.Number == Workspace.PendingChangeNumber && Workspace.IsBusy())
 				{
@@ -1643,9 +1651,9 @@ namespace UnrealGameSync
 						e.Graphics.DrawImage(Properties.Resources.Icons, MaxX, IconY, InfoIcon, GraphicsUnit.Pixel);
 					}
 
-					int MinX = e.Bounds.Left;
+					float MinX = e.Bounds.Left;
 
-					TextRenderer.DrawText(e.Graphics, Progress.Item1, CurrentFont, new Rectangle(MinX, e.Bounds.Top, MaxX - MinX, e.Bounds.Height), TextColor, TextFormatFlags.Left | TextFormatFlags.VerticalCenter | TextFormatFlags.EndEllipsis | TextFormatFlags.NoPrefix);
+					TextRenderer.DrawText(e.Graphics, Progress.Item1, CurrentFont, new Rectangle((int)MinX, e.Bounds.Top, (int)(MaxX - MinX), e.Bounds.Height), TextColor, TextFormatFlags.Left | TextFormatFlags.VerticalCenter | TextFormatFlags.EndEllipsis | TextFormatFlags.NoPrefix);
 				}
 				else
 				{
@@ -1656,17 +1664,17 @@ namespace UnrealGameSync
 					{
 						EventData Review = EventMonitor.GetReviewByCurrentUser(Change.Number);
 
-						MaxX -= FrownIcon.Width;
+						MaxX -= FrownIcon.Width * DpiScaleX;
 						e.Graphics.DrawImage(Properties.Resources.Icons, MaxX, IconY, (Review == null || !EventMonitor.IsPositiveReview(Review.Type))? FrownIcon : DisabledFrownIcon, GraphicsUnit.Pixel);
 
-						MaxX -= HappyIcon.Width;
+						MaxX -= HappyIcon.Width * DpiScaleX;
 						e.Graphics.DrawImage(Properties.Resources.Icons, MaxX, IconY, (Review == null || !EventMonitor.IsNegativeReview(Review.Type))? HappyIcon : DisabledHappyIcon, GraphicsUnit.Pixel);
 					}
 					else if(e.ItemIndex == BuildList.HoverItem && bAllowSync)
 					{
 						Rectangle SyncBadgeRectangle = GetSyncBadgeRectangle(e.SubItem.Bounds);
 						DrawBadge(e.Graphics, SyncBadgeRectangle, "Sync", bHoverSync? Color.FromArgb(140, 180, 230) : Color.FromArgb(112, 146, 190), true, true);
-						MaxX = SyncBadgeRectangle.Left - 2;
+						MaxX = SyncBadgeRectangle.Left - (int)(2 * DpiScaleX);
 					}
 
 					string SummaryText;
@@ -1699,7 +1707,7 @@ namespace UnrealGameSync
 						SummaryText = SummaryText.Substring(0, SummaryText.IndexOf('\n')).TrimEnd() + "...";
 					}
 
-					TextRenderer.DrawText(e.Graphics, SummaryText, CurrentFont, new Rectangle(e.Bounds.Left, e.Bounds.Top, MaxX - e.Bounds.Left, e.Bounds.Height), TextColor, TextFormatFlags.Left | TextFormatFlags.VerticalCenter | TextFormatFlags.EndEllipsis | TextFormatFlags.NoPrefix);
+					TextRenderer.DrawText(e.Graphics, SummaryText, CurrentFont, new Rectangle(e.Bounds.Left, e.Bounds.Top, (int)MaxX - e.Bounds.Left, e.Bounds.Height), TextColor, TextFormatFlags.Left | TextFormatFlags.VerticalCenter | TextFormatFlags.EndEllipsis | TextFormatFlags.NoPrefix);
 				}
 			}
 			else
@@ -2663,14 +2671,17 @@ namespace UnrealGameSync
 				}
 			}
 
-			int StatusPanelHeight = 148;
-			if(Alert != null)
+			using(Graphics Graphics = CreateGraphics())
 			{
-				StatusPanelHeight += 40;
-			}
-			if(StatusLayoutPanel.RowStyles[0].Height != StatusPanelHeight)
-			{
-				StatusLayoutPanel.RowStyles[0].Height = StatusPanelHeight;
+				int StatusPanelHeight = (int)(148.0f * Graphics.DpiY / 96.0f);
+				if(Alert != null)
+				{
+					StatusPanelHeight += (int)(40.0f * Graphics.DpiY / 96.0f);
+				}
+				if(StatusLayoutPanel.RowStyles[0].Height != StatusPanelHeight)
+				{
+					StatusLayoutPanel.RowStyles[0].Height = StatusPanelHeight;
+				}
 			}
 
 			StatusPanel.Set(Lines, Caption, Alert, TintColor);
diff --git a/Engine/Source/Programs/UnrealGameSync/UnrealGameSync/Forms/BuildStepWindow.cs b/Engine/Source/Programs/UnrealGameSync/UnrealGameSync/Forms/BuildStepWindow.cs
--- a/Engine/Source/Programs/UnrealGameSync/UnrealGameSync/Forms/BuildStepWindow.cs
+++ b/Engine/Source/Programs/UnrealGameSync/UnrealGameSync/Forms/BuildStepWindow.cs
@@ -29,6 +29,9 @@ namespace UnrealGameSync
 			Variables = InVariables;
 
 			InitializeComponent();
+
+			this.MaximumSize = new System.Drawing.Size(32768, Height);
+			this.MinimumSize = new System.Drawing.Size(Width, Height);
 		}
 
 		private void BuildTaskWindow_Load(object sender, EventArgs e)
diff --git a/Engine/Source/Programs/UnrealGameSync/UnrealGameSync/Forms/ChangelistWindow.cs b/Engine/Source/Programs/UnrealGameSync/UnrealGameSync/Forms/ChangelistWindow.cs
--- a/Engine/Source/Programs/UnrealGameSync/UnrealGameSync/Forms/ChangelistWindow.cs
+++ b/Engine/Source/Programs/UnrealGameSync/UnrealGameSync/Forms/ChangelistWindow.cs
@@ -1,4 +1,4 @@
-﻿// Copyright 1998-2018 Epic Games, Inc. All Rights Reserved.
+// Copyright 1998-2018 Epic Games, Inc. All Rights Reserved.
 
 using System;
 using System.Collections.Generic;
@@ -20,6 +20,9 @@ namespace UnrealGameSync
 		{
 			InitializeComponent();
 
+			this.MaximumSize = new System.Drawing.Size(32768, Height);
+			this.MinimumSize = new System.Drawing.Size(Width, Height);
+
 			ChangeNumber = InChangeNumber;
 			ChangeNumberTextBox.Text = (ChangeNumber > 0)? ChangeNumber.ToString() : "";
 		}
diff --git a/Engine/Source/Programs/UnrealGameSync/UnrealGameSync/Forms/ModifyBuildStepsWindow.cs b/Engine/Source/Programs/UnrealGameSync/UnrealGameSync/Forms/ModifyBuildStepsWindow.cs
--- a/Engine/Source/Programs/UnrealGameSync/UnrealGameSync/Forms/ModifyBuildStepsWindow.cs
+++ b/Engine/Source/Programs/UnrealGameSync/UnrealGameSync/Forms/ModifyBuildStepsWindow.cs
@@ -1,4 +1,4 @@
-﻿// Copyright 1998-2018 Epic Games, Inc. All Rights Reserved.
+// Copyright 1998-2018 Epic Games, Inc. All Rights Reserved.
 
 using System;
 using System.Collections.Generic;
@@ -32,6 +32,16 @@ namespace UnrealGameSync
 
 			InitializeComponent();
 
+			using(Graphics Graphics = Graphics.FromHwnd(IntPtr.Zero))
+			{
+				float DpiScaleX = Graphics.DpiX / 96.0f;
+
+				NormalSyncColumn.Width = (int)(104 * DpiScaleX);
+				ScheduledSyncColumn.Width = (int)(104 * DpiScaleX);
+				ShowAsToolColumn.Width = (int)(104 * DpiScaleX);
+				DescriptionColumn.Width = BuildStepList.ClientSize.Width - NormalSyncColumn.Width - ScheduledSyncColumn.Width - ShowAsToolColumn.Width - 10;
+			}
+
 			BuildStepList.Font = SystemFonts.IconTitleFont;
 		}
 
diff --git a/Engine/Source/Programs/UnrealGameSync/UnrealGameSync/Properties/app.manifest b/Engine/Source/Programs/UnrealGameSync/UnrealGameSync/Properties/app.manifest
--- a/Engine/Source/Programs/UnrealGameSync/UnrealGameSync/Properties/app.manifest
+++ b/Engine/Source/Programs/UnrealGameSync/UnrealGameSync/Properties/app.manifest
@@ -1,5 +1,5 @@
-﻿<?xml version="1.0" encoding="utf-8"?>
-<asmv1:assembly manifestVersion="1.0" xmlns="urn:schemas-microsoft-com:asm.v1" xmlns:asmv1="urn:schemas-microsoft-com:asm.v1" xmlns:asmv2="urn:schemas-microsoft-com:asm.v2" xmlns:xsi="http://www.w3.org/2001/XMLSchema-instance">
+<?xml version="1.0" encoding="utf-8"?>
+<asmv1:assembly manifestVersion="1.0" xmlns="urn:schemas-microsoft-com:asm.v1" xmlns:asmv1="urn:schemas-microsoft-com:asm.v1" xmlns:asmv2="urn:schemas-microsoft-com:asm.v2" xmlns:asmv3="urn:schemas-microsoft-com:asm.v3" xmlns:xsi="http://www.w3.org/2001/XMLSchema-instance">
   <assemblyIdentity version="1.0.0.0" name="MyApplication.app" />
   <trustInfo xmlns="urn:schemas-microsoft-com:asm.v2">
     <security>
@@ -38,7 +38,13 @@
       <!--<supportedOS Id="{1f676c76-80e1-4239-95bb-83d0f6d0da78}"/>-->
     </application>
   </compatibility>
-  <!-- Enable themes for Windows common controls and dialogs (Windows XP and later) -->
+  <asmv3:application>
+    <asmv3:windowsSettings xmlns="http://schemas.microsoft.com/SMI/2005/WindowsSettings">
+      <dpiAware>true</dpiAware>
+    </asmv3:windowsSettings>
+  </asmv3:application>
+
+	<!-- Enable themes for Windows common controls and dialogs (Windows XP and later) -->
   <!-- <dependency>
     <dependentAssembly>
       <assemblyIdentity
