UGS: Prevent "path too long" errors when enumerating files to delete. If the path is invalid, we probably never synced it in the first place, and P4 failed silently.

diff --git a/Engine/Source/Programs/UnrealGameSync/UnrealGameSync/Workspace.cs b/Engine/Source/Programs/UnrealGameSync/UnrealGameSync/Workspace.cs
--- a/Engine/Source/Programs/UnrealGameSync/UnrealGameSync/Workspace.cs
+++ b/Engine/Source/Programs/UnrealGameSync/UnrealGameSync/Workspace.cs
@@ -385,9 +385,7 @@ namespace UnrealGameSync
 							}
 							catch(PathTooLongException)
 							{
-								Log.WriteLine("The local path for {0} exceeds the maximum allowed by Windows. Re-sync your workspace to a directory with a shorter name, or delete the file from the server.", HaveFile.Path);
-								StatusMessage = "File exceeds maximum path length allowed by Windows.";
-								return WorkspaceUpdateResult.FailedToSync;
+								// We don't actually care about this when looking for files to remove. Perforce may think that it's synced the path, and silently failed. Just ignore it.
 							}
 						}
 
