Adding regex to rip the first two p4 path fragments out of a project name to look for wildcard matches in the db
Also actually now allows ability to get ids for specific projects, versus before where all ids were returned for projects that would have matched via wildcard
Pushed to Live 8/28/18

diff --git a/Engine/Source/Programs/UnrealGameSync/UnrealGameSyncMetadataServer/Connectors/SqlConnector.cs b/Engine/Source/Programs/UnrealGameSync/UnrealGameSyncMetadataServer/Connectors/SqlConnector.cs
--- a/Engine/Source/Programs/UnrealGameSync/UnrealGameSyncMetadataServer/Connectors/SqlConnector.cs
+++ b/Engine/Source/Programs/UnrealGameSync/UnrealGameSyncMetadataServer/Connectors/SqlConnector.cs
@@ -15,12 +15,12 @@ namespace UnrealGameSyncMetadataServer.Connectors
 		private static string ConnectionString = System.Configuration.ConfigurationManager.AppSettings["ConnectionString"];
 		public static LatestData GetLastIds(string Project = null)
 		{
-			// Get ids going back 25 builds for the project being asked for
-			// Do this by grouping by ChangeNumber to get unique entries, then take the 25th id
+			// Get ids going back 432 builds for the project being asked for
+			// Do this by grouping by ChangeNumber to get unique entries, then take the 432nd id
 			long LastEventId = 0;
 			long LastCommentId = 0;
 			long LastBuildId = 0;
-			string ProjectLikeString = "%" + (Project == null ? String.Empty : Project) + "%";
+			string ProjectLikeString = "%" + (Project == null ? String.Empty : GetProjectStream(Project)) + "%";
 			using (SQLiteConnection Connection = new SQLiteConnection(ConnectionString))
 			{
 				Connection.Open();
@@ -66,7 +66,7 @@ namespace UnrealGameSyncMetadataServer.Connectors
 		public static List<EventData> GetUserVotes(string Project, long LastEventId)
 		{
 			List<EventData> ReturnedEvents = new List<EventData>();
-			string ProjectLikeString = "%" + (Project == null ? String.Empty : Project) + "%";
+			string ProjectLikeString = "%" + (Project == null ? String.Empty : GetProjectStream(Project)) + "%";
 			using (SQLiteConnection Connection = new SQLiteConnection(ConnectionString))
 			{
 				Connection.Open();
@@ -99,7 +99,7 @@ namespace UnrealGameSyncMetadataServer.Connectors
 		public static List<CommentData> GetComments(string Project, long LastCommentId)
 		{
 			List<CommentData> ReturnedComments = new List<CommentData>();
-			string ProjectLikeString = "%" + (Project == null ? String.Empty : Project) + "%";
+			string ProjectLikeString = "%" + (Project == null ? String.Empty : GetProjectStream(Project)) + "%";
 			using (SQLiteConnection Connection = new SQLiteConnection(ConnectionString))
 			{
 				Connection.Open();
@@ -131,7 +131,7 @@ namespace UnrealGameSyncMetadataServer.Connectors
 		public static List<BuildData> GetBuilds(string Project, long LastBuildId)
 		{
 			List<BuildData> ReturnedBuilds = new List<BuildData>();
-			string ProjectLikeString = "%" + (Project == null ? String.Empty : Project) + "%";
+			string ProjectLikeString = "%" + (Project == null ? String.Empty : GetProjectStream(Project)) + "%";
 			using (SQLiteConnection Connection = new SQLiteConnection(ConnectionString))
 			{
 				Connection.Open();
@@ -290,6 +290,17 @@ namespace UnrealGameSyncMetadataServer.Connectors
 				}
 			}
 		}
+		private static string GetProjectStream(string Project)
+		{
+			// Get first two fragments of the p4 path.  If it doesn't work, just return back the project.
+			Regex StreamPattern = new Regex("(\\/\\/[a-zA-Z0-9\\.\\-_]{1,}\\/[a-zA-Z0-9\\.\\-_]{1,})");
+			Match StreamMatch = StreamPattern.Match(Project);
+			if(StreamMatch.Success)
+			{
+				return StreamMatch.Groups[1].Value;
+			}
+			return Project;
+		}
 		private static bool MatchesWildcard(string Wildcard, string Project)
 		{
 			return Wildcard.EndsWith("...") && Project.StartsWith(Wildcard.Substring(0, Wildcard.Length - 4), StringComparison.InvariantCultureIgnoreCase);
