UGS: When browsing for a project file, always default to the directory containing the current selection.

diff --git a/Engine/Source/Programs/UnrealGameSync/UnrealGameSync/Forms/OpenProjectWindow.cs b/Engine/Source/Programs/UnrealGameSync/UnrealGameSync/Forms/OpenProjectWindow.cs
--- a/Engine/Source/Programs/UnrealGameSync/UnrealGameSync/Forms/OpenProjectWindow.cs
+++ b/Engine/Source/Programs/UnrealGameSync/UnrealGameSync/Forms/OpenProjectWindow.cs
@@ -324,6 +324,15 @@ namespace UnrealGameSync
 			OpenFileDialog Dialog = new OpenFileDialog();
 			Dialog.Filter = "Project files (*.uproject)|*.uproject|Project directory lists (*.uprojectdirs)|*.uprojectdirs|All supported files (*.uproject;*.uprojectdirs)|*.uproject;*.uprojectdirs|All files (*.*)|*.*" ;
 			Dialog.FilterIndex = Settings.FilterIndex;
+			
+			try
+			{
+				Dialog.InitialDirectory = Path.GetDirectoryName(LocalFileTextBox.Text);
+			}
+			catch
+			{
+			}
+
 			if(Dialog.ShowDialog(this) == System.Windows.Forms.DialogResult.OK)
 			{
 				string FullName = Path.GetFullPath(Dialog.FileName);
