UGS: Halve P4 update rate (now once every 2m)

diff --git a/Engine/Source/Programs/UnrealGameSync/UnrealGameSync/PerforceMonitor.cs b/Engine/Source/Programs/UnrealGameSync/UnrealGameSync/PerforceMonitor.cs
--- a/Engine/Source/Programs/UnrealGameSync/UnrealGameSync/PerforceMonitor.cs
+++ b/Engine/Source/Programs/UnrealGameSync/UnrealGameSync/PerforceMonitor.cs
@@ -187,7 +187,7 @@ namespace UnrealGameSync
 				}
 
 				// Wait for another request, or scan for new builds after a timeout
-				RefreshEvent.WaitOne((IsActive? 1 : 10) * 60 * 1000);
+				RefreshEvent.WaitOne((IsActive? 2 : 10) * 60 * 1000);
 			}
 		}
 
diff --git a/Engine/Source/Programs/UnrealGameSync/UnrealGameSync/Properties/AssemblyInfo.cs b/Engine/Source/Programs/UnrealGameSync/UnrealGameSync/Properties/AssemblyInfo.cs
--- a/Engine/Source/Programs/UnrealGameSync/UnrealGameSync/Properties/AssemblyInfo.cs
+++ b/Engine/Source/Programs/UnrealGameSync/UnrealGameSync/Properties/AssemblyInfo.cs
@@ -34,5 +34,5 @@ using System.Runtime.InteropServices;
 // You can specify all the values or you can default the Build and Revision Numbers 
 // by using the '*' as shown below:
 // [assembly: AssemblyVersion("1.0.*")]
-[assembly: AssemblyVersion("1.154.*")]
-[assembly: AssemblyFileVersion("1.154.0.0")]
+[assembly: AssemblyVersion("1.155.*")]
+[assembly: AssemblyFileVersion("1.155.0.0")]
