UGS: Do not try to read the server address from the p4 info command. If a proxy is being used, or the DNS endpoint changes, it will be invalidated.

diff --git a/Engine/Source/Programs/UnrealGameSync/UnrealGameSync/Perforce.cs b/Engine/Source/Programs/UnrealGameSync/UnrealGameSync/Perforce.cs
--- a/Engine/Source/Programs/UnrealGameSync/UnrealGameSync/Perforce.cs
+++ b/Engine/Source/Programs/UnrealGameSync/UnrealGameSync/Perforce.cs
@@ -155,7 +155,6 @@ namespace UnrealGameSync
 		public string UserName;
 		public string HostName;
 		public string ClientAddress;
-		public string ServerAddress;
 		public TimeSpan ServerTimeZone;
 
 		public PerforceInfoRecord(Dictionary<string, string> Tags)
@@ -163,7 +162,6 @@ namespace UnrealGameSync
 			Tags.TryGetValue("userName", out UserName);
 			Tags.TryGetValue("clientHost", out HostName);
 			Tags.TryGetValue("clientAddress", out ClientAddress);
-			Tags.TryGetValue("serverAddress", out ServerAddress);
 
 			string ServerDateTime;
 			if(Tags.TryGetValue("serverDate", out ServerDateTime))
@@ -475,6 +473,8 @@ namespace UnrealGameSync
 		delegate bool HandleRecordDelegate(Dictionary<string, string> Tags);
 		delegate bool HandleOutputDelegate(PerforceOutputLine Line);
 
+		public const string DefaultServerAndPort = "perforce:1666";
+
 		public readonly string ServerAndPort;
 		public readonly string UserName;
 		public readonly string ClientName;
diff --git a/Engine/Source/Programs/UnrealGameSync/UnrealGameSync/PerforceModalTask.cs b/Engine/Source/Programs/UnrealGameSync/UnrealGameSync/PerforceModalTask.cs
--- a/Engine/Source/Programs/UnrealGameSync/UnrealGameSync/PerforceModalTask.cs
+++ b/Engine/Source/Programs/UnrealGameSync/UnrealGameSync/PerforceModalTask.cs
@@ -67,28 +67,24 @@ namespace UnrealGameSync
 				}
 			}
 
+			// If the server setting is still null, user the default
+			if(ServerAndPort == null)
+			{
+				ServerAndPort = PerforceConnection.DefaultServerAndPort;
+			}
+
 			// Update the server and username from the reported server info if it's not set
-			if(ServerAndPort == null || UserName == null)
+			if(UserName == null)
 			{
 				PerforceConnection Perforce = new PerforceConnection(UserName, null, ServerAndPort);
 
 				PerforceInfoRecord PerforceInfo;
-				if(!Perforce.Info(out PerforceInfo, Log))
-				{
-					return false;
-				}
-				if(ServerAndPort == null && !String.IsNullOrEmpty(PerforceInfo.ServerAddress))
-				{
-					ServerAndPort = PerforceInfo.ServerAddress;
-				}
-				if(UserName == null && !String.IsNullOrEmpty(PerforceInfo.UserName))
-				{
-					UserName = PerforceInfo.UserName;
-				}
-				if(ServerAndPort == null || UserName == null)
+				if(!Perforce.Info(out PerforceInfo, Log) || String.IsNullOrEmpty(PerforceInfo.UserName))
 				{
 					return false;
 				}
+
+				UserName = PerforceInfo.UserName;
 			}
 
 			// Otherwise succeed
diff --git a/Engine/Source/Programs/UnrealGameSync/UnrealGameSync/Properties/AssemblyInfo.cs b/Engine/Source/Programs/UnrealGameSync/UnrealGameSync/Properties/AssemblyInfo.cs
--- a/Engine/Source/Programs/UnrealGameSync/UnrealGameSync/Properties/AssemblyInfo.cs
+++ b/Engine/Source/Programs/UnrealGameSync/UnrealGameSync/Properties/AssemblyInfo.cs
@@ -34,5 +34,5 @@ using System.Runtime.InteropServices;
 // You can specify all the values or you can default the Build and Revision Numbers 
 // by using the '*' as shown below:
 // [assembly: AssemblyVersion("1.0.*")]
-[assembly: AssemblyVersion("1.155.*")]
-[assembly: AssemblyFileVersion("1.155.0.0")]
+[assembly: AssemblyVersion("1.156.*")]
+[assembly: AssemblyFileVersion("1.156.0.0")]
diff --git a/Engine/Source/Programs/UnrealGameSync/UnrealGameSyncLauncher/Program.cs b/Engine/Source/Programs/UnrealGameSync/UnrealGameSyncLauncher/Program.cs
--- a/Engine/Source/Programs/UnrealGameSync/UnrealGameSyncLauncher/Program.cs
+++ b/Engine/Source/Programs/UnrealGameSync/UnrealGameSyncLauncher/Program.cs
@@ -110,7 +110,7 @@ namespace UnrealGameSyncLauncher
 			{
 				if(Key != null)
 				{
-					ServerAndPort = Key.GetValue("Server", ServerAndPort) as string;
+					ServerAndPort = Key.GetValue("ServerAndPort", ServerAndPort) as string;
 					UserName = Key.GetValue("UserName", UserName) as string;
 					DepotPath = Key.GetValue("DepotPath", DepotPath) as string;
 				}
@@ -123,13 +123,16 @@ namespace UnrealGameSyncLauncher
 			{
 				using (RegistryKey Key = Registry.CurrentUser.CreateSubKey("SOFTWARE\\Epic Games\\UnrealGameSync"))
 				{
+					// Delete this legacy setting
+					try { Key.DeleteValue("Server"); } catch(Exception) { }
+
 					if(String.IsNullOrEmpty(ServerAndPort))
 					{
-						try { Key.DeleteValue("Server"); } catch(Exception) { }
+						try { Key.DeleteValue("ServerAndPort"); } catch(Exception) { }
 					}
 					else
 					{
-						Key.SetValue("Server", ServerAndPort);
+						Key.SetValue("ServerAndPort", ServerAndPort);
 					}
 
 					if(String.IsNullOrEmpty(UserName))
diff --git a/Engine/Source/Programs/UnrealGameSync/UnrealGameSyncLauncher/Properties/AssemblyInfo.cs b/Engine/Source/Programs/UnrealGameSync/UnrealGameSyncLauncher/Properties/AssemblyInfo.cs
--- a/Engine/Source/Programs/UnrealGameSync/UnrealGameSyncLauncher/Properties/AssemblyInfo.cs
+++ b/Engine/Source/Programs/UnrealGameSync/UnrealGameSyncLauncher/Properties/AssemblyInfo.cs
@@ -34,5 +34,5 @@ using System.Runtime.InteropServices;
 // You can specify all the values or you can default the Build and Revision Numbers 
 // by using the '*' as shown below:
 // [assembly: AssemblyVersion("1.0.*")]
-[assembly: AssemblyVersion("1.16.0.0")]
-[assembly: AssemblyFileVersion("1.16.0.0")]
+[assembly: AssemblyVersion("1.17.0.0")]
+[assembly: AssemblyFileVersion("1.17.0.0")]
