UGS: Fix any corrupt server paths that have been saved which point to p4-nodeb.epicgames.net

#rnx

diff --git a/Engine/Source/Programs/UnrealGameSync/UnrealGameSync/UserSettings.cs b/Engine/Source/Programs/UnrealGameSync/UnrealGameSync/UserSettings.cs
--- a/Engine/Source/Programs/UnrealGameSync/UnrealGameSync/UserSettings.cs
+++ b/Engine/Source/Programs/UnrealGameSync/UnrealGameSync/UserSettings.cs
@@ -69,6 +69,12 @@ namespace UnrealGameSync
 					ServerAndPort = null;
 				}
 
+				// Fixup for code that was saving server host name rather than DNS entry
+				if(ServerAndPort != null && ServerAndPort.Equals("p4-nodeb.epicgames.net:1666", StringComparison.OrdinalIgnoreCase))
+				{
+					ServerAndPort = "perforce:1666";
+				}
+
 				string UserName = Object.GetValue("UserName", null);
 				if(String.IsNullOrWhiteSpace(UserName))
 				{
