UGS: Fix exception when executing post-sync steps, due to variable being added twice.

#rnx

diff --git a/Engine/Source/Programs/UnrealGameSync/UnrealGameSync/Properties/AssemblyInfo.cs b/Engine/Source/Programs/UnrealGameSync/UnrealGameSync/Properties/AssemblyInfo.cs
--- a/Engine/Source/Programs/UnrealGameSync/UnrealGameSync/Properties/AssemblyInfo.cs
+++ b/Engine/Source/Programs/UnrealGameSync/UnrealGameSync/Properties/AssemblyInfo.cs
@@ -34,5 +34,5 @@ using System.Runtime.InteropServices;
 // You can specify all the values or you can default the Build and Revision Numbers 
 // by using the '*' as shown below:
 // [assembly: AssemblyVersion("1.0.*")]
-[assembly: AssemblyVersion("1.157.*")]
-[assembly: AssemblyFileVersion("1.157.0.0")]
+[assembly: AssemblyVersion("1.158.*")]
+[assembly: AssemblyFileVersion("1.158.0.0")]
diff --git a/Engine/Source/Programs/UnrealGameSync/UnrealGameSync/Workspace.cs b/Engine/Source/Programs/UnrealGameSync/UnrealGameSync/Workspace.cs
--- a/Engine/Source/Programs/UnrealGameSync/UnrealGameSync/Workspace.cs
+++ b/Engine/Source/Programs/UnrealGameSync/UnrealGameSync/Workspace.cs
@@ -697,8 +697,8 @@ namespace UnrealGameSync
 							Log.WriteLine("Executing post-sync steps...");
 
 							Dictionary<string, string> PostSyncVariables = new Dictionary<string, string>(Context.Variables);
-							PostSyncVariables.Add("Change", PendingChangeNumber.ToString());
-							PostSyncVariables.Add("CodeChange", VersionChangeNumber.ToString());
+							PostSyncVariables["Change"] = PendingChangeNumber.ToString();
+							PostSyncVariables["CodeChange"] = VersionChangeNumber.ToString();
 
 							foreach (string PostSyncStep in PostSyncSteps.Select(x => x.Trim()))
 							{
