UGS: Rename BuildData to BadgeData.

#rnx

diff --git a/Engine/Source/Programs/UnrealGameSync/UnrealGameSync/Controls/WorkspaceControl.cs b/Engine/Source/Programs/UnrealGameSync/UnrealGameSync/Controls/WorkspaceControl.cs
--- a/Engine/Source/Programs/UnrealGameSync/UnrealGameSync/Controls/WorkspaceControl.cs
+++ b/Engine/Source/Programs/UnrealGameSync/UnrealGameSync/Controls/WorkspaceControl.cs
@@ -206,7 +206,7 @@ namespace UnrealGameSync
 		Font BadgeFont;
 		List<KeyValuePair<string, string>> BadgeNameAndGroupPairs = new List<KeyValuePair<string, string>>();
 		Dictionary<string, Size> BadgeLabelToSize = new Dictionary<string, Size>();
-		List<KeyValuePair<string, BuildData>> ServiceBadges = new List<KeyValuePair<string, BuildData>>();
+		List<KeyValuePair<string, BadgeData>> ServiceBadges = new List<KeyValuePair<string, BadgeData>>();
 
 		string OriginalExecutableFileName;
 
@@ -472,9 +472,9 @@ namespace UnrealGameSync
 			ServiceBadges.Clear();
 			foreach(string ServiceBadgeName in ServiceBadgeNames)
 			{
-				BuildData LatestBuild;
-				EventMonitor.TryGetLatestBuild(ServiceBadgeName, out LatestBuild);
-				ServiceBadges.Add(new KeyValuePair<string, BuildData>(ServiceBadgeName, LatestBuild));
+				BadgeData LatestBuild;
+				EventMonitor.TryGetLatestBadge(ServiceBadgeName, out LatestBuild);
+				ServiceBadges.Add(new KeyValuePair<string, BadgeData>(ServiceBadgeName, LatestBuild));
 			}
 		}
 
@@ -1136,9 +1136,9 @@ namespace UnrealGameSync
 				EventSummary Summary = EventMonitor.GetSummaryForChange(((PerforceChangeSummary)Item.Tag).Number);
 				if(Summary != null)
 				{
-					foreach(BuildData Build in Summary.Builds)
+					foreach(BadgeData Badge in Summary.Badges)
 					{
-						string BadgeSlot = Build.BadgeName;
+						string BadgeSlot = Badge.BadgeName;
 						if(!BadgeNameToGroup.ContainsKey(BadgeSlot))
 						{
 							BadgeNameToGroup.Add(BadgeSlot, "XXXX");
@@ -1148,7 +1148,7 @@ namespace UnrealGameSync
 			}
 
 			// Remove anything that's a service badge
-			foreach(KeyValuePair<string, BuildData> ServiceBadge in ServiceBadges)
+			foreach(KeyValuePair<string, BadgeData> ServiceBadge in ServiceBadges)
 			{
 				BadgeNameToGroup.Remove(ServiceBadge.Key);
 			}
@@ -1282,22 +1282,22 @@ namespace UnrealGameSync
 			ContentBadges.UnionWith(PerforceMonitor.LatestProjectConfigFile.GetValues("Notifications.ContentBadges", new string[0]));
 
 			// Find the most recent build of each type, and the last time it succeeded
-			Dictionary<string, BuildData> TypeToLastBuild = new Dictionary<string,BuildData>();
-			Dictionary<string, BuildData> TypeToLastSucceededBuild = new Dictionary<string,BuildData>();
+			Dictionary<string, BadgeData> TypeToLastBuild = new Dictionary<string,BadgeData>();
+			Dictionary<string, BadgeData> TypeToLastSucceededBuild = new Dictionary<string,BadgeData>();
 			for(int Idx = SortedChangeNumbers.Count - 1; Idx >= 0; Idx--)
 			{
 				EventSummary Summary = EventMonitor.GetSummaryForChange(SortedChangeNumbers[Idx]);
 				if(Summary != null)
 				{
-					foreach(BuildData Build in Summary.Builds)
+					foreach(BadgeData Badge in Summary.Badges)
 					{
-						if(!TypeToLastBuild.ContainsKey(Build.BuildType) && (Build.Result == BuildDataResult.Success || Build.Result == BuildDataResult.Warning || Build.Result == BuildDataResult.Failure))
+						if(!TypeToLastBuild.ContainsKey(Badge.BuildType) && (Badge.Result == BadgeResult.Success || Badge.Result == BadgeResult.Warning || Badge.Result == BadgeResult.Failure))
 						{
-							TypeToLastBuild.Add(Build.BuildType, Build);
+							TypeToLastBuild.Add(Badge.BuildType, Badge);
 						}
-						if(!TypeToLastSucceededBuild.ContainsKey(Build.BuildType) && Build.Result == BuildDataResult.Success)
+						if(!TypeToLastSucceededBuild.ContainsKey(Badge.BuildType) && Badge.Result == BadgeResult.Success)
 						{
-							TypeToLastSucceededBuild.Add(Build.BuildType, Build);
+							TypeToLastSucceededBuild.Add(Badge.BuildType, Badge);
 						}
 					}
 				}
@@ -1305,10 +1305,10 @@ namespace UnrealGameSync
 
 			// Find all the build types that the user needs to be notified about.
 			int RequireNotificationForChange = -1;
-			List<BuildData> NotifyBuilds = new List<BuildData>();
-			foreach(BuildData LastBuild in TypeToLastBuild.Values.OrderBy(x => x.BuildType))
+			List<BadgeData> NotifyBuilds = new List<BadgeData>();
+			foreach(BadgeData LastBuild in TypeToLastBuild.Values.OrderBy(x => x.BuildType))
 			{
-				if(LastBuild.Result == BuildDataResult.Failure || LastBuild.Result == BuildDataResult.Warning)
+				if(LastBuild.Result == BadgeResult.Failure || LastBuild.Result == BadgeResult.Warning)
 				{
 					// Get the last submitted changelist by this user of the correct type
 					int LastChangeByCurrentUserOfType;
@@ -1325,7 +1325,7 @@ namespace UnrealGameSync
 					if(LastChangeByCurrentUserOfType > 0 && LastBuild.ChangeNumber >= LastChangeByCurrentUserOfType)
 					{
 						// And check that there wasn't a successful build after we submitted (if there was, we're in the clear)
-						BuildData LastSuccessfulBuild;
+						BadgeData LastSuccessfulBuild;
 						if(!TypeToLastSucceededBuild.TryGetValue(LastBuild.BuildType, out LastSuccessfulBuild) || LastSuccessfulBuild.ChangeNumber < LastChangeByCurrentUserOfType)
 						{
 							// Add it to the list of notifications
@@ -1357,7 +1357,7 @@ namespace UnrealGameSync
 				}
 
 				// Show the balloon tooltip
-				if(NotifyBuilds.Any(x => x.Result == BuildDataResult.Failure))
+				if(NotifyBuilds.Any(x => x.Result == BadgeResult.Failure))
 				{
 					string Title = String.Format("{0} Errors", PlatformList.ToString());
 					string Message = String.Format("CIS failed after your last submitted changelist ({0}).", RequireNotificationForChange);
@@ -1375,7 +1375,7 @@ namespace UnrealGameSync
 				NotificationWindow.OnMoreInformation = () => { Owner.ShowAndActivate(); SelectChange(HighlightChange); };
 				
 				// Don't show messages for this change again
-				foreach(BuildData NotifyBuild in NotifyBuilds)
+				foreach(BadgeData NotifyBuild in NotifyBuilds)
 				{
 					NotifiedBuildTypeToChangeNumber[NotifyBuild.BuildType] = RequireNotificationForChange;
 				}
@@ -1961,28 +1961,28 @@ namespace UnrealGameSync
 		{
 			List<BadgeInfo> Badges = new List<BadgeInfo>();
 
-			if(Summary != null && Summary.Builds.Count > 0)
+			if(Summary != null && Summary.Badges.Count > 0)
 			{
 				// Create a lookup for build data for each badge name
-				Dictionary<string, BuildData> BadgeNameToBuildData = new Dictionary<string, BuildData>();
-				foreach(BuildData Build in Summary.Builds)
+				Dictionary<string, BadgeData> BadgeNameToBuildData = new Dictionary<string, BadgeData>();
+				foreach(BadgeData Badge in Summary.Badges)
 				{
-					BadgeNameToBuildData[Build.BadgeName] = Build;
+					BadgeNameToBuildData[Badge.BadgeName] = Badge;
 				}
 
 				// Add all the badges, sorted by group
 				foreach(KeyValuePair<string, string> BadgeNameAndGroup in BadgeNameAndGroupPairs)
 				{
-					BuildData Build;
-					BadgeNameToBuildData.TryGetValue(BadgeNameAndGroup.Key, out Build);
+					BadgeData BadgeData;
+					BadgeNameToBuildData.TryGetValue(BadgeNameAndGroup.Key, out BadgeData);
 
-					BadgeInfo Badge = CreateBadge(ChangeNumber, BadgeNameAndGroup.Key, BadgeNameAndGroup.Value, Build);
-					if(MaxBuildBadgeChars != -1 && Badge.Label.Length > MaxBuildBadgeChars)
+					BadgeInfo BadgeInfo = CreateBadge(ChangeNumber, BadgeNameAndGroup.Key, BadgeNameAndGroup.Value, BadgeData);
+					if(MaxBuildBadgeChars != -1 && BadgeInfo.Label.Length > MaxBuildBadgeChars)
 					{
-						Badge.ToolTip = Badge.Label;
-						Badge.Label = Badge.Label.Substring(0, MaxBuildBadgeChars);
+						BadgeInfo.ToolTip = BadgeInfo.Label;
+						BadgeInfo.Label = BadgeInfo.Label.Substring(0, MaxBuildBadgeChars);
 					}
-					Badges.Add(Badge);
+					Badges.Add(BadgeInfo);
 				}
 			}
 
@@ -1991,27 +1991,27 @@ namespace UnrealGameSync
 			return Badges;
 		}
 
-		private BadgeInfo CreateBadge(int ChangeNumber, string BadgeName, string BadgeGroup, BuildData Build)
+		private BadgeInfo CreateBadge(int ChangeNumber, string BadgeName, string BadgeGroup, BadgeData BadgeData)
 		{
 			string BadgeLabel = BadgeName;
 			Color BadgeColor = Color.FromArgb(0, Color.White);
 
-			if(Build != null)
+			if(BadgeData != null)
 			{
-				BadgeLabel = Build.BadgeLabel;
-				BadgeColor = GetBuildBadgeColor(Build.Result);
+				BadgeLabel = BadgeData.BadgeLabel;
+				BadgeColor = GetBuildBadgeColor(BadgeData.Result);
 			}
 
 			Color HoverBadgeColor = Color.FromArgb(BadgeColor.A, Math.Min(BadgeColor.R + 32, 255), Math.Min(BadgeColor.G + 32, 255), Math.Min(BadgeColor.B + 32, 255));
 
 			string UniqueId = String.Format("{0}:{1}", ChangeNumber, BadgeName);
-			return new BadgeInfo(BadgeLabel, BadgeGroup, UniqueId, BadgeColor, HoverBadgeColor, Build);
+			return new BadgeInfo(BadgeLabel, BadgeGroup, UniqueId, BadgeColor, HoverBadgeColor, BadgeData);
 		}
 
 		private Dictionary<string, List<BadgeInfo>> CreateCustomBadges(int ChangeNumber, EventSummary Summary)
 		{
 			Dictionary<string, List<BadgeInfo>> ColumnNameToBadges = new Dictionary<string, List<BadgeInfo>>();
-			if(Summary != null && Summary.Builds.Count > 0)
+			if(Summary != null && Summary.Badges.Count > 0)
 			{
 				foreach(ColumnHeader CustomColumn in CustomColumns)
 				{
@@ -2023,7 +2023,7 @@ namespace UnrealGameSync
 						string[] BadgeNames = Config.GetValue("Badges", "").Split(new char[]{ ',' }, StringSplitOptions.RemoveEmptyEntries);
 						foreach(string BadgeName in BadgeNames)
 						{
-							BadgeInfo Badge = CreateBadge(ChangeNumber, BadgeName, "XXXX", Summary.Builds.FirstOrDefault(x => x.BadgeName == BadgeName));
+							BadgeInfo Badge = CreateBadge(ChangeNumber, BadgeName, "XXXX", Summary.Badges.FirstOrDefault(x => x.BadgeName == BadgeName));
 							Badges.Add(Badge);
 						}
 
@@ -2036,21 +2036,21 @@ namespace UnrealGameSync
 			return ColumnNameToBadges;
 		}
 
-		private static Color GetBuildBadgeColor(BuildDataResult Result)
+		private static Color GetBuildBadgeColor(BadgeResult Result)
 		{
-			if(Result == BuildDataResult.Starting)
+			if(Result == BadgeResult.Starting)
 			{
 				return Color.FromArgb(128, 192, 255);
 			}
-			else if(Result == BuildDataResult.Warning)
+			else if(Result == BadgeResult.Warning)
 			{
 				return Color.FromArgb(255, 192, 0);
 			}
-			else if(Result == BuildDataResult.Failure)
+			else if(Result == BadgeResult.Failure)
 			{
 				return Color.FromArgb(192, 64, 0);
 			}
-			else if(Result == BuildDataResult.Skipped)
+			else if(Result == BadgeResult.Skipped)
 			{
 				return Color.FromArgb(192, 192, 192);
 			}
@@ -2586,12 +2586,12 @@ namespace UnrealGameSync
 				ProgramsLine.AddText("  |  ");
 				ProgramsLine.AddLink("Windows Explorer", FontStyle.Regular, () => { Process.Start("explorer.exe", String.Format("\"{0}\"", Path.GetDirectoryName(SelectedFileName))); });
 
-				foreach(KeyValuePair<string, BuildData> ServiceBadge in ServiceBadges)
+				foreach(KeyValuePair<string, BadgeData> ServiceBadge in ServiceBadges)
 				{
 					ProgramsLine.AddText("  |  ");
 					if(ServiceBadge.Value == null)
 					{
-						ProgramsLine.AddBadge(ServiceBadge.Key, GetBuildBadgeColor(BuildDataResult.Skipped), null);
+						ProgramsLine.AddBadge(ServiceBadge.Key, GetBuildBadgeColor(BadgeResult.Skipped), null);
 					}
 					else
 					{
@@ -3006,13 +3006,13 @@ namespace UnrealGameSync
 							Point BuildListLocation = GetBadgeListLocation(LayoutInfo.BuildBadges, HitTest.SubItem.Bounds, HorizontalAlign.Center, VerticalAlignment.Middle);
 							BuildListLocation.X = Math.Max(BuildListLocation.X, HitTest.SubItem.Bounds.Left);
 
-							BadgeInfo Badge = HitTestBadge(Args.Location, LayoutInfo.BuildBadges, BuildListLocation);
-							if(Badge != null)
+							BadgeInfo BadgeInfo = HitTestBadge(Args.Location, LayoutInfo.BuildBadges, BuildListLocation);
+							if(BadgeInfo != null)
 							{
-								BuildData Build = (BuildData)Badge.UserData;
-								if(Build != null)
+								BadgeData BadgeData = (BadgeData)BadgeInfo.UserData;
+								if(BadgeData != null)
 								{
-									Process.Start(Build.Url);
+									Process.Start(BadgeData.Url);
 								}
 							}
 						}
@@ -3029,13 +3029,13 @@ namespace UnrealGameSync
 							{
 								Point ListLocation = GetBadgeListLocation(Badges, HitTest.SubItem.Bounds, HorizontalAlign.Center, VerticalAlignment.Middle);
 
-								BadgeInfo Badge = HitTestBadge(Args.Location, Badges, ListLocation);
-								if(Badge != null)
+								BadgeInfo BadgeInfo = HitTestBadge(Args.Location, Badges, ListLocation);
+								if(BadgeInfo != null)
 								{
-									BuildData Build = (BuildData)Badge.UserData;
-									if(Build != null)
+									BadgeData BadgeData = (BadgeData)BadgeInfo.UserData;
+									if(BadgeData != null)
 									{
-										Process.Start(Build.Url);
+										Process.Start(BadgeData.Url);
 									}
 								}
 							}
diff --git a/Engine/Source/Programs/UnrealGameSync/UnrealGameSync/EventMonitor.cs b/Engine/Source/Programs/UnrealGameSync/UnrealGameSync/EventMonitor.cs
--- a/Engine/Source/Programs/UnrealGameSync/UnrealGameSync/EventMonitor.cs
+++ b/Engine/Source/Programs/UnrealGameSync/UnrealGameSync/EventMonitor.cs
@@ -34,12 +34,14 @@ namespace UnrealGameSync
 		Investigating,
 		Resolved,
 	}
+
 	class LatestData
 	{
 		public long LastEventId { get; set; }
 		public long LastCommentId { get; set; }
 		public long LastBuildId { get; set; }
 	}
+
 	class EventData
 	{
 		public long Id { get; set; }
@@ -58,7 +60,7 @@ namespace UnrealGameSync
 		public string Project { get; set; }
 	}
 
-	enum BuildDataResult
+	enum BadgeResult
 	{
 		Starting,
 		Failure,
@@ -67,23 +69,23 @@ namespace UnrealGameSync
 		Skipped,
 	}
 
-	class BuildData
+	class BadgeData
 	{
 		public long Id { get; set; }
 		public int ChangeNumber { get; set; }
 		public string BuildType { get; set; }
-		public BuildDataResult Result { get; set; }
+		public BadgeResult Result { get; set; }
 		public string Url { get; set; }
 		public string Project { get; set; }
 
 		public bool IsSuccess
 		{
-			get { return Result == BuildDataResult.Success || Result == BuildDataResult.Warning; }
+			get { return Result == BadgeResult.Success || Result == BadgeResult.Warning; }
 		}
 
 		public bool IsFailure
 		{
-			get { return Result == BuildDataResult.Failure; }
+			get { return Result == BadgeResult.Failure; }
 		}
 
 		public string BadgeName
@@ -135,7 +137,7 @@ namespace UnrealGameSync
 		public List<EventData> Reviews = new List<EventData>();
 		public List<string> CurrentUsers = new List<string>();
 		public EventData LastStarReview;
-		public List<BuildData> Builds = new List<BuildData>();
+		public List<BadgeData> Badges = new List<BadgeData>();
 		public List<CommentData> Comments = new List<CommentData>();
 	}
 
@@ -150,10 +152,10 @@ namespace UnrealGameSync
 		ConcurrentQueue<EventData> IncomingEvents = new ConcurrentQueue<EventData>();
 		ConcurrentQueue<CommentData> OutgoingComments = new ConcurrentQueue<CommentData>();
 		ConcurrentQueue<CommentData> IncomingComments = new ConcurrentQueue<CommentData>();
-		ConcurrentQueue<BuildData> IncomingBuilds = new ConcurrentQueue<BuildData>();
+		ConcurrentQueue<BadgeData> IncomingBadges = new ConcurrentQueue<BadgeData>();
 		Dictionary<int, EventSummary> ChangeNumberToSummary = new Dictionary<int, EventSummary>();
 		Dictionary<string, EventData> UserNameToLastSyncEvent = new Dictionary<string, EventData>(StringComparer.InvariantCultureIgnoreCase);
-		Dictionary<string, BuildData> BadgeNameToLatestBuild = new Dictionary<string, BuildData>();
+		Dictionary<string, BadgeData> BadgeNameToLatestData = new Dictionary<string, BadgeData>();
 		BoundedLogWriter LogWriter;
 		bool bDisposing;
 		LatestData LatestIds;
@@ -258,10 +260,10 @@ namespace UnrealGameSync
 				ApplyEventUpdate(Event);
 			}
 
-			BuildData Build;
-			while(IncomingBuilds.TryDequeue(out Build))
+			BadgeData Badge;
+			while(IncomingBadges.TryDequeue(out Badge))
 			{
-				ApplyBuildUpdate(Build);
+				ApplyBadgeUpdate(Badge);
 			}
 
 			CommentData Comment;
@@ -320,7 +322,7 @@ namespace UnrealGameSync
 
 				// Add the new review, and find the new verdict for this change
 				Summary.Reviews.Add(Event);
-				Summary.Verdict = GetVerdict(Summary.Reviews, Summary.Builds);
+				Summary.Verdict = GetVerdict(Summary.Reviews, Summary.Badges);
 			}
 			else
 			{
@@ -328,16 +330,16 @@ namespace UnrealGameSync
 			}
 		}
 
-		void ApplyBuildUpdate(BuildData Build)
+		void ApplyBadgeUpdate(BadgeData Badge)
 		{
-			EventSummary Summary = FindOrAddSummary(Build.ChangeNumber);
+			EventSummary Summary = FindOrAddSummary(Badge.ChangeNumber);
 
-			BuildData ExistingBuild = Summary.Builds.Find(x => x.ChangeNumber == Build.ChangeNumber && x.BuildType == Build.BuildType);
-			if(ExistingBuild != null)
+			BadgeData ExistingBadge = Summary.Badges.Find(x => x.ChangeNumber == Badge.ChangeNumber && x.BuildType == Badge.BuildType);
+			if(ExistingBadge != null)
 			{
-				if(ExistingBuild.Id <= Build.Id)
+				if(ExistingBadge.Id <= Badge.Id)
 				{
-					Summary.Builds.Remove(ExistingBuild);
+					Summary.Badges.Remove(ExistingBadge);
 				}
 				else
 				{
@@ -345,13 +347,13 @@ namespace UnrealGameSync
 				}
 			}
 
-			Summary.Builds.Add(Build);
-			Summary.Verdict = GetVerdict(Summary.Reviews, Summary.Builds);
+			Summary.Badges.Add(Badge);
+			Summary.Verdict = GetVerdict(Summary.Reviews, Summary.Badges);
 
-			BuildData LatestBuild;
-			if(!BadgeNameToLatestBuild.TryGetValue(Build.BadgeName, out LatestBuild) || Build.ChangeNumber > LatestBuild.ChangeNumber || (Build.ChangeNumber == LatestBuild.ChangeNumber && Build.Id > LatestBuild.Id))
+			BadgeData LatestBadge;
+			if(!BadgeNameToLatestData.TryGetValue(Badge.BadgeName, out LatestBadge) || Badge.ChangeNumber > LatestBadge.ChangeNumber || (Badge.ChangeNumber == LatestBadge.ChangeNumber && Badge.Id > LatestBadge.Id))
 			{
-				BadgeNameToLatestBuild[Build.BadgeName] = Build;
+				BadgeNameToLatestData[Badge.BadgeName] = Badge;
 			}
 		}
 
@@ -391,7 +393,7 @@ namespace UnrealGameSync
 			return true;
 		}
 
-		static ReviewVerdict GetVerdict(IEnumerable<EventData> Events, IEnumerable<BuildData> Builds)
+		static ReviewVerdict GetVerdict(IEnumerable<EventData> Events, IEnumerable<BadgeData> Badges)
 		{
 			int NumPositiveReviews = Events.Count(x => x.Type == EventType.Good);
 			int NumNegativeReviews = Events.Count(x => x.Type == EventType.Bad);
@@ -407,11 +409,11 @@ namespace UnrealGameSync
 				return GetVerdict(NumCompiles, NumFailedCompiles);
 			}
 
-			int NumBuilds = Builds.Count(x => x.BuildType == "Editor" && x.IsSuccess);
-			int NumFailedBuilds = Builds.Count(x => x.BuildType == "Editor" && x.IsFailure);
-			if(NumBuilds > 0 || NumFailedBuilds > 0)
+			int NumBadges = Badges.Count(x => x.BuildType == "Editor" && x.IsSuccess);
+			int NumFailedBadges = Badges.Count(x => x.BuildType == "Editor" && x.IsFailure);
+			if(NumBadges > 0 || NumFailedBadges > 0)
 			{
-				return GetVerdict(NumBuilds, NumFailedBuilds);
+				return GetVerdict(NumBadges, NumFailedBadges);
 			}
 
 			return ReviewVerdict.Unknown;
@@ -483,7 +485,7 @@ namespace UnrealGameSync
 					ReadEventsFromBackend();
 
 					// Send a notification that we're ready to update
-					if((IncomingEvents.Count > 0 || IncomingBuilds.Count > 0 || IncomingComments.Count > 0) && OnUpdatesReady != null)
+					if((IncomingEvents.Count > 0 || IncomingBadges.Count > 0 || IncomingComments.Count > 0) && OnUpdatesReady != null)
 					{
 						OnUpdatesReady();
 					}
@@ -568,10 +570,10 @@ namespace UnrealGameSync
 				//////////////
 				/// Bulids
 				//////////////
-				List<BuildData> Builds = RESTApi.GET<List<BuildData>>(ApiUrl, "build", string.Format("project={0}", Project), string.Format("lastbuildid={0}", LatestIds.LastBuildId));
-				foreach (BuildData Build in Builds)
+				List<BadgeData> Builds = RESTApi.GET<List<BadgeData>>(ApiUrl, "build", string.Format("project={0}", Project), string.Format("lastbuildid={0}", LatestIds.LastBuildId));
+				foreach (BadgeData Build in Builds)
 				{
-					IncomingBuilds.Enqueue(Build);
+					IncomingBadges.Enqueue(Build);
 					LatestIds.LastBuildId = Math.Max(LatestIds.LastBuildId, Build.Id);
 				}
 
@@ -671,9 +673,9 @@ namespace UnrealGameSync
 			return Summary;
 		}
 
-		public bool TryGetLatestBuild(string BuildType, out BuildData BuildData)
+		public bool TryGetLatestBadge(string BuildType, out BadgeData BadgeData)
 		{
-			return BadgeNameToLatestBuild.TryGetValue(BuildType, out BuildData);
+			return BadgeNameToLatestData.TryGetValue(BuildType, out BadgeData);
 		}
 
 		public static bool IsReview(EventType Type)
