UnrealGameSync: Fix tooltips for builds showing for placeholder (Alpha = 0) badges.

#rnx

diff --git a/Engine/Source/Programs/UnrealGameSync/UnrealGameSync/Controls/WorkspaceControl.cs b/Engine/Source/Programs/UnrealGameSync/UnrealGameSync/Controls/WorkspaceControl.cs
--- a/Engine/Source/Programs/UnrealGameSync/UnrealGameSync/Controls/WorkspaceControl.cs
+++ b/Engine/Source/Programs/UnrealGameSync/UnrealGameSync/Controls/WorkspaceControl.cs
@@ -3183,9 +3183,16 @@ namespace UnrealGameSync
 						BadgeInfo Badge = HitTestBadge(e.Location, LayoutInfo.BuildBadges, BuildListLocation);
 						NewHoverBadgeUniqueId = (Badge != null)? Badge.UniqueId : null;
 
-						if(Badge != null && Badge.ToolTip != null && HoverBadgeUniqueId != NewHoverBadgeUniqueId)
+						if(HoverBadgeUniqueId != NewHoverBadgeUniqueId)
 						{
-							BuildListToolTip.Show(Badge.ToolTip, BuildList, new Point(BuildListLocation.X + Badge.Offset, HitTest.Item.Bounds.Bottom + 2));
+							if(Badge != null && Badge.ToolTip != null && Badge.BackgroundColor.A != 0)
+							{
+								BuildListToolTip.Show(Badge.ToolTip, BuildList, new Point(BuildListLocation.X + Badge.Offset, HitTest.Item.Bounds.Bottom + 2));
+							}
+							else
+							{
+								BuildListToolTip.Hide(BuildList);
+							}
 						}
 					}
 				}
