Updating web.template.config and global.asax.cs to specify connection strings based on .NET3.5+ instaed of the old .NET2 way

diff --git a/Engine/Source/Programs/UnrealGameSync/UnrealGameSyncMetadataServer/Global.asax.cs b/Engine/Source/Programs/UnrealGameSync/UnrealGameSyncMetadataServer/Global.asax.cs
--- a/Engine/Source/Programs/UnrealGameSync/UnrealGameSyncMetadataServer/Global.asax.cs
+++ b/Engine/Source/Programs/UnrealGameSync/UnrealGameSyncMetadataServer/Global.asax.cs
@@ -14,7 +14,7 @@ namespace UnrealGameSyncMetadataServer
 		{
 			try
 			{
-				string ConnectionString = System.Configuration.ConfigurationManager.AppSettings["ConnectionString"];
+				string ConnectionString = System.Configuration.ConfigurationManager.ConnectionStrings["ConnectionString"].ConnectionString;
 				string FileName = Regex.Match(ConnectionString, "Data Source=(.+);.+").Groups[1].Value;
 				Directory.CreateDirectory(new FileInfo(FileName).Directory.FullName);
 				if (!File.Exists(FileName))
diff --git a/Engine/Source/Programs/UnrealGameSync/UnrealGameSyncMetadataServer/Web.template.config b/Engine/Source/Programs/UnrealGameSync/UnrealGameSyncMetadataServer/Web.template.config
--- a/Engine/Source/Programs/UnrealGameSync/UnrealGameSyncMetadataServer/Web.template.config
+++ b/Engine/Source/Programs/UnrealGameSync/UnrealGameSyncMetadataServer/Web.template.config
@@ -4,13 +4,7 @@
   http://go.microsoft.com/fwlink/?LinkId=301879
   -->
 <configuration>
-  <appSettings>
-    <!-- Leave this connection string blank! It should go in either the .debug or .release transform config
-         There is a custom build step on BeforeBuild that generates a transformed Web.config on the fly. The generated Web.config should not be checked into source control, 
-         and this template.config should be checked in in it's place. -->
-    <add key="ConnectionString"
-      value=""/>
-  </appSettings>
+	<appSettings/>
   <!--
     For a description of web.config changes see http://go.microsoft.com/fwlink/?LinkId=235367.
 
@@ -19,7 +13,13 @@
         <httpRuntime targetFramework="4.5" />
       </system.Web>
   -->
-  <system.web>
+	<!-- Leave this connection string blank! It should go in either the .debug or .release transform config
+         There is a custom build step on BeforeBuild that generates a transformed Web.config on the fly. The generated Web.config should not be checked into source control, 
+         and this template.config should be checked in in it's place. -->
+	<connectionStrings>
+		<add name="ConnectionString" connectionString=""/>
+	</connectionStrings>
+	<system.web>
     <compilation debug="true" targetFramework="4.5" />
     <httpRuntime targetFramework="4.5" />
   </system.web>
