changing sqlconnector over to .NET3.5 connection string version

diff --git a/Engine/Source/Programs/UnrealGameSync/UnrealGameSyncMetadataServer/Connectors/SqlConnector.cs b/Engine/Source/Programs/UnrealGameSync/UnrealGameSyncMetadataServer/Connectors/SqlConnector.cs
--- a/Engine/Source/Programs/UnrealGameSync/UnrealGameSyncMetadataServer/Connectors/SqlConnector.cs
+++ b/Engine/Source/Programs/UnrealGameSync/UnrealGameSyncMetadataServer/Connectors/SqlConnector.cs
@@ -12,7 +12,7 @@ namespace UnrealGameSyncMetadataServer.Connectors
 {
 	public static class SqlConnector
 	{
-		private static string ConnectionString = System.Configuration.ConfigurationManager.AppSettings["ConnectionString"];
+		private static string ConnectionString = System.Configuration.ConfigurationManager.ConnectionStrings["ConnectionString"].ConnectionString;
 		public static LatestData GetLastIds(string Project = null)
 		{
 			// Get ids going back 432 builds for the project being asked for
