UGS: Fix exception caused by custom build step that was temporarily submitted to a project's UnrealGameSync.ini file.

diff --git a/Engine/Source/Programs/UnrealGameSync/UnrealGameSync/BuildStep.cs b/Engine/Source/Programs/UnrealGameSync/UnrealGameSync/BuildStep.cs
--- a/Engine/Source/Programs/UnrealGameSync/UnrealGameSync/BuildStep.cs
+++ b/Engine/Source/Programs/UnrealGameSync/UnrealGameSync/BuildStep.cs
@@ -1,4 +1,4 @@
-﻿// Copyright 1998-2018 Epic Games, Inc. All Rights Reserved.
+// Copyright 1998-2018 Epic Games, Inc. All Rights Reserved.
 
 using System;
 using System.Collections.Generic;
@@ -104,6 +104,21 @@ namespace UnrealGameSync
 			}
 		}
 
+		public bool IsValid()
+		{
+			switch(Type)
+			{
+				case BuildStepType.Compile:
+					return Target != null && Platform != null && Configuration != null;
+				case BuildStepType.Cook:
+					return FileName != null;
+				case BuildStepType.Other:
+					return FileName != null;
+				default:
+					return false;
+			}
+		}
+
 		public static void MergeBuildStepObjects(Dictionary<Guid, ConfigObject> BuildStepObjects, IEnumerable<ConfigObject> ModifyObjects)
 		{
 			foreach(ConfigObject ModifyObject in ModifyObjects)
diff --git a/Engine/Source/Programs/UnrealGameSync/UnrealGameSync/Controls/WorkspaceControl.cs b/Engine/Source/Programs/UnrealGameSync/UnrealGameSync/Controls/WorkspaceControl.cs
--- a/Engine/Source/Programs/UnrealGameSync/UnrealGameSync/Controls/WorkspaceControl.cs
+++ b/Engine/Source/Programs/UnrealGameSync/UnrealGameSync/Controls/WorkspaceControl.cs
@@ -3964,13 +3964,16 @@ namespace UnrealGameSync
 				List<ConfigObject> ModifiedBuildSteps = new List<ConfigObject>();
 				foreach(BuildStep Step in UserSteps)
 				{
-					ConfigObject DefaultObject;
-					ProjectBuildStepObjects.TryGetValue(Step.UniqueId, out DefaultObject);
-
-					ConfigObject UserConfigObject = Step.ToConfigObject(DefaultObject);
-					if(UserConfigObject != null)
+					if(Step.IsValid())
 					{
-						ModifiedBuildSteps.Add(UserConfigObject);
+						ConfigObject DefaultObject;
+						ProjectBuildStepObjects.TryGetValue(Step.UniqueId, out DefaultObject);
+
+						ConfigObject UserConfigObject = Step.ToConfigObject(DefaultObject);
+						if(UserConfigObject != null && UserConfigObject.Pairs.Any(x => x.Key != "UniqueId"))
+						{
+							ModifiedBuildSteps.Add(UserConfigObject);
+						}
 					}
 				}
 
diff --git a/Engine/Source/Programs/UnrealGameSync/UnrealGameSync/Properties/AssemblyInfo.cs b/Engine/Source/Programs/UnrealGameSync/UnrealGameSync/Properties/AssemblyInfo.cs
--- a/Engine/Source/Programs/UnrealGameSync/UnrealGameSync/Properties/AssemblyInfo.cs
+++ b/Engine/Source/Programs/UnrealGameSync/UnrealGameSync/Properties/AssemblyInfo.cs
@@ -34,5 +34,5 @@ using System.Runtime.InteropServices;
 // You can specify all the values or you can default the Build and Revision Numbers 
 // by using the '*' as shown below:
 // [assembly: AssemblyVersion("1.0.*")]
-[assembly: AssemblyVersion("1.158.*")]
-[assembly: AssemblyFileVersion("1.158.0.0")]
+[assembly: AssemblyVersion("1.159.*")]
+[assembly: AssemblyFileVersion("1.159.0.0")]
diff --git a/Engine/Source/Programs/UnrealGameSync/UnrealGameSync/Workspace.cs b/Engine/Source/Programs/UnrealGameSync/UnrealGameSync/Workspace.cs
--- a/Engine/Source/Programs/UnrealGameSync/UnrealGameSync/Workspace.cs
+++ b/Engine/Source/Programs/UnrealGameSync/UnrealGameSync/Workspace.cs
@@ -905,79 +905,82 @@ namespace UnrealGameSync
 
 						Log.WriteLine(Step.StatusText);
 
-						switch(Step.Type)
+						if(Step.IsValid())
 						{
-							case BuildStepType.Compile:
-								using(TelemetryStopwatch StepStopwatch = new TelemetryStopwatch("Compile:" + Step.Target, TelemetryProjectPath))
-								{
-									string CommandLine = String.Format("{0} {1} {2} {3} -NoHotReloadFromIDE", Step.Target, Step.Platform, Step.Configuration, Utility.ExpandVariables(Step.Arguments, Context.Variables));
-									if(!Context.Options.HasFlag(WorkspaceUpdateOptions.UseIncrementalBuilds) || bForceClean)
+							switch(Step.Type)
+							{
+								case BuildStepType.Compile:
+									using(TelemetryStopwatch StepStopwatch = new TelemetryStopwatch("Compile:" + Step.Target, TelemetryProjectPath))
 									{
-										Log.WriteLine("ubt> Running {0} {1} -clean", UnrealBuildToolPath, CommandLine);
-										Utility.ExecuteProcess(UnrealBuildToolPath, null, CommandLine + " -clean", null, new ProgressTextWriter(Progress, new PrefixedTextWriter("ubt> ", Log)));
-									}
-
-									Log.WriteLine("ubt> Running {0} {1} -progress", UnrealBuildToolPath, CommandLine);
+										string CommandLine = String.Format("{0} {1} {2} {3} -NoHotReloadFromIDE", Step.Target, Step.Platform, Step.Configuration, Utility.ExpandVariables(Step.Arguments ?? "", Context.Variables));
+										if(!Context.Options.HasFlag(WorkspaceUpdateOptions.UseIncrementalBuilds) || bForceClean)
+										{
+											Log.WriteLine("ubt> Running {0} {1} -clean", UnrealBuildToolPath, CommandLine);
+											Utility.ExecuteProcess(UnrealBuildToolPath, null, CommandLine + " -clean", null, new ProgressTextWriter(Progress, new PrefixedTextWriter("ubt> ", Log)));
+										}
 
-									int ResultFromBuild = Utility.ExecuteProcess(UnrealBuildToolPath, null, CommandLine + " -progress", null, new ProgressTextWriter(Progress, new PrefixedTextWriter("ubt> ", Log)));
-									if(ResultFromBuild != 0)
-									{
-										StepStopwatch.Stop("Failed");
-										StatusMessage = String.Format("Failed to compile {0}.", Step.Target);
-										return (HasModifiedSourceFiles() || Context.UserBuildStepObjects.Count > 0)? WorkspaceUpdateResult.FailedToCompile : WorkspaceUpdateResult.FailedToCompileWithCleanWorkspace;
-									}
+										Log.WriteLine("ubt> Running {0} {1} -progress", UnrealBuildToolPath, CommandLine);
 
-									StepStopwatch.Stop("Success");
-								}
-								break;
-							case BuildStepType.Cook:
-								using(TelemetryStopwatch StepStopwatch = new TelemetryStopwatch("Cook/Launch: " + Path.GetFileNameWithoutExtension(Step.FileName), TelemetryProjectPath))
-								{
-									string LocalRunUAT = Path.Combine(LocalRootPath, "Engine", "Build", "BatchFiles", "RunUAT.bat");
-									string Arguments = String.Format("/C \"\"{0}\" -profile=\"{1}\"\"", LocalRunUAT, Path.Combine(LocalRootPath, Step.FileName));
-									Log.WriteLine("uat> Running {0} {1}", LocalRunUAT, Arguments);
+										int ResultFromBuild = Utility.ExecuteProcess(UnrealBuildToolPath, null, CommandLine + " -progress", null, new ProgressTextWriter(Progress, new PrefixedTextWriter("ubt> ", Log)));
+										if(ResultFromBuild != 0)
+										{
+											StepStopwatch.Stop("Failed");
+											StatusMessage = String.Format("Failed to compile {0}.", Step.Target);
+											return (HasModifiedSourceFiles() || Context.UserBuildStepObjects.Count > 0)? WorkspaceUpdateResult.FailedToCompile : WorkspaceUpdateResult.FailedToCompileWithCleanWorkspace;
+										}
 
-									int ResultFromUAT = Utility.ExecuteProcess(CmdExe, null, Arguments, null, new ProgressTextWriter(Progress, new PrefixedTextWriter("uat> ", Log)));
-									if(ResultFromUAT != 0)
-									{
-										StepStopwatch.Stop("Failed");
-										StatusMessage = String.Format("Cook failed. ({0})", ResultFromUAT);
-										return WorkspaceUpdateResult.FailedToCompile;
+										StepStopwatch.Stop("Success");
 									}
-
-									StepStopwatch.Stop("Success");
-								}
-								break;
-							case BuildStepType.Other:
-								using(TelemetryStopwatch StepStopwatch = new TelemetryStopwatch("Custom: " + Path.GetFileNameWithoutExtension(Step.FileName), TelemetryProjectPath))
-								{
-									string ToolFileName = Path.Combine(LocalRootPath, Utility.ExpandVariables(Step.FileName, Context.Variables));
-									string ToolWorkingDir = String.IsNullOrWhiteSpace(Step.WorkingDir) ? Path.GetDirectoryName(ToolFileName) : Utility.ExpandVariables(Step.WorkingDir, Context.Variables);
-									string ToolArguments = Utility.ExpandVariables(Step.Arguments, Context.Variables);
-									Log.WriteLine("tool> Running {0} {1}", ToolFileName, ToolArguments);
-
-									if(Step.bUseLogWindow)
+									break;
+								case BuildStepType.Cook:
+									using(TelemetryStopwatch StepStopwatch = new TelemetryStopwatch("Cook/Launch: " + Path.GetFileNameWithoutExtension(Step.FileName), TelemetryProjectPath))
 									{
-										int ResultFromTool = Utility.ExecuteProcess(ToolFileName, ToolWorkingDir, ToolArguments, null, new ProgressTextWriter(Progress, new PrefixedTextWriter("tool> ", Log)));
-										if(ResultFromTool != 0)
+										string LocalRunUAT = Path.Combine(LocalRootPath, "Engine", "Build", "BatchFiles", "RunUAT.bat");
+										string Arguments = String.Format("/C \"\"{0}\" -profile=\"{1}\"\"", LocalRunUAT, Path.Combine(LocalRootPath, Step.FileName));
+										Log.WriteLine("uat> Running {0} {1}", LocalRunUAT, Arguments);
+
+										int ResultFromUAT = Utility.ExecuteProcess(CmdExe, null, Arguments, null, new ProgressTextWriter(Progress, new PrefixedTextWriter("uat> ", Log)));
+										if(ResultFromUAT != 0)
 										{
 											StepStopwatch.Stop("Failed");
-											StatusMessage = String.Format("Tool terminated with exit code {0}.", ResultFromTool);
+											StatusMessage = String.Format("Cook failed. ({0})", ResultFromUAT);
 											return WorkspaceUpdateResult.FailedToCompile;
 										}
+
+										StepStopwatch.Stop("Success");
 									}
-									else
+									break;
+								case BuildStepType.Other:
+									using(TelemetryStopwatch StepStopwatch = new TelemetryStopwatch("Custom: " + Path.GetFileNameWithoutExtension(Step.FileName), TelemetryProjectPath))
 									{
-										ProcessStartInfo StartInfo = new ProcessStartInfo(ToolFileName, ToolArguments);
-										StartInfo.WorkingDirectory = ToolWorkingDir;
-										using(Process.Start(StartInfo))
+										string ToolFileName = Path.Combine(LocalRootPath, Utility.ExpandVariables(Step.FileName, Context.Variables));
+										string ToolWorkingDir = String.IsNullOrWhiteSpace(Step.WorkingDir) ? Path.GetDirectoryName(ToolFileName) : Utility.ExpandVariables(Step.WorkingDir, Context.Variables);
+										string ToolArguments = Utility.ExpandVariables(Step.Arguments ?? "", Context.Variables);
+										Log.WriteLine("tool> Running {0} {1}", ToolFileName, ToolArguments);
+
+										if(Step.bUseLogWindow)
 										{
+											int ResultFromTool = Utility.ExecuteProcess(ToolFileName, ToolWorkingDir, ToolArguments, null, new ProgressTextWriter(Progress, new PrefixedTextWriter("tool> ", Log)));
+											if(ResultFromTool != 0)
+											{
+												StepStopwatch.Stop("Failed");
+												StatusMessage = String.Format("Tool terminated with exit code {0}.", ResultFromTool);
+												return WorkspaceUpdateResult.FailedToCompile;
+											}
+										}
+										else
+										{
+											ProcessStartInfo StartInfo = new ProcessStartInfo(ToolFileName, ToolArguments);
+											StartInfo.WorkingDirectory = ToolWorkingDir;
+											using(Process.Start(StartInfo))
+											{
+											}
 										}
-									}
 
-									StepStopwatch.Stop("Success");
-								}
-								break;
+										StepStopwatch.Stop("Success");
+									}
+									break;
+							}
 						}
 
 						Log.WriteLine();
