UGS: Prevent arrays being serialized as blank items.

diff --git a/Engine/Source/Programs/UnrealGameSync/UnrealGameSync/ConfigFile.cs b/Engine/Source/Programs/UnrealGameSync/UnrealGameSync/ConfigFile.cs
--- a/Engine/Source/Programs/UnrealGameSync/UnrealGameSync/ConfigFile.cs
+++ b/Engine/Source/Programs/UnrealGameSync/UnrealGameSync/ConfigFile.cs
@@ -1,4 +1,4 @@
-﻿// Copyright 1998-2018 Epic Games, Inc. All Rights Reserved.
+// Copyright 1998-2018 Epic Games, Inc. All Rights Reserved.
 
 using System;
 using System.Collections.Generic;
@@ -300,7 +300,7 @@ namespace UnrealGameSync
 
 		public void SetValues(string Key, string[] Values)
 		{
-			if(Values == null)
+			if(Values == null || Values.Length == 0)
 			{
 				RemoveValue(Key);
 			}
