UGS: Allow executing arbitrary programs from description badges.

diff --git a/Engine/Source/Programs/UnrealGameSync/UnrealGameSync/Controls/WorkspaceControl.cs b/Engine/Source/Programs/UnrealGameSync/UnrealGameSync/Controls/WorkspaceControl.cs
--- a/Engine/Source/Programs/UnrealGameSync/UnrealGameSync/Controls/WorkspaceControl.cs
+++ b/Engine/Source/Programs/UnrealGameSync/UnrealGameSync/Controls/WorkspaceControl.cs
@@ -69,26 +69,26 @@ namespace UnrealGameSync
 			public int Height;
 			public Color BackgroundColor;
 			public Color HoverBackgroundColor;
-			public object UserData;
+			public Action ClickHandler;
 			public string ToolTip;
 
 			public BadgeInfo(string Label, string Group, Color BadgeColor)
-				: this(Label, Group, null, BadgeColor, BadgeColor)
+				: this(Label, Group, null, BadgeColor, BadgeColor, null)
 			{
 			}
 
-			public BadgeInfo(string Label, string Group, string UniqueId, Color BackgroundColor, Color HoverBackgroundColor, object UserData = null)
+			public BadgeInfo(string Label, string Group, string UniqueId, Color BackgroundColor, Color HoverBackgroundColor, Action ClickHandler)
 			{
 				this.Label = Label;
 				this.Group = Group;
 				this.UniqueId = UniqueId;
 				this.BackgroundColor = BackgroundColor;
 				this.HoverBackgroundColor = HoverBackgroundColor;
-				this.UserData = UserData;
+				this.ClickHandler = ClickHandler;
 			}
 
 			public BadgeInfo(BadgeInfo Other)
-				: this(Other.Label, Other.Group, Other.UniqueId, Other.BackgroundColor, Other.HoverBackgroundColor, Other.UserData)
+				: this(Other.Label, Other.Group, Other.UniqueId, Other.BackgroundColor, Other.HoverBackgroundColor, Other.ClickHandler)
 			{
 				this.Offset = Other.Offset;
 				this.Width = Other.Width;
@@ -1833,6 +1833,7 @@ namespace UnrealGameSync
 						string Color = BadgeDefinitionObject.GetValue("Color", "#909090");
 						string HoverColor = BadgeDefinitionObject.GetValue("HoverColor", "#b0b0b0");
 						string Url = BadgeDefinitionObject.GetValue("Url", null);
+						string Arguments = BadgeDefinitionObject.GetValue("Arguments", null);
 						if(!String.IsNullOrEmpty(Name) && !String.IsNullOrEmpty(Pattern))
 						{
 							foreach(Match MatchResult in Regex.Matches(Description, Pattern, RegexOptions.Multiline))
@@ -1841,7 +1842,25 @@ namespace UnrealGameSync
 								Color HoverBadgeColor = System.Drawing.ColorTranslator.FromHtml(HoverColor);
 
 								string UniqueId = String.IsNullOrEmpty(Url)? null : String.Format("Description:{0}:{1}", Change.Number, Badges.Count);
-								Badges.Add(new BadgeInfo(ReplaceRegexMatches(Name, MatchResult), Group, UniqueId, BadgeColor, HoverBadgeColor, ReplaceRegexMatches(Url, MatchResult)));
+
+								string ExpandedUrl = ReplaceRegexMatches(Url, MatchResult);
+								string ExpandedArguments = ReplaceRegexMatches(Arguments, MatchResult);
+
+								Action ClickHandler;
+								if(String.IsNullOrEmpty(ExpandedUrl))
+								{
+									ClickHandler = null;
+								}
+								else if(String.IsNullOrEmpty(ExpandedArguments))
+								{
+									ClickHandler = () => Process.Start(ExpandedUrl);
+								}
+								else
+								{
+									ClickHandler = () => Process.Start(ExpandedUrl, ExpandedArguments);
+								}
+
+								Badges.Add(new BadgeInfo(ReplaceRegexMatches(Name, MatchResult), Group, UniqueId, BadgeColor, HoverBadgeColor, ClickHandler));
 							}
 						}
 					}
@@ -2004,8 +2023,18 @@ namespace UnrealGameSync
 
 			Color HoverBadgeColor = Color.FromArgb(BadgeColor.A, Math.Min(BadgeColor.R + 32, 255), Math.Min(BadgeColor.G + 32, 255), Math.Min(BadgeColor.B + 32, 255));
 
+			Action ClickHandler;
+			if(BadgeData == null || BadgeData.Url == null)
+			{
+				ClickHandler = null;
+			}
+			else
+			{
+				ClickHandler = () => Process.Start(BadgeData.Url);
+			}
+
 			string UniqueId = String.Format("{0}:{1}", ChangeNumber, BadgeName);
-			return new BadgeInfo(BadgeLabel, BadgeGroup, UniqueId, BadgeColor, HoverBadgeColor, BadgeData);
+			return new BadgeInfo(BadgeLabel, BadgeGroup, UniqueId, BadgeColor, HoverBadgeColor, ClickHandler);
 		}
 
 		private Dictionary<string, List<BadgeInfo>> CreateCustomBadges(int ChangeNumber, EventSummary Summary)
@@ -2989,9 +3018,9 @@ namespace UnrealGameSync
 							foreach (BadgeInfo Badge in LayoutInfo.DescriptionBadges)
 							{
 								Rectangle BadgeBounds = Badge.GetBounds(BuildListLocation);
-								if(BadgeBounds.Contains(Args.Location))
+								if(BadgeBounds.Contains(Args.Location) && Badge.ClickHandler != null)
 								{
-									Process.Start((string)Badge.UserData);
+									Badge.ClickHandler();
 									break;
 								}
 							}
@@ -3007,13 +3036,9 @@ namespace UnrealGameSync
 							BuildListLocation.X = Math.Max(BuildListLocation.X, HitTest.SubItem.Bounds.Left);
 
 							BadgeInfo BadgeInfo = HitTestBadge(Args.Location, LayoutInfo.BuildBadges, BuildListLocation);
-							if(BadgeInfo != null)
+							if(BadgeInfo != null && BadgeInfo.ClickHandler != null)
 							{
-								BadgeData BadgeData = (BadgeData)BadgeInfo.UserData;
-								if(BadgeData != null)
-								{
-									Process.Start(BadgeData.Url);
-								}
+								BadgeInfo.ClickHandler();
 							}
 						}
 					}
@@ -3030,13 +3055,9 @@ namespace UnrealGameSync
 								Point ListLocation = GetBadgeListLocation(Badges, HitTest.SubItem.Bounds, HorizontalAlign.Center, VerticalAlignment.Middle);
 
 								BadgeInfo BadgeInfo = HitTestBadge(Args.Location, Badges, ListLocation);
-								if(BadgeInfo != null)
+								if(BadgeInfo != null && BadgeInfo.ClickHandler != null)
 								{
-									BadgeData BadgeData = (BadgeData)BadgeInfo.UserData;
-									if(BadgeData != null)
-									{
-										Process.Start(BadgeData.Url);
-									}
+									BadgeInfo.ClickHandler();
 								}
 							}
 						}
