UGS: Allow any explicitly specified streams in the quick-select list. Also saves polling the server for the lists of streams.

diff --git a/Engine/Source/Programs/UnrealGameSync/UnrealGameSync/Controls/WorkspaceControl.cs b/Engine/Source/Programs/UnrealGameSync/UnrealGameSync/Controls/WorkspaceControl.cs
--- a/Engine/Source/Programs/UnrealGameSync/UnrealGameSync/Controls/WorkspaceControl.cs
+++ b/Engine/Source/Programs/UnrealGameSync/UnrealGameSync/Controls/WorkspaceControl.cs
@@ -2827,41 +2827,39 @@ namespace UnrealGameSync
 			bool bShownContextMenu = false;
 			if(StreamName != null)
 			{
-				IReadOnlyList<string> OtherStreamFilter = Workspace.ProjectStreamFilter;
-				if(OtherStreamFilter != null)
+				IReadOnlyList<string> OtherStreamNames = Workspace.ProjectStreamFilter;
+				if(OtherStreamNames != null)
 				{
-					IReadOnlyList<string> OtherStreamNames = PerforceMonitor.OtherStreamNames.Where(x => OtherStreamFilter.Any(y => y.Equals(x, StringComparison.InvariantCultureIgnoreCase)) || x.Equals(StreamName, StringComparison.InvariantCultureIgnoreCase)).ToList().AsReadOnly();
-					if(OtherStreamNames.Count > 0)
-					{
-						StreamContextMenu.Items.Clear();
+					StreamContextMenu.Items.Clear();
 
-						foreach (string OtherStreamName in OtherStreamNames.OrderBy(x => x).Where(x => !x.EndsWith("/Dev-Binaries")))
-						{
-							string ThisStreamName = OtherStreamName; // Local for lambda capture
+					ToolStripMenuItem CurrentStreamItem = new ToolStripMenuItem(StreamName, null, new EventHandler((S, E) => SelectStream(StreamName)));
+					CurrentStreamItem.Checked = true;
+					StreamContextMenu.Items.Add(CurrentStreamItem);
+					
+					StreamContextMenu.Items.Add(new ToolStripSeparator());
 
+					foreach (string OtherStreamName in OtherStreamNames.OrderBy(x => x).Where(x => !x.EndsWith("/Dev-Binaries")))
+					{
+						string ThisStreamName = OtherStreamName; // Local for lambda capture
+						if(String.Compare(StreamName, OtherStreamName, StringComparison.InvariantCultureIgnoreCase) != 0)
+						{
 							ToolStripMenuItem Item = new ToolStripMenuItem(ThisStreamName, null, new EventHandler((S, E) => SelectStream(ThisStreamName)));
-							if(String.Compare(StreamName, OtherStreamName, StringComparison.InvariantCultureIgnoreCase) == 0)
-							{
-								Item.Checked = true;
-								StreamContextMenu.Items.Insert(0, Item);
-								StreamContextMenu.Items.Insert(1, new ToolStripSeparator());
-							}
-							else
-							{
-								StreamContextMenu.Items.Add(Item);
-							}
+							StreamContextMenu.Items.Add(Item);
 						}
+					}
 
+					if(StreamContextMenu.Items.Count > 2)
+					{
 						StreamContextMenu.Items.Add(new ToolStripSeparator());
-						StreamContextMenu.Items.Add(new ToolStripMenuItem("Select Other...", null, new EventHandler((S, E) => SelectOtherStreamDialog())));
+					}
 
-						int X = (Bounds.Left + Bounds.Right) / 2 + StreamContextMenu.Bounds.Width / 2;
-						int Y = Bounds.Bottom + 2;
-						StreamContextMenu.Show(StatusPanel, new Point(X, Y), ToolStripDropDownDirection.Left);
+					StreamContextMenu.Items.Add(new ToolStripMenuItem("Select Other...", null, new EventHandler((S, E) => SelectOtherStreamDialog())));
 
-						bShownContextMenu = true;
+					int X = (Bounds.Left + Bounds.Right) / 2 + StreamContextMenu.Bounds.Width / 2;
+					int Y = Bounds.Bottom + 2;
+					StreamContextMenu.Show(StatusPanel, new Point(X, Y), ToolStripDropDownDirection.Left);
 
-					}
+					bShownContextMenu = true;
 				}
 			}
 			if(!bShownContextMenu)
diff --git a/Engine/Source/Programs/UnrealGameSync/UnrealGameSync/PerforceMonitor.cs b/Engine/Source/Programs/UnrealGameSync/UnrealGameSync/PerforceMonitor.cs
--- a/Engine/Source/Programs/UnrealGameSync/UnrealGameSync/PerforceMonitor.cs
+++ b/Engine/Source/Programs/UnrealGameSync/UnrealGameSync/PerforceMonitor.cs
@@ -59,7 +59,6 @@ namespace UnrealGameSync
 			PendingMaxChangesValue = 100;
 			LastChangeByCurrentUser = -1;
 			LastCodeChangeByCurrentUser = -1;
-			OtherStreamNames = new List<string>();
 			bIsEnterpriseProject = bInIsEnterpriseProject;
 			LatestProjectConfigFile = InProjectConfigFile;
 			LocalConfigFiles = InLocalConfigFiles;
@@ -116,12 +115,6 @@ namespace UnrealGameSync
 			set { lock(this){ if(value != PendingMaxChangesValue){ PendingMaxChangesValue = value; RefreshEvent.Set(); } } }
 		}
 
-		public IReadOnlyList<string> OtherStreamNames
-		{
-			get;
-			private set;
-		}
-
 		void PollForUpdates()
 		{
 			string StreamName;
@@ -155,17 +148,6 @@ namespace UnrealGameSync
 							OnStreamChange();
 						}
 
-						// Update the stream list
-						if(StreamName != null)
-						{
-							List<string> NewOtherStreamNames;
-							if(!Perforce.FindStreams(PerforceUtils.GetClientOrDepotDirectoryName(StreamName) + "/*", out NewOtherStreamNames, LogWriter))
-							{
-								NewOtherStreamNames = new List<string>();
-							}
-							OtherStreamNames = NewOtherStreamNames;
-						}
-
 						// Check for any p4 changes
 						if(!UpdateChanges())
 						{
