UGS: Fix crash when message of the day contains an escaped bracket.

#rnx

diff --git a/Engine/Source/Programs/UnrealGameSync/UnrealGameSync/Controls/WorkspaceControl.cs b/Engine/Source/Programs/UnrealGameSync/UnrealGameSync/Controls/WorkspaceControl.cs
--- a/Engine/Source/Programs/UnrealGameSync/UnrealGameSync/Controls/WorkspaceControl.cs
+++ b/Engine/Source/Programs/UnrealGameSync/UnrealGameSync/Controls/WorkspaceControl.cs
@@ -2770,6 +2770,7 @@ namespace UnrealGameSync
 					if(Idx + 1 < Text.Length && Text[Idx + 1] == '[')
 					{
 						ElementText.Append(Text[Idx]);
+						Idx += 2;
 						continue;
 					}
 
