- Throttle events and comments to every 90 seconds

diff --git a/Engine/Source/Programs/UnrealGameSync/UnrealGameSync/EventMonitor.cs b/Engine/Source/Programs/UnrealGameSync/UnrealGameSync/EventMonitor.cs
--- a/Engine/Source/Programs/UnrealGameSync/UnrealGameSync/EventMonitor.cs
+++ b/Engine/Source/Programs/UnrealGameSync/UnrealGameSync/EventMonitor.cs
@@ -462,10 +462,13 @@ namespace UnrealGameSync
 		{
 			EventData Event = null;
 			CommentData Comment = null;
-			while(!bDisposing)
+			bool bUpdateThrottledRequests = true;
+			double RequestThrottle = 90; // seconds to wait for throttled request;
+			Stopwatch Timer = Stopwatch.StartNew();
+			while (!bDisposing)
 			{
 				// If there's no connection string, just empty out the queue
-				if(ApiUrl != null)
+				if (ApiUrl != null)
 				{
 					// Post all the reviews to the database. We don't send them out of order, so keep the review outside the queue until the next update if it fails
 					while(Event != null || OutgoingEvents.TryDequeue(out Event))
@@ -481,8 +484,14 @@ namespace UnrealGameSync
 						Comment = null;
 					}
 
-					// Read all the new reviews
-					ReadEventsFromBackend();
+					if(Timer.Elapsed > TimeSpan.FromSeconds(RequestThrottle))
+					{
+						bUpdateThrottledRequests = true;
+						Timer.Reset();
+					}
+
+					// Read all the new reviews, pass whether or not to fire the throttled requests
+					ReadEventsFromBackend(bUpdateThrottledRequests);
 
 					// Send a notification that we're ready to update
 					if((IncomingEvents.Count > 0 || IncomingBadges.Count > 0 || IncomingComments.Count > 0) && OnUpdatesReady != null)
@@ -492,7 +501,7 @@ namespace UnrealGameSync
 				}
 
 				// Wait for something else to do
-				RefreshEvent.WaitOne(30 * 1000);
+				bUpdateThrottledRequests = RefreshEvent.WaitOne(30 * 1000);
 			}
 		}
 
@@ -529,7 +538,7 @@ namespace UnrealGameSync
 			}
 		}
 
-		bool ReadEventsFromBackend()
+		bool ReadEventsFromBackend(bool bFireThrottledRequests)
 		{
 			try
 			{
@@ -547,26 +556,6 @@ namespace UnrealGameSync
 					LatestIds.LastEventId = InitialIds.LastEventId;
 				}
 
-				//////////////
-				/// Reviews 
-				//////////////
-				List<EventData> Events = RESTApi.GET<List<EventData>>(ApiUrl, "event", string.Format("project={0}", Project), string.Format("lasteventid={0}", LatestIds.LastEventId));
-				foreach(EventData Review in Events)
-				{
-					IncomingEvents.Enqueue(Review);
-					LatestIds.LastEventId = Math.Max(LatestIds.LastEventId, Review.Id);
-				}
-
-				//////////////
-				/// Comments 
-				//////////////
-				List<CommentData> Comments = RESTApi.GET<List<CommentData>>(ApiUrl, "comment", string.Format("project={0}", Project), string.Format("lastcommentid={0}", LatestIds.LastCommentId));
-				foreach (CommentData Comment in Comments)
-				{
-					IncomingComments.Enqueue(Comment);
-					LatestIds.LastCommentId = Math.Max(LatestIds.LastCommentId, Comment.Id);
-				}
-
 				//////////////
 				/// Bulids
 				//////////////
@@ -577,6 +566,32 @@ namespace UnrealGameSync
 					LatestIds.LastBuildId = Math.Max(LatestIds.LastBuildId, Build.Id);
 				}
 
+				//////////////////////////
+				/// Throttled Requests
+				//////////////////////////
+				if (bFireThrottledRequests)
+				{
+					//////////////
+					/// Reviews 
+					//////////////
+					List<EventData> Events = RESTApi.GET<List<EventData>>(ApiUrl, "event", string.Format("project={0}", Project), string.Format("lasteventid={0}", LatestIds.LastEventId));
+					foreach (EventData Review in Events)
+					{
+						IncomingEvents.Enqueue(Review);
+						LatestIds.LastEventId = Math.Max(LatestIds.LastEventId, Review.Id);
+					}
+
+					//////////////
+					/// Comments 
+					//////////////
+					List<CommentData> Comments = RESTApi.GET<List<CommentData>>(ApiUrl, "comment", string.Format("project={0}", Project), string.Format("lastcommentid={0}", LatestIds.LastCommentId));
+					foreach (CommentData Comment in Comments)
+					{
+						IncomingComments.Enqueue(Comment);
+						LatestIds.LastCommentId = Math.Max(LatestIds.LastCommentId, Comment.Id);
+					}
+				}
+
 				LastStatusMessage = String.Format("Last update took {0}ms", Timer.ElapsedMilliseconds);
 				LogWriter.WriteLine("Done in {0}ms.", Timer.ElapsedMilliseconds);
 				return true;
