- Implement request retry
- Remove unneccessary check for RestURL
- Have SendRequest return an int instead of a string. 0 on 200, 1 otherwise
- Print response content to the log if there is any
- Add missing using around HttpWebResponse object

diff --git a/Engine/Source/Programs/UnrealGameSync/PostBadgeStatus/Program.cs b/Engine/Source/Programs/UnrealGameSync/PostBadgeStatus/Program.cs
--- a/Engine/Source/Programs/UnrealGameSync/PostBadgeStatus/Program.cs
+++ b/Engine/Source/Programs/UnrealGameSync/PostBadgeStatus/Program.cs
@@ -7,6 +7,7 @@ using System.IO;
 using System.Linq;
 using System.Net;
 using System.Text;
+using System.Threading;
 using System.Threading.Tasks;
 using System.Web.Script.Serialization;
 
@@ -32,36 +33,46 @@ namespace WriteBadgeStatus
 				Console.WriteLine("  PostBadgeStatus -Name=<Name> -Change=<CL> -Project=<DepotPath> -RestUrl=<Url> -Status=<Status> -Url=<Url>");
 				return 1;
 			}
-			if (RestUrl != null)
+
+			BuildData Build = new BuildData
+			{
+				BuildType = Name,
+				Url = Url,
+				Project = Project,
+				ArchivePath = ""
+			};
+			if (!int.TryParse(Change, out Build.ChangeNumber))
+			{
+				Console.WriteLine("Change must be an integer!");
+				return 1;
+			}
+			if (!Enum.TryParse<BuildData.BuildDataResult>(Status, true, out Build.Result))
+			{
+				Console.WriteLine("Change must be Starting, Failure, Warning, Success, or Skipped!");
+				return 1;
+			}
+			int NumRetries = 0;
+			while (true)
 			{
-				BuildData Build = new BuildData
-				{
-					BuildType = Name,
-					Url = Url,
-					Project = Project,
-					ArchivePath = ""
-				};
-				if(!int.TryParse(Change, out Build.ChangeNumber))
-				{
-					Console.WriteLine("Change must be an integer!");
-					return 1;
-				}
-				if (!Enum.TryParse<BuildData.BuildDataResult>(Status, true, out Build.Result))
-				{
-					Console.WriteLine("Change must be Starting, Failure, Warning, Success, or Skipped!");
-					return 1;
-				}
 				try
 				{
-					SendRequest(RestUrl, "Build", "POST", new JavaScriptSerializer().Serialize(Build));
+					return SendRequest(RestUrl, "Build", "POST", new JavaScriptSerializer().Serialize(Build));
 				}
-				catch(Exception ex)
+				catch (Exception ex)
 				{
-					Console.WriteLine(string.Format("An exception was thrown attempting to send the request: {0}", ex.Message));
-					return 1;
+					if (++NumRetries <= 3)
+					{
+						Console.WriteLine(string.Format("An exception was thrown attempting to send the request: {0}, retrying...", ex.Message));
+						// Wait 5 seconds and retry;
+						Thread.Sleep(5000);
+					}
+					else
+					{
+						Console.WriteLine("Failed to POST due to multiple failures. Aborting.");
+						return 1;
+					}
 				}
 			}
-			return 0;
 		}
 
 		static string ParseParam(List<string> Arguments, string ParamName)
@@ -108,7 +119,7 @@ namespace WriteBadgeStatus
 			}
 		}
 
-		static string SendRequest(string URI, string Resource, string Method, string RequestBody = null, params string[] QueryParams)
+		static int SendRequest(string URI, string Resource, string Method, string RequestBody = null, params string[] QueryParams)
 		{
 			// set up the query string
 			StringBuilder TargetURI = new StringBuilder(string.Format("{0}/api/{1}", URI, Resource));
@@ -143,13 +154,17 @@ namespace WriteBadgeStatus
 			}
 			try
 			{
-				WebResponse Repsonse = Request.GetResponse();
-				string ResponseContent = null;
-				using (StreamReader ResponseReader = new System.IO.StreamReader(Repsonse.GetResponseStream(), Encoding.Default))
+				using (HttpWebResponse Response = (HttpWebResponse)Request.GetResponse())
 				{
-					ResponseContent = ResponseReader.ReadToEnd();
+					string ResponseContent = null;
+					using (StreamReader ResponseReader = new System.IO.StreamReader(Response.GetResponseStream(), Encoding.Default))
+					{
+						ResponseContent = ResponseReader.ReadToEnd();
+						Console.WriteLine(ResponseContent);
+						return Response.StatusCode == HttpStatusCode.OK ? 0 : 1;
+					}
 				}
-				return ResponseContent;
+				
 			}
 			catch (WebException ex)
 			{
