UGS: Prevent exception when browsing for project.

#rnx

diff --git a/Engine/Source/Programs/UnrealGameSync/UnrealGameSync/Forms/OpenProjectWindow.cs b/Engine/Source/Programs/UnrealGameSync/UnrealGameSync/Forms/OpenProjectWindow.cs
--- a/Engine/Source/Programs/UnrealGameSync/UnrealGameSync/Forms/OpenProjectWindow.cs
+++ b/Engine/Source/Programs/UnrealGameSync/UnrealGameSync/Forms/OpenProjectWindow.cs
@@ -325,12 +325,15 @@ namespace UnrealGameSync
 			Dialog.Filter = "Project files (*.uproject)|*.uproject|Project directory lists (*.uprojectdirs)|*.uprojectdirs|All supported files (*.uproject;*.uprojectdirs)|*.uproject;*.uprojectdirs|All files (*.*)|*.*" ;
 			Dialog.FilterIndex = Settings.FilterIndex;
 			
-			try
-			{
-				Dialog.InitialDirectory = Path.GetDirectoryName(LocalFileTextBox.Text);
-			}
-			catch
+			if(!String.IsNullOrEmpty(LocalFileTextBox.Text))
 			{
+				try
+				{
+					Dialog.InitialDirectory = Path.GetDirectoryName(LocalFileTextBox.Text);
+				}
+				catch
+				{
+				}
 			}
 
 			if(Dialog.ShowDialog(this) == System.Windows.Forms.DialogResult.OK)
