UGS: Rename PerforceSettingsWindow to PerforceSyncSettingsWindow

#rnx

diff --git a/Engine/Source/Programs/UnrealGameSync/UnrealGameSync/Controls/WorkspaceControl.cs b/Engine/Source/Programs/UnrealGameSync/UnrealGameSync/Controls/WorkspaceControl.cs
--- a/Engine/Source/Programs/UnrealGameSync/UnrealGameSync/Controls/WorkspaceControl.cs
+++ b/Engine/Source/Programs/UnrealGameSync/UnrealGameSync/Controls/WorkspaceControl.cs
@@ -4218,7 +4218,7 @@ namespace UnrealGameSync
 
 		private void OptionsContextMenu_PerforceSettings_Click(object sender, EventArgs e)
 		{
-			PerforceSettingsWindow Perforce = new PerforceSettingsWindow(Settings);
+			PerforceSyncSettingsWindow Perforce = new PerforceSyncSettingsWindow(Settings);
 			Perforce.ShowDialog();
 		}
 
diff --git a/Engine/Source/Programs/UnrealGameSync/UnrealGameSync/Forms/PerforceSettingsWindow.cs b/Engine/Source/Programs/UnrealGameSync/UnrealGameSync/Forms/PerforceSettingsWindow.cs
deleted file mode 100644
diff --git a/Engine/Source/Programs/UnrealGameSync/UnrealGameSync/Forms/PerforceSyncSettingsWindow.cs b/Engine/Source/Programs/UnrealGameSync/UnrealGameSync/Forms/PerforceSyncSettingsWindow.cs
new file mode 100644
--- /dev/null
+++ b/Engine/Source/Programs/UnrealGameSync/UnrealGameSync/Forms/PerforceSyncSettingsWindow.cs
@@ -0,0 +1,71 @@
+// Copyright 1998-2018 Epic Games, Inc. All Rights Reserved.
+
+using System;
+using System.Collections.Generic;
+using System.ComponentModel;
+using System.Data;
+using System.Drawing;
+using System.Linq;
+using System.Text;
+using System.Threading.Tasks;
+using System.Windows.Forms;
+
+namespace UnrealGameSync
+{
+	partial class PerforceSyncSettingsWindow : Form
+	{
+		UserSettings Settings;
+
+		public PerforceSyncSettingsWindow(UserSettings Settings)
+		{
+			this.Settings = Settings;
+			InitializeComponent();
+		}
+
+		private void PerforceSettingsWindow_Load(object sender, EventArgs e)
+		{
+			PerforceSyncOptions SyncOptions = Settings.SyncOptions;
+			NumRetriesTextBox.Text = (SyncOptions.NumRetries > 0)? SyncOptions.NumRetries.ToString() : "";
+			NumThreadsTextBox.Text = (SyncOptions.NumThreads > 0)? SyncOptions.NumThreads.ToString() : "";
+			TcpBufferSizeText.Text = (SyncOptions.TcpBufferSize > 0)? SyncOptions.TcpBufferSize.ToString() : "";
+		}
+
+		private void OkButton_Click(object sender, EventArgs e)
+		{
+			int NewNumRetries = 0;
+			if(NumRetriesTextBox.Text.Length > 0 && !int.TryParse(NumRetriesTextBox.Text, out NewNumRetries))
+			{
+				MessageBox.Show("Invalid value for number of retries");
+				return;
+			}
+
+			int NewNumThreads = 0;
+			if(NumThreadsTextBox.Text.Length > 0 && !int.TryParse(NumThreadsTextBox.Text, out NewNumThreads))
+			{
+				MessageBox.Show("Invalid value for number of threads");
+				return;
+			}
+
+			int NewTcpBufferSize = 0;
+			if(TcpBufferSizeText.Text.Length > 0 && !int.TryParse(TcpBufferSizeText.Text, out NewTcpBufferSize))
+			{
+				MessageBox.Show("Invalid value for TCP buffer size");
+				return;
+			}
+
+			Settings.SyncOptions.NumRetries = NewNumRetries;
+			Settings.SyncOptions.NumThreads = NewNumThreads;
+			Settings.SyncOptions.TcpBufferSize = NewTcpBufferSize;
+			Settings.Save();
+
+			DialogResult = System.Windows.Forms.DialogResult.OK;
+			Close();
+		}
+
+		private void CancButton_Click(object sender, EventArgs e)
+		{
+			DialogResult = System.Windows.Forms.DialogResult.Cancel;
+			Close();
+		}
+	}
+}
diff --git a/Engine/Source/Programs/UnrealGameSync/UnrealGameSync/UnrealGameSync.csproj b/Engine/Source/Programs/UnrealGameSync/UnrealGameSync/UnrealGameSync.csproj
--- a/Engine/Source/Programs/UnrealGameSync/UnrealGameSync/UnrealGameSync.csproj
+++ b/Engine/Source/Programs/UnrealGameSync/UnrealGameSync/UnrealGameSync.csproj
@@ -206,11 +206,11 @@
     <Compile Include="Forms\PasswordWindow.Designer.cs">
       <DependentUpon>PasswordWindow.cs</DependentUpon>
     </Compile>
-    <Compile Include="Forms\PerforceSettingsWindow.cs">
+    <Compile Include="Forms\PerforceSyncSettingsWindow.cs">
       <SubType>Form</SubType>
     </Compile>
-    <Compile Include="Forms\PerforceSettingsWindow.Designer.cs">
-      <DependentUpon>PerforceSettingsWindow.cs</DependentUpon>
+    <Compile Include="Forms\PerforceSyncSettingsWindow.Designer.cs">
+      <DependentUpon>PerforceSyncSettingsWindow.cs</DependentUpon>
     </Compile>
     <Compile Include="Forms\ProgramsRunningWindow.cs">
       <SubType>Form</SubType>
@@ -369,8 +369,8 @@
     <EmbeddedResource Include="Forms\PasswordWindow.resx">
       <DependentUpon>PasswordWindow.cs</DependentUpon>
     </EmbeddedResource>
-    <EmbeddedResource Include="Forms\PerforceSettingsWindow.resx">
-      <DependentUpon>PerforceSettingsWindow.cs</DependentUpon>
+    <EmbeddedResource Include="Forms\PerforceSyncSettingsWindow.resx">
+      <DependentUpon>PerforceSyncSettingsWindow.cs</DependentUpon>
     </EmbeddedResource>
     <EmbeddedResource Include="Forms\ProgramsRunningWindow.resx">
       <DependentUpon>ProgramsRunningWindow.cs</DependentUpon>
