UGS: Move the site-specific UGS settings into a separate file (DeploymentSettings.cs) that can be shared between UnrealGameSync and UnrealGameSyncLauncher.

diff --git a/Engine/Source/Programs/UnrealGameSync/UnrealGameSync/DeploymentSettings.cs b/Engine/Source/Programs/UnrealGameSync/UnrealGameSync/DeploymentSettings.cs
new file mode 100644
--- /dev/null
+++ b/Engine/Source/Programs/UnrealGameSync/UnrealGameSync/DeploymentSettings.cs
@@ -0,0 +1,26 @@
+using System;
+using System.Collections.Generic;
+using System.Linq;
+using System.Text;
+using System.Threading.Tasks;
+
+namespace UnrealGameSync
+{
+	/// <summary>
+	/// This class contains settings for a site-specific deployment of UGS. Epic's internal implementation uses a static constructor in a NotForLicensees folder to initialize these values.
+	/// </summary>
+	static partial class DeploymentSettings
+	{
+		/// <summary>
+		/// SQL connection string used to connect to the database for telemetry and review data. The 'Program' class is a partial class, to allow an
+		/// opportunistically included C# source file in NotForLicensees/ProgramSettings.cs to override this value in a static constructor.
+		/// </summary>
+		public static readonly string ApiUrl = null;
+
+		/// <summary>
+		/// Specifies the depot path to sync down the stable version of UGS from, without a trailing slash (eg. //depot/UnrealGameSync/bin). This is a site-specific setting. 
+		/// The UnrealGameSync executable should be located at Release/UnrealGameSync.exe under this path, with any dependent DLLs.
+		/// </summary>
+		public static readonly string DefaultDepotPath = null;
+	}
+}
diff --git a/Engine/Source/Programs/UnrealGameSync/UnrealGameSync/Program.cs b/Engine/Source/Programs/UnrealGameSync/UnrealGameSync/Program.cs
--- a/Engine/Source/Programs/UnrealGameSync/UnrealGameSync/Program.cs
+++ b/Engine/Source/Programs/UnrealGameSync/UnrealGameSync/Program.cs
@@ -15,14 +15,8 @@ using System.Windows.Forms;
 
 namespace UnrealGameSync
 {
-	static partial class Program
+	static class Program
 	{
-		/// <summary>
-		/// SQL connection string used to connect to the database for telemetry and review data. The 'Program' class is a partial class, to allow an
-		/// opportunistically included C# source file in NotForLicensees/ProgramSettings.cs to override this value in a static constructor.
-		/// </summary>
-		public static readonly string ApiUrl = null;
-
 		public static string SyncVersion = null;
 
 		[STAThread]
@@ -92,12 +86,12 @@ namespace UnrealGameSync
 			string DataFolder = Path.Combine(Environment.GetFolderPath(Environment.SpecialFolder.LocalApplicationData), "UnrealGameSync");
 			Directory.CreateDirectory(DataFolder);
 
-			using(TelemetryWriter Telemetry = new TelemetryWriter(ApiUrl, Path.Combine(DataFolder, "Telemetry.log")))
+			using(TelemetryWriter Telemetry = new TelemetryWriter(DeploymentSettings.ApiUrl, Path.Combine(DataFolder, "Telemetry.log")))
 			{
 				AppDomain.CurrentDomain.UnhandledException += CurrentDomain_UnhandledException;
 				using(UpdateMonitor UpdateMonitor = new UpdateMonitor(new PerforceConnection(UserName, null, ServerAndPort), UpdatePath))
 				{
-					ProgramApplicationContext Context = new ProgramApplicationContext(UpdateMonitor, ApiUrl, DataFolder, ActivateEvent, bRestoreState, UpdateSpawn, ProjectFileName, bUnstable);
+					ProgramApplicationContext Context = new ProgramApplicationContext(UpdateMonitor, DeploymentSettings.ApiUrl, DataFolder, ActivateEvent, bRestoreState, UpdateSpawn, ProjectFileName, bUnstable);
 					Application.Run(Context);
 
 					if(UpdateMonitor.IsUpdateAvailable && UpdateSpawn != null)
diff --git a/Engine/Source/Programs/UnrealGameSync/UnrealGameSync/UnrealGameSync.csproj b/Engine/Source/Programs/UnrealGameSync/UnrealGameSync/UnrealGameSync.csproj
--- a/Engine/Source/Programs/UnrealGameSync/UnrealGameSync/UnrealGameSync.csproj
+++ b/Engine/Source/Programs/UnrealGameSync/UnrealGameSync/UnrealGameSync.csproj
@@ -149,6 +149,7 @@
       <SubType>Component</SubType>
     </Compile>
     <Compile Include="DeleteFilesTask.cs" />
+    <Compile Include="DeploymentSettings.cs" />
     <Compile Include="DetectProjectSettingsTask.cs" />
     <Compile Include="FileFilter.cs" />
     <Compile Include="FindFoldersToCleanTask.cs" />
@@ -317,6 +318,7 @@
       <DependentUpon>SyncFilter.cs</DependentUpon>
     </Compile>
     <Compile Include="ModalTask.cs" />
+    <Compile Include="NotForLicensees\DeploymentSettingsInternal.cs" />
     <Compile Include="PerforceModalTask.cs" />
     <Compile Include="PerforceMonitor.cs" />
     <Compile Include="ProgramApplicationContext.cs" />
@@ -441,11 +443,7 @@
     <None Include="Resources\App.config" />
   </ItemGroup>
   <Choose>
-    <When Condition="Exists('NotForLicensees\ProgramSettings.cs')">
-      <ItemGroup>
-        <Compile Include="NotForLicensees\ProgramSettings.cs" />
-      </ItemGroup>
-    </When>
+    <When Condition="Exists('NotForLicensees\ProgramSettings.cs')" />
   </Choose>
   <ItemGroup>
     <EmbeddedResource Include="Resources\Icons.png" />
diff --git a/Engine/Source/Programs/UnrealGameSync/UnrealGameSyncLauncher/Program.cs b/Engine/Source/Programs/UnrealGameSync/UnrealGameSyncLauncher/Program.cs
--- a/Engine/Source/Programs/UnrealGameSync/UnrealGameSyncLauncher/Program.cs
+++ b/Engine/Source/Programs/UnrealGameSync/UnrealGameSyncLauncher/Program.cs
@@ -15,14 +15,8 @@ using UnrealGameSync;
 
 namespace UnrealGameSyncLauncher
 {
-	static partial class Program
+	static class Program
 	{
-		/// <summary>
-		/// Specifies the depot path to sync down the stable version of UGS from, without a trailing slash (eg. //depot/UnrealGameSync/bin). This is a site-specific setting. 
-		/// The UnrealGameSync executable should be located at Release/UnrealGameSync.exe under this path, with any dependent DLLs.
-		/// </summary>
-		static readonly string DefaultDepotPath;
-
 		[STAThread]
 		static int Main(string[] Args)
 		{
@@ -70,7 +64,7 @@ namespace UnrealGameSyncLauncher
 				// Read the settings
 				string ServerAndPort = null;
 				string UserName = null;
-				string DepotPath = DefaultDepotPath;
+				string DepotPath = DeploymentSettings.DefaultDepotPath;
 				ReadSettings(ref ServerAndPort, ref UserName, ref DepotPath);
 
 				// If the shift key is held down, immediately show the settings window
@@ -144,7 +138,7 @@ namespace UnrealGameSyncLauncher
 						Key.SetValue("UserName", UserName);
 					}
 
-					if(String.IsNullOrEmpty(DepotPath) || (DefaultDepotPath != null && String.Equals(DepotPath, DefaultDepotPath, StringComparison.InvariantCultureIgnoreCase)))
+					if(String.IsNullOrEmpty(DepotPath) || (DeploymentSettings.DefaultDepotPath != null && String.Equals(DepotPath, DeploymentSettings.DefaultDepotPath, StringComparison.InvariantCultureIgnoreCase)))
 					{
 						try { Key.DeleteValue("DepotPath"); } catch(Exception) { }
 					}
diff --git a/Engine/Source/Programs/UnrealGameSync/UnrealGameSyncLauncher/UnrealGameSyncLauncher.csproj b/Engine/Source/Programs/UnrealGameSync/UnrealGameSyncLauncher/UnrealGameSyncLauncher.csproj
--- a/Engine/Source/Programs/UnrealGameSync/UnrealGameSyncLauncher/UnrealGameSyncLauncher.csproj
+++ b/Engine/Source/Programs/UnrealGameSync/UnrealGameSyncLauncher/UnrealGameSyncLauncher.csproj
@@ -89,6 +89,9 @@
     <Compile Include="..\UnrealGameSync\Forms\PasswordWindow.Designer.cs">
       <Link>UnrealGameSync\PasswordWindow.Designer.cs</Link>
     </Compile>
+    <Compile Include="..\UnrealGameSync\DeploymentSettings.cs">
+      <Link>Shared\DeploymentSettings.cs</Link>
+    </Compile>
     <Compile Include="..\UnrealGameSync\ModalTask.cs">
       <Link>Shared\ModalTask.cs</Link>
     </Compile>
@@ -151,9 +154,11 @@
   </ItemGroup>
   <ItemGroup />
   <Choose>
-    <When Condition="Exists('NotForLicensees\ProgramSettings.cs')">
+    <When Condition="Exists('..\UnrealGameSync\NotForLicensees\DeploymentSettingsInternal.cs')">
       <ItemGroup>
-        <Compile Include="NotForLicensees\ProgramSettings.cs" />
+        <Compile Include="..\UnrealGameSync\NotForLicensees\DeploymentSettingsInternal.cs">
+			<Link>Shared\NotForLicensees\DeploymentSettingsInternal.cs</Link>
+		</Compile>
       </ItemGroup>
     </When>
   </Choose>
