UGS: Fallback to using the default UnrealGameSyncLauncher.exe location if there is no AutoUpdate INI file present.

diff --git a/Engine/Source/Programs/UnrealGameSync/UnrealGameSync/Program.cs b/Engine/Source/Programs/UnrealGameSync/UnrealGameSync/Program.cs
--- a/Engine/Source/Programs/UnrealGameSync/UnrealGameSync/Program.cs
+++ b/Engine/Source/Programs/UnrealGameSync/UnrealGameSync/Program.cs
@@ -83,6 +83,23 @@ namespace UnrealGameSync
 				}
 			}
 
+			if(SyncVersion != null)
+			{
+				if(UpdateSpawn == null)
+				{
+					string DefaultLauncherPath = Path.Combine(Environment.GetFolderPath(Environment.SpecialFolder.ProgramFilesX86), "UnrealGameSync", "UnrealGameSyncLauncher.exe");
+					if(File.Exists(DefaultLauncherPath))
+					{
+						UpdateSpawn = DefaultLauncherPath;
+					}
+				}
+
+				if(UpdatePath == null && UpdateSpawn != null)
+				{
+					UpdatePath = String.Format("{0}/{1}/...@>{2}", DeploymentSettings.DefaultDepotPath, bUnstable? "UnstableRelease" : "Release", SyncVersion);
+				}
+			}
+
 			string DataFolder = Path.Combine(Environment.GetFolderPath(Environment.SpecialFolder.LocalApplicationData), "UnrealGameSync");
 			Directory.CreateDirectory(DataFolder);
 
