UGS: Add a Perforce settings window that allows modifying the default connection, auto-update path, and whether the unstable build is being used. If modified, this causes a restart of the program and allows the launcher to pick up on the changes.

diff --git a/Engine/Source/Programs/UnrealGameSync/UnrealGameSync/Controls/WorkspaceControl.cs b/Engine/Source/Programs/UnrealGameSync/UnrealGameSync/Controls/WorkspaceControl.cs
--- a/Engine/Source/Programs/UnrealGameSync/UnrealGameSync/Controls/WorkspaceControl.cs
+++ b/Engine/Source/Programs/UnrealGameSync/UnrealGameSync/Controls/WorkspaceControl.cs
@@ -41,6 +41,7 @@ namespace UnrealGameSync
 		void SetTabNames(TabLabels TabNames);
 		void SetupScheduledSync();
 		void UpdateProgress();
+		void RequestRestart();
 	}
 
 	partial class WorkspaceControl : UserControl, IMainWindowTabPanel
@@ -178,6 +179,7 @@ namespace UnrealGameSync
 		SynchronizationContext MainThreadSynchronizationContext;
 		bool bIsDisposing;
 
+		bool bUnstable;
 		string EditorTargetName;
 		bool bIsEnterpriseProject;
 		PerforceMonitor PerforceMonitor;
@@ -231,7 +233,7 @@ namespace UnrealGameSync
 		bool bUpdateBuildMetadataPosted;
 		bool bUpdateReviewsPosted;
 
-		public WorkspaceControl(IWorkspaceControlOwner InOwner, string InApiUrl, string InOriginalExecutableFileName, DetectProjectSettingsTask DetectSettings, LineBasedTextWriter InLog, UserSettings InSettings)
+		public WorkspaceControl(IWorkspaceControlOwner InOwner, string InApiUrl, string InOriginalExecutableFileName, bool bInUnstable, DetectProjectSettingsTask DetectSettings, LineBasedTextWriter InLog, UserSettings InSettings)
 		{
 			InitializeComponent();
 
@@ -241,6 +243,7 @@ namespace UnrealGameSync
 			ApiUrl = InApiUrl;
 			DataFolder = DetectSettings.DataFolder;
 			OriginalExecutableFileName = InOriginalExecutableFileName;
+			bUnstable = bInUnstable;
 			Log = InLog;
 			Settings = InSettings;
 			WorkspaceSettings = InSettings.FindOrAddWorkspace(DetectSettings.BranchClientPath);
@@ -4218,8 +4221,11 @@ namespace UnrealGameSync
 
 		private void OptionsContextMenu_PerforceSettings_Click(object sender, EventArgs e)
 		{
-			PerforceSyncSettingsWindow Perforce = new PerforceSyncSettingsWindow(Settings);
-			Perforce.ShowDialog();
+			PerforceSettingsWindow Perforce = new PerforceSettingsWindow(bUnstable, Settings, Log);
+			if(Perforce.ShowDialog() == DialogResult.OK)
+			{
+				Owner.RequestRestart();
+			}
 		}
 
 		private void OptionsContextMenu_TabNames_Stream_Click(object sender, EventArgs e)
diff --git a/Engine/Source/Programs/UnrealGameSync/UnrealGameSync/Forms/MainWindow.cs b/Engine/Source/Programs/UnrealGameSync/UnrealGameSync/Forms/MainWindow.cs
--- a/Engine/Source/Programs/UnrealGameSync/UnrealGameSync/Forms/MainWindow.cs
+++ b/Engine/Source/Programs/UnrealGameSync/UnrealGameSync/Forms/MainWindow.cs
@@ -9,9 +9,15 @@ using System.Linq;
 using System.Runtime.InteropServices;
 using System.Windows.Forms;
 using System.Threading;
+using System.Reflection;
 
 namespace UnrealGameSync
 {
+	interface IMainWindowOwner
+	{
+		void RequestRestart();
+	}
+
 	interface IMainWindowTabPanel : IDisposable
 	{
 		void Activate();
@@ -53,6 +59,7 @@ namespace UnrealGameSync
 
 		private const int WM_SETREDRAW = 11; 
 
+		IMainWindowOwner MainWindowOwner;
 		SynchronizationContext MainThreadSynchronizationContext;
 
 		string ApiUrl;
@@ -70,18 +77,21 @@ namespace UnrealGameSync
 		System.Threading.Timer ScheduleSettledTimer;
 
 		string OriginalExecutableFileName;
+		bool bUnstable;
 
 		IMainWindowTabPanel CurrentTabPanel;
 
-		public MainWindow(string InApiUrl, string InDataFolder, bool bInRestoreStateOnLoad, string InOriginalExecutableFileName, List<DetectProjectSettingsResult> StartupProjects, LineBasedTextWriter InLog, UserSettings InSettings)
+		public MainWindow(IMainWindowOwner Owner, string InApiUrl, string InDataFolder, bool bInRestoreStateOnLoad, string InOriginalExecutableFileName, bool bInUnstable, List<DetectProjectSettingsResult> StartupProjects, LineBasedTextWriter InLog, UserSettings InSettings)
 		{
 			InitializeComponent();
 
+			MainWindowOwner = Owner;
 			MainThreadSynchronizationContext = SynchronizationContext.Current;
 			ApiUrl = InApiUrl;
 			DataFolder = InDataFolder;
 			bRestoreStateOnLoad = bInRestoreStateOnLoad;
 			OriginalExecutableFileName = InOriginalExecutableFileName;
+			bUnstable = bInUnstable;
 			Log = InLog;
 			Settings = InSettings;
 
@@ -124,6 +134,11 @@ namespace UnrealGameSync
 			}
 
 			StartScheduleTimer();
+
+			if(bUnstable)
+			{
+				Text += String.Format(" (UNSTABLE BUILD {0})", Assembly.GetExecutingAssembly().GetName().Version);
+			}
 		}
 
 		protected override void OnHandleCreated(EventArgs e)
@@ -753,7 +768,7 @@ namespace UnrealGameSync
 			}
 
 			// Now that we have the project settings, we can construct the tab
-			WorkspaceControl NewWorkspace = new WorkspaceControl(this, ApiUrl, OriginalExecutableFileName, ProjectSettings, Log, Settings);
+			WorkspaceControl NewWorkspace = new WorkspaceControl(this, ApiUrl, OriginalExecutableFileName, bUnstable, ProjectSettings, Log, Settings);
 			NewWorkspace.Parent = TabPanel;
 			NewWorkspace.Location = new Point(0, 0);
 			NewWorkspace.Size = new Size(TabPanel.Width, TabPanel.Height);
@@ -945,5 +960,10 @@ namespace UnrealGameSync
 				}
 			}
 		}
+
+		public void RequestRestart()
+		{
+			MainWindowOwner.RequestRestart();
+		}
 	}
 }
diff --git a/Engine/Source/Programs/UnrealGameSync/UnrealGameSync/Forms/PerforceSettingsWindow.cs b/Engine/Source/Programs/UnrealGameSync/UnrealGameSync/Forms/PerforceSettingsWindow.cs
new file mode 100644
--- /dev/null
+++ b/Engine/Source/Programs/UnrealGameSync/UnrealGameSync/Forms/PerforceSettingsWindow.cs
@@ -0,0 +1,144 @@
+// Copyright 1998-2018 Epic Games, Inc. All Rights Reserved.
+
+using System;
+using System.Collections.Generic;
+using System.ComponentModel;
+using System.Data;
+using System.Drawing;
+using System.IO;
+using System.Linq;
+using System.Runtime.InteropServices;
+using System.Text;
+using System.Threading;
+using System.Threading.Tasks;
+using System.Windows.Forms;
+using UnrealGameSync;
+
+namespace UnrealGameSync
+{
+	partial class PerforceSettingsWindow : Form
+	{
+		class PerforceTestConnectionTask : IPerforceModalTask
+		{
+			string DepotPath;
+
+			public PerforceTestConnectionTask(string DepotPath)
+			{
+				this.DepotPath = DepotPath ?? DeploymentSettings.DefaultDepotPath;
+			}
+
+			public bool Run(PerforceConnection Perforce, TextWriter Log, out string ErrorMessage)
+			{
+				string CheckFilePath = String.Format("{0}/Release/UnrealGameSync.exe", DepotPath);
+
+				List<PerforceFileRecord> FileRecords;
+				if(!Perforce.FindFiles(CheckFilePath, out FileRecords, Log) || FileRecords.Count == 0)
+				{
+					ErrorMessage = String.Format("Unable to find {0}", CheckFilePath);
+					return false;
+				}
+
+				ErrorMessage = null;
+				return true;
+			}
+		}
+
+		UserSettings Settings;
+		TextWriter Log;
+
+		string InitialServerAndPort;
+		string InitialUserName;
+		string InitialDepotPath;
+		bool bInitialUnstable;
+
+		public PerforceSettingsWindow(bool bUnstable, UserSettings Settings, TextWriter Log)
+		{
+			InitializeComponent();
+
+			this.Settings = Settings;
+			this.Log = Log;
+
+			Utility.ReadGlobalPerforceSettings(ref InitialServerAndPort, ref InitialUserName, ref InitialDepotPath);
+			bInitialUnstable = bUnstable;
+
+			this.ServerTextBox.Text = InitialServerAndPort;
+			this.ServerTextBox.Select(ServerTextBox.TextLength, 0);
+
+			this.UserNameTextBox.Text = InitialUserName;
+			this.UserNameTextBox.Select(UserNameTextBox.TextLength, 0);
+
+			this.DepotPathTextBox.Text = InitialDepotPath;
+			this.DepotPathTextBox.Select(DepotPathTextBox.TextLength, 0);
+			this.DepotPathTextBox.CueBanner = DeploymentSettings.DefaultDepotPath;
+
+			this.UseUnstableBuildCheckBox.Checked = bUnstable;
+		}
+
+		private void OkBtn_Click(object sender, EventArgs e)
+		{
+			// Update the settings
+			string ServerAndPort = ServerTextBox.Text.Trim();
+			if(ServerAndPort.Length == 0)
+			{
+				ServerAndPort = null;
+			}
+
+			string UserName = UserNameTextBox.Text.Trim();
+			if(UserName.Length == 0)
+			{
+				UserName = null;
+			}
+
+			string DepotPath = DepotPathTextBox.Text.Trim();
+			if(DepotPath.Length == 0 || DepotPath == DeploymentSettings.DefaultDepotPath)
+			{
+				DepotPath = null;
+			}
+
+			bool bUnstable = UseUnstableBuildCheckBox.Checked;
+
+			if(ServerAndPort == InitialServerAndPort && UserName == InitialUserName && DepotPath == InitialDepotPath && bUnstable == bInitialUnstable)
+			{
+				DialogResult = DialogResult.Cancel;
+				Close();
+			}
+			else
+			{
+				// Try to log in to the new server, and check the application is there
+				if(ServerAndPort != InitialServerAndPort || UserName != InitialUserName || DepotPath != InitialDepotPath)
+				{
+					string ErrorMessage;
+					ModalTaskResult Result = PerforceModalTask.Execute(this, null, ServerAndPort, UserName, new PerforceTestConnectionTask(DepotPath), "Connecting", "Checking connection, please wait...", Log, out ErrorMessage);
+					if(Result != ModalTaskResult.Succeeded)
+					{
+						if(Result == ModalTaskResult.Failed)
+						{
+							MessageBox.Show(ErrorMessage, "Unable to connect");
+						}
+						return;
+					}
+				}
+
+				if(MessageBox.Show("UnrealGameSync must be restarted to apply these settings.\n\nWould you like to restart now?", "Restart Required", MessageBoxButtons.OKCancel) == DialogResult.OK)
+				{
+					Utility.SaveGlobalPerforceSettings(ServerAndPort, UserName, DepotPath);
+
+					DialogResult = DialogResult.OK;
+					Close();
+				}
+			}
+		}
+
+		private void CancelBtn_Click(object sender, EventArgs e)
+		{
+			DialogResult = DialogResult.Cancel;
+			Close();
+		}
+
+		private void AdvancedBtn_Click(object sender, EventArgs e)
+		{
+			PerforceSyncSettingsWindow Window = new PerforceSyncSettingsWindow(Settings);
+			Window.ShowDialog();
+		}
+	}
+}
diff --git a/Engine/Source/Programs/UnrealGameSync/UnrealGameSync/ProgramApplicationContext.cs b/Engine/Source/Programs/UnrealGameSync/UnrealGameSync/ProgramApplicationContext.cs
--- a/Engine/Source/Programs/UnrealGameSync/UnrealGameSync/ProgramApplicationContext.cs
+++ b/Engine/Source/Programs/UnrealGameSync/UnrealGameSync/ProgramApplicationContext.cs
@@ -14,7 +14,7 @@ using System.Windows.Forms;
 
 namespace UnrealGameSync
 {
-	class ProgramApplicationContext : ApplicationContext
+	class ProgramApplicationContext : ApplicationContext, IMainWindowOwner
 	{
 		SynchronizationContext MainThreadSynchronizationContext;
 
@@ -178,11 +178,7 @@ namespace UnrealGameSync
 			DetectStartupProjectSettingsWindow = null;
 
 			// Create the main window
-			MainWindowInstance = new MainWindow(ApiUrl, DataFolder, bRestoreState, UpdateSpawn ?? Assembly.GetExecutingAssembly().Location, DetectStartupProjectSettingsTask.Results, Log, Settings);
-			if(bUnstable)
-			{
-				MainWindowInstance.Text += String.Format(" (UNSTABLE BUILD {0})", Assembly.GetExecutingAssembly().GetName().Version);
-			}
+			MainWindowInstance = new MainWindow(this, ApiUrl, DataFolder, bRestoreState, UpdateSpawn ?? Assembly.GetExecutingAssembly().Location, bUnstable, DetectStartupProjectSettingsTask.Results, Log, Settings);
 			if(bVisible)
 			{
 				MainWindowInstance.Show();
@@ -198,6 +194,11 @@ namespace UnrealGameSync
 			DetectStartupProjectSettingsTask = null;
 		}
 
+		public void RequestRestart()
+		{
+			UpdateMonitor.TriggerUpdate(true);
+		}
+
 		private void MainWindowInstance_FormClosed(object sender, FormClosedEventArgs e)
 		{
 			ExitThread();
@@ -208,7 +209,7 @@ namespace UnrealGameSync
 			// Check if we're trying to reopen with the unstable version; if so, trigger an update to trigger a restart with the new executable
 			if(!bUnstable && (Control.ModifierKeys & Keys.Shift) != 0)
 			{
-				UpdateMonitor.TriggerUpdate();
+				UpdateMonitor.TriggerUpdate(true);
 			}
 			else if(MainWindowInstance != null)
 			{
@@ -221,18 +222,21 @@ namespace UnrealGameSync
 			MainThreadSynchronizationContext.Post((o) => OnActivationListenerCallback(), null);
 		}
 
-		private void OnUpdateAvailable()
+		private void OnUpdateAvailable(bool bForce)
 		{
-			if(MainWindowInstance != null && !bIsClosing && MainWindowInstance.CanPerformUpdate())
+			if(MainWindowInstance != null && !bIsClosing)
 			{
-				bIsClosing = true;
-				MainWindowInstance.ForceClose();
+				if(bForce || MainWindowInstance.CanPerformUpdate())
+				{
+					bIsClosing = true;
+					MainWindowInstance.ForceClose();
+				}
 			}
 		}
 
-		private void OnUpdateAvailableCallback()
+		private void OnUpdateAvailableCallback(bool bForce)
 		{ 
-			MainThreadSynchronizationContext.Post((o) => OnUpdateAvailable(), null);
+			MainThreadSynchronizationContext.Post((o) => OnUpdateAvailable(bForce), null);
 		}
 
 		protected override void Dispose(bool bDisposing)
diff --git a/Engine/Source/Programs/UnrealGameSync/UnrealGameSync/UnrealGameSync.csproj b/Engine/Source/Programs/UnrealGameSync/UnrealGameSync/UnrealGameSync.csproj
--- a/Engine/Source/Programs/UnrealGameSync/UnrealGameSync/UnrealGameSync.csproj
+++ b/Engine/Source/Programs/UnrealGameSync/UnrealGameSync/UnrealGameSync.csproj
@@ -207,6 +207,12 @@
     <Compile Include="Forms\PasswordWindow.Designer.cs">
       <DependentUpon>PasswordWindow.cs</DependentUpon>
     </Compile>
+    <Compile Include="Forms\PerforceSettingsWindow.cs">
+      <SubType>Form</SubType>
+    </Compile>
+    <Compile Include="Forms\PerforceSettingsWindow.Designer.cs">
+      <DependentUpon>PerforceSettingsWindow.cs</DependentUpon>
+    </Compile>
     <Compile Include="Forms\PerforceSyncSettingsWindow.cs">
       <SubType>Form</SubType>
     </Compile>
@@ -371,6 +377,9 @@
     <EmbeddedResource Include="Forms\PasswordWindow.resx">
       <DependentUpon>PasswordWindow.cs</DependentUpon>
     </EmbeddedResource>
+    <EmbeddedResource Include="Forms\PerforceSettingsWindow.resx">
+      <DependentUpon>PerforceSettingsWindow.cs</DependentUpon>
+    </EmbeddedResource>
     <EmbeddedResource Include="Forms\PerforceSyncSettingsWindow.resx">
       <DependentUpon>PerforceSyncSettingsWindow.cs</DependentUpon>
     </EmbeddedResource>
diff --git a/Engine/Source/Programs/UnrealGameSync/UnrealGameSync/UpdateMonitor.cs b/Engine/Source/Programs/UnrealGameSync/UnrealGameSync/UpdateMonitor.cs
--- a/Engine/Source/Programs/UnrealGameSync/UnrealGameSync/UpdateMonitor.cs
+++ b/Engine/Source/Programs/UnrealGameSync/UnrealGameSync/UpdateMonitor.cs
@@ -1,4 +1,4 @@
-﻿// Copyright 1998-2018 Epic Games, Inc. All Rights Reserved.
+// Copyright 1998-2018 Epic Games, Inc. All Rights Reserved.
 
 using System;
 using System.Collections.Generic;
@@ -17,7 +17,7 @@ namespace UnrealGameSync
 		Thread WorkerThread;
 		ManualResetEvent QuitEvent;
 
-		public event Action OnUpdateAvailable;
+		public event Action<bool> OnUpdateAvailable;
 
 		public UpdateMonitor(PerforceConnection InPerforce, string InWatchPath)
 		{
@@ -68,17 +68,17 @@ namespace UnrealGameSync
 				List<PerforceChangeSummary> Changes;
 				if(Perforce.FindChanges(WatchPath, 1, out Changes, Log) && Changes.Count > 0)
 				{
-					TriggerUpdate();
+					TriggerUpdate(false);
 				}
 			}
 		}
 
-		public void TriggerUpdate()
+		public void TriggerUpdate(bool bForce)
 		{
 			IsUpdateAvailable = true;
 			if(OnUpdateAvailable != null)
 			{
-				OnUpdateAvailable();
+				OnUpdateAvailable(bForce);
 			}
 		}
 	}
diff --git a/Engine/Source/Programs/UnrealGameSync/UnrealGameSync/Utility.cs b/Engine/Source/Programs/UnrealGameSync/UnrealGameSync/Utility.cs
--- a/Engine/Source/Programs/UnrealGameSync/UnrealGameSync/Utility.cs
+++ b/Engine/Source/Programs/UnrealGameSync/UnrealGameSync/Utility.cs
@@ -2,6 +2,7 @@
 
 #define USE_NEW_PROCESS_JOBS
 
+using Microsoft.Win32;
 using System;
 using System.Collections.Generic;
 using System.Diagnostics;
@@ -11,6 +12,7 @@ using System.Runtime.InteropServices;
 using System.Text;
 using System.Threading.Tasks;
 using System.Web.Script.Serialization;
+using System.Windows.Forms;
 
 namespace UnrealGameSync
 {
@@ -289,5 +291,61 @@ namespace UnrealGameSync
 			}
 			return ProjectConfigFileNames;
 		}
+
+		public static void ReadGlobalPerforceSettings(ref string ServerAndPort, ref string UserName, ref string DepotPath)
+		{
+			using (RegistryKey Key = Registry.CurrentUser.OpenSubKey("SOFTWARE\\Epic Games\\UnrealGameSync", false))
+			{
+				if(Key != null)
+				{
+					ServerAndPort = Key.GetValue("ServerAndPort", ServerAndPort) as string;
+					UserName = Key.GetValue("UserName", UserName) as string;
+					DepotPath = Key.GetValue("DepotPath", DepotPath) as string;
+				}
+			}
+		}
+
+		public static void SaveGlobalPerforceSettings(string ServerAndPort, string UserName, string DepotPath)
+		{
+			try
+			{
+				using (RegistryKey Key = Registry.CurrentUser.CreateSubKey("SOFTWARE\\Epic Games\\UnrealGameSync"))
+				{
+					// Delete this legacy setting
+					try { Key.DeleteValue("Server"); } catch(Exception) { }
+
+					if(String.IsNullOrEmpty(ServerAndPort))
+					{
+						try { Key.DeleteValue("ServerAndPort"); } catch(Exception) { }
+					}
+					else
+					{
+						Key.SetValue("ServerAndPort", ServerAndPort);
+					}
+
+					if(String.IsNullOrEmpty(UserName))
+					{
+						try { Key.DeleteValue("UserName"); } catch(Exception) { }
+					}
+					else
+					{
+						Key.SetValue("UserName", UserName);
+					}
+
+					if(String.IsNullOrEmpty(DepotPath) || (DeploymentSettings.DefaultDepotPath != null && String.Equals(DepotPath, DeploymentSettings.DefaultDepotPath, StringComparison.InvariantCultureIgnoreCase)))
+					{
+						try { Key.DeleteValue("DepotPath"); } catch(Exception) { }
+					}
+					else
+					{
+						Key.SetValue("DepotPath", DepotPath);
+					}
+				}
+			}
+			catch(Exception Ex)
+			{
+				MessageBox.Show("Unable to save settings.\n\n" + Ex.ToString());
+			}
+		}
 	}
 }
diff --git a/Engine/Source/Programs/UnrealGameSync/UnrealGameSyncLauncher/Forms/SettingsWindow.cs b/Engine/Source/Programs/UnrealGameSync/UnrealGameSyncLauncher/Forms/SettingsWindow.cs
--- a/Engine/Source/Programs/UnrealGameSync/UnrealGameSyncLauncher/Forms/SettingsWindow.cs
+++ b/Engine/Source/Programs/UnrealGameSync/UnrealGameSyncLauncher/Forms/SettingsWindow.cs
@@ -82,7 +82,7 @@ namespace UnrealGameSyncLauncher
 				DepotPath = null;
 			}
 
-			Program.SaveSettings(ServerAndPort, UserName, DepotPath);
+			Utility.SaveGlobalPerforceSettings(ServerAndPort, UserName, DepotPath);
 
 			// Create the task for connecting to this server
 			StringWriter Log = new StringWriter();
@@ -93,7 +93,7 @@ namespace UnrealGameSyncLauncher
 			ModalTaskResult Result = PerforceModalTask.Execute(this, null, ServerAndPort, UserName, SyncApplication, "Updating", "Checking for updates, please wait...", Log, out ErrorMessage);
 			if(Result == ModalTaskResult.Succeeded)
 			{
-				Program.SaveSettings(ServerAndPort, UserName, DepotPath);
+				Utility.SaveGlobalPerforceSettings(ServerAndPort, UserName, DepotPath);
 				DialogResult = DialogResult.OK;
 				Close();
 			}
diff --git a/Engine/Source/Programs/UnrealGameSync/UnrealGameSyncLauncher/Program.cs b/Engine/Source/Programs/UnrealGameSync/UnrealGameSyncLauncher/Program.cs
--- a/Engine/Source/Programs/UnrealGameSync/UnrealGameSyncLauncher/Program.cs
+++ b/Engine/Source/Programs/UnrealGameSync/UnrealGameSyncLauncher/Program.cs
@@ -65,7 +65,7 @@ namespace UnrealGameSyncLauncher
 				string ServerAndPort = null;
 				string UserName = null;
 				string DepotPath = DeploymentSettings.DefaultDepotPath;
-				ReadSettings(ref ServerAndPort, ref UserName, ref DepotPath);
+				Utility.ReadGlobalPerforceSettings(ref ServerAndPort, ref UserName, ref DepotPath);
 
 				// If the shift key is held down, immediately show the settings window
 				if((Control.ModifierKeys & Keys.Shift) != 0)
@@ -98,62 +98,6 @@ namespace UnrealGameSyncLauncher
 			return 1;
 		}
 
-		public static void ReadSettings(ref string ServerAndPort, ref string UserName, ref string DepotPath)
-		{
-			using (RegistryKey Key = Registry.CurrentUser.OpenSubKey("SOFTWARE\\Epic Games\\UnrealGameSync", false))
-			{
-				if(Key != null)
-				{
-					ServerAndPort = Key.GetValue("ServerAndPort", ServerAndPort) as string;
-					UserName = Key.GetValue("UserName", UserName) as string;
-					DepotPath = Key.GetValue("DepotPath", DepotPath) as string;
-				}
-			}
-		}
-
-		public static void SaveSettings(string ServerAndPort, string UserName, string DepotPath)
-		{
-			try
-			{
-				using (RegistryKey Key = Registry.CurrentUser.CreateSubKey("SOFTWARE\\Epic Games\\UnrealGameSync"))
-				{
-					// Delete this legacy setting
-					try { Key.DeleteValue("Server"); } catch(Exception) { }
-
-					if(String.IsNullOrEmpty(ServerAndPort))
-					{
-						try { Key.DeleteValue("ServerAndPort"); } catch(Exception) { }
-					}
-					else
-					{
-						Key.SetValue("ServerAndPort", ServerAndPort);
-					}
-
-					if(String.IsNullOrEmpty(UserName))
-					{
-						try { Key.DeleteValue("UserName"); } catch(Exception) { }
-					}
-					else
-					{
-						Key.SetValue("UserName", UserName);
-					}
-
-					if(String.IsNullOrEmpty(DepotPath) || (DeploymentSettings.DefaultDepotPath != null && String.Equals(DepotPath, DeploymentSettings.DefaultDepotPath, StringComparison.InvariantCultureIgnoreCase)))
-					{
-						try { Key.DeleteValue("DepotPath"); } catch(Exception) { }
-					}
-					else
-					{
-						Key.SetValue("DepotPath", DepotPath);
-					}
-				}
-			}
-			catch(Exception Ex)
-			{
-				MessageBox.Show("Unable to save settings.\n\n" + Ex.ToString());
-			}
-		}
-
 		public static bool SyncAndRun(PerforceConnection Perforce, string BaseDepotPath, bool bUnstable, string[] Args, Mutex InstanceMutex, TextWriter LogWriter)
 		{
 			try
