UGS: Fix unstable checkbox not being respected when restarting.

#rnx

diff --git a/Engine/Source/Programs/UnrealGameSync/UnrealGameSync/Controls/WorkspaceControl.cs b/Engine/Source/Programs/UnrealGameSync/UnrealGameSync/Controls/WorkspaceControl.cs
--- a/Engine/Source/Programs/UnrealGameSync/UnrealGameSync/Controls/WorkspaceControl.cs
+++ b/Engine/Source/Programs/UnrealGameSync/UnrealGameSync/Controls/WorkspaceControl.cs
@@ -41,7 +41,7 @@ namespace UnrealGameSync
 		void SetTabNames(TabLabels TabNames);
 		void SetupScheduledSync();
 		void UpdateProgress();
-		void RequestRestart();
+		void RequestRestart(bool? bUnstableBuild);
 	}
 
 	partial class WorkspaceControl : UserControl, IMainWindowTabPanel
@@ -4224,7 +4224,7 @@ namespace UnrealGameSync
 			PerforceSettingsWindow Perforce = new PerforceSettingsWindow(bUnstable, Settings, Log);
 			if(Perforce.ShowDialog() == DialogResult.OK)
 			{
-				Owner.RequestRestart();
+				Owner.RequestRestart(Perforce.UseUnstableBuildCheckBox.Checked);
 			}
 		}
 
diff --git a/Engine/Source/Programs/UnrealGameSync/UnrealGameSync/Forms/MainWindow.cs b/Engine/Source/Programs/UnrealGameSync/UnrealGameSync/Forms/MainWindow.cs
--- a/Engine/Source/Programs/UnrealGameSync/UnrealGameSync/Forms/MainWindow.cs
+++ b/Engine/Source/Programs/UnrealGameSync/UnrealGameSync/Forms/MainWindow.cs
@@ -15,7 +15,7 @@ namespace UnrealGameSync
 {
 	interface IMainWindowOwner
 	{
-		void RequestRestart();
+		void RequestRestart(bool? bUnstableBuild);
 	}
 
 	interface IMainWindowTabPanel : IDisposable
@@ -961,9 +961,9 @@ namespace UnrealGameSync
 			}
 		}
 
-		public void RequestRestart()
+		public void RequestRestart(bool? bUnstableBuild)
 		{
-			MainWindowOwner.RequestRestart();
+			MainWindowOwner.RequestRestart(bUnstableBuild);
 		}
 	}
 }
diff --git a/Engine/Source/Programs/UnrealGameSync/UnrealGameSync/Program.cs b/Engine/Source/Programs/UnrealGameSync/UnrealGameSync/Program.cs
--- a/Engine/Source/Programs/UnrealGameSync/UnrealGameSync/Program.cs
+++ b/Engine/Source/Programs/UnrealGameSync/UnrealGameSync/Program.cs
@@ -97,7 +97,8 @@ namespace UnrealGameSync
 					if(UpdateMonitor.IsUpdateAvailable && UpdateSpawn != null)
 					{
 						InstanceMutex.Close();
-						Utility.SpawnProcess(UpdateSpawn, "-restorestate" + (bUnstable? " -unstable" : ""));
+						bool bLaunchUnstable = Context.bRelaunchUnstable ?? bUnstable;
+						Utility.SpawnProcess(UpdateSpawn, "-restorestate" + (bLaunchUnstable? " -unstable" : ""));
 					}
 				}
 			}
diff --git a/Engine/Source/Programs/UnrealGameSync/UnrealGameSync/ProgramApplicationContext.cs b/Engine/Source/Programs/UnrealGameSync/UnrealGameSync/ProgramApplicationContext.cs
--- a/Engine/Source/Programs/UnrealGameSync/UnrealGameSync/ProgramApplicationContext.cs
+++ b/Engine/Source/Programs/UnrealGameSync/UnrealGameSync/ProgramApplicationContext.cs
@@ -26,6 +26,8 @@ namespace UnrealGameSync
 		bool bUnstable;
 		bool bIsClosing;
 
+		public bool? bRelaunchUnstable;
+
 		TimestampLogWriter Log;
 		UserSettings Settings;
 		ActivationListener ActivationListener;
@@ -194,8 +196,9 @@ namespace UnrealGameSync
 			DetectStartupProjectSettingsTask = null;
 		}
 
-		public void RequestRestart()
+		public void RequestRestart(bool? bRelaunchUnstable)
 		{
+			this.bRelaunchUnstable = bRelaunchUnstable;
 			UpdateMonitor.TriggerUpdate(true);
 		}
 
