UGS: Improve startup time (especially when working remotely)

* Always use the cached workspace name rather than trying to detect it each time.
* Try to find a workspace owned by the current user before falling back to checking for shared workspaces. This significantly reduces the number of results that need to be returned from the server.

diff --git a/Engine/Source/Programs/UnrealGameSync/UnrealGameSync/DetectProjectSettingsTask.cs b/Engine/Source/Programs/UnrealGameSync/UnrealGameSync/DetectProjectSettingsTask.cs
--- a/Engine/Source/Programs/UnrealGameSync/UnrealGameSync/DetectProjectSettingsTask.cs
+++ b/Engine/Source/Programs/UnrealGameSync/UnrealGameSync/DetectProjectSettingsTask.cs
@@ -97,8 +97,8 @@ namespace UnrealGameSync
 			// Configure the time zone
 			ServerTimeZone = PerforceInfo.ServerTimeZone;
 
-			// If we're using the legacy path of specifying a file, figure out the workspace name now
-			if(SelectedProject.Type == UserSelectedProjectType.Client)
+			// Use the cached client path to the file if it's available; it's much quicker than trying to find the correct workspace.
+			if(!String.IsNullOrEmpty(SelectedProject.ClientPath))
 			{
 				// Get the client path
 				NewSelectedClientFileName = SelectedProject.ClientPath;
@@ -121,7 +121,7 @@ namespace UnrealGameSync
 					return false;
 				}
 			}
-			else if(SelectedProject.Type == UserSelectedProjectType.Local)
+			else
 			{
 				// Use the path as the selected filename
 				NewSelectedFileName = SelectedProject.LocalPath;
@@ -133,67 +133,40 @@ namespace UnrealGameSync
 					return false;
 				}
 
-				// Find all the clients on this machine
-				Log.WriteLine("Enumerating clients on local machine...");
+				// Find all the clients for this user
+				Log.WriteLine("Enumerating clients for {0}...", PerforceInfo.UserName);
+
 				List<PerforceClientRecord> Clients;
-				if(!Perforce.FindClients(out Clients, Log))
+				if(!Perforce.FindClients(PerforceInfo.UserName, out Clients, Log))
 				{
 					ErrorMessage = String.Format("Couldn't find any clients for this host.");
 					return false;
 				}
 
-				// Find any clients which are valid. If this is not exactly one, we should fail.
-				List<PerforceConnection> CandidateClients = new List<PerforceConnection>();
-				foreach(PerforceClientRecord Client in Clients)
+				List<PerforceConnection> CandidateClients = FilterClients(Clients, Perforce.ServerAndPort, PerforceInfo.HostName, PerforceInfo.UserName);
+				if(CandidateClients.Count == 0)
 				{
-					// Make sure the client is well formed
-					if(!String.IsNullOrEmpty(Client.Name) && (!String.IsNullOrEmpty(Client.Host) || !String.IsNullOrEmpty(Client.Owner)) && !String.IsNullOrEmpty(Client.Root))
+					// Search through all workspaces. We may find a suitable workspace which is for any user.
+					Log.WriteLine("Enumerating shared clients...");
+					if(!Perforce.FindClients("", out Clients, Log))
 					{
-						// Require either a username or host name match
-						if((String.IsNullOrEmpty(Client.Host) || String.Compare(Client.Host, PerforceInfo.HostName, StringComparison.InvariantCultureIgnoreCase) == 0) && (String.IsNullOrEmpty(Client.Owner) || String.Compare(Client.Owner, PerforceInfo.UserName, StringComparison.InvariantCultureIgnoreCase) == 0))
-						{
-							if(!Utility.SafeIsFileUnderDirectory(NewSelectedFileName, Client.Root))
-							{
-								Log.WriteLine("Rejecting {0} due to root mismatch ({1})", Client.Name, Client.Root);
-								continue;
-							}
-
-							PerforceConnection CandidateClient = new PerforceConnection(PerforceInfo.UserName, Client.Name, Perforce.ServerAndPort);
-
-							bool bFileExists;
-							if(!CandidateClient.FileExists(NewSelectedFileName, out bFileExists, Log) || !bFileExists)
-							{
-								Log.WriteLine("Rejecting {0} due to file not existing in workspace", Client.Name);
-								continue;
-							}
-
-							List<PerforceFileRecord> Records;
-							if(!CandidateClient.Stat(NewSelectedFileName, out Records, Log))
-							{
-								Log.WriteLine("Rejecting {0} due to {1} not in depot", Client.Name, NewSelectedFileName);
-								continue;
-							}
-
-							Records.RemoveAll(x => !x.IsMapped);
-							if(Records.Count == 0)
-							{
-								Log.WriteLine("Rejecting {0} due to {1} matching records", Client.Name, Records.Count);
-								continue;
-							}
-
-							Log.WriteLine("Found valid client {0}", Client.Name);
-							CandidateClients.Add(CandidateClient);
-						}
+						ErrorMessage = "Failed to enumerate clients.";
+						return false;
+					}
+
+					// Filter this list of clients
+					CandidateClients = FilterClients(Clients, Perforce.ServerAndPort, PerforceInfo.HostName, PerforceInfo.UserName);
+
+					// If we still couldn't find any, fail.
+					if(CandidateClients.Count == 0)
+					{
+						ErrorMessage = String.Format("Couldn't find any Perforce workspace containing {0}. Check your connection settings.", NewSelectedFileName);
+						return false;
 					}
 				}
 
 				// Check there's only one client
-				if(CandidateClients.Count == 0)
-				{
-					ErrorMessage = String.Format("Couldn't find any Perforce workspace containing {0}. Check your connection settings.", NewSelectedFileName);
-					return false;
-				}
-				else if(CandidateClients.Count > 1)
+				if(CandidateClients.Count > 1)
 				{
 					ErrorMessage = String.Format("Found multiple workspaces containing {0}:\n\n{1}\n\nCannot determine which to use.", Path.GetFileName(NewSelectedFileName), String.Join("\n", CandidateClients.Select(x => x.ClientName)));
 					return false;
@@ -209,10 +182,6 @@ namespace UnrealGameSync
 					return false;
 				}
 			}
-			else
-			{
-				throw new InvalidDataException("Invalid selected project type");
-			}
 
 			// Normalize the filename
 			NewSelectedFileName = Path.GetFullPath(NewSelectedFileName).Replace('/', Path.DirectorySeparatorChar);
@@ -336,6 +305,54 @@ namespace UnrealGameSync
 			return true;
 		}
 
+		List<PerforceConnection> FilterClients(List<PerforceClientRecord> Clients, string ServerAndPort, string HostName, string UserName)
+		{
+			List<PerforceConnection> CandidateClients = new List<PerforceConnection>();
+			foreach(PerforceClientRecord Client in Clients)
+			{
+				// Make sure the client is well formed
+				if(!String.IsNullOrEmpty(Client.Name) && (!String.IsNullOrEmpty(Client.Host) || !String.IsNullOrEmpty(Client.Owner)) && !String.IsNullOrEmpty(Client.Root))
+				{
+					// Require either a username or host name match
+					if((String.IsNullOrEmpty(Client.Host) || String.Compare(Client.Host, HostName, StringComparison.InvariantCultureIgnoreCase) == 0) && (String.IsNullOrEmpty(Client.Owner) || String.Compare(Client.Owner, UserName, StringComparison.InvariantCultureIgnoreCase) == 0))
+					{
+						if(!Utility.SafeIsFileUnderDirectory(NewSelectedFileName, Client.Root))
+						{
+							Log.WriteLine("Rejecting {0} due to root mismatch ({1})", Client.Name, Client.Root);
+							continue;
+						}
+
+						PerforceConnection CandidateClient = new PerforceConnection(UserName, Client.Name, ServerAndPort);
+
+						bool bFileExists;
+						if(!CandidateClient.FileExists(NewSelectedFileName, out bFileExists, Log) || !bFileExists)
+						{
+							Log.WriteLine("Rejecting {0} due to file not existing in workspace", Client.Name);
+							continue;
+						}
+
+						List<PerforceFileRecord> Records;
+						if(!CandidateClient.Stat(NewSelectedFileName, out Records, Log))
+						{
+							Log.WriteLine("Rejecting {0} due to {1} not in depot", Client.Name, NewSelectedFileName);
+							continue;
+						}
+
+						Records.RemoveAll(x => !x.IsMapped);
+						if(Records.Count == 0)
+						{
+							Log.WriteLine("Rejecting {0} due to {1} matching records", Client.Name, Records.Count);
+							continue;
+						}
+
+						Log.WriteLine("Found valid client {0}", Client.Name);
+						CandidateClients.Add(CandidateClient);
+					}
+				}
+			}
+			return CandidateClients;
+		}
+
 		bool TryGetStreamPrefix(PerforceConnection Perforce, string StreamName, TextWriter Log, out string StreamPrefix)
 		{ 
 			string CurrentStreamName = StreamName;
