Missing file.

#rnx

diff --git a/Engine/Source/Programs/UnrealGameSync/UnrealGameSync/Perforce.cs b/Engine/Source/Programs/UnrealGameSync/UnrealGameSync/Perforce.cs
--- a/Engine/Source/Programs/UnrealGameSync/UnrealGameSync/Perforce.cs
+++ b/Engine/Source/Programs/UnrealGameSync/UnrealGameSync/Perforce.cs
@@ -599,13 +599,13 @@ namespace UnrealGameSync
 
 		public bool FindClients(string ForUserName, out List<PerforceClientRecord> Clients, TextWriter Log)
 		{
-			return RunCommand(String.Format("clients -u{0}", ForUserName), out Clients, CommandOptions.NoClient, Log);
+			return RunCommand(String.Format("clients -u \"{0}\"", ForUserName), out Clients, CommandOptions.NoClient, Log);
 		}
 
 		public bool FindClients(string ForUserName, out List<string> ClientNames, TextWriter Log)
 		{
 			List<string> Lines;
-			if(!RunCommand(String.Format("clients -u{0}", ForUserName), out Lines, CommandOptions.None, Log))
+			if(!RunCommand(String.Format("clients -u \"{0}\"", ForUserName), out Lines, CommandOptions.None, Log))
 			{
 				ClientNames = null;
 				return false;
