UGS: Move settings for auto-update into a separate "Application Settings" dialog, and fix changes to unstable flag not taking effect.

diff --git a/Engine/Source/Programs/UnrealGameSync/UnrealGameSync/Controls/WorkspaceControl.cs b/Engine/Source/Programs/UnrealGameSync/UnrealGameSync/Controls/WorkspaceControl.cs
--- a/Engine/Source/Programs/UnrealGameSync/UnrealGameSync/Controls/WorkspaceControl.cs
+++ b/Engine/Source/Programs/UnrealGameSync/UnrealGameSync/Controls/WorkspaceControl.cs
@@ -41,7 +41,7 @@ namespace UnrealGameSync
 		void SetTabNames(TabLabels TabNames);
 		void SetupScheduledSync();
 		void UpdateProgress();
-		void RequestRestart(bool? bUnstableBuild);
+		void ModifyApplicationSettings();
 	}
 
 	partial class WorkspaceControl : UserControl, IMainWindowTabPanel
@@ -4219,13 +4219,15 @@ namespace UnrealGameSync
 			}
 		}
 
+		private void OptionsContextMenu_ApplicationSettings_Click(object sender, EventArgs e)
+		{
+			Owner.ModifyApplicationSettings();
+		}
+
 		private void OptionsContextMenu_PerforceSettings_Click(object sender, EventArgs e)
 		{
-			PerforceSettingsWindow Perforce = new PerforceSettingsWindow(bUnstable, Settings, Log);
-			if(Perforce.ShowDialog() == DialogResult.OK)
-			{
-				Owner.RequestRestart(Perforce.UseUnstableBuildCheckBox.Checked);
-			}
+			PerforceSyncSettingsWindow Window = new PerforceSyncSettingsWindow(Settings);
+			Window.ShowDialog();
 		}
 
 		private void OptionsContextMenu_TabNames_Stream_Click(object sender, EventArgs e)
diff --git a/Engine/Source/Programs/UnrealGameSync/UnrealGameSync/Forms/MainWindow.cs b/Engine/Source/Programs/UnrealGameSync/UnrealGameSync/Forms/MainWindow.cs
--- a/Engine/Source/Programs/UnrealGameSync/UnrealGameSync/Forms/MainWindow.cs
+++ b/Engine/Source/Programs/UnrealGameSync/UnrealGameSync/Forms/MainWindow.cs
@@ -13,11 +13,6 @@ using System.Reflection;
 
 namespace UnrealGameSync
 {
-	interface IMainWindowOwner
-	{
-		void RequestRestart(bool? bUnstableBuild);
-	}
-
 	interface IMainWindowTabPanel : IDisposable
 	{
 		void Activate();
@@ -59,7 +54,7 @@ namespace UnrealGameSync
 
 		private const int WM_SETREDRAW = 11; 
 
-		IMainWindowOwner MainWindowOwner;
+		UpdateMonitor UpdateMonitor;
 		SynchronizationContext MainThreadSynchronizationContext;
 
 		string ApiUrl;
@@ -81,11 +76,11 @@ namespace UnrealGameSync
 
 		IMainWindowTabPanel CurrentTabPanel;
 
-		public MainWindow(IMainWindowOwner Owner, string InApiUrl, string InDataFolder, bool bInRestoreStateOnLoad, string InOriginalExecutableFileName, bool bInUnstable, List<DetectProjectSettingsResult> StartupProjects, LineBasedTextWriter InLog, UserSettings InSettings)
+		public MainWindow(UpdateMonitor InUpdateMonitor, string InApiUrl, string InDataFolder, bool bInRestoreStateOnLoad, string InOriginalExecutableFileName, bool bInUnstable, List<DetectProjectSettingsResult> StartupProjects, LineBasedTextWriter InLog, UserSettings InSettings)
 		{
 			InitializeComponent();
 
-			MainWindowOwner = Owner;
+			UpdateMonitor = InUpdateMonitor;
 			MainThreadSynchronizationContext = SynchronizationContext.Current;
 			ApiUrl = InApiUrl;
 			DataFolder = InDataFolder;
@@ -961,9 +956,14 @@ namespace UnrealGameSync
 			}
 		}
 
-		public void RequestRestart(bool? bUnstableBuild)
+		public void ModifyApplicationSettings()
 		{
-			MainWindowOwner.RequestRestart(bUnstableBuild);
+			ApplicationSettingsWindow ApplicationSettings = new ApplicationSettingsWindow(UpdateMonitor.Perforce.ServerAndPort, UpdateMonitor.Perforce.UserName, bUnstable, Log);
+			if(ApplicationSettings.ShowDialog() == DialogResult.OK)
+			{
+				UpdateMonitor.RelaunchUnstable = ApplicationSettings.UseUnstableBuildCheckBox.Checked;
+				UpdateMonitor.TriggerUpdate(UpdateType.UserInitiated);
+			}
 		}
 	}
 }
diff --git a/Engine/Source/Programs/UnrealGameSync/UnrealGameSync/Forms/PerforceSettingsWindow.cs b/Engine/Source/Programs/UnrealGameSync/UnrealGameSync/Forms/PerforceSettingsWindow.cs
--- a/Engine/Source/Programs/UnrealGameSync/UnrealGameSync/Forms/PerforceSettingsWindow.cs
+++ b/Engine/Source/Programs/UnrealGameSync/UnrealGameSync/Forms/PerforceSettingsWindow.cs
@@ -16,7 +16,7 @@ using UnrealGameSync;
 
 namespace UnrealGameSync
 {
-	partial class PerforceSettingsWindow : Form
+	partial class ApplicationSettingsWindow : Form
 	{
 		class PerforceTestConnectionTask : IPerforceModalTask
 		{
@@ -43,7 +43,6 @@ namespace UnrealGameSync
 			}
 		}
 
-		UserSettings Settings;
 		TextWriter Log;
 
 		string InitialServerAndPort;
@@ -51,21 +50,22 @@ namespace UnrealGameSync
 		string InitialDepotPath;
 		bool bInitialUnstable;
 
-		public PerforceSettingsWindow(bool bUnstable, UserSettings Settings, TextWriter Log)
+		public ApplicationSettingsWindow(string ServerAndPort, string UserName, bool bUnstable, TextWriter Log)
 		{
 			InitializeComponent();
 
-			this.Settings = Settings;
 			this.Log = Log;
 
 			Utility.ReadGlobalPerforceSettings(ref InitialServerAndPort, ref InitialUserName, ref InitialDepotPath);
 			bInitialUnstable = bUnstable;
 
-			this.ServerTextBox.Text = InitialServerAndPort;
+			this.ServerTextBox.Text = ServerAndPort;
 			this.ServerTextBox.Select(ServerTextBox.TextLength, 0);
+			this.ServerTextBox.CueBanner = (ServerAndPort == null)? "Default" : String.Format("Default ({0})", ServerAndPort);
 
-			this.UserNameTextBox.Text = InitialUserName;
+			this.UserNameTextBox.Text = UserName;
 			this.UserNameTextBox.Select(UserNameTextBox.TextLength, 0);
+			this.UserNameTextBox.CueBanner = (UserName == null)? "Default" : String.Format("Default ({0})", UserName);
 
 			this.DepotPathTextBox.Text = InitialDepotPath;
 			this.DepotPathTextBox.Select(DepotPathTextBox.TextLength, 0);
@@ -134,11 +134,5 @@ namespace UnrealGameSync
 			DialogResult = DialogResult.Cancel;
 			Close();
 		}
-
-		private void AdvancedBtn_Click(object sender, EventArgs e)
-		{
-			PerforceSyncSettingsWindow Window = new PerforceSyncSettingsWindow(Settings);
-			Window.ShowDialog();
-		}
 	}
 }
diff --git a/Engine/Source/Programs/UnrealGameSync/UnrealGameSync/Program.cs b/Engine/Source/Programs/UnrealGameSync/UnrealGameSync/Program.cs
--- a/Engine/Source/Programs/UnrealGameSync/UnrealGameSync/Program.cs
+++ b/Engine/Source/Programs/UnrealGameSync/UnrealGameSync/Program.cs
@@ -97,7 +97,7 @@ namespace UnrealGameSync
 					if(UpdateMonitor.IsUpdateAvailable && UpdateSpawn != null)
 					{
 						InstanceMutex.Close();
-						bool bLaunchUnstable = Context.bRelaunchUnstable ?? bUnstable;
+						bool bLaunchUnstable = UpdateMonitor.RelaunchUnstable ?? bUnstable;
 						Utility.SpawnProcess(UpdateSpawn, "-restorestate" + (bLaunchUnstable? " -unstable" : ""));
 					}
 				}
diff --git a/Engine/Source/Programs/UnrealGameSync/UnrealGameSync/ProgramApplicationContext.cs b/Engine/Source/Programs/UnrealGameSync/UnrealGameSync/ProgramApplicationContext.cs
--- a/Engine/Source/Programs/UnrealGameSync/UnrealGameSync/ProgramApplicationContext.cs
+++ b/Engine/Source/Programs/UnrealGameSync/UnrealGameSync/ProgramApplicationContext.cs
@@ -14,7 +14,7 @@ using System.Windows.Forms;
 
 namespace UnrealGameSync
 {
-	class ProgramApplicationContext : ApplicationContext, IMainWindowOwner
+	class ProgramApplicationContext : ApplicationContext
 	{
 		SynchronizationContext MainThreadSynchronizationContext;
 
@@ -26,8 +26,6 @@ namespace UnrealGameSync
 		bool bUnstable;
 		bool bIsClosing;
 
-		public bool? bRelaunchUnstable;
-
 		TimestampLogWriter Log;
 		UserSettings Settings;
 		ActivationListener ActivationListener;
@@ -180,7 +178,7 @@ namespace UnrealGameSync
 			DetectStartupProjectSettingsWindow = null;
 
 			// Create the main window
-			MainWindowInstance = new MainWindow(this, ApiUrl, DataFolder, bRestoreState, UpdateSpawn ?? Assembly.GetExecutingAssembly().Location, bUnstable, DetectStartupProjectSettingsTask.Results, Log, Settings);
+			MainWindowInstance = new MainWindow(UpdateMonitor, ApiUrl, DataFolder, bRestoreState, UpdateSpawn ?? Assembly.GetExecutingAssembly().Location, bUnstable, DetectStartupProjectSettingsTask.Results, Log, Settings);
 			if(bVisible)
 			{
 				MainWindowInstance.Show();
@@ -196,12 +194,6 @@ namespace UnrealGameSync
 			DetectStartupProjectSettingsTask = null;
 		}
 
-		public void RequestRestart(bool? bRelaunchUnstable)
-		{
-			this.bRelaunchUnstable = bRelaunchUnstable;
-			UpdateMonitor.TriggerUpdate(true);
-		}
-
 		private void MainWindowInstance_FormClosed(object sender, FormClosedEventArgs e)
 		{
 			ExitThread();
@@ -212,7 +204,7 @@ namespace UnrealGameSync
 			// Check if we're trying to reopen with the unstable version; if so, trigger an update to trigger a restart with the new executable
 			if(!bUnstable && (Control.ModifierKeys & Keys.Shift) != 0)
 			{
-				UpdateMonitor.TriggerUpdate(true);
+				UpdateMonitor.TriggerUpdate(UpdateType.UserInitiated);
 			}
 			else if(MainWindowInstance != null)
 			{
@@ -225,11 +217,11 @@ namespace UnrealGameSync
 			MainThreadSynchronizationContext.Post((o) => OnActivationListenerCallback(), null);
 		}
 
-		private void OnUpdateAvailable(bool bForce)
+		private void OnUpdateAvailable(UpdateType Type)
 		{
 			if(MainWindowInstance != null && !bIsClosing)
 			{
-				if(bForce || MainWindowInstance.CanPerformUpdate())
+				if(Type == UpdateType.UserInitiated || MainWindowInstance.CanPerformUpdate())
 				{
 					bIsClosing = true;
 					MainWindowInstance.ForceClose();
@@ -237,9 +229,9 @@ namespace UnrealGameSync
 			}
 		}
 
-		private void OnUpdateAvailableCallback(bool bForce)
+		private void OnUpdateAvailableCallback(UpdateType Type)
 		{ 
-			MainThreadSynchronizationContext.Post((o) => OnUpdateAvailable(bForce), null);
+			MainThreadSynchronizationContext.Post((o) => OnUpdateAvailable(Type), null);
 		}
 
 		protected override void Dispose(bool bDisposing)
diff --git a/Engine/Source/Programs/UnrealGameSync/UnrealGameSync/UpdateMonitor.cs b/Engine/Source/Programs/UnrealGameSync/UnrealGameSync/UpdateMonitor.cs
--- a/Engine/Source/Programs/UnrealGameSync/UnrealGameSync/UpdateMonitor.cs
+++ b/Engine/Source/Programs/UnrealGameSync/UnrealGameSync/UpdateMonitor.cs
@@ -10,14 +10,31 @@ using System.Threading.Tasks;
 
 namespace UnrealGameSync
 {
+	enum UpdateType
+	{
+		Background,
+		UserInitiated,
+	}
+
 	class UpdateMonitor : IDisposable
 	{
-		PerforceConnection Perforce;
 		string WatchPath;
 		Thread WorkerThread;
 		ManualResetEvent QuitEvent;
 
-		public event Action<bool> OnUpdateAvailable;
+		public event Action<UpdateType> OnUpdateAvailable;
+
+		public PerforceConnection Perforce
+		{
+			get;
+			private set;
+		}
+
+		public bool? RelaunchUnstable
+		{
+			get;
+			set;
+		}
 
 		public UpdateMonitor(PerforceConnection InPerforce, string InWatchPath)
 		{
@@ -68,17 +85,17 @@ namespace UnrealGameSync
 				List<PerforceChangeSummary> Changes;
 				if(Perforce.FindChanges(WatchPath, 1, out Changes, Log) && Changes.Count > 0)
 				{
-					TriggerUpdate(false);
+					TriggerUpdate(UpdateType.Background);
 				}
 			}
 		}
 
-		public void TriggerUpdate(bool bForce)
+		public void TriggerUpdate(UpdateType UpdateType)
 		{
 			IsUpdateAvailable = true;
 			if(OnUpdateAvailable != null)
 			{
-				OnUpdateAvailable(bForce);
+				OnUpdateAvailable(UpdateType);
 			}
 		}
 	}
