UGS: Rename PerforceSettingsWindow to ApplicationSettingsWindow.

#rnx

diff --git a/Engine/Source/Programs/UnrealGameSync/UnrealGameSync/Forms/ApplicationSettingsWindow.cs b/Engine/Source/Programs/UnrealGameSync/UnrealGameSync/Forms/ApplicationSettingsWindow.cs
new file mode 100644
--- /dev/null
+++ b/Engine/Source/Programs/UnrealGameSync/UnrealGameSync/Forms/ApplicationSettingsWindow.cs
@@ -0,0 +1,138 @@
+// Copyright 1998-2018 Epic Games, Inc. All Rights Reserved.
+
+using System;
+using System.Collections.Generic;
+using System.ComponentModel;
+using System.Data;
+using System.Drawing;
+using System.IO;
+using System.Linq;
+using System.Runtime.InteropServices;
+using System.Text;
+using System.Threading;
+using System.Threading.Tasks;
+using System.Windows.Forms;
+using UnrealGameSync;
+
+namespace UnrealGameSync
+{
+	partial class ApplicationSettingsWindow : Form
+	{
+		class PerforceTestConnectionTask : IPerforceModalTask
+		{
+			string DepotPath;
+
+			public PerforceTestConnectionTask(string DepotPath)
+			{
+				this.DepotPath = DepotPath ?? DeploymentSettings.DefaultDepotPath;
+			}
+
+			public bool Run(PerforceConnection Perforce, TextWriter Log, out string ErrorMessage)
+			{
+				string CheckFilePath = String.Format("{0}/Release/UnrealGameSync.exe", DepotPath);
+
+				List<PerforceFileRecord> FileRecords;
+				if(!Perforce.FindFiles(CheckFilePath, out FileRecords, Log) || FileRecords.Count == 0)
+				{
+					ErrorMessage = String.Format("Unable to find {0}", CheckFilePath);
+					return false;
+				}
+
+				ErrorMessage = null;
+				return true;
+			}
+		}
+
+		TextWriter Log;
+
+		string InitialServerAndPort;
+		string InitialUserName;
+		string InitialDepotPath;
+		bool bInitialUnstable;
+
+		public ApplicationSettingsWindow(string ServerAndPort, string UserName, bool bUnstable, TextWriter Log)
+		{
+			InitializeComponent();
+
+			this.Log = Log;
+
+			Utility.ReadGlobalPerforceSettings(ref InitialServerAndPort, ref InitialUserName, ref InitialDepotPath);
+			bInitialUnstable = bUnstable;
+
+			this.ServerTextBox.Text = ServerAndPort;
+			this.ServerTextBox.Select(ServerTextBox.TextLength, 0);
+			this.ServerTextBox.CueBanner = (ServerAndPort == null)? "Default" : String.Format("Default ({0})", ServerAndPort);
+
+			this.UserNameTextBox.Text = UserName;
+			this.UserNameTextBox.Select(UserNameTextBox.TextLength, 0);
+			this.UserNameTextBox.CueBanner = (UserName == null)? "Default" : String.Format("Default ({0})", UserName);
+
+			this.DepotPathTextBox.Text = InitialDepotPath;
+			this.DepotPathTextBox.Select(DepotPathTextBox.TextLength, 0);
+			this.DepotPathTextBox.CueBanner = DeploymentSettings.DefaultDepotPath;
+
+			this.UseUnstableBuildCheckBox.Checked = bUnstable;
+		}
+
+		private void OkBtn_Click(object sender, EventArgs e)
+		{
+			// Update the settings
+			string ServerAndPort = ServerTextBox.Text.Trim();
+			if(ServerAndPort.Length == 0)
+			{
+				ServerAndPort = null;
+			}
+
+			string UserName = UserNameTextBox.Text.Trim();
+			if(UserName.Length == 0)
+			{
+				UserName = null;
+			}
+
+			string DepotPath = DepotPathTextBox.Text.Trim();
+			if(DepotPath.Length == 0 || DepotPath == DeploymentSettings.DefaultDepotPath)
+			{
+				DepotPath = null;
+			}
+
+			bool bUnstable = UseUnstableBuildCheckBox.Checked;
+
+			if(ServerAndPort == InitialServerAndPort && UserName == InitialUserName && DepotPath == InitialDepotPath && bUnstable == bInitialUnstable)
+			{
+				DialogResult = DialogResult.Cancel;
+				Close();
+			}
+			else
+			{
+				// Try to log in to the new server, and check the application is there
+				if(ServerAndPort != InitialServerAndPort || UserName != InitialUserName || DepotPath != InitialDepotPath)
+				{
+					string ErrorMessage;
+					ModalTaskResult Result = PerforceModalTask.Execute(this, null, ServerAndPort, UserName, new PerforceTestConnectionTask(DepotPath), "Connecting", "Checking connection, please wait...", Log, out ErrorMessage);
+					if(Result != ModalTaskResult.Succeeded)
+					{
+						if(Result == ModalTaskResult.Failed)
+						{
+							MessageBox.Show(ErrorMessage, "Unable to connect");
+						}
+						return;
+					}
+				}
+
+				if(MessageBox.Show("UnrealGameSync must be restarted to apply these settings.\n\nWould you like to restart now?", "Restart Required", MessageBoxButtons.OKCancel) == DialogResult.OK)
+				{
+					Utility.SaveGlobalPerforceSettings(ServerAndPort, UserName, DepotPath);
+
+					DialogResult = DialogResult.OK;
+					Close();
+				}
+			}
+		}
+
+		private void CancelBtn_Click(object sender, EventArgs e)
+		{
+			DialogResult = DialogResult.Cancel;
+			Close();
+		}
+	}
+}
diff --git a/Engine/Source/Programs/UnrealGameSync/UnrealGameSync/Forms/PerforceSettingsWindow.cs b/Engine/Source/Programs/UnrealGameSync/UnrealGameSync/Forms/PerforceSettingsWindow.cs
deleted file mode 100644
diff --git a/Engine/Source/Programs/UnrealGameSync/UnrealGameSync/UnrealGameSync.csproj b/Engine/Source/Programs/UnrealGameSync/UnrealGameSync/UnrealGameSync.csproj
--- a/Engine/Source/Programs/UnrealGameSync/UnrealGameSync/UnrealGameSync.csproj
+++ b/Engine/Source/Programs/UnrealGameSync/UnrealGameSync/UnrealGameSync.csproj
@@ -207,11 +207,11 @@
     <Compile Include="Forms\PasswordWindow.Designer.cs">
       <DependentUpon>PasswordWindow.cs</DependentUpon>
     </Compile>
-    <Compile Include="Forms\PerforceSettingsWindow.cs">
+    <Compile Include="Forms\ApplicationSettingsWindow.cs">
       <SubType>Form</SubType>
     </Compile>
-    <Compile Include="Forms\PerforceSettingsWindow.Designer.cs">
-      <DependentUpon>PerforceSettingsWindow.cs</DependentUpon>
+    <Compile Include="Forms\ApplicationSettingsWindow.Designer.cs">
+      <DependentUpon>ApplicationSettingsWindow.cs</DependentUpon>
     </Compile>
     <Compile Include="Forms\PerforceSyncSettingsWindow.cs">
       <SubType>Form</SubType>
@@ -377,8 +377,8 @@
     <EmbeddedResource Include="Forms\PasswordWindow.resx">
       <DependentUpon>PasswordWindow.cs</DependentUpon>
     </EmbeddedResource>
-    <EmbeddedResource Include="Forms\PerforceSettingsWindow.resx">
-      <DependentUpon>PerforceSettingsWindow.cs</DependentUpon>
+    <EmbeddedResource Include="Forms\ApplicationSettingsWindow.resx">
+      <DependentUpon>ApplicationSettingsWindow.cs</DependentUpon>
     </EmbeddedResource>
     <EmbeddedResource Include="Forms\PerforceSyncSettingsWindow.resx">
       <DependentUpon>PerforceSyncSettingsWindow.cs</DependentUpon>
