UGS: Fix the default label in the auto-update settings dialog not showing the correct defaults.

#rnx

diff --git a/Engine/Source/Programs/UnrealGameSync/UnrealGameSync/Forms/ApplicationSettingsWindow.cs b/Engine/Source/Programs/UnrealGameSync/UnrealGameSync/Forms/ApplicationSettingsWindow.cs
--- a/Engine/Source/Programs/UnrealGameSync/UnrealGameSync/Forms/ApplicationSettingsWindow.cs
+++ b/Engine/Source/Programs/UnrealGameSync/UnrealGameSync/Forms/ApplicationSettingsWindow.cs
@@ -18,6 +18,33 @@ namespace UnrealGameSync
 {
 	partial class ApplicationSettingsWindow : Form
 	{
+		class GetDefaultSettingsTask : IModalTask
+		{
+			TextWriter Log;
+
+			public string ServerAndPort;
+			public string UserName;
+
+			public GetDefaultSettingsTask(TextWriter Log)
+			{
+				this.Log = Log;
+			}
+
+			public bool Run(out string ErrorMessage)
+			{
+				if(PerforceModalTask.TryGetServerSettings(null, ref ServerAndPort, ref UserName, Log))
+				{
+					ErrorMessage = null;
+					return true;
+				}
+				else
+				{
+					ErrorMessage = "Unable to query server settings";
+					return false;
+				}
+			}
+		}
+
 		class PerforceTestConnectionTask : IPerforceModalTask
 		{
 			string DepotPath;
@@ -50,7 +77,7 @@ namespace UnrealGameSync
 		string InitialDepotPath;
 		bool bInitialUnstable;
 
-		public ApplicationSettingsWindow(string ServerAndPort, string UserName, bool bUnstable, TextWriter Log)
+		private ApplicationSettingsWindow(string DefaultServerAndPort, string DefaultUserName, bool bUnstable, TextWriter Log)
 		{
 			InitializeComponent();
 
@@ -58,14 +85,14 @@ namespace UnrealGameSync
 
 			Utility.ReadGlobalPerforceSettings(ref InitialServerAndPort, ref InitialUserName, ref InitialDepotPath);
 			bInitialUnstable = bUnstable;
-
-			this.ServerTextBox.Text = ServerAndPort;
+			
+			this.ServerTextBox.Text = InitialServerAndPort;
 			this.ServerTextBox.Select(ServerTextBox.TextLength, 0);
-			this.ServerTextBox.CueBanner = (ServerAndPort == null)? "Default" : String.Format("Default ({0})", ServerAndPort);
+			this.ServerTextBox.CueBanner = (DefaultServerAndPort == null)? "Default" : String.Format("Default ({0})", DefaultServerAndPort);
 
-			this.UserNameTextBox.Text = UserName;
+			this.UserNameTextBox.Text = InitialUserName;
 			this.UserNameTextBox.Select(UserNameTextBox.TextLength, 0);
-			this.UserNameTextBox.CueBanner = (UserName == null)? "Default" : String.Format("Default ({0})", UserName);
+			this.UserNameTextBox.CueBanner = (DefaultUserName == null)? "Default" : String.Format("Default ({0})", DefaultUserName);
 
 			this.DepotPathTextBox.Text = InitialDepotPath;
 			this.DepotPathTextBox.Select(DepotPathTextBox.TextLength, 0);
@@ -74,6 +101,24 @@ namespace UnrealGameSync
 			this.UseUnstableBuildCheckBox.Checked = bUnstable;
 		}
 
+		public static bool? ShowModal(IWin32Window Owner, bool bUnstable, TextWriter Log)
+		{
+			GetDefaultSettingsTask DefaultSettings = new GetDefaultSettingsTask(Log);
+
+			string ErrorMessage;
+			ModalTask.Execute(Owner, DefaultSettings, "Checking Settings", "Checking settings, please wait...", out ErrorMessage);
+
+			ApplicationSettingsWindow ApplicationSettings = new ApplicationSettingsWindow(DefaultSettings.ServerAndPort, DefaultSettings.UserName, bUnstable, Log);
+			if(ApplicationSettings.ShowDialog() == DialogResult.OK)
+			{
+				return ApplicationSettings.UseUnstableBuildCheckBox.Checked;
+			}
+			else
+			{
+				return null;
+			}
+		}
+
 		private void OkBtn_Click(object sender, EventArgs e)
 		{
 			// Update the settings
diff --git a/Engine/Source/Programs/UnrealGameSync/UnrealGameSync/Forms/MainWindow.cs b/Engine/Source/Programs/UnrealGameSync/UnrealGameSync/Forms/MainWindow.cs
--- a/Engine/Source/Programs/UnrealGameSync/UnrealGameSync/Forms/MainWindow.cs
+++ b/Engine/Source/Programs/UnrealGameSync/UnrealGameSync/Forms/MainWindow.cs
@@ -958,11 +958,10 @@ namespace UnrealGameSync
 
 		public void ModifyApplicationSettings()
 		{
-			ApplicationSettingsWindow ApplicationSettings = new ApplicationSettingsWindow(UpdateMonitor.Perforce.ServerAndPort, UpdateMonitor.Perforce.UserName, bUnstable, Log);
-			if(ApplicationSettings.ShowDialog() == DialogResult.OK)
+			bool? bRelaunchUnstable = ApplicationSettingsWindow.ShowModal(this, bUnstable, Log);
+			if(bRelaunchUnstable.HasValue)
 			{
-				UpdateMonitor.RelaunchUnstable = ApplicationSettings.UseUnstableBuildCheckBox.Checked;
-				UpdateMonitor.TriggerUpdate(UpdateType.UserInitiated);
+				UpdateMonitor.TriggerUpdate(UpdateType.UserInitiated, bRelaunchUnstable);
 			}
 		}
 	}
diff --git a/Engine/Source/Programs/UnrealGameSync/UnrealGameSync/ProgramApplicationContext.cs b/Engine/Source/Programs/UnrealGameSync/UnrealGameSync/ProgramApplicationContext.cs
--- a/Engine/Source/Programs/UnrealGameSync/UnrealGameSync/ProgramApplicationContext.cs
+++ b/Engine/Source/Programs/UnrealGameSync/UnrealGameSync/ProgramApplicationContext.cs
@@ -201,12 +201,7 @@ namespace UnrealGameSync
 
 		private void OnActivationListenerCallback()
 		{
-			// Check if we're trying to reopen with the unstable version; if so, trigger an update to trigger a restart with the new executable
-			if(!bUnstable && (Control.ModifierKeys & Keys.Shift) != 0)
-			{
-				UpdateMonitor.TriggerUpdate(UpdateType.UserInitiated);
-			}
-			else if(MainWindowInstance != null)
+			if(MainWindowInstance != null)
 			{
 				MainWindowInstance.ShowAndActivate();
 			}
diff --git a/Engine/Source/Programs/UnrealGameSync/UnrealGameSync/UpdateMonitor.cs b/Engine/Source/Programs/UnrealGameSync/UnrealGameSync/UpdateMonitor.cs
--- a/Engine/Source/Programs/UnrealGameSync/UnrealGameSync/UpdateMonitor.cs
+++ b/Engine/Source/Programs/UnrealGameSync/UnrealGameSync/UpdateMonitor.cs
@@ -33,7 +33,7 @@ namespace UnrealGameSync
 		public bool? RelaunchUnstable
 		{
 			get;
-			set;
+			private set;
 		}
 
 		public UpdateMonitor(PerforceConnection InPerforce, string InWatchPath)
@@ -85,13 +85,14 @@ namespace UnrealGameSync
 				List<PerforceChangeSummary> Changes;
 				if(Perforce.FindChanges(WatchPath, 1, out Changes, Log) && Changes.Count > 0)
 				{
-					TriggerUpdate(UpdateType.Background);
+					TriggerUpdate(UpdateType.Background, null);
 				}
 			}
 		}
 
-		public void TriggerUpdate(UpdateType UpdateType)
+		public void TriggerUpdate(UpdateType UpdateType, bool? RelaunchUnstable)
 		{
+			this.RelaunchUnstable = RelaunchUnstable;
 			IsUpdateAvailable = true;
 			if(OnUpdateAvailable != null)
 			{
