UGS: Move settings for running at startup and closing to the notification area into the application settings dialog.

diff --git a/Engine/Source/Programs/UnrealGameSync/UnrealGameSync/Controls/WorkspaceControl.cs b/Engine/Source/Programs/UnrealGameSync/UnrealGameSync/Controls/WorkspaceControl.cs
--- a/Engine/Source/Programs/UnrealGameSync/UnrealGameSync/Controls/WorkspaceControl.cs
+++ b/Engine/Source/Programs/UnrealGameSync/UnrealGameSync/Controls/WorkspaceControl.cs
@@ -3279,8 +3279,6 @@ namespace UnrealGameSync
 			OptionsContextMenu_ScheduledSync.Checked = Settings.bScheduleEnabled;
 			OptionsContextMenu_TimeZone_Local.Checked = Settings.bShowLocalTimes;
 			OptionsContextMenu_TimeZone_PerforceServer.Checked = !Settings.bShowLocalTimes;
-			OptionsContextMenu_AutomaticallyRunAtStartup.Checked = IsAutomaticallyRunAtStartup();
-			OptionsContextMenu_KeepInTray.Checked = Settings.bKeepInTray;
 			OptionsContextMenu_ShowChanges_ShowUnreviewed.Checked = Settings.bShowUnreviewedChanges;
 			OptionsContextMenu_ShowChanges_ShowAutomated.Checked = Settings.bShowAutomatedChanges;
 			OptionsContextMenu_TabNames_Stream.Checked = Settings.TabLabels == TabLabels.Stream;
@@ -4092,25 +4090,6 @@ namespace UnrealGameSync
 			return UserBuildStepObjects.Values.Select(x => new BuildStep(x)).OrderBy(x => x.OrderIndex).ToList();
 		}
 
-		private void OptionsContextMenu_AutomaticallyRunAtStartup_Click(object sender, EventArgs e)
-		{
-			RegistryKey Key = Registry.CurrentUser.CreateSubKey("SOFTWARE\\Microsoft\\Windows\\CurrentVersion\\Run");
-			if(IsAutomaticallyRunAtStartup())
-			{
-	            Key.DeleteValue("UnrealGameSync", false);
-			}
-			else
-			{
-				Key.SetValue("UnrealGameSync", String.Format("\"{0}\" -RestoreState", OriginalExecutableFileName));
-			}
-		}
-
-		private bool IsAutomaticallyRunAtStartup()
-		{
-			RegistryKey Key = Registry.CurrentUser.OpenSubKey("SOFTWARE\\Microsoft\\Windows\\CurrentVersion\\Run");
-			return (Key.GetValue("UnrealGameSync") != null);
-		}
-
 		private void OptionsContextMenu_SyncPrecompiledEditor_Click(object sender, EventArgs e)
 		{
 			OptionsContextMenu_SyncPrecompiledEditor.Checked ^= true;
@@ -4204,12 +4183,6 @@ namespace UnrealGameSync
 			}
 		}
 
-		private void OptionsContextMenu_KeepInTray_Click(object sender, EventArgs e)
-		{
-			Settings.bKeepInTray ^= true;
-			Settings.Save();
-		}
-
 		private void BuildList_KeyDown(object Sender, KeyEventArgs Args)
 		{
 			if(Args.Control && Args.KeyCode == Keys.C && BuildList.SelectedItems.Count > 0)
diff --git a/Engine/Source/Programs/UnrealGameSync/UnrealGameSync/Forms/ApplicationSettingsWindow.cs b/Engine/Source/Programs/UnrealGameSync/UnrealGameSync/Forms/ApplicationSettingsWindow.cs
--- a/Engine/Source/Programs/UnrealGameSync/UnrealGameSync/Forms/ApplicationSettingsWindow.cs
+++ b/Engine/Source/Programs/UnrealGameSync/UnrealGameSync/Forms/ApplicationSettingsWindow.cs
@@ -1,5 +1,6 @@
 // Copyright 1998-2018 Epic Games, Inc. All Rights Reserved.
 
+using Microsoft.Win32;
 using System;
 using System.Collections.Generic;
 using System.ComponentModel;
@@ -70,6 +71,8 @@ namespace UnrealGameSync
 			}
 		}
 
+		string OriginalExecutableFileName;
+		UserSettings Settings;
 		TextWriter Log;
 
 		string InitialServerAndPort;
@@ -77,15 +80,22 @@ namespace UnrealGameSync
 		string InitialDepotPath;
 		bool bInitialUnstable;
 
-		private ApplicationSettingsWindow(string DefaultServerAndPort, string DefaultUserName, bool bUnstable, TextWriter Log)
+		bool? bRestartUnstable;
+
+		private ApplicationSettingsWindow(string DefaultServerAndPort, string DefaultUserName, bool bUnstable, string OriginalExecutableFileName, UserSettings Settings, TextWriter Log)
 		{
 			InitializeComponent();
 
+			this.OriginalExecutableFileName = OriginalExecutableFileName;
+			this.Settings = Settings;
 			this.Log = Log;
 
 			Utility.ReadGlobalPerforceSettings(ref InitialServerAndPort, ref InitialUserName, ref InitialDepotPath);
 			bInitialUnstable = bUnstable;
-			
+
+			this.AutomaticallyRunAtStartupCheckBox.Checked = IsAutomaticallyRunAtStartup();
+			this.KeepInTrayCheckBox.Checked = Settings.bKeepInTray;
+						
 			this.ServerTextBox.Text = InitialServerAndPort;
 			this.ServerTextBox.Select(ServerTextBox.TextLength, 0);
 			this.ServerTextBox.CueBanner = (DefaultServerAndPort == null)? "Default" : String.Format("Default ({0})", DefaultServerAndPort);
@@ -101,17 +111,17 @@ namespace UnrealGameSync
 			this.UseUnstableBuildCheckBox.Checked = bUnstable;
 		}
 
-		public static bool? ShowModal(IWin32Window Owner, bool bUnstable, TextWriter Log)
+		public static bool? ShowModal(IWin32Window Owner, bool bUnstable, string OriginalExecutableFileName, UserSettings Settings, TextWriter Log)
 		{
 			GetDefaultSettingsTask DefaultSettings = new GetDefaultSettingsTask(Log);
 
 			string ErrorMessage;
 			ModalTask.Execute(Owner, DefaultSettings, "Checking Settings", "Checking settings, please wait...", out ErrorMessage);
 
-			ApplicationSettingsWindow ApplicationSettings = new ApplicationSettingsWindow(DefaultSettings.ServerAndPort, DefaultSettings.UserName, bUnstable, Log);
+			ApplicationSettingsWindow ApplicationSettings = new ApplicationSettingsWindow(DefaultSettings.ServerAndPort, DefaultSettings.UserName, bUnstable, OriginalExecutableFileName, Settings, Log);
 			if(ApplicationSettings.ShowDialog() == DialogResult.OK)
 			{
-				return ApplicationSettings.UseUnstableBuildCheckBox.Checked;
+				return ApplicationSettings.bRestartUnstable;
 			}
 			else
 			{
@@ -119,6 +129,12 @@ namespace UnrealGameSync
 			}
 		}
 
+		private bool IsAutomaticallyRunAtStartup()
+		{
+			RegistryKey Key = Registry.CurrentUser.OpenSubKey("SOFTWARE\\Microsoft\\Windows\\CurrentVersion\\Run");
+			return (Key.GetValue("UnrealGameSync") != null);
+		}
+
 		private void OkBtn_Click(object sender, EventArgs e)
 		{
 			// Update the settings
@@ -142,12 +158,7 @@ namespace UnrealGameSync
 
 			bool bUnstable = UseUnstableBuildCheckBox.Checked;
 
-			if(ServerAndPort == InitialServerAndPort && UserName == InitialUserName && DepotPath == InitialDepotPath && bUnstable == bInitialUnstable)
-			{
-				DialogResult = DialogResult.Cancel;
-				Close();
-			}
-			else
+			if(ServerAndPort != InitialServerAndPort || UserName != InitialUserName || DepotPath != InitialDepotPath || bUnstable != bInitialUnstable)
 			{
 				// Try to log in to the new server, and check the application is there
 				if(ServerAndPort != InitialServerAndPort || UserName != InitialUserName || DepotPath != InitialDepotPath)
@@ -164,14 +175,33 @@ namespace UnrealGameSync
 					}
 				}
 
-				if(MessageBox.Show("UnrealGameSync must be restarted to apply these settings.\n\nWould you like to restart now?", "Restart Required", MessageBoxButtons.OKCancel) == DialogResult.OK)
+				if(MessageBox.Show("UnrealGameSync must be restarted to apply these settings.\n\nWould you like to restart now?", "Restart Required", MessageBoxButtons.OKCancel) != DialogResult.OK)
 				{
-					Utility.SaveGlobalPerforceSettings(ServerAndPort, UserName, DepotPath);
-
-					DialogResult = DialogResult.OK;
-					Close();
+					return;
 				}
+
+				bRestartUnstable = UseUnstableBuildCheckBox.Checked;
+				Utility.SaveGlobalPerforceSettings(ServerAndPort, UserName, DepotPath);
 			}
+
+			RegistryKey Key = Registry.CurrentUser.CreateSubKey("SOFTWARE\\Microsoft\\Windows\\CurrentVersion\\Run");
+			if(IsAutomaticallyRunAtStartup())
+			{
+	            Key.DeleteValue("UnrealGameSync", false);
+			}
+			else
+			{
+				Key.SetValue("UnrealGameSync", String.Format("\"{0}\" -RestoreState", OriginalExecutableFileName));
+			}
+
+			if(Settings.bKeepInTray != KeepInTrayCheckBox.Checked)
+			{
+				Settings.bKeepInTray = KeepInTrayCheckBox.Checked;
+				Settings.Save();
+			}
+
+			DialogResult = DialogResult.OK;
+			Close();
 		}
 
 		private void CancelBtn_Click(object sender, EventArgs e)
diff --git a/Engine/Source/Programs/UnrealGameSync/UnrealGameSync/Forms/MainWindow.cs b/Engine/Source/Programs/UnrealGameSync/UnrealGameSync/Forms/MainWindow.cs
--- a/Engine/Source/Programs/UnrealGameSync/UnrealGameSync/Forms/MainWindow.cs
+++ b/Engine/Source/Programs/UnrealGameSync/UnrealGameSync/Forms/MainWindow.cs
@@ -958,7 +958,7 @@ namespace UnrealGameSync
 
 		public void ModifyApplicationSettings()
 		{
-			bool? bRelaunchUnstable = ApplicationSettingsWindow.ShowModal(this, bUnstable, Log);
+			bool? bRelaunchUnstable = ApplicationSettingsWindow.ShowModal(this, bUnstable, OriginalExecutableFileName, Settings, Log);
 			if(bRelaunchUnstable.HasValue)
 			{
 				UpdateMonitor.TriggerUpdate(UpdateType.UserInitiated, bRelaunchUnstable);
