UGS: Allow detection of project settings at startup to happen in parallel, to improve startup time.

diff --git a/Engine/Source/Programs/UnrealGameSync/UnrealGameSync/DetectProjectSettingsTask.cs b/Engine/Source/Programs/UnrealGameSync/UnrealGameSync/DetectProjectSettingsTask.cs
--- a/Engine/Source/Programs/UnrealGameSync/UnrealGameSync/DetectProjectSettingsTask.cs
+++ b/Engine/Source/Programs/UnrealGameSync/UnrealGameSync/DetectProjectSettingsTask.cs
@@ -399,12 +399,12 @@ namespace UnrealGameSync
 
 	class DetectMultipleProjectSettingsTask : IModalTask, IDisposable
 	{
-		public List<DetectProjectSettingsTask> Tasks;
-		public List<DetectProjectSettingsResult> Results;
+		public DetectProjectSettingsTask[] Tasks;
+		public DetectProjectSettingsResult[] Results;
 
 		public DetectMultipleProjectSettingsTask(IEnumerable<DetectProjectSettingsTask> Tasks)
 		{
-			this.Tasks = new List<DetectProjectSettingsTask>(Tasks);
+			this.Tasks = Tasks.ToArray();
 		}
 
 		public void Dispose()
@@ -415,7 +415,7 @@ namespace UnrealGameSync
 				{
 					Task.Dispose();
 				}
-				foreach(DetectProjectSettingsResult Result in Results)
+				foreach(DetectProjectSettingsResult Result in Results.Where(x => x != null))
 				{
 					Result.Dispose();
 				}
@@ -425,17 +425,19 @@ namespace UnrealGameSync
 
 		public bool Run(out string ErrorMessage)
 		{
-			Results = new List<DetectProjectSettingsResult>();
-			for(int Idx = 0; Idx < Tasks.Count; Idx++)
-			{
-				string TaskErrorMessage;
-				bool bTaskSucceeded = Tasks[Idx].Run(out TaskErrorMessage);
-				Results.Add(new DetectProjectSettingsResult(Tasks[Idx], bTaskSucceeded, TaskErrorMessage));
-				Tasks[Idx] = null;
-			}
+			Results = new DetectProjectSettingsResult[Tasks.Length];
+			Parallel.For(0, Tasks.Length, new ParallelOptions(){ MaxDegreeOfParallelism = 4 }, Idx => RunTask(Idx));
 
 			ErrorMessage = null;
 			return true;
 		}
+
+		void RunTask(int Idx)
+		{
+			string TaskErrorMessage;
+			bool bTaskSucceeded = Tasks[Idx].Run(out TaskErrorMessage);
+			Results[Idx] = new DetectProjectSettingsResult(Tasks[Idx], bTaskSucceeded, TaskErrorMessage);
+			Tasks[Idx] = null;
+		}
 	}
 }
diff --git a/Engine/Source/Programs/UnrealGameSync/UnrealGameSync/Forms/MainWindow.cs b/Engine/Source/Programs/UnrealGameSync/UnrealGameSync/Forms/MainWindow.cs
--- a/Engine/Source/Programs/UnrealGameSync/UnrealGameSync/Forms/MainWindow.cs
+++ b/Engine/Source/Programs/UnrealGameSync/UnrealGameSync/Forms/MainWindow.cs
@@ -76,7 +76,7 @@ namespace UnrealGameSync
 
 		IMainWindowTabPanel CurrentTabPanel;
 
-		public MainWindow(UpdateMonitor InUpdateMonitor, string InApiUrl, string InDataFolder, bool bInRestoreStateOnLoad, string InOriginalExecutableFileName, bool bInUnstable, List<DetectProjectSettingsResult> StartupProjects, LineBasedTextWriter InLog, UserSettings InSettings)
+		public MainWindow(UpdateMonitor InUpdateMonitor, string InApiUrl, string InDataFolder, bool bInRestoreStateOnLoad, string InOriginalExecutableFileName, bool bInUnstable, DetectProjectSettingsResult[] StartupProjects, LineBasedTextWriter InLog, UserSettings InSettings)
 		{
 			InitializeComponent();
 
diff --git a/Engine/Source/Programs/UnrealGameSync/UnrealGameSync/OutputAdapters.cs b/Engine/Source/Programs/UnrealGameSync/UnrealGameSync/OutputAdapters.cs
--- a/Engine/Source/Programs/UnrealGameSync/UnrealGameSync/OutputAdapters.cs
+++ b/Engine/Source/Programs/UnrealGameSync/UnrealGameSync/OutputAdapters.cs
@@ -53,6 +53,16 @@ namespace UnrealGameSync
 		}
 	}
 
+	class BufferedTextWriter : LineBasedTextWriter
+	{
+		public List<string> Lines = new List<string>();
+
+		protected override void FlushLine(string Line)
+		{
+			Lines.Add(Line);
+		}
+	}
+
 	class PrefixedTextWriter : LineBasedTextWriter
 	{
 		string Prefix;
diff --git a/Engine/Source/Programs/UnrealGameSync/UnrealGameSync/ProgramApplicationContext.cs b/Engine/Source/Programs/UnrealGameSync/UnrealGameSync/ProgramApplicationContext.cs
--- a/Engine/Source/Programs/UnrealGameSync/UnrealGameSync/ProgramApplicationContext.cs
+++ b/Engine/Source/Programs/UnrealGameSync/UnrealGameSync/ProgramApplicationContext.cs
@@ -40,6 +40,7 @@ namespace UnrealGameSync
 		ToolStripSeparator NotifyMenu_ExitSeparator;
 		ToolStripMenuItem NotifyMenu_Exit;
 
+		List<BufferedTextWriter> StartupLogs = new List<BufferedTextWriter>();
 		DetectMultipleProjectSettingsTask DetectStartupProjectSettingsTask;
 		ModalTaskWindow DetectStartupProjectSettingsWindow;
 		MainWindow MainWindowInstance;
@@ -147,8 +148,11 @@ namespace UnrealGameSync
 			List<DetectProjectSettingsTask> Tasks = new List<DetectProjectSettingsTask>();
 			foreach(UserSelectedProjectSettings OpenProject in Settings.OpenProjects)
 			{
-				Log.WriteLine("Detecting settings for {0}", OpenProject);
-				Tasks.Add(new DetectProjectSettingsTask(OpenProject, DataFolder, new PrefixedTextWriter("  ", Log)));
+				BufferedTextWriter StartupLog = new BufferedTextWriter();
+				StartupLog.WriteLine("Detecting settings for {0}", OpenProject);
+				StartupLogs.Add(StartupLog);
+
+				Tasks.Add(new DetectProjectSettingsTask(OpenProject, DataFolder, new TimestampLogWriter(new PrefixedTextWriter("  ", StartupLog))));
 			}
 
 			// Detect settings for the project we want to open
@@ -177,8 +181,20 @@ namespace UnrealGameSync
 			DetectStartupProjectSettingsWindow.Close();
 			DetectStartupProjectSettingsWindow = null;
 
+			// Copy all the logs to the main log
+			foreach(BufferedTextWriter StartupLog in StartupLogs)
+			{
+				foreach(string Line in StartupLog.Lines)
+				{
+					Log.WriteLine("{0}", Line);
+				}
+			}
+
+			// Get a list of all the valid projects to open
+			DetectProjectSettingsResult[] StartupProjects = DetectStartupProjectSettingsTask.Results.Where(x => x != null).ToArray();
+
 			// Create the main window
-			MainWindowInstance = new MainWindow(UpdateMonitor, ApiUrl, DataFolder, bRestoreState, UpdateSpawn ?? Assembly.GetExecutingAssembly().Location, bUnstable, DetectStartupProjectSettingsTask.Results, Log, Settings);
+			MainWindowInstance = new MainWindow(UpdateMonitor, ApiUrl, DataFolder, bRestoreState, UpdateSpawn ?? Assembly.GetExecutingAssembly().Location, bUnstable, StartupProjects, Log, Settings);
 			if(bVisible)
 			{
 				MainWindowInstance.Show();
