UGS: Cache the contents of INI files locally, so we can reuse them rather than having to print them from the server each time the program starts. Improves startup times when working through a proxy.

diff --git a/Engine/Source/Programs/UnrealGameSync/UnrealGameSync/Controls/WorkspaceControl.cs b/Engine/Source/Programs/UnrealGameSync/UnrealGameSync/Controls/WorkspaceControl.cs
--- a/Engine/Source/Programs/UnrealGameSync/UnrealGameSync/Controls/WorkspaceControl.cs
+++ b/Engine/Source/Programs/UnrealGameSync/UnrealGameSync/Controls/WorkspaceControl.cs
@@ -304,7 +304,7 @@ namespace UnrealGameSync
 
 			string ProjectLogBaseName = Path.Combine(DataFolder, String.Format("{0}@{1}", PerforceClient.ClientName, DetectSettings.BranchClientPath.Replace("//" + PerforceClient.ClientName + "/", "").Trim('/').Replace("/", "$")));
 
-			PerforceMonitor = new PerforceMonitor(PerforceClient, DetectSettings.BranchClientPath, DetectSettings.NewSelectedClientFileName, DetectSettings.NewSelectedProjectIdentifier, ProjectLogBaseName + ".p4.log", DetectSettings.bIsEnterpriseProject, DetectSettings.LatestProjectConfigFile, DetectSettings.LocalConfigFiles);
+			PerforceMonitor = new PerforceMonitor(PerforceClient, DetectSettings.BranchClientPath, DetectSettings.NewSelectedClientFileName, DetectSettings.NewSelectedProjectIdentifier, ProjectLogBaseName + ".p4.log", DetectSettings.bIsEnterpriseProject, DetectSettings.LatestProjectConfigFile, DetectSettings.CacheFolder, DetectSettings.LocalConfigFiles);
 			PerforceMonitor.OnUpdate += UpdateBuildListCallback;
 			PerforceMonitor.OnUpdateMetadata += UpdateBuildMetadataCallback;
 			PerforceMonitor.OnStreamChange += StreamChangedCallback;
diff --git a/Engine/Source/Programs/UnrealGameSync/UnrealGameSync/DetectProjectSettingsTask.cs b/Engine/Source/Programs/UnrealGameSync/UnrealGameSync/DetectProjectSettingsTask.cs
--- a/Engine/Source/Programs/UnrealGameSync/UnrealGameSync/DetectProjectSettingsTask.cs
+++ b/Engine/Source/Programs/UnrealGameSync/UnrealGameSync/DetectProjectSettingsTask.cs
@@ -25,15 +25,17 @@ namespace UnrealGameSync
 		public Image ProjectLogo;
 		public TimeSpan ServerTimeZone;
 		public string DataFolder;
+		public string CacheFolder;
 		public bool bIsEnterpriseProject;
 		public ConfigFile LatestProjectConfigFile;
 		public List<KeyValuePair<string, DateTime>> LocalConfigFiles;
 		TextWriter Log;
 
-		public DetectProjectSettingsTask(UserSelectedProjectSettings SelectedProject, string InDataFolder, TextWriter InLog)
+		public DetectProjectSettingsTask(UserSelectedProjectSettings SelectedProject, string InDataFolder, string InCacheFolder, TextWriter InLog)
 		{
 			this.SelectedProject = SelectedProject;
 			DataFolder = InDataFolder;
+			CacheFolder = InCacheFolder;
 			Log = InLog;
 		}
 
@@ -298,7 +300,7 @@ namespace UnrealGameSync
 
 			// Read the initial config file
 			LocalConfigFiles = new List<KeyValuePair<string, DateTime>>();
-			LatestProjectConfigFile = PerforceMonitor.ReadProjectConfigFile(PerforceClient, BranchClientPath, NewSelectedClientFileName, LocalConfigFiles, Log);
+			LatestProjectConfigFile = PerforceMonitor.ReadProjectConfigFile(PerforceClient, BranchClientPath, NewSelectedClientFileName, CacheFolder, LocalConfigFiles, Log);
 
 			// Succeed!
 			ErrorMessage = null;
diff --git a/Engine/Source/Programs/UnrealGameSync/UnrealGameSync/Forms/MainWindow.cs b/Engine/Source/Programs/UnrealGameSync/UnrealGameSync/Forms/MainWindow.cs
--- a/Engine/Source/Programs/UnrealGameSync/UnrealGameSync/Forms/MainWindow.cs
+++ b/Engine/Source/Programs/UnrealGameSync/UnrealGameSync/Forms/MainWindow.cs
@@ -59,6 +59,7 @@ namespace UnrealGameSync
 
 		string ApiUrl;
 		string DataFolder;
+		string CacheFolder;
 		LineBasedTextWriter Log;
 		UserSettings Settings;
 		int TabMenu_TabIdx = -1;
@@ -76,7 +77,7 @@ namespace UnrealGameSync
 
 		IMainWindowTabPanel CurrentTabPanel;
 
-		public MainWindow(UpdateMonitor InUpdateMonitor, string InApiUrl, string InDataFolder, bool bInRestoreStateOnLoad, string InOriginalExecutableFileName, bool bInUnstable, DetectProjectSettingsResult[] StartupProjects, LineBasedTextWriter InLog, UserSettings InSettings)
+		public MainWindow(UpdateMonitor InUpdateMonitor, string InApiUrl, string InDataFolder, string InCacheFolder, bool bInRestoreStateOnLoad, string InOriginalExecutableFileName, bool bInUnstable, DetectProjectSettingsResult[] StartupProjects, LineBasedTextWriter InLog, UserSettings InSettings)
 		{
 			InitializeComponent();
 
@@ -84,6 +85,7 @@ namespace UnrealGameSync
 			MainThreadSynchronizationContext = SynchronizationContext.Current;
 			ApiUrl = InApiUrl;
 			DataFolder = InDataFolder;
+			CacheFolder = InCacheFolder;
 			bRestoreStateOnLoad = bInRestoreStateOnLoad;
 			OriginalExecutableFileName = InOriginalExecutableFileName;
 			bUnstable = bInUnstable;
@@ -616,7 +618,7 @@ namespace UnrealGameSync
 		public void OpenNewProject()
 		{
 			DetectProjectSettingsTask DetectedProjectSettings;
-			if(OpenProjectWindow.ShowModal(this, null, out DetectedProjectSettings, Settings, DataFolder, Log))
+			if(OpenProjectWindow.ShowModal(this, null, out DetectedProjectSettings, Settings, DataFolder, CacheFolder, Log))
 			{
 				int NewTabIdx = TryOpenProject(DetectedProjectSettings, -1, OpenProjectOptions.None);
 				if(NewTabIdx != -1)
@@ -668,7 +670,7 @@ namespace UnrealGameSync
 		public void EditSelectedProject(int TabIdx, UserSelectedProjectSettings SelectedProject)
 		{
 			DetectProjectSettingsTask DetectedProjectSettings;
-			if(OpenProjectWindow.ShowModal(this, SelectedProject, out DetectedProjectSettings, Settings, DataFolder, Log))
+			if(OpenProjectWindow.ShowModal(this, SelectedProject, out DetectedProjectSettings, Settings, DataFolder, CacheFolder, Log))
 			{
 				int NewTabIdx = TryOpenProject(DetectedProjectSettings, TabIdx, OpenProjectOptions.None);
 				if(NewTabIdx != -1)
@@ -686,7 +688,7 @@ namespace UnrealGameSync
 		int TryOpenProject(UserSelectedProjectSettings Project, int ReplaceTabIdx, OpenProjectOptions Options = OpenProjectOptions.None)
 		{
 			Log.WriteLine("Detecting settings for {0}", Project);
-			using(DetectProjectSettingsTask DetectProjectSettings = new DetectProjectSettingsTask(Project, DataFolder, new PrefixedTextWriter("  ", Log)))
+			using(DetectProjectSettingsTask DetectProjectSettings = new DetectProjectSettingsTask(Project, DataFolder, CacheFolder, new PrefixedTextWriter("  ", Log)))
 			{
 				string ErrorMessage;
 
diff --git a/Engine/Source/Programs/UnrealGameSync/UnrealGameSync/Forms/OpenProjectWindow.cs b/Engine/Source/Programs/UnrealGameSync/UnrealGameSync/Forms/OpenProjectWindow.cs
--- a/Engine/Source/Programs/UnrealGameSync/UnrealGameSync/Forms/OpenProjectWindow.cs
+++ b/Engine/Source/Programs/UnrealGameSync/UnrealGameSync/Forms/OpenProjectWindow.cs
@@ -19,16 +19,18 @@ namespace UnrealGameSync
 		string UserName;
 		DetectProjectSettingsTask DetectedProjectSettings;
 		string DataFolder;
+		string CacheFolder;
 		TextWriter Log;
 		UserSettings Settings;
 
-		private OpenProjectWindow(UserSelectedProjectSettings Project, UserSettings Settings, string DataFolder, TextWriter Log)
+		private OpenProjectWindow(UserSelectedProjectSettings Project, UserSettings Settings, string DataFolder, string CacheFolder, TextWriter Log)
 		{
 			InitializeComponent();
 
 			this.Settings = Settings;
 			this.DetectedProjectSettings = null;
 			this.DataFolder = DataFolder;
+			this.CacheFolder = CacheFolder;
 			this.Log = Log;
 
 			if(Project == null)
@@ -77,9 +79,9 @@ namespace UnrealGameSync
 			UpdateOkButton();
 		}
 
-		public static bool ShowModal(IWin32Window Owner, UserSelectedProjectSettings Project, out DetectProjectSettingsTask NewDetectedProjectSettings, UserSettings Settings, string DataFolder, TextWriter Log)
+		public static bool ShowModal(IWin32Window Owner, UserSelectedProjectSettings Project, out DetectProjectSettingsTask NewDetectedProjectSettings, UserSettings Settings, string DataFolder, string CacheFolder, TextWriter Log)
 		{
-			OpenProjectWindow Window = new OpenProjectWindow(Project, Settings, DataFolder, Log);
+			OpenProjectWindow Window = new OpenProjectWindow(Project, Settings, DataFolder, CacheFolder, Log);
 			if(Window.ShowDialog(Owner) == DialogResult.OK)
 			{
 				NewDetectedProjectSettings = Window.DetectedProjectSettings;
@@ -270,7 +272,7 @@ namespace UnrealGameSync
 			UserSelectedProjectSettings SelectedProject;
 			if(TryGetSelectedProject(out SelectedProject))
 			{
-				DetectProjectSettingsTask NewDetectedProjectSettings = new DetectProjectSettingsTask(SelectedProject, DataFolder, Log);
+				DetectProjectSettingsTask NewDetectedProjectSettings = new DetectProjectSettingsTask(SelectedProject, DataFolder, CacheFolder, Log);
 				try
 				{
 					string ProjectFileName = null;
diff --git a/Engine/Source/Programs/UnrealGameSync/UnrealGameSync/Perforce.cs b/Engine/Source/Programs/UnrealGameSync/UnrealGameSync/Perforce.cs
--- a/Engine/Source/Programs/UnrealGameSync/UnrealGameSync/Perforce.cs
+++ b/Engine/Source/Programs/UnrealGameSync/UnrealGameSync/Perforce.cs
@@ -199,6 +199,7 @@ namespace UnrealGameSync
 		public int Revision;
 		public bool IsMapped;
 		public bool Unmap;
+		public string Digest;
 
 		public PerforceFileRecord(Dictionary<string, string> Tags)
 		{
@@ -245,6 +246,8 @@ namespace UnrealGameSync
 			{
 				int.TryParse(RevisionString, out Revision);
 			}
+
+			Tags.TryGetValue("digest", out Digest);
 		}
 	}
 
@@ -1080,6 +1083,20 @@ namespace UnrealGameSync
 			return RunCommand(String.Format("fstat \"{0}\"", Filter), out FileRecords, CommandOptions.IgnoreFilesNotOnClientError | CommandOptions.IgnoreNoSuchFilesError, Log);
 		}
 
+		public bool Stat(string Options, List<string> Files, out List<PerforceFileRecord> FileRecords, TextWriter Log)
+		{
+			StringBuilder Arguments = new StringBuilder("fstat");
+			if(!String.IsNullOrEmpty(Options))
+			{
+				Arguments.AppendFormat(" {0}", Options);
+			}
+			foreach(string File in Files)
+			{
+				Arguments.AppendFormat(" \"{0}\"", File);
+			}
+			return RunCommand(Arguments.ToString(), out FileRecords, CommandOptions.IgnoreFilesNotOnClientError | CommandOptions.IgnoreNoSuchFilesError, Log);
+		}
+
 		public bool Sync(string Filter, TextWriter Log)
 		{
 			return RunCommand("sync " + Filter, CommandOptions.IgnoreFilesUpToDateError, Log);
diff --git a/Engine/Source/Programs/UnrealGameSync/UnrealGameSync/PerforceMonitor.cs b/Engine/Source/Programs/UnrealGameSync/UnrealGameSync/PerforceMonitor.cs
--- a/Engine/Source/Programs/UnrealGameSync/UnrealGameSync/PerforceMonitor.cs
+++ b/Engine/Source/Programs/UnrealGameSync/UnrealGameSync/PerforceMonitor.cs
@@ -43,6 +43,7 @@ namespace UnrealGameSync
 		BoundedLogWriter LogWriter;
 		bool bIsEnterpriseProject;
 		bool bDisposing;
+		string CacheFolder;
 		List<KeyValuePair<string, DateTime>> LocalConfigFiles;
 
 		public event Action OnUpdate;
@@ -50,7 +51,7 @@ namespace UnrealGameSync
 		public event Action OnStreamChange;
 		public event Action OnLoginExpired;
 
-		public PerforceMonitor(PerforceConnection InPerforce, string InBranchClientPath, string InSelectedClientFileName, string InSelectedProjectIdentifier, string InLogPath, bool bInIsEnterpriseProject, ConfigFile InProjectConfigFile, List<KeyValuePair<string, DateTime>> InLocalConfigFiles)
+		public PerforceMonitor(PerforceConnection InPerforce, string InBranchClientPath, string InSelectedClientFileName, string InSelectedProjectIdentifier, string InLogPath, bool bInIsEnterpriseProject, ConfigFile InProjectConfigFile, string InCacheFolder, List<KeyValuePair<string, DateTime>> InLocalConfigFiles)
 		{
 			Perforce = InPerforce;
 			BranchClientPath = InBranchClientPath;
@@ -61,6 +62,7 @@ namespace UnrealGameSync
 			LastCodeChangeByCurrentUser = -1;
 			bIsEnterpriseProject = bInIsEnterpriseProject;
 			LatestProjectConfigFile = InProjectConfigFile;
+			CacheFolder = InCacheFolder;
 			LocalConfigFiles = InLocalConfigFiles;
 
 			LogWriter = new BoundedLogWriter(InLogPath);
@@ -499,57 +501,56 @@ namespace UnrealGameSync
 		void UpdateProjectConfigFile()
 		{
 			LocalConfigFiles.Clear();
-			LatestProjectConfigFile = ReadProjectConfigFile(Perforce, BranchClientPath, SelectedClientFileName, LocalConfigFiles, LogWriter);
+			LatestProjectConfigFile = ReadProjectConfigFile(Perforce, BranchClientPath, SelectedClientFileName, CacheFolder, LocalConfigFiles, LogWriter);
 		}
 
-		public static ConfigFile ReadProjectConfigFile(PerforceConnection Perforce, string BranchClientPath, string SelectedClientFileName, List<KeyValuePair<string, DateTime>> LocalConfigFiles, TextWriter Log)
+		public static ConfigFile ReadProjectConfigFile(PerforceConnection Perforce, string BranchClientPath, string SelectedClientFileName, string CacheFolder, List<KeyValuePair<string, DateTime>> LocalConfigFiles, TextWriter Log)
 		{
 			List<string> ConfigFilePaths = Utility.GetConfigFileLocations(BranchClientPath, SelectedClientFileName, '/');
 
-			List<PerforceFileRecord> OpenFiles;
-			Perforce.GetOpenFiles(String.Format("{0}/....ini", BranchClientPath), out OpenFiles, Log);
-
 			ConfigFile ProjectConfig = new ConfigFile();
-			foreach(string ConfigFilePath in ConfigFilePaths)
-			{
-				List<string> Lines = null;
 
-				// If this file is open for edit, read the local version
-				if(OpenFiles != null && OpenFiles.Any(x => x.ClientPath.Equals(ConfigFilePath, StringComparison.InvariantCultureIgnoreCase)))
+			List<PerforceFileRecord> FileRecords;
+			if(Perforce.Stat("-Ol", ConfigFilePaths, out FileRecords, Log))
+			{
+				foreach(PerforceFileRecord FileRecord in FileRecords)
 				{
-					try
+					List<string> Lines = null;
+
+					// If this file is open for edit, read the local version
+					string LocalFileName = FileRecord.ClientPath;
+					if(LocalFileName != null && File.Exists(LocalFileName) && (File.GetAttributes(LocalFileName) & FileAttributes.ReadOnly) == 0)
 					{
-						string LocalFileName;
-						if(Perforce.ConvertToLocalPath(ConfigFilePath, out LocalFileName, Log))
+						try
 						{
 							DateTime LastModifiedTime = File.GetLastWriteTimeUtc(LocalFileName);
 							LocalConfigFiles.Add(new KeyValuePair<string, DateTime>(LocalFileName, LastModifiedTime));
 							Lines = File.ReadAllLines(LocalFileName).ToList();
 						}
+						catch(Exception Ex)
+						{
+							Log.WriteLine("Failed to read local config file for {0}: {1}", LocalFileName, Ex.ToString());
+						}
 					}
-					catch(Exception Ex)
-					{
-						Log.WriteLine("Failed to read local config file for {0}: {1}", ConfigFilePath, Ex.ToString());
-					}
-				}
 
-				// Otherwise try to get it from perforce
-				if(Lines == null)
-				{
-					Perforce.Print(ConfigFilePath, out Lines, Log);
-				}
-
-				// Merge the text with the config file
-				if(Lines != null)
-				{
-					try
+					// Otherwise try to get it from perforce
+					if(Lines == null)
 					{
-						ProjectConfig.Parse(Lines.ToArray());
-						Log.WriteLine("Read config file from {0}", ConfigFilePath);
+						Utility.TryPrintFileUsingCache(Perforce, FileRecord.DepotPath, CacheFolder, FileRecord.Digest, out Lines, Log);
 					}
-					catch(Exception Ex)
+
+					// Merge the text with the config file
+					if(Lines != null)
 					{
-						Log.WriteLine("Failed to read config file from {0}: {1}", ConfigFilePath, Ex.ToString());
+						try
+						{
+							ProjectConfig.Parse(Lines.ToArray());
+							Log.WriteLine("Read config file from {0}", FileRecord.DepotPath);
+						}
+						catch(Exception Ex)
+						{
+							Log.WriteLine("Failed to read config file from {0}: {1}", FileRecord.DepotPath, Ex.ToString());
+						}
 					}
 				}
 			}
diff --git a/Engine/Source/Programs/UnrealGameSync/UnrealGameSync/ProgramApplicationContext.cs b/Engine/Source/Programs/UnrealGameSync/UnrealGameSync/ProgramApplicationContext.cs
--- a/Engine/Source/Programs/UnrealGameSync/UnrealGameSync/ProgramApplicationContext.cs
+++ b/Engine/Source/Programs/UnrealGameSync/UnrealGameSync/ProgramApplicationContext.cs
@@ -21,6 +21,7 @@ namespace UnrealGameSync
 		UpdateMonitor UpdateMonitor;
 		string ApiUrl;
 		string DataFolder;
+		string CacheFolder;
 		bool bRestoreState;
 		string UpdateSpawn;
 		bool bUnstable;
@@ -50,10 +51,15 @@ namespace UnrealGameSync
 			this.UpdateMonitor = UpdateMonitor;
 			this.ApiUrl = ApiUrl;
 			this.DataFolder = DataFolder;
+			this.CacheFolder = Path.Combine(DataFolder, "Cache");
 			this.bRestoreState = bRestoreState;
 			this.UpdateSpawn = UpdateSpawn;
 			this.bUnstable = bUnstable;
 
+			// Create the directories
+			Directory.CreateDirectory(DataFolder);
+			Directory.CreateDirectory(CacheFolder);
+
 			// Make sure a synchronization context is set. We spawn a bunch of threads (eg. UpdateMonitor) at startup, and need to make sure we can post messages 
 			// back to the main thread at any time.
 			if(SynchronizationContext.Current == null)
@@ -152,7 +158,7 @@ namespace UnrealGameSync
 				StartupLog.WriteLine("Detecting settings for {0}", OpenProject);
 				StartupLogs.Add(StartupLog);
 
-				Tasks.Add(new DetectProjectSettingsTask(OpenProject, DataFolder, new TimestampLogWriter(new PrefixedTextWriter("  ", StartupLog))));
+				Tasks.Add(new DetectProjectSettingsTask(OpenProject, DataFolder, CacheFolder, new TimestampLogWriter(new PrefixedTextWriter("  ", StartupLog))));
 			}
 
 			// Detect settings for the project we want to open
@@ -190,11 +196,14 @@ namespace UnrealGameSync
 				}
 			}
 
+			// Clear out the cache folder
+			Utility.ClearPrintCache(CacheFolder);
+
 			// Get a list of all the valid projects to open
 			DetectProjectSettingsResult[] StartupProjects = DetectStartupProjectSettingsTask.Results.Where(x => x != null).ToArray();
 
 			// Create the main window
-			MainWindowInstance = new MainWindow(UpdateMonitor, ApiUrl, DataFolder, bRestoreState, UpdateSpawn ?? Assembly.GetExecutingAssembly().Location, bUnstable, StartupProjects, Log, Settings);
+			MainWindowInstance = new MainWindow(UpdateMonitor, ApiUrl, DataFolder, CacheFolder, bRestoreState, UpdateSpawn ?? Assembly.GetExecutingAssembly().Location, bUnstable, StartupProjects, Log, Settings);
 			if(bVisible)
 			{
 				MainWindowInstance.Show();
diff --git a/Engine/Source/Programs/UnrealGameSync/UnrealGameSync/Utility.cs b/Engine/Source/Programs/UnrealGameSync/UnrealGameSync/Utility.cs
--- a/Engine/Source/Programs/UnrealGameSync/UnrealGameSync/Utility.cs
+++ b/Engine/Source/Programs/UnrealGameSync/UnrealGameSync/Utility.cs
@@ -347,5 +347,77 @@ namespace UnrealGameSync
 				MessageBox.Show("Unable to save settings.\n\n" + Ex.ToString());
 			}
 		}
+
+		public static bool TryPrintFileUsingCache(PerforceConnection Perforce, string DepotPath, string CacheFolder, string Digest, out List<string> Lines, TextWriter Log)
+		{
+			string CacheFile = Path.Combine(CacheFolder, Digest);
+			if(File.Exists(CacheFile))
+			{
+				Log.WriteLine("Reading cached copy of {0} from {1}", DepotPath, CacheFile);
+				Lines = new List<string>(File.ReadAllLines(CacheFile));
+				try
+				{
+					File.SetLastWriteTimeUtc(CacheFile, DateTime.UtcNow);
+				}
+				catch(Exception Ex)
+				{
+					Log.WriteLine("Exception touching cache file {0}: {1}", CacheFile, Ex.ToString());
+				}
+				return true;
+			}
+			else
+			{
+				string TempFile = String.Format("{0}.{1}.temp", CacheFile, Guid.NewGuid());
+				if(!Perforce.PrintToFile(DepotPath, TempFile, Log))
+				{
+					Lines = null;
+					return false;
+				}
+				else
+				{
+					Lines = new List<string>(File.ReadAllLines(TempFile));
+					try
+					{
+						File.SetAttributes(TempFile, FileAttributes.Normal);
+						File.SetLastWriteTimeUtc(TempFile, DateTime.UtcNow);
+						File.Move(TempFile, CacheFile);
+					}
+					catch
+					{
+						try
+						{
+							File.Delete(TempFile);
+						}
+						catch
+						{
+						}
+					}
+					return true;
+				}
+			}
+		}
+
+		public static void ClearPrintCache(string CacheFolder)
+		{
+			DirectoryInfo CacheDir = new DirectoryInfo(CacheFolder);
+			if(CacheDir.Exists)
+			{
+				DateTime DeleteTime = DateTime.UtcNow - TimeSpan.FromDays(5.0);
+				foreach(FileInfo CacheFile in CacheDir.EnumerateFiles())
+				{
+					if(CacheFile.LastWriteTimeUtc < DeleteTime || CacheFile.Name.EndsWith(".temp", StringComparison.OrdinalIgnoreCase))
+					{
+						try
+						{
+							CacheFile.Attributes = FileAttributes.Normal;
+							CacheFile.Delete();
+						}
+						catch
+						{
+						}
+					}
+				}
+			}
+		}
 	}
 }
