UGS: Use the local project file to determine if a project is an enterprise project (if it's available). Reduces number of P4 calls during startup.

diff --git a/Engine/Source/Programs/UnrealGameSync/UnrealGameSync/DetectProjectSettingsTask.cs b/Engine/Source/Programs/UnrealGameSync/UnrealGameSync/DetectProjectSettingsTask.cs
--- a/Engine/Source/Programs/UnrealGameSync/UnrealGameSync/DetectProjectSettingsTask.cs
+++ b/Engine/Source/Programs/UnrealGameSync/UnrealGameSync/DetectProjectSettingsTask.cs
@@ -290,11 +290,24 @@ namespace UnrealGameSync
 				}
 			}
 
-			// Figure out if it's an enterprise project
-			List<string> ProjectLines;
-			if(NewSelectedClientFileName.EndsWith(".uproject", StringComparison.InvariantCultureIgnoreCase) && PerforceClient.Print(NewSelectedClientFileName, out ProjectLines, Log))
+			// Figure out if it's an enterprise project. Use the synced version if we have it.
+			if(NewSelectedClientFileName.EndsWith(".uproject", StringComparison.InvariantCultureIgnoreCase))
 			{
-				string Text = String.Join("\n", ProjectLines);
+				string Text;
+				if(File.Exists(NewSelectedFileName))
+				{
+					Text = File.ReadAllText(NewSelectedFileName);
+				}
+				else
+				{
+					List<string> ProjectLines;
+					if(PerforceClient.Print(NewSelectedClientFileName, out ProjectLines, Log))
+					{
+						ErrorMessage = String.Format("Unable to get contents of {0}", NewSelectedClientFileName);
+						return false;
+					}
+					Text = String.Join("\n", ProjectLines);
+				}
 				bIsEnterpriseProject = Utility.IsEnterpriseProjectFromText(Text);
 			}
 
