UGS: Move query for server timezone to PerforceMonitor, to reduce startup time in the modal dialog.

diff --git a/Engine/Source/Programs/UnrealGameSync/UnrealGameSync/Controls/WorkspaceControl.cs b/Engine/Source/Programs/UnrealGameSync/UnrealGameSync/Controls/WorkspaceControl.cs
--- a/Engine/Source/Programs/UnrealGameSync/UnrealGameSync/Controls/WorkspaceControl.cs
+++ b/Engine/Source/Programs/UnrealGameSync/UnrealGameSync/Controls/WorkspaceControl.cs
@@ -198,8 +198,6 @@ namespace UnrealGameSync
 
 		Dictionary<string, int> NotifiedBuildTypeToChangeNumber = new Dictionary<string,int>();
 
-		TimeSpan ServerTimeZone;
-
 		string HoverBadgeUniqueId = null;
 		bool bHoverSync;
 		PerforceChangeSummary ContextMenuChange;
@@ -283,7 +281,6 @@ namespace UnrealGameSync
 			EditorTargetName = DetectSettings.NewProjectEditorTarget;
 			bIsEnterpriseProject = DetectSettings.bIsEnterpriseProject;
 			StreamName = DetectSettings.StreamName;
-			ServerTimeZone = DetectSettings.ServerTimeZone;
 
 			// Update the branch directory
 			BranchDirectoryName = DetectSettings.BranchDirectoryName;
@@ -1010,7 +1007,7 @@ namespace UnrealGameSync
 							DateTime DisplayTime = Change.Date;
 							if(Settings.bShowLocalTimes)
 							{
-								DisplayTime = (DisplayTime - ServerTimeZone).ToLocalTime();
+								DisplayTime = (DisplayTime - PerforceMonitor.ServerTimeZone).ToLocalTime();
 							}
 
 							string[] SubItemLabels = new string[BuildList.Columns.Count];
diff --git a/Engine/Source/Programs/UnrealGameSync/UnrealGameSync/DetectProjectSettingsTask.cs b/Engine/Source/Programs/UnrealGameSync/UnrealGameSync/DetectProjectSettingsTask.cs
--- a/Engine/Source/Programs/UnrealGameSync/UnrealGameSync/DetectProjectSettingsTask.cs
+++ b/Engine/Source/Programs/UnrealGameSync/UnrealGameSync/DetectProjectSettingsTask.cs
@@ -23,7 +23,6 @@ namespace UnrealGameSync
 		public string NewSelectedClientFileName;
 		public string StreamName;
 		public Image ProjectLogo;
-		public TimeSpan ServerTimeZone;
 		public string DataFolder;
 		public string CacheFolder;
 		public bool bIsEnterpriseProject;
@@ -88,17 +87,6 @@ namespace UnrealGameSync
 
 		public bool Run(PerforceConnection Perforce, TextWriter Log, out string ErrorMessage)
 		{
-			// Get the perforce server settings
-			PerforceInfoRecord PerforceInfo;
-			if(!Perforce.Info(out PerforceInfo, Log))
-			{
-				ErrorMessage = String.Format("Couldn't get Perforce server info");
-				return false;
-			}
-
-			// Configure the time zone
-			ServerTimeZone = PerforceInfo.ServerTimeZone;
-
 			// Use the cached client path to the file if it's available; it's much quicker than trying to find the correct workspace.
 			if(!String.IsNullOrEmpty(SelectedProject.ClientPath))
 			{
@@ -125,6 +113,14 @@ namespace UnrealGameSync
 			}
 			else
 			{
+				// Get the perforce server settings
+				PerforceInfoRecord PerforceInfo;
+				if(!Perforce.Info(out PerforceInfo, Log))
+				{
+					ErrorMessage = String.Format("Couldn't get Perforce server info");
+					return false;
+				}
+
 				// Use the path as the selected filename
 				NewSelectedFileName = SelectedProject.LocalPath;
 
diff --git a/Engine/Source/Programs/UnrealGameSync/UnrealGameSync/PerforceMonitor.cs b/Engine/Source/Programs/UnrealGameSync/UnrealGameSync/PerforceMonitor.cs
--- a/Engine/Source/Programs/UnrealGameSync/UnrealGameSync/PerforceMonitor.cs
+++ b/Engine/Source/Programs/UnrealGameSync/UnrealGameSync/PerforceMonitor.cs
@@ -51,6 +51,12 @@ namespace UnrealGameSync
 		public event Action OnStreamChange;
 		public event Action OnLoginExpired;
 
+		public TimeSpan ServerTimeZone
+		{
+			get;
+			private set;
+		}
+
 		public PerforceMonitor(PerforceConnection InPerforce, string InBranchClientPath, string InSelectedClientFileName, string InSelectedProjectIdentifier, string InLogPath, bool bInIsEnterpriseProject, ConfigFile InProjectConfigFile, string InCacheFolder, List<KeyValuePair<string, DateTime>> InLocalConfigFiles)
 		{
 			Perforce = InPerforce;
@@ -125,6 +131,13 @@ namespace UnrealGameSync
 				StreamName = null;
 			}
 
+			// Get the perforce server settings
+			PerforceInfoRecord PerforceInfo;
+			if(Perforce.Info(out PerforceInfo, LogWriter))
+			{
+				ServerTimeZone = PerforceInfo.ServerTimeZone;
+			}
+
 			// Try to update the zipped binaries list before anything else, because it causes a state change in the UI
 			UpdateZippedBinaries();
 
