UGS: Remove a few more additional unnecessary P4 calls on startup.

diff --git a/Engine/Source/Programs/UnrealGameSync/UnrealGameSync/DetectProjectSettingsTask.cs b/Engine/Source/Programs/UnrealGameSync/UnrealGameSync/DetectProjectSettingsTask.cs
--- a/Engine/Source/Programs/UnrealGameSync/UnrealGameSync/DetectProjectSettingsTask.cs
+++ b/Engine/Source/Programs/UnrealGameSync/UnrealGameSync/DetectProjectSettingsTask.cs
@@ -104,11 +104,18 @@ namespace UnrealGameSync
 				// Create the client
 				PerforceClient = new PerforceConnection(Perforce.UserName, ClientName, Perforce.ServerAndPort);
 
-				// Figure out the path on the client
-				if(!PerforceClient.ConvertToLocalPath(NewSelectedClientFileName, out NewSelectedFileName, Log))
+				// Figure out the path on the client. Use the cached location if it's valid.
+				if(SelectedProject.LocalPath != null && File.Exists(SelectedProject.LocalPath))
 				{
-					ErrorMessage = String.Format("Couldn't get client path for {0}", NewSelectedFileName);
-					return false;
+					NewSelectedFileName = SelectedProject.LocalPath;
+				}
+				else
+				{
+					if(!PerforceClient.ConvertToLocalPath(NewSelectedClientFileName, out NewSelectedFileName, Log))
+					{
+						ErrorMessage = String.Format("Couldn't get client path for {0}", NewSelectedFileName);
+						return false;
+					}
 				}
 			}
 			else
@@ -191,7 +198,12 @@ namespace UnrealGameSync
 			SelectedProject = new UserSelectedProjectSettings(Perforce.ServerAndPort, Perforce.UserName, SelectedProject.Type, NewSelectedClientFileName, NewSelectedFileName);
 
 			// Figure out where the engine is in relation to it
-			for(int EndIdx = NewSelectedClientFileName.Length - 1;;EndIdx--)
+			int EndIdx = NewSelectedClientFileName.Length - 1;
+			if(EndIdx != -1 && NewSelectedClientFileName.EndsWith(".uproject", StringComparison.InvariantCultureIgnoreCase))
+			{
+				EndIdx = NewSelectedClientFileName.LastIndexOf('/') - 1;
+			}
+			for(;;EndIdx--)
 			{
 				if(EndIdx < 2)
 				{
@@ -200,25 +212,23 @@ namespace UnrealGameSync
 				}
 				if(NewSelectedClientFileName[EndIdx] == '/')
 				{
-					bool bFileExists;
-					if(PerforceClient.FileExists(NewSelectedClientFileName.Substring(0, EndIdx) + "/Engine/Build/Build.version", out bFileExists, Log) && bFileExists)
+					List<PerforceFileRecord> FileRecords;
+					if(PerforceClient.Stat(NewSelectedClientFileName.Substring(0, EndIdx) + "/Engine/Build/Build.version", out FileRecords, Log) && FileRecords.Count > 0)
 					{
+						if(FileRecords[0].ClientPath == null)
+						{
+							ErrorMessage = String.Format("Missing client path for {0}", FileRecords[0].DepotPath);
+							return false;
+						}
+
 						BranchClientPath = NewSelectedClientFileName.Substring(0, EndIdx);
+						BranchDirectoryName = Path.GetFullPath(Path.Combine(Path.GetDirectoryName(FileRecords[0].ClientPath), "..", ".."));
 						break;
 					}
 				}
 			}
 			Log.WriteLine("Found branch root at {0}", BranchClientPath);
 
-			// Get the local branch root
-			string BuildVersionPath;
-			if(!PerforceClient.ConvertToLocalPath(BranchClientPath + "/Engine/Build/Build.version", out BuildVersionPath, Log))
-			{
-				ErrorMessage = String.Format("Couldn't get local path for Engine/Build/Build.version");
-				return false;
-			}
-			BranchDirectoryName = Path.GetFullPath(Path.Combine(Path.GetDirectoryName(BuildVersionPath), "..", ".."));
-
 			// Find the editor target for this project
 			if(NewSelectedFileName.EndsWith(".uproject", StringComparison.InvariantCultureIgnoreCase))
 			{
