UGS: Allow scripting UGS via a local socket. To enable, set the HKEY_CURRENT_USER\Software\Epic Games\UnrealGameSync\AutomationPort to a port number (eg. 30422) and restart UGS.

The generalized format of automation exchanges is:

Client:
    4 bytes: Request type (see AutomationRequestType)
    4 bytes: Request data length (X)
    X bytes: Request data

Server:
    4 bytes: Result code (see AutomationRequestResult)
    4 bytes: Response data length (X)
    X bytes: Response data

At the moment, the only request type supported, SyncProject, has a code of zero. The request data consists of two, UTF-8, null terminated strings - the stream name, and relative path to the file within in (starting with a '/' character). On success, the response will be the path to the project file on the local machine. The request will prompt the user for the workspace they would like to use for the sync (defaulting to any suitable workspace they may already have), change streams if necessary, and sync it with their current settings.

diff --git a/Engine/Source/Programs/UnrealGameSync/UnrealGameSync/AutomationServer.cs b/Engine/Source/Programs/UnrealGameSync/UnrealGameSync/AutomationServer.cs
new file mode 100644
--- /dev/null
+++ b/Engine/Source/Programs/UnrealGameSync/UnrealGameSync/AutomationServer.cs
@@ -0,0 +1,234 @@
+using Microsoft.Win32;
+using System;
+using System.Collections.Generic;
+using System.IO;
+using System.Linq;
+using System.Net;
+using System.Net.Sockets;
+using System.Text;
+using System.Threading;
+using System.Threading.Tasks;
+
+namespace UnrealGameSync
+{
+	enum AutomationRequestType
+	{
+		SyncProject
+	}
+
+	class AutomationRequestInput
+	{
+		public AutomationRequestType Type;
+		public byte[] Data;
+
+		public AutomationRequestInput(AutomationRequestType Type, byte[] Data)
+		{
+			this.Type = Type;
+			this.Data = Data;
+		}
+
+		public static AutomationRequestInput Read(Stream InputStream)
+		{
+			BinaryReader Reader = new BinaryReader(InputStream);
+			
+			int Type = Reader.ReadInt32();
+			int InputSize = Reader.ReadInt32();
+			byte[] Input = Reader.ReadBytes(InputSize);
+
+			return new AutomationRequestInput((AutomationRequestType)Type, Input);
+		}
+
+		public void Write(Stream OutputStream)
+		{
+			BinaryWriter Writer = new BinaryWriter(OutputStream);
+
+			Writer.Write((int)Type);
+			Writer.Write(Data.Length);
+			Writer.Write(Data);
+		}
+	}
+
+	enum AutomationRequestResult
+	{
+		Ok,
+		Invalid,
+		Busy,
+		Canceled,
+		Error
+	}
+
+	class AutomationRequestOutput
+	{
+		public AutomationRequestResult Result;
+		public byte[] Data;
+
+		public AutomationRequestOutput(AutomationRequestResult Result)
+		{
+			this.Result = Result;
+			this.Data = new byte[0];
+		}
+
+		public AutomationRequestOutput(AutomationRequestResult Result, byte[] Data)
+		{
+			this.Result = Result;
+			this.Data = Data;
+		}
+
+		public static AutomationRequestOutput Read(Stream InputStream)
+		{
+			using(BinaryReader Reader = new BinaryReader(InputStream))
+			{
+				AutomationRequestResult Result = (AutomationRequestResult)Reader.ReadInt32();
+				int DataSize = Reader.ReadInt32();
+				byte[] Data = Reader.ReadBytes(DataSize);
+				return new AutomationRequestOutput(Result, Data);
+			}
+		}
+
+		public void Write(Stream OutputStream)
+		{
+			using(BinaryWriter Writer = new BinaryWriter(OutputStream))
+			{
+				Writer.Write((int)Result);
+				Writer.Write(Data.Length);
+				Writer.Write(Data);
+			}
+		}
+	}
+
+	class AutomationRequest : IDisposable
+	{
+		public AutomationRequestInput Input;
+		public AutomationRequestOutput Output;
+		public ManualResetEventSlim Complete;
+
+		public AutomationRequest(AutomationRequestInput Input)
+		{
+			this.Input = Input;
+			this.Complete = new ManualResetEventSlim(false);
+		}
+
+		public void SetOutput(AutomationRequestOutput Output)
+		{
+			this.Output = Output;
+			Complete.Set();
+		}
+
+		public void Dispose()
+		{
+			if(Complete != null)
+			{
+				Complete.Dispose();
+				Complete = null;
+			}
+		}
+	}
+
+	class AutomationServer : IDisposable
+	{
+		TcpListener Listener;
+		TcpClient CurrentClient;
+		Thread BackgroundThread;
+		Action<AutomationRequest> PostRequest;
+		bool bDisposing;
+		TextWriter Log;
+
+		public AutomationServer(Action<AutomationRequest> PostRequest, TextWriter Log)
+		{
+			this.PostRequest = PostRequest;
+			this.Log = Log;
+
+			object PortValue = Registry.GetValue("HKEY_CURRENT_USER\\Software\\Epic Games\\UnrealGameSync", "AutomationPort", null);
+			if(PortValue != null)
+			{
+				try
+				{
+					int PortNumber = (int)PortValue;
+
+					Listener = new TcpListener(IPAddress.Loopback, PortNumber);
+					Listener.Start();
+
+					BackgroundThread = new Thread(() => Run());
+					BackgroundThread.Start();
+				}
+				catch(Exception Ex)
+				{
+					Log.WriteLine("Unable to start automation server: {0}", Ex.ToString());
+				}
+			}
+		}
+
+		public void Run()
+		{
+			try
+			{
+				for(;;)
+				{
+					TcpClient Client = CurrentClient = Listener.AcceptTcpClient();
+					try
+					{
+						Log.WriteLine("Accepted connection from {0}", Client.Client.RemoteEndPoint);
+
+						NetworkStream Stream = Client.GetStream();
+
+						AutomationRequestInput Input = AutomationRequestInput.Read(Stream);
+						Log.WriteLine("Received input: {0} (+{1} bytes)", Input.Type, Input.Data.Length);
+
+						AutomationRequestOutput Output;
+						using(AutomationRequest Request = new AutomationRequest(Input))
+						{
+							PostRequest(Request);
+							Request.Complete.Wait();
+							Output = Request.Output;
+						}
+
+						Output.Write(Stream);
+						Log.WriteLine("Sent output: {0} (+{1} bytes)", Output.Result, Output.Data.Length);
+					}
+					catch(Exception Ex)
+					{
+						Log.WriteLine("Exception: {0}", Ex.ToString());
+					}
+					finally
+					{
+						Client.Close();
+						Log.WriteLine("Closed connection.");
+					}
+					CurrentClient = null;
+				}
+			}
+			catch(Exception Ex)
+			{
+				if(!bDisposing)
+				{
+					Log.WriteLine("Exception: {0}", Ex.ToString());
+				}
+			}
+			Log.WriteLine("Closing socket.");
+		}
+		
+		public void Dispose()
+		{
+			bDisposing = true;
+
+			TcpClient Client = CurrentClient;
+			if(Client != null)
+			{
+				try { Client.Close(); } catch { }
+				Client = null;
+			}
+
+			if(Listener != null)
+			{
+				Listener.Stop();
+				Listener = null;
+			}
+
+			if(BackgroundThread != null)
+			{
+				BackgroundThread.Join();
+				BackgroundThread = null;
+			}
+		}
+	}
+}
diff --git a/Engine/Source/Programs/UnrealGameSync/UnrealGameSync/Controls/WorkspaceControl.cs b/Engine/Source/Programs/UnrealGameSync/UnrealGameSync/Controls/WorkspaceControl.cs
--- a/Engine/Source/Programs/UnrealGameSync/UnrealGameSync/Controls/WorkspaceControl.cs
+++ b/Engine/Source/Programs/UnrealGameSync/UnrealGameSync/Controls/WorkspaceControl.cs
@@ -44,6 +44,9 @@ namespace UnrealGameSync
 		void ModifyApplicationSettings();
 	}
 
+	delegate void WorkspaceStartupCallback(WorkspaceControl Workspace, bool bCancel);
+	delegate void WorkspaceUpdateCallback(WorkspaceUpdateResult Result);
+
 	partial class WorkspaceControl : UserControl, IMainWindowTabPanel
 	{
 		enum HorizontalAlignment
@@ -231,6 +234,11 @@ namespace UnrealGameSync
 		bool bUpdateBuildMetadataPosted;
 		bool bUpdateReviewsPosted;
 
+		WorkspaceUpdateCallback UpdateCallback;
+
+		System.Threading.Timer StartupTimer;
+		List<WorkspaceStartupCallback> StartupCallbacks;
+
 		public WorkspaceControl(IWorkspaceControlOwner InOwner, string InApiUrl, string InOriginalExecutableFileName, bool bInUnstable, DetectProjectSettingsTask DetectSettings, LineBasedTextWriter InLog, UserSettings InSettings)
 		{
 			InitializeComponent();
@@ -331,6 +339,51 @@ namespace UnrealGameSync
 
 			PerforceMonitor.Start();
 			EventMonitor.Start();
+
+			StartupTimer = new System.Threading.Timer(x => MainThreadSynchronizationContext.Post((o) => { if(!IsDisposed){ StartupTimerElapsed(false); } }, null), null, TimeSpan.FromSeconds(20.0), TimeSpan.FromMilliseconds(-1.0));
+			StartupCallbacks = new List<WorkspaceStartupCallback>();
+		}
+
+		private void CheckForStartupComplete()
+		{
+			if(StartupTimer != null)
+			{
+				int LatestChangeNumber;
+				if(FindChangeToSync(Settings.SyncType, out LatestChangeNumber))
+				{
+					StartupTimerElapsed(false);
+				}
+			}
+		}
+
+		private void StartupTimerElapsed(bool bCancel)
+		{
+			if(StartupTimer != null)
+			{
+				StartupTimer.Dispose();
+				StartupTimer = null;
+			}
+
+			if(StartupCallbacks != null)
+			{
+				foreach(WorkspaceStartupCallback StartupCallback in StartupCallbacks)
+				{
+					StartupCallback(this, bCancel);
+				}
+				StartupCallbacks = null;
+			}
+		}
+
+		public void AddStartupCallback(WorkspaceStartupCallback StartupCallback)
+		{
+			if(StartupTimer == null)
+			{
+				StartupCallback(this, false);
+			}
+			else
+			{
+				StartupCallbacks.Add(StartupCallback);
+			}
 		}
 
 		private void UpdateColumnSettings()
@@ -517,6 +570,20 @@ namespace UnrealGameSync
 
 			UpdateTimer.Stop();
 
+			if(StartupCallbacks != null)
+			{
+				foreach(WorkspaceStartupCallback StartupCallback in StartupCallbacks)
+				{
+					StartupCallback(this, true);
+				}
+				StartupCallbacks = null;
+			}
+
+			if(StartupTimer != null)
+			{
+				StartupTimer.Dispose();
+				StartupTimer = null;
+			}
 			if(NotificationWindow != null)
 			{
 				NotificationWindow.Dispose();
@@ -687,6 +754,11 @@ namespace UnrealGameSync
 		}
 
 		void StartSync(int ChangeNumber)
+		{
+			StartSync(ChangeNumber, null);
+		}
+
+		void StartSync(int ChangeNumber, WorkspaceUpdateCallback Callback)
 		{
 			WorkspaceUpdateOptions Options = WorkspaceUpdateOptions.Sync | WorkspaceUpdateOptions.SyncArchives | WorkspaceUpdateOptions.GenerateProjectFiles;
 			if(Settings.bAutoResolveConflicts)
@@ -713,10 +785,15 @@ namespace UnrealGameSync
 			{
 				Options |= WorkspaceUpdateOptions.SyncAllProjects;
 			}
-			StartWorkspaceUpdate(ChangeNumber, Options);
+			StartWorkspaceUpdate(ChangeNumber, Options, Callback);
 		}
 
 		void StartWorkspaceUpdate(int ChangeNumber, WorkspaceUpdateOptions Options)
+		{
+			StartWorkspaceUpdate(ChangeNumber, Options, null);
+		}
+
+		void StartWorkspaceUpdate(int ChangeNumber, WorkspaceUpdateOptions Options, WorkspaceUpdateCallback Callback)
 		{
 			if((Options & (WorkspaceUpdateOptions.Sync | WorkspaceUpdateOptions.Build)) != 0 && GetProcessesRunningInWorkspace().Length > 0)
 			{
@@ -753,11 +830,13 @@ namespace UnrealGameSync
 				}
 				Context.ArchiveTypeToDepotPath.Add(EditorArchiveType, EditorArchivePath);
 			}
-			StartWorkspaceUpdate(Context);
+			StartWorkspaceUpdate(Context, Callback);
 		}
 
-		void StartWorkspaceUpdate(WorkspaceUpdateContext Context)
+		void StartWorkspaceUpdate(WorkspaceUpdateContext Context, WorkspaceUpdateCallback Callback)
 		{
+			UpdateCallback = Callback;
+
 			Context.StartTime = DateTime.UtcNow;
 			Context.PerforceSyncOptions = (PerforceSyncOptions)Settings.SyncOptions.Clone();
 
@@ -810,6 +889,12 @@ namespace UnrealGameSync
 
 				Workspace.CancelUpdate();
 
+				if(UpdateCallback != null)
+				{
+					UpdateCallback(WorkspaceUpdateResult.Canceled);
+					UpdateCallback = null;
+				}
+
 				UpdateTimer.Stop();
 
 				UpdateSyncActionCheckboxes();
@@ -863,7 +948,7 @@ namespace UnrealGameSync
 				DeleteWindow Window = new DeleteWindow(Context.DeleteFiles);
 				if(Window.ShowDialog(this) == DialogResult.OK)
 				{
-					StartWorkspaceUpdate(Context);
+					StartWorkspaceUpdate(Context, UpdateCallback);
 					return;
 				}
 			}
@@ -875,7 +960,7 @@ namespace UnrealGameSync
 				ClobberWindow Window = new ClobberWindow(Context.ClobberFiles);
 				if(Window.ShowDialog(this) == DialogResult.OK)
 				{
-					StartWorkspaceUpdate(Context);
+					StartWorkspaceUpdate(Context, UpdateCallback);
 					return;
 				}
 			}
@@ -899,6 +984,12 @@ namespace UnrealGameSync
 				}
 			}
 
+			if(UpdateCallback != null)
+			{
+				UpdateCallback(Result);
+				UpdateCallback = null;
+			}
+
 			DesiredTaskbarState = Tuple.Create((Result == WorkspaceUpdateResult.Success)? TaskbarState.NoProgress : TaskbarState.Error, 0.0f);
 			Owner.UpdateProgress();
 
@@ -1176,6 +1267,7 @@ namespace UnrealGameSync
 			UpdateServiceBadges();
 			UpdateStatusPanel();
 			UpdateBuildFailureNotification();
+			CheckForStartupComplete();
 		}
 
 		void UpdateMaxBuildBadgeChars()
@@ -1270,6 +1362,7 @@ namespace UnrealGameSync
 			EventMonitor.ApplyUpdates();
 			Refresh();
 			UpdateBuildFailureNotification();
+			CheckForStartupComplete();
 		}
 
 		void UpdateBuildFailureNotification()
@@ -2259,7 +2352,7 @@ namespace UnrealGameSync
 					}
 					else
 					{
-						StartSync(Change.Number);
+						StartSync(Change.Number, null);
 					}
 				}
 			}
@@ -2316,7 +2409,7 @@ namespace UnrealGameSync
 						}
 		
 						WorkspaceUpdateContext Context = new WorkspaceUpdateContext(Workspace.CurrentChangeNumber, Options, null, GetDefaultBuildStepObjects(), ProjectSettings.BuildSteps, null, GetWorkspaceVariables(Workspace.CurrentChangeNumber));
-						StartWorkspaceUpdate(Context);
+						StartWorkspaceUpdate(Context, null);
 					}
 				}
 			}
@@ -2888,7 +2981,7 @@ namespace UnrealGameSync
 				}
 				else if(Workspace.Perforce != null)
 				{
-					if(!Workspace.Perforce.HasOpenFiles(Log) || MessageBox.Show("You have files open for edit in this workspace. If you continue, you will not be able to submit them until you switch back.\n\nContinue switching workspaces?", "Files checked out", MessageBoxButtons.YesNo) == DialogResult.Yes)
+					if(!Workspace.Perforce.HasOpenFiles(Log) || MessageBox.Show("You have files open for edit in this workspace. If you continue, you will not be able to submit them until you switch back.\n\nContinue switching streams?", "Files checked out", MessageBoxButtons.YesNo) == DialogResult.Yes)
 					{
 						if(Workspace.Perforce.SwitchStream(NewStreamName, Log))
 						{
@@ -3002,7 +3095,7 @@ namespace UnrealGameSync
 						Rectangle SyncBadgeRectangle = GetSyncBadgeRectangle(HitTest.Item.SubItems[StatusColumn.Index].Bounds);
 						if(SyncBadgeRectangle.Contains(Args.Location) && CanSyncChange(Change.Number))
 						{
-							StartSync(Change.Number);
+							StartSync(Change.Number, null);
 						}
 					}
 
@@ -3353,6 +3446,11 @@ namespace UnrealGameSync
 		}
 
 		public void SyncLatestChange()
+		{
+			SyncLatestChange(null);
+		}
+
+		public void SyncLatestChange(WorkspaceUpdateCallback Callback)
 		{
 			if(Workspace != null)
 			{
@@ -3369,7 +3467,7 @@ namespace UnrealGameSync
 				{
 					Owner.ShowAndActivate();
 					SelectChange(ChangeNumber);
-					StartSync(ChangeNumber);
+					StartSync(ChangeNumber, Callback);
 				}
 			}
 		}
@@ -3920,7 +4018,7 @@ namespace UnrealGameSync
 				else
 				{
 					WorkspaceUpdateContext Context = new WorkspaceUpdateContext(Workspace.CurrentChangeNumber, WorkspaceUpdateOptions.Build, null, GetDefaultBuildStepObjects(), ProjectSettings.BuildSteps, new HashSet<Guid>{ UniqueId }, GetWorkspaceVariables(Workspace.CurrentChangeNumber));
-					StartWorkspaceUpdate(Context);
+					StartWorkspaceUpdate(Context, null);
 				}
 			}
 		}
diff --git a/Engine/Source/Programs/UnrealGameSync/UnrealGameSync/Forms/AutomatedSyncWindow.cs b/Engine/Source/Programs/UnrealGameSync/UnrealGameSync/Forms/AutomatedSyncWindow.cs
new file mode 100644
--- /dev/null
+++ b/Engine/Source/Programs/UnrealGameSync/UnrealGameSync/Forms/AutomatedSyncWindow.cs
@@ -0,0 +1,241 @@
+using System;
+using System.Collections.Generic;
+using System.ComponentModel;
+using System.Data;
+using System.Drawing;
+using System.IO;
+using System.Linq;
+using System.Text;
+using System.Threading.Tasks;
+using System.Windows.Forms;
+
+namespace UnrealGameSync
+{
+	partial class AutomatedSyncWindow : Form
+	{
+		class FindDefaultWorkspaceTask : IPerforceModalTask
+		{
+			string StreamName;
+			public string WorkspaceName;
+
+			public FindDefaultWorkspaceTask(string StreamName)
+			{
+				this.StreamName = StreamName;
+			}
+
+			public bool Run(PerforceConnection Perforce, TextWriter Log, out string ErrorMessage)
+			{
+				PerforceInfoRecord Info;
+				if(!Perforce.Info(out Info, Log))
+				{
+					ErrorMessage = "Unable to query Perforce info.";
+					return false;
+				}
+				
+				List<PerforceClientRecord> Clients;
+				if(!Perforce.FindClients(Info.UserName, out Clients, Log))
+				{
+					ErrorMessage = "Unable to enumerate clients from Perforce";
+					return false;
+				}
+
+				List<PerforceClientRecord> CandidateClients = new List<PerforceClientRecord>();
+				foreach(PerforceClientRecord Client in Clients)
+				{
+					if(Client.Host == null || Client.Host.Equals(Info.HostName, StringComparison.OrdinalIgnoreCase))
+					{
+						if(Client.Stream != null && Client.Stream.Equals(StreamName, StringComparison.OrdinalIgnoreCase))
+						{
+							CandidateClients.Add(Client);
+						}
+					}
+				}
+
+				if(CandidateClients.Count == 1)
+				{
+					WorkspaceName = CandidateClients[0].Name;
+				}
+
+				ErrorMessage = null;
+				return true;
+			}
+		}
+
+		class ValidateWorkspaceTask : IPerforceModalTask
+		{
+			public string WorkspaceName;
+			public string StreamName;
+			public string ServerAndPort;
+			public string UserName;
+			public bool bRequiresStreamSwitch;
+			public bool bHasOpenFiles;
+
+			public ValidateWorkspaceTask(string WorkspaceName, string StreamName)
+			{
+				this.WorkspaceName = WorkspaceName;
+				this.StreamName = StreamName;
+			}
+
+			public bool Run(PerforceConnection Perforce, TextWriter Log, out string ErrorMessage)
+			{
+				this.ServerAndPort = Perforce.ServerAndPort;
+				this.UserName = Perforce.UserName;
+
+				PerforceSpec Spec;
+				if(!Perforce.TryGetClientSpec(WorkspaceName, out Spec, Log))
+				{
+					ErrorMessage = String.Format("Unable to get info for client '{0}'", WorkspaceName);
+					return false;
+				}
+
+				string CurrentStreamName = Spec.GetField("Stream");
+				if(CurrentStreamName == null || CurrentStreamName != StreamName)
+				{
+					bRequiresStreamSwitch = true;
+					bHasOpenFiles = Perforce.HasOpenFiles(Log);
+				}
+			
+				ErrorMessage = null;
+				return true;
+			}
+		}
+
+		public class WorkspaceInfo
+		{
+			public string ServerAndPort;
+			public string UserName;
+			public string WorkspaceName;
+			public bool bRequiresStreamSwitch;
+		}
+		
+		string StreamName;
+		TextWriter Log;
+
+		string ServerAndPort;
+		string UserName;
+		WorkspaceInfo SelectedWorkspaceInfo;
+
+		private AutomatedSyncWindow(string StreamName, string ProjectPath, string WorkspaceName, TextWriter Log)
+		{
+			this.StreamName = StreamName;
+			this.Log = Log;
+
+			InitializeComponent();
+
+			ActiveControl = WorkspaceNameTextBox;
+
+			MinimumSize = Size;
+			MaximumSize = new Size(32768, Size.Height);
+
+			ProjectTextBox.Text = StreamName + ProjectPath;
+
+			if(WorkspaceName != null)
+			{
+				WorkspaceNameTextBox.Text = WorkspaceName;
+				WorkspaceNameTextBox.Select(WorkspaceNameTextBox.Text.Length, 0);
+			}
+
+			UpdateServerLabel();
+			UpdateOkButton();
+		}
+
+		public static bool ShowModal(IWin32Window Owner, string StreamName, string ProjectPath, out WorkspaceInfo WorkspaceInfo, TextWriter Log)
+		{
+			FindDefaultWorkspaceTask FindWorkspace = new FindDefaultWorkspaceTask(StreamName);
+
+			string ErrorMessage;
+			PerforceModalTask.Execute(Owner, null, null, null, FindWorkspace, "Finding workspace", "Finding default workspace, please wait...", Log, out ErrorMessage);
+
+			AutomatedSyncWindow Window = new AutomatedSyncWindow(StreamName, ProjectPath, FindWorkspace.WorkspaceName, Log);
+			if(Window.ShowDialog() == DialogResult.OK)
+			{
+				WorkspaceInfo = Window.SelectedWorkspaceInfo;
+				return true;
+			}
+			else
+			{
+				WorkspaceInfo = null;
+				return false;
+			}
+		}
+
+		private void ChangeLink_LinkClicked(object sender, LinkLabelLinkClickedEventArgs e)
+		{
+			string NewServerAndPort;
+			string NewUserName;
+			if(ConnectWindow.ShowModal(this, ServerAndPort, UserName, out NewServerAndPort, out NewUserName))
+			{
+				ServerAndPort = NewServerAndPort;
+				UserName = NewUserName;
+				UpdateServerLabel();
+			}
+		}
+
+		private void UpdateServerLabel()
+		{
+			ServerLabel.Text = OpenProjectWindow.GetServerLabelText(ServerAndPort, UserName);
+			ChangeLink.Location = new Point(ServerLabel.Right + 5, ChangeLink.Location.Y);
+		}
+
+		private void WorkspaceNameNewBtn_Click(object sender, EventArgs e)
+		{
+			string WorkspaceName;
+			if(NewWorkspaceWindow.ShowModal(this, ServerAndPort, UserName, StreamName, WorkspaceNameTextBox.Text, Log, out WorkspaceName))
+			{
+				WorkspaceNameTextBox.Text = WorkspaceName;
+				UpdateOkButton();
+			}
+		}
+
+		private void WorkspaceNameBrowseBtn_Click(object sender, EventArgs e)
+		{
+			string WorkspaceName = WorkspaceNameTextBox.Text;
+			if(SelectWorkspaceWindow.ShowModal(this, ServerAndPort, UserName, WorkspaceName, Log, out WorkspaceName))
+			{
+				WorkspaceNameTextBox.Text = WorkspaceName;
+				UpdateOkButton();
+			}
+		}
+
+		private void UpdateOkButton()
+		{
+			OkBtn.Enabled = (WorkspaceNameTextBox.Text.Length > 0);
+		}
+
+		private void OkBtn_Click(object sender, EventArgs e)
+		{
+			ValidateWorkspaceTask ValidateWorkspace = new ValidateWorkspaceTask(WorkspaceNameTextBox.Text, StreamName);
+
+			string ErrorMessage;
+			ModalTaskResult Result = PerforceModalTask.Execute(Owner, null, ServerAndPort, UserName, ValidateWorkspace, "Checking workspace", "Checking workspace, please wait...", Log, out ErrorMessage);
+			if(Result == ModalTaskResult.Failed)
+			{
+				MessageBox.Show(ErrorMessage);
+			}
+			else if(Result == ModalTaskResult.Succeeded)
+			{
+				if(ValidateWorkspace.bRequiresStreamSwitch)
+				{
+					string Message;
+					if(ValidateWorkspace.bHasOpenFiles)
+					{
+						Message = String.Format("You have files open for edit in this workspace. If you switch this workspace to {0}, you will not be able to submit them until you switch back.\n\nContinue switching streams?", StreamName);
+					}
+					else
+					{
+						Message = String.Format("Switch this workspace to {0}?", StreamName);
+					}
+					if(MessageBox.Show(Message, "Switch Streams", MessageBoxButtons.YesNo) != DialogResult.Yes)
+					{
+						return;
+					}
+				}
+
+				SelectedWorkspaceInfo = new WorkspaceInfo(){ ServerAndPort = ValidateWorkspace.ServerAndPort, UserName = ValidateWorkspace.UserName, WorkspaceName = ValidateWorkspace.WorkspaceName, bRequiresStreamSwitch = ValidateWorkspace.bRequiresStreamSwitch };
+
+				DialogResult = DialogResult.OK;
+				Close();
+			}
+		}
+	}
+}
diff --git a/Engine/Source/Programs/UnrealGameSync/UnrealGameSync/Forms/MainWindow.cs b/Engine/Source/Programs/UnrealGameSync/UnrealGameSync/Forms/MainWindow.cs
--- a/Engine/Source/Programs/UnrealGameSync/UnrealGameSync/Forms/MainWindow.cs
+++ b/Engine/Source/Programs/UnrealGameSync/UnrealGameSync/Forms/MainWindow.cs
@@ -10,6 +10,7 @@ using System.Runtime.InteropServices;
 using System.Windows.Forms;
 using System.Threading;
 using System.Reflection;
+using System.Text;
 
 namespace UnrealGameSync
 {
@@ -77,6 +78,9 @@ namespace UnrealGameSync
 
 		IMainWindowTabPanel CurrentTabPanel;
 
+		AutomationServer AutomationServer;
+		TextWriter AutomationLog;
+
 		public MainWindow(UpdateMonitor InUpdateMonitor, string InApiUrl, string InDataFolder, string InCacheFolder, bool bInRestoreStateOnLoad, string InOriginalExecutableFileName, bool bInUnstable, DetectProjectSettingsResult[] StartupProjects, LineBasedTextWriter InLog, UserSettings InSettings)
 		{
 			InitializeComponent();
@@ -136,6 +140,121 @@ namespace UnrealGameSync
 			{
 				Text += String.Format(" (UNSTABLE BUILD {0})", Assembly.GetExecutingAssembly().GetName().Version);
 			}
+
+			AutomationLog = new TimestampLogWriter(new BoundedLogWriter(Path.Combine(DataFolder, "Automation.log")));
+			AutomationServer = new AutomationServer(Request => { MainThreadSynchronizationContext.Post(Obj => PostAutomationRequest(Request), null); }, AutomationLog);
+		}
+
+		void PostAutomationRequest(AutomationRequest Request)
+		{
+			try
+			{
+				if(!CanFocus)
+				{
+					Request.SetOutput(new AutomationRequestOutput(AutomationRequestResult.Busy));
+				}
+				else if(Request.Input.Type == AutomationRequestType.SyncProject)
+				{
+					AutomationRequestResult Result = StartAutomatedSync(Request);
+					if(Result != AutomationRequestResult.Ok)
+					{
+						Request.SetOutput(new AutomationRequestOutput(Result));
+					}
+				}
+				else
+				{
+					Request.SetOutput(new AutomationRequestOutput(AutomationRequestResult.Invalid));
+				}
+			}
+			catch(Exception Ex)
+			{
+				Log.WriteLine("Exception running automation request: {0}", Ex);
+				Request.SetOutput(new AutomationRequestOutput(AutomationRequestResult.Invalid));
+			}
+		}
+
+		AutomationRequestResult StartAutomatedSync(AutomationRequest Request)
+		{
+			ShowAndActivate();
+
+			BinaryReader Reader = new BinaryReader(new MemoryStream(Request.Input.Data));
+			string StreamName = Reader.ReadString();
+			string ProjectPath = Reader.ReadString();
+
+			AutomatedSyncWindow.WorkspaceInfo WorkspaceInfo;
+			if(!AutomatedSyncWindow.ShowModal(this, StreamName, ProjectPath, out WorkspaceInfo, Log))
+			{
+				return AutomationRequestResult.Canceled;
+			}
+
+			if(WorkspaceInfo.bRequiresStreamSwitch)
+			{
+				// Close any tab containing this window
+				for(int ExistingTabIdx = 0; ExistingTabIdx < TabControl.GetTabCount(); ExistingTabIdx++)
+				{
+					WorkspaceControl ExistingWorkspace = TabControl.GetTabData(ExistingTabIdx) as WorkspaceControl;
+					if(ExistingWorkspace != null && ExistingWorkspace.ClientName.Equals(WorkspaceInfo.WorkspaceName))
+					{
+						TabControl.RemoveTab(ExistingTabIdx);
+						break;
+					}
+				}
+
+				// Switch the stream
+				PerforceConnection Perforce = new PerforceConnection(WorkspaceInfo.UserName, WorkspaceInfo.WorkspaceName, WorkspaceInfo.ServerAndPort);
+				if(!Perforce.SwitchStream(StreamName, Log))
+				{
+					Log.WriteLine("Unable to switch stream");
+					return AutomationRequestResult.Error;
+				}
+			}
+
+			UserSelectedProjectSettings SelectedProject = new UserSelectedProjectSettings(WorkspaceInfo.ServerAndPort, WorkspaceInfo.UserName, UserSelectedProjectType.Client, String.Format("//{0}{1}", WorkspaceInfo.WorkspaceName, ProjectPath), null);
+
+			int TabIdx = TryOpenProject(SelectedProject, -1, OpenProjectOptions.None);
+			if(TabIdx == -1)
+			{
+				Log.WriteLine("Unable to open project");
+				return AutomationRequestResult.Error;
+			}
+
+			WorkspaceControl Workspace = TabControl.GetTabData(TabIdx) as WorkspaceControl;
+			if(Workspace == null)
+			{
+				Log.WriteLine("Workspace was unable to open");
+				return AutomationRequestResult.Error;
+			}
+
+			Workspace.AddStartupCallback((Control, bCancel) => StartAutomatedSyncAfterStartup(Control, bCancel, Request));
+			return AutomationRequestResult.Ok;
+		}
+
+		private void StartAutomatedSyncAfterStartup(WorkspaceControl Workspace, bool bCancel, AutomationRequest Request)
+		{
+			if(bCancel)
+			{
+				Request.SetOutput(new AutomationRequestOutput(AutomationRequestResult.Canceled));
+			}
+			else
+			{
+				Workspace.SyncLatestChange(Result => CompleteAutomatedSync(Result, Workspace.SelectedFileName, Request));
+			}
+		}
+
+		void CompleteAutomatedSync(WorkspaceUpdateResult Result, string SelectedFileName, AutomationRequest Request)
+		{
+			if(Result == WorkspaceUpdateResult.Success)
+			{
+				Request.SetOutput(new AutomationRequestOutput(AutomationRequestResult.Ok, Encoding.UTF8.GetBytes(SelectedFileName)));
+			}
+			else if(Result == WorkspaceUpdateResult.Canceled)
+			{
+				Request.SetOutput(new AutomationRequestOutput(AutomationRequestResult.Canceled));
+			}
+			else
+			{
+				Request.SetOutput(new AutomationRequestOutput(AutomationRequestResult.Error));
+			}
 		}
 
 		protected override void OnHandleCreated(EventArgs e)
@@ -234,6 +353,18 @@ namespace UnrealGameSync
 
 			StopScheduleTimer();
 
+			if(AutomationServer != null)
+			{
+				AutomationServer.Dispose();
+				AutomationServer = null;
+			}
+
+			if(AutomationLog != null)
+			{
+				AutomationLog.Close();
+				AutomationLog = null;
+			}
+
 			base.Dispose(disposing);
 		}
 
diff --git a/Engine/Source/Programs/UnrealGameSync/UnrealGameSync/Forms/NewWorkspaceWindow.cs b/Engine/Source/Programs/UnrealGameSync/UnrealGameSync/Forms/NewWorkspaceWindow.cs
--- a/Engine/Source/Programs/UnrealGameSync/UnrealGameSync/Forms/NewWorkspaceWindow.cs
+++ b/Engine/Source/Programs/UnrealGameSync/UnrealGameSync/Forms/NewWorkspaceWindow.cs
@@ -112,7 +112,7 @@ namespace UnrealGameSync
 		NewWorkspaceSettings Settings;
 		string DefaultRootPath;
 
-		private NewWorkspaceWindow(string ServerAndPort, string UserName, string DefaultStream, PerforceInfoRecord Info, List<PerforceClientRecord> Clients, TextWriter Log)
+		private NewWorkspaceWindow(string ServerAndPort, string UserName, string ForceStream, string DefaultStream, PerforceInfoRecord Info, List<PerforceClientRecord> Clients, TextWriter Log)
 		{
 			InitializeComponent();
 
@@ -159,17 +159,28 @@ namespace UnrealGameSync
 				}
 			}
 
-			StreamTextBox.Text = DefaultStream ?? "";
+			if(ForceStream != null)
+			{
+				StreamTextBox.Text = ForceStream;
+				StreamTextBox.Enabled = false;
+			}
+			else
+			{
+				StreamTextBox.Text = DefaultStream ?? "";
+				StreamTextBox.Enabled = true;
+			}
 			StreamTextBox.SelectionStart = StreamTextBox.Text.Length;
 			StreamTextBox.SelectionLength = 0;
 			StreamTextBox.Focus();
 
+			StreamBrowseBtn.Enabled = (ForceStream == null);
+
 			UpdateOkButton();
 			UpdateNameCueBanner();
 			UpdateRootDirCueBanner();
 		}
 
-		public static bool ShowModal(IWin32Window Owner, string ServerAndPort, string UserName, string CurrentWorkspaceName, TextWriter Log, out string WorkspaceName)
+		public static bool ShowModal(IWin32Window Owner, string ServerAndPort, string UserName, string ForceStreamName, string CurrentWorkspaceName, TextWriter Log, out string WorkspaceName)
 		{
 			FindWorkspaceSettingsTask FindSettings = new FindWorkspaceSettingsTask(CurrentWorkspaceName);
 
@@ -182,7 +193,7 @@ namespace UnrealGameSync
 				}
 			}
 
-			NewWorkspaceWindow Window = new NewWorkspaceWindow(ServerAndPort, UserName, FindSettings.CurrentStream, FindSettings.Info, FindSettings.Clients, Log);
+			NewWorkspaceWindow Window = new NewWorkspaceWindow(ServerAndPort, UserName, ForceStreamName, FindSettings.CurrentStream, FindSettings.Info, FindSettings.Clients, Log);
 			if(Window.ShowDialog(Owner) == DialogResult.OK)
 			{
 				WorkspaceName = Window.Settings.Name;
diff --git a/Engine/Source/Programs/UnrealGameSync/UnrealGameSync/Forms/OpenProjectWindow.cs b/Engine/Source/Programs/UnrealGameSync/UnrealGameSync/Forms/OpenProjectWindow.cs
--- a/Engine/Source/Programs/UnrealGameSync/UnrealGameSync/Forms/OpenProjectWindow.cs
+++ b/Engine/Source/Programs/UnrealGameSync/UnrealGameSync/Forms/OpenProjectWindow.cs
@@ -113,11 +113,11 @@ namespace UnrealGameSync
 			UpdateWorkspacePathBrowseButton();
 		}
 
-		private void UpdateServerLabel()
+		public static string GetServerLabelText(string ServerAndPort, string UserName)
 		{
 			if(ServerAndPort == null && UserName == null)
 			{
-				ServerLabel.Text = "Using default Perforce server settings.";
+				return "Using default Perforce server settings.";
 			}
 			else
 			{
@@ -139,9 +139,13 @@ namespace UnrealGameSync
 				{
 					Text.AppendFormat("server '{0}'.", ServerAndPort);
 				}
-				ServerLabel.Text = Text.ToString();
+				return Text.ToString();
 			}
+		}
 
+		private void UpdateServerLabel()
+		{
+			ServerLabel.Text = GetServerLabelText(ServerAndPort, UserName);
 			ChangeLink.Location = new Point(ServerLabel.Right + 5, ChangeLink.Location.Y);
 		}
 
@@ -162,7 +166,7 @@ namespace UnrealGameSync
 			WorkspaceRadioBtn.Checked = true;
 			
 			string WorkspaceName;
-			if(NewWorkspaceWindow.ShowModal(this, ServerAndPort, UserName, WorkspaceNameTextBox.Text, Log, out WorkspaceName))
+			if(NewWorkspaceWindow.ShowModal(this, ServerAndPort, UserName, null, WorkspaceNameTextBox.Text, Log, out WorkspaceName))
 			{
 				WorkspaceNameTextBox.Text = WorkspaceName;
 				UpdateOkButton();
diff --git a/Engine/Source/Programs/UnrealGameSync/UnrealGameSync/UnrealGameSync.csproj b/Engine/Source/Programs/UnrealGameSync/UnrealGameSync/UnrealGameSync.csproj
--- a/Engine/Source/Programs/UnrealGameSync/UnrealGameSync/UnrealGameSync.csproj
+++ b/Engine/Source/Programs/UnrealGameSync/UnrealGameSync/UnrealGameSync.csproj
@@ -112,6 +112,7 @@
   <ItemGroup>
     <Compile Include="ActivationListener.cs" />
     <Compile Include="ArchiveManifest.cs" />
+    <Compile Include="AutomationServer.cs" />
     <Compile Include="BuildStep.cs" />
     <Compile Include="ChildProcess.cs" />
     <Compile Include="Controls\CustomListViewControl.cs">
@@ -159,6 +160,12 @@
     <Compile Include="Forms\ArgumentsWindow.Designer.cs">
       <DependentUpon>ArgumentsWindow.cs</DependentUpon>
     </Compile>
+    <Compile Include="Forms\AutomatedSyncWindow.cs">
+      <SubType>Form</SubType>
+    </Compile>
+    <Compile Include="Forms\AutomatedSyncWindow.Designer.cs">
+      <DependentUpon>AutomatedSyncWindow.cs</DependentUpon>
+    </Compile>
     <Compile Include="Forms\CleanWorkspaceWindow.cs">
       <SubType>Form</SubType>
     </Compile>
@@ -353,6 +360,9 @@
     <EmbeddedResource Include="Forms\ArgumentsWindow.resx">
       <DependentUpon>ArgumentsWindow.cs</DependentUpon>
     </EmbeddedResource>
+    <EmbeddedResource Include="Forms\AutomatedSyncWindow.resx">
+      <DependentUpon>AutomatedSyncWindow.cs</DependentUpon>
+    </EmbeddedResource>
     <EmbeddedResource Include="Forms\CleanWorkspaceWindow.resx">
       <DependentUpon>CleanWorkspaceWindow.cs</DependentUpon>
     </EmbeddedResource>
