UGS: Add a Filter menu to allow customizing which changes should be displayed in the changes list.

diff --git a/Engine/Source/Programs/UnrealGameSync/UnrealGameSync/Controls/WorkspaceControl.cs b/Engine/Source/Programs/UnrealGameSync/UnrealGameSync/Controls/WorkspaceControl.cs
--- a/Engine/Source/Programs/UnrealGameSync/UnrealGameSync/Controls/WorkspaceControl.cs
+++ b/Engine/Source/Programs/UnrealGameSync/UnrealGameSync/Controls/WorkspaceControl.cs
@@ -1164,6 +1164,14 @@ namespace UnrealGameSync
 
 		bool ShouldShowChange(PerforceChangeSummary Change, string[] ExcludeChanges)
 		{
+			if(ProjectSettings.FilterBadges.Count > 0)
+			{
+				EventSummary Summary = EventMonitor.GetSummaryForChange(Change.Number);
+				if(Summary == null || !Summary.Badges.Any(x => ProjectSettings.FilterBadges.Contains(x.BadgeName)))
+				{
+					return false;
+				}
+			}
 			if(!Settings.bShowAutomatedChanges)
 			{
 				foreach(string ExcludeChange in ExcludeChanges)
@@ -1268,6 +1276,12 @@ namespace UnrealGameSync
 			UpdateStatusPanel();
 			UpdateBuildFailureNotification();
 			CheckForStartupComplete();
+
+			// If we are filtering by badges, we may also need to update the build list
+			if(ProjectSettings.FilterBadges.Count > 0)
+			{
+				UpdateBuildList();
+			}
 		}
 
 		void UpdateMaxBuildBadgeChars()
@@ -4360,8 +4374,7 @@ namespace UnrealGameSync
 			Settings.bShowAutomatedChanges ^= true;
 			Settings.Save();
 
-			UpdateBuildList();
-			UpdateNumRequestedBuilds(true);
+			UpdateBuildListFilter();
 		}
 
 		private void BuildList_ColumnWidthChanging(object sender, ColumnWidthChangingEventArgs e)
@@ -4657,5 +4670,80 @@ namespace UnrealGameSync
 				PerforceMonitor.IsActive = Visible;
 			}
 		}
+
+		private void FilterButton_Click(object sender, EventArgs e)
+		{
+			FilterContextMenu_ResetToDefault.Checked = !Settings.bShowAutomatedChanges && ProjectSettings.FilterBadges.Count == 0;
+
+			FilterContextMenu_Badges.DropDownItems.Clear();
+			FilterContextMenu_Badges.Checked = ProjectSettings.FilterBadges.Count > 0;
+
+			HashSet<string> BadgeNames = new HashSet<string>(ProjectSettings.FilterBadges, StringComparer.OrdinalIgnoreCase);
+			BadgeNames.ExceptWith(BadgeNameAndGroupPairs.Select(x => x.Key));
+
+			List<KeyValuePair<string, string>> DisplayBadgeNameAndGroupPairs = new List<KeyValuePair<string, string>>(BadgeNameAndGroupPairs);
+			DisplayBadgeNameAndGroupPairs.AddRange(BadgeNames.Select(x => new KeyValuePair<string, string>(x, "User")));
+
+			string LastGroup = null;
+			foreach(KeyValuePair<string, string> BadgeNameAndGroupPair in DisplayBadgeNameAndGroupPairs)
+			{
+				if(LastGroup != BadgeNameAndGroupPair.Value)
+				{
+					if(LastGroup != null)
+					{
+						FilterContextMenu_Badges.DropDownItems.Add(new ToolStripSeparator());
+					}
+					LastGroup = BadgeNameAndGroupPair.Value;
+				}
+
+				ToolStripMenuItem Item = new ToolStripMenuItem(BadgeNameAndGroupPair.Key);
+				Item.Checked = ProjectSettings.FilterBadges.Contains(BadgeNameAndGroupPair.Key, StringComparer.OrdinalIgnoreCase);
+				Item.Click += (Sender, Args) => FilterContextMenu_Badge_Click(BadgeNameAndGroupPair.Key);
+				FilterContextMenu_Badges.DropDownItems.Add(Item);
+			}
+
+			FilterContextMenu_Badges.Enabled = FilterContextMenu_Badges.DropDownItems.Count > 0;
+
+			FilterContextMenu_ShowBuildMachineChanges.Checked = Settings.bShowAutomatedChanges;
+			FilterContextMenu.Show(FilterButton, new Point(0, FilterButton.Height));
+		}
+
+		private void FilterContextMenu_Badge_Click(string BadgeName)
+		{
+			if(ProjectSettings.FilterBadges.Contains(BadgeName))
+			{
+				ProjectSettings.FilterBadges.Remove(BadgeName);
+			}
+			else
+			{
+				ProjectSettings.FilterBadges.Add(BadgeName);
+			}
+
+			UpdateBuildListFilter();
+		}
+
+		private void FilterContextMenu_ResetToDefault_Click(object sender, EventArgs e)
+		{
+			ProjectSettings.FilterBadges.Clear();
+
+			Settings.bShowAutomatedChanges = false;
+			Settings.Save();
+
+			UpdateBuildListFilter();
+		}
+
+		private void FilterContextMenu_ShowBuildMachineChanges_Click(object sender, EventArgs e)
+		{
+			Settings.bShowAutomatedChanges ^= true;
+			Settings.Save();
+
+			UpdateBuildListFilter();
+		}
+
+		private void UpdateBuildListFilter()
+		{
+			UpdateBuildList();
+			UpdateNumRequestedBuilds(true);
+		}
 	}
 }
diff --git a/Engine/Source/Programs/UnrealGameSync/UnrealGameSync/UserSettings.cs b/Engine/Source/Programs/UnrealGameSync/UnrealGameSync/UserSettings.cs
--- a/Engine/Source/Programs/UnrealGameSync/UnrealGameSync/UserSettings.cs
+++ b/Engine/Source/Programs/UnrealGameSync/UnrealGameSync/UserSettings.cs
@@ -171,6 +171,7 @@ namespace UnrealGameSync
 	class UserProjectSettings
 	{
 		public List<ConfigObject> BuildSteps = new List<ConfigObject>();
+		public HashSet<string> FilterBadges = new HashSet<string>(StringComparer.OrdinalIgnoreCase);
 	}
 
 	class UserSettings
@@ -495,6 +496,7 @@ namespace UnrealGameSync
 	
 				ConfigSection ProjectSection = ConfigFile.FindOrAddSection(ClientProjectFileName);
 				CurrentProject.BuildSteps.AddRange(ProjectSection.GetValues("BuildStep", new string[0]).Select(x => new ConfigObject(x)));
+				CurrentProject.FilterBadges.UnionWith(ProjectSection.GetValues("FilterBadges", new string[0]));
 			}
 			return CurrentProject;
 		}
@@ -612,6 +614,7 @@ namespace UnrealGameSync
 				ConfigSection ProjectSection = ConfigFile.FindOrAddSection(CurrentProjectKey);
 				ProjectSection.Clear();
 				ProjectSection.SetValues("BuildStep", CurrentProject.BuildSteps.Select(x => x.ToString()).ToArray());
+				ProjectSection.SetValues("FilterBadges", CurrentProject.FilterBadges.ToArray());
 			}
 
 			// Perforce settings
