UBT: Allow filtering changes by type (code/content).

diff --git a/Engine/Source/Programs/UnrealGameSync/UnrealGameSync/Controls/WorkspaceControl.cs b/Engine/Source/Programs/UnrealGameSync/UnrealGameSync/Controls/WorkspaceControl.cs
--- a/Engine/Source/Programs/UnrealGameSync/UnrealGameSync/Controls/WorkspaceControl.cs
+++ b/Engine/Source/Programs/UnrealGameSync/UnrealGameSync/Controls/WorkspaceControl.cs
@@ -1164,6 +1164,22 @@ namespace UnrealGameSync
 
 		bool ShouldShowChange(PerforceChangeSummary Change, string[] ExcludeChanges)
 		{
+			if(ProjectSettings.FilterType != FilterType.None)
+			{
+				PerforceChangeDetails Details;
+				if(!PerforceMonitor.TryGetChangeDetails(Change.Number, out Details))
+				{
+					return false;
+				}
+				if(ProjectSettings.FilterType == FilterType.Code && !Details.bContainsCode)
+				{
+					return false;
+				}
+				if(ProjectSettings.FilterType == FilterType.Content && !Details.bContainsContent)
+				{
+					return false;
+				}
+			}
 			if(ProjectSettings.FilterBadges.Count > 0)
 			{
 				EventSummary Summary = EventMonitor.GetSummaryForChange(Change.Number);
@@ -4673,7 +4689,12 @@ namespace UnrealGameSync
 
 		private void FilterButton_Click(object sender, EventArgs e)
 		{
-			FilterContextMenu_ResetToDefault.Checked = !Settings.bShowAutomatedChanges && ProjectSettings.FilterBadges.Count == 0;
+			FilterContextMenu_ResetToDefault.Checked = !Settings.bShowAutomatedChanges && ProjectSettings.FilterType == FilterType.None && ProjectSettings.FilterBadges.Count == 0;
+
+			FilterContextMenu_Type.Checked = ProjectSettings.FilterType != FilterType.None;
+			FilterContextMenu_Type_ShowAll.Checked = ProjectSettings.FilterType == FilterType.None;
+			FilterContextMenu_Type_Code.Checked = ProjectSettings.FilterType == FilterType.Code;
+			FilterContextMenu_Type_Content.Checked = ProjectSettings.FilterType == FilterType.Content;
 
 			FilterContextMenu_Badges.DropDownItems.Clear();
 			FilterContextMenu_Badges.Checked = ProjectSettings.FilterBadges.Count > 0;
@@ -4740,6 +4761,30 @@ namespace UnrealGameSync
 			UpdateBuildListFilter();
 		}
 
+		private void FilterContextMenu_Type_ShowAll_Click(object sender, EventArgs e)
+		{
+			ProjectSettings.FilterType = FilterType.None;
+			Settings.Save();
+
+			UpdateBuildListFilter();
+		}
+
+		private void FilterContextMenu_Type_Code_Click(object sender, EventArgs e)
+		{
+			ProjectSettings.FilterType = FilterType.Code;
+			Settings.Save();
+
+			UpdateBuildListFilter();
+		}
+
+		private void FilterContextMenu_Type_Content_Click(object sender, EventArgs e)
+		{
+			ProjectSettings.FilterType = FilterType.Content;
+			Settings.Save();
+
+			UpdateBuildListFilter();
+		}
+
 		private void UpdateBuildListFilter()
 		{
 			UpdateBuildList();
diff --git a/Engine/Source/Programs/UnrealGameSync/UnrealGameSync/UserSettings.cs b/Engine/Source/Programs/UnrealGameSync/UnrealGameSync/UserSettings.cs
--- a/Engine/Source/Programs/UnrealGameSync/UnrealGameSync/UserSettings.cs
+++ b/Engine/Source/Programs/UnrealGameSync/UnrealGameSync/UserSettings.cs
@@ -39,6 +39,13 @@ namespace UnrealGameSync
 		Local
 	}
 
+	enum FilterType
+	{
+		None,
+		Code,
+		Content
+	}
+
 	class UserSelectedProjectSettings
 	{
 		public readonly string ServerAndPort;
@@ -171,6 +178,7 @@ namespace UnrealGameSync
 	class UserProjectSettings
 	{
 		public List<ConfigObject> BuildSteps = new List<ConfigObject>();
+		public FilterType FilterType;
 		public HashSet<string> FilterBadges = new HashSet<string>(StringComparer.OrdinalIgnoreCase);
 	}
 
@@ -496,6 +504,10 @@ namespace UnrealGameSync
 	
 				ConfigSection ProjectSection = ConfigFile.FindOrAddSection(ClientProjectFileName);
 				CurrentProject.BuildSteps.AddRange(ProjectSection.GetValues("BuildStep", new string[0]).Select(x => new ConfigObject(x)));
+				if(!Enum.TryParse(ProjectSection.GetValue("FilterType", ""), true, out CurrentProject.FilterType))
+				{
+					CurrentProject.FilterType = FilterType.None;
+				}
 				CurrentProject.FilterBadges.UnionWith(ProjectSection.GetValues("FilterBadges", new string[0]));
 			}
 			return CurrentProject;
@@ -614,6 +626,10 @@ namespace UnrealGameSync
 				ConfigSection ProjectSection = ConfigFile.FindOrAddSection(CurrentProjectKey);
 				ProjectSection.Clear();
 				ProjectSection.SetValues("BuildStep", CurrentProject.BuildSteps.Select(x => x.ToString()).ToArray());
+				if(CurrentProject.FilterType != FilterType.None)
+				{
+					ProjectSection.SetValue("FilterType", CurrentProject.FilterType.ToString());
+				}
 				ProjectSection.SetValues("FilterBadges", CurrentProject.FilterBadges.ToArray());
 			}
 
