UGS: Fix issues when creating main window when taskbar icon would not be present, caused by accessing TabPanel size during main window creation. This would cause child windows to be created recursively, and allow the window to be shown without creating a taskbar button.

diff --git a/Engine/Source/Programs/UnrealGameSync/UnrealGameSync/Forms/MainWindow.cs b/Engine/Source/Programs/UnrealGameSync/UnrealGameSync/Forms/MainWindow.cs
--- a/Engine/Source/Programs/UnrealGameSync/UnrealGameSync/Forms/MainWindow.cs
+++ b/Engine/Source/Programs/UnrealGameSync/UnrealGameSync/Forms/MainWindow.cs
@@ -421,8 +421,7 @@ namespace UnrealGameSync
 			ErrorPanel.BorderStyle = BorderStyle.FixedSingle;
 			ErrorPanel.BackColor = System.Drawing.Color.FromArgb(((int)(((byte)(250)))), ((int)(((byte)(250)))), ((int)(((byte)(250)))));
 			ErrorPanel.Location = new Point(0, 0);
-			ErrorPanel.Size = new Size(TabPanel.Width, TabPanel.Height);
-			ErrorPanel.Anchor = AnchorStyles.Left | AnchorStyles.Top | AnchorStyles.Right | AnchorStyles.Bottom;
+			ErrorPanel.Dock = DockStyle.Fill;
 			ErrorPanel.Hide();
 
 			string SummaryText = String.Format("Unable to open '{0}'.", Project.ToString());
@@ -899,8 +898,7 @@ namespace UnrealGameSync
 			WorkspaceControl NewWorkspace = new WorkspaceControl(this, ApiUrl, OriginalExecutableFileName, bUnstable, ProjectSettings, Log, Settings);
 			NewWorkspace.Parent = TabPanel;
 			NewWorkspace.Location = new Point(0, 0);
-			NewWorkspace.Size = new Size(TabPanel.Width, TabPanel.Height);
-			NewWorkspace.Anchor = AnchorStyles.Left | AnchorStyles.Top | AnchorStyles.Right | AnchorStyles.Bottom;
+			NewWorkspace.Dock = DockStyle.Fill;
 			NewWorkspace.Hide();
 
 			// Add the tab
