UGS: Fix 'Default' filter not resetting selected type.

diff --git a/Engine/Source/Programs/UnrealGameSync/UnrealGameSync/Controls/WorkspaceControl.cs b/Engine/Source/Programs/UnrealGameSync/UnrealGameSync/Controls/WorkspaceControl.cs
--- a/Engine/Source/Programs/UnrealGameSync/UnrealGameSync/Controls/WorkspaceControl.cs
+++ b/Engine/Source/Programs/UnrealGameSync/UnrealGameSync/Controls/WorkspaceControl.cs
@@ -4689,7 +4689,7 @@ namespace UnrealGameSync
 
 		private void FilterButton_Click(object sender, EventArgs e)
 		{
-			FilterContextMenu_ResetToDefault.Checked = !Settings.bShowAutomatedChanges && ProjectSettings.FilterType == FilterType.None && ProjectSettings.FilterBadges.Count == 0;
+			FilterContextMenu_Default.Checked = !Settings.bShowAutomatedChanges && ProjectSettings.FilterType == FilterType.None && ProjectSettings.FilterBadges.Count == 0;
 
 			FilterContextMenu_Type.Checked = ProjectSettings.FilterType != FilterType.None;
 			FilterContextMenu_Type_ShowAll.Checked = ProjectSettings.FilterType == FilterType.None;
@@ -4743,9 +4743,10 @@ namespace UnrealGameSync
 			UpdateBuildListFilter();
 		}
 
-		private void FilterContextMenu_ResetToDefault_Click(object sender, EventArgs e)
+		private void FilterContextMenu_Default_Click(object sender, EventArgs e)
 		{
 			ProjectSettings.FilterBadges.Clear();
+			ProjectSettings.FilterType = FilterType.None;
 
 			Settings.bShowAutomatedChanges = false;
 			Settings.Save();
