UGS: Fix another instance of the window handle being created while the MainWindow constructor is still being run, and add an assert to prevent it happening in the future.

diff --git a/Engine/Source/Programs/UnrealGameSync/UnrealGameSync/Forms/MainWindow.cs b/Engine/Source/Programs/UnrealGameSync/UnrealGameSync/Forms/MainWindow.cs
--- a/Engine/Source/Programs/UnrealGameSync/UnrealGameSync/Forms/MainWindow.cs
+++ b/Engine/Source/Programs/UnrealGameSync/UnrealGameSync/Forms/MainWindow.cs
@@ -11,6 +11,7 @@ using System.Windows.Forms;
 using System.Threading;
 using System.Reflection;
 using System.Text;
+using System.Diagnostics;
 
 namespace UnrealGameSync
 {
@@ -81,6 +82,8 @@ namespace UnrealGameSync
 		AutomationServer AutomationServer;
 		TextWriter AutomationLog;
 
+		bool bAllowCreatingHandle;
+
 		public MainWindow(UpdateMonitor InUpdateMonitor, string InApiUrl, string InDataFolder, string InCacheFolder, bool bInRestoreStateOnLoad, string InOriginalExecutableFileName, bool bInUnstable, DetectProjectSettingsResult[] StartupProjects, LineBasedTextWriter InLog, UserSettings InSettings)
 		{
 			InitializeComponent();
@@ -106,6 +109,9 @@ namespace UnrealGameSync
 
 			SetupDefaultControl();
 
+			// While creating tab controls during startup, we need to prevent layout calls resulting in the window handle being created too early. Disable layout calls here.
+			TabPanel.SuspendLayout();
+
 			int SelectTabIdx = -1;
 			foreach(DetectProjectSettingsResult StartupProject in StartupProjects)
 			{
@@ -134,6 +140,8 @@ namespace UnrealGameSync
 				TabControl.SelectTab(0);
 			}
 
+			TabPanel.ResumeLayout(false);
+
 			StartScheduleTimer();
 
 			if(bUnstable)
@@ -143,6 +151,8 @@ namespace UnrealGameSync
 
 			AutomationLog = new TimestampLogWriter(new BoundedLogWriter(Path.Combine(DataFolder, "Automation.log")));
 			AutomationServer = new AutomationServer(Request => { MainThreadSynchronizationContext.Post(Obj => PostAutomationRequest(Request), null); }, AutomationLog);
+
+			bAllowCreatingHandle = true;
 		}
 
 		void PostAutomationRequest(AutomationRequest Request)
@@ -260,6 +270,8 @@ namespace UnrealGameSync
 		protected override void OnHandleCreated(EventArgs e)
 		{
 			base.OnHandleCreated(e);
+
+			Debug.Assert(bAllowCreatingHandle, "Window handle should not be created before constructor has run.");
 		}
 
 		void TabControl_OnButtonClick(int ButtonIdx, Point Location, MouseButtons Buttons)
@@ -897,7 +909,6 @@ namespace UnrealGameSync
 			// Now that we have the project settings, we can construct the tab
 			WorkspaceControl NewWorkspace = new WorkspaceControl(this, ApiUrl, OriginalExecutableFileName, bUnstable, ProjectSettings, Log, Settings);
 			NewWorkspace.Parent = TabPanel;
-			NewWorkspace.Location = new Point(0, 0);
 			NewWorkspace.Dock = DockStyle.Fill;
 			NewWorkspace.Hide();
 
