UGS: Preseve the location/size/state of the main window between runs.

diff --git a/Engine/Source/Programs/UnrealGameSync/UnrealGameSync/Forms/MainWindow.cs b/Engine/Source/Programs/UnrealGameSync/UnrealGameSync/Forms/MainWindow.cs
--- a/Engine/Source/Programs/UnrealGameSync/UnrealGameSync/Forms/MainWindow.cs
+++ b/Engine/Source/Programs/UnrealGameSync/UnrealGameSync/Forms/MainWindow.cs
@@ -99,6 +99,10 @@ namespace UnrealGameSync
 			Log = InLog;
 			Settings = InSettings;
 
+			// While creating tab controls during startup, we need to prevent layout calls resulting in the window handle being created too early. Disable layout calls here.
+			SuspendLayout();
+			TabPanel.SuspendLayout();
+
 			TabControl.OnTabChanged += TabControl_OnTabChanged;
 			TabControl.OnNewTabClick += TabControl_OnNewTabClick;
 			TabControl.OnTabClicked += TabControl_OnTabClicked;
@@ -109,9 +113,6 @@ namespace UnrealGameSync
 
 			SetupDefaultControl();
 
-			// While creating tab controls during startup, we need to prevent layout calls resulting in the window handle being created too early. Disable layout calls here.
-			TabPanel.SuspendLayout();
-
 			int SelectTabIdx = -1;
 			foreach(DetectProjectSettingsResult StartupProject in StartupProjects)
 			{
@@ -140,8 +141,6 @@ namespace UnrealGameSync
 				TabControl.SelectTab(0);
 			}
 
-			TabPanel.ResumeLayout(false);
-
 			StartScheduleTimer();
 
 			if(bUnstable)
@@ -152,6 +151,10 @@ namespace UnrealGameSync
 			AutomationLog = new TimestampLogWriter(new BoundedLogWriter(Path.Combine(DataFolder, "Automation.log")));
 			AutomationServer = new AutomationServer(Request => { MainThreadSynchronizationContext.Post(Obj => PostAutomationRequest(Request), null); }, AutomationLog);
 
+			// Allow creating controls from now on
+			TabPanel.ResumeLayout(false);
+			ResumeLayout(false);
+
 			bAllowCreatingHandle = true;
 		}
 
@@ -386,9 +389,6 @@ namespace UnrealGameSync
 			{
 				Hide();
 				EventArgs.Cancel = true;
-
-				Settings.bWindowVisible = Visible;
-				Settings.Save();
 			}
 			else
 			{
@@ -403,10 +403,20 @@ namespace UnrealGameSync
 				}
 
 				StopScheduleTimer();
+			}
 
-				Settings.bWindowVisible = Visible;
-				Settings.Save();
+			Settings.bWindowVisible = Visible;
+			Settings.WindowState = WindowState;
+			if(WindowState == FormWindowState.Normal)
+			{
+				Settings.WindowBounds = new Rectangle(Location, Size);
 			}
+			else
+			{
+				Settings.WindowBounds = RestoreBounds;
+			}
+
+			Settings.Save();
 		}
 
 		private void SetupDefaultControl()
@@ -1106,5 +1116,15 @@ namespace UnrealGameSync
 				UpdateMonitor.TriggerUpdate(UpdateType.UserInitiated, bRelaunchUnstable);
 			}
 		}
+
+		private void MainWindow_Load(object sender, EventArgs e)
+		{
+			if(Settings.WindowBounds != null)
+			{
+				Location = Settings.WindowBounds.Value.Location;
+				Size = Settings.WindowBounds.Value.Size;
+			}
+			WindowState = Settings.WindowState;
+		}
 	}
 }
diff --git a/Engine/Source/Programs/UnrealGameSync/UnrealGameSync/UserSettings.cs b/Engine/Source/Programs/UnrealGameSync/UnrealGameSync/UserSettings.cs
--- a/Engine/Source/Programs/UnrealGameSync/UnrealGameSync/UserSettings.cs
+++ b/Engine/Source/Programs/UnrealGameSync/UnrealGameSync/UserSettings.cs
@@ -7,6 +7,7 @@ using System.IO;
 using System.Linq;
 using System.Text;
 using System.Threading.Tasks;
+using System.Windows.Forms;
 
 namespace UnrealGameSync
 {
@@ -212,6 +213,8 @@ namespace UnrealGameSync
 
 		// Window settings
 		public bool bWindowVisible;
+		public FormWindowState WindowState;
+		public Rectangle? WindowBounds;
 		
 		// Schedule settings
 		public bool bScheduleEnabled;
@@ -347,6 +350,11 @@ namespace UnrealGameSync
 
 			// Window settings
 			bWindowVisible = ConfigFile.GetValue("Window.Visible", true);
+			if(!Enum.TryParse(ConfigFile.GetValue("Window.State", ""), true, out WindowState))
+			{
+				WindowState = FormWindowState.Normal;
+			}
+			WindowBounds = ParseRectangleValue(ConfigFile.GetValue("Window.Bounds", ""));
 
 			// Schedule settings
 			bScheduleEnabled = ConfigFile.GetValue("Schedule.Enabled", false);
@@ -376,6 +384,37 @@ namespace UnrealGameSync
 			}
 		}
 
+		static Rectangle? ParseRectangleValue(string Text)
+		{
+			ConfigObject Object = new ConfigObject(Text);
+
+			int X = Object.GetValue("X", -1);
+			int Y = Object.GetValue("Y", -1);
+			int W = Object.GetValue("W", -1);
+			int H = Object.GetValue("H", -1);
+
+			if(X == -1 || Y == -1 || W == -1 || H == -1)
+			{
+				return null;
+			}
+			else
+			{
+				return new Rectangle(X, Y, W, H);
+			}
+		}
+
+		static string FormatRectangleValue(Rectangle Value)
+		{
+			ConfigObject Object = new ConfigObject();
+
+			Object.SetValue("X", Value.X);
+			Object.SetValue("Y", Value.Y);
+			Object.SetValue("W", Value.Width);
+			Object.SetValue("H", Value.Height);
+
+			return Object.ToString();
+		}
+
 		public UserWorkspaceSettings FindOrAddWorkspace(string ClientBranchPath)
 		{
 			// Update the current workspace
@@ -569,6 +608,11 @@ namespace UnrealGameSync
 			ConfigSection WindowSection = ConfigFile.FindOrAddSection("Window");
 			WindowSection.Clear();
 			WindowSection.SetValue("Visible", bWindowVisible);
+			WindowSection.SetValue("State", WindowState.ToString());
+			if(WindowBounds != null)
+			{
+				WindowSection.SetValue("Bounds", FormatRectangleValue(WindowBounds.Value));
+			}
 
 			// Current workspace settings
 			foreach(KeyValuePair<string, UserWorkspaceSettings> Pair in WorkspaceKeyToSettings)
