UGS: Fix exception for files that aren't mapped to the workspace.

#rnx

diff --git a/Engine/Source/Programs/UnrealGameSync/UnrealGameSync/Utility.cs b/Engine/Source/Programs/UnrealGameSync/UnrealGameSync/Utility.cs
--- a/Engine/Source/Programs/UnrealGameSync/UnrealGameSync/Utility.cs
+++ b/Engine/Source/Programs/UnrealGameSync/UnrealGameSync/Utility.cs
@@ -350,6 +350,11 @@ namespace UnrealGameSync
 
 		public static bool TryPrintFileUsingCache(PerforceConnection Perforce, string DepotPath, string CacheFolder, string Digest, out List<string> Lines, TextWriter Log)
 		{
+			if(Digest == null)
+			{
+				return Perforce.Print(DepotPath, out Lines, Log);
+			}
+
 			string CacheFile = Path.Combine(CacheFolder, Digest);
 			if(File.Exists(CacheFile))
 			{
