UGS: Fixed crash when creating a new workspace if logging into the server fails.

diff --git a/Engine/Source/Programs/UnrealGameSync/UnrealGameSync/Forms/NewWorkspaceWindow.cs b/Engine/Source/Programs/UnrealGameSync/UnrealGameSync/Forms/NewWorkspaceWindow.cs
--- a/Engine/Source/Programs/UnrealGameSync/UnrealGameSync/Forms/NewWorkspaceWindow.cs
+++ b/Engine/Source/Programs/UnrealGameSync/UnrealGameSync/Forms/NewWorkspaceWindow.cs
@@ -191,6 +191,9 @@ namespace UnrealGameSync
 				{
 					MessageBox.Show(ErrorMessage);
 				}
+
+				WorkspaceName = null;
+				return false;
 			}
 
 			NewWorkspaceWindow Window = new NewWorkspaceWindow(ServerAndPort, UserName, ForceStreamName, FindSettings.CurrentStream, FindSettings.Info, FindSettings.Clients, Log);
