UGS: Fix fstat calls failing if access to any file is forbitten by p4 protect.

#rnx

diff --git a/Engine/Source/Programs/UnrealGameSync/UnrealGameSync/Perforce.cs b/Engine/Source/Programs/UnrealGameSync/UnrealGameSync/Perforce.cs
--- a/Engine/Source/Programs/UnrealGameSync/UnrealGameSync/Perforce.cs
+++ b/Engine/Source/Programs/UnrealGameSync/UnrealGameSync/Perforce.cs
@@ -471,6 +471,7 @@ namespace UnrealGameSync
 			IgnoreFilesNotInClientViewError = 0x80,
 			IgnoreEnterPassword = 0x100,
 			IgnoreFilesNotOnClientError = 0x200,
+			IgnoreProtectedNamespaceError = 0x400,
 		}
 
 		delegate bool HandleRecordDelegate(Dictionary<string, string> Tags);
@@ -734,7 +735,7 @@ namespace UnrealGameSync
 		public bool FileExists(string Filter, out bool bExists, TextWriter Log)
 		{
 			List<PerforceFileRecord> FileRecords;
-			if(RunCommand(String.Format("fstat \"{0}\"", Filter), out FileRecords, CommandOptions.IgnoreNoSuchFilesError | CommandOptions.IgnoreFilesNotInClientViewError, Log))
+			if(RunCommand(String.Format("fstat \"{0}\"", Filter), out FileRecords, CommandOptions.IgnoreNoSuchFilesError | CommandOptions.IgnoreFilesNotInClientViewError | CommandOptions.IgnoreProtectedNamespaceError, Log))
 			{
 				bExists = (FileRecords.Count > 0);
 				return true;
@@ -1080,7 +1081,7 @@ namespace UnrealGameSync
 
 		public bool Stat(string Filter, out List<PerforceFileRecord> FileRecords, TextWriter Log)
 		{
-			return RunCommand(String.Format("fstat \"{0}\"", Filter), out FileRecords, CommandOptions.IgnoreFilesNotOnClientError | CommandOptions.IgnoreNoSuchFilesError, Log);
+			return RunCommand(String.Format("fstat \"{0}\"", Filter), out FileRecords, CommandOptions.IgnoreFilesNotOnClientError | CommandOptions.IgnoreNoSuchFilesError | CommandOptions.IgnoreProtectedNamespaceError, Log);
 		}
 
 		public bool Stat(string Options, List<string> Files, out List<PerforceFileRecord> FileRecords, TextWriter Log)
@@ -1094,7 +1095,7 @@ namespace UnrealGameSync
 			{
 				Arguments.AppendFormat(" \"{0}\"", File);
 			}
-			return RunCommand(Arguments.ToString(), out FileRecords, CommandOptions.IgnoreFilesNotOnClientError | CommandOptions.IgnoreNoSuchFilesError, Log);
+			return RunCommand(Arguments.ToString(), out FileRecords, CommandOptions.IgnoreFilesNotOnClientError | CommandOptions.IgnoreNoSuchFilesError | CommandOptions.IgnoreProtectedNamespaceError, Log);
 		}
 
 		public bool Sync(string Filter, TextWriter Log)
@@ -1459,6 +1460,10 @@ namespace UnrealGameSync
 			{
 				return true;
 			}
+			else if(Options.HasFlag(CommandOptions.IgnoreProtectedNamespaceError) && Text.StartsWith("error: ") && Text.EndsWith(" - protected namespace - access denied."))
+			{
+				return true;
+			}
 			else if(Options.HasFlag(CommandOptions.IgnoreEnterPassword) && Text.StartsWith("Enter password:"))
 			{
 				return true;
