UGS: Use local time for grouping changes.

diff --git a/Engine/Source/Programs/UnrealGameSync/UnrealGameSync/Controls/WorkspaceControl.cs b/Engine/Source/Programs/UnrealGameSync/UnrealGameSync/Controls/WorkspaceControl.cs
--- a/Engine/Source/Programs/UnrealGameSync/UnrealGameSync/Controls/WorkspaceControl.cs
+++ b/Engine/Source/Programs/UnrealGameSync/UnrealGameSync/Controls/WorkspaceControl.cs
@@ -1078,7 +1078,13 @@ namespace UnrealGameSync
 
 							ListViewGroup Group;
 
-							string GroupName = Change.Date.ToString("D");//"dddd\\,\\ h\\.mmtt");
+							DateTime DisplayTime = Change.Date;
+							if(Settings.bShowLocalTimes)
+							{
+								DisplayTime = (DisplayTime - PerforceMonitor.ServerTimeZone).ToLocalTime();
+							}
+
+							string GroupName = DisplayTime.ToString("D");//"dddd\\,\\ h\\.mmtt");
 							for(int Idx = 0;;Idx++)
 							{
 								if(Idx == BuildList.Groups.Count)
@@ -1095,12 +1101,6 @@ namespace UnrealGameSync
 								}
 							}
 
-							DateTime DisplayTime = Change.Date;
-							if(Settings.bShowLocalTimes)
-							{
-								DisplayTime = (DisplayTime - PerforceMonitor.ServerTimeZone).ToLocalTime();
-							}
-
 							string[] SubItemLabels = new string[BuildList.Columns.Count];
 							SubItemLabels[ChangeColumn.Index] = String.Format("{0}", Change.Number);
 							SubItemLabels[TimeColumn.Index] = DisplayTime.ToString("h\\.mmtt");
