UGS: Fix exception when clicking on a badge with an empty URL string.

#rnx

diff --git a/Engine/Source/Programs/UnrealGameSync/UnrealGameSync/Controls/WorkspaceControl.cs b/Engine/Source/Programs/UnrealGameSync/UnrealGameSync/Controls/WorkspaceControl.cs
--- a/Engine/Source/Programs/UnrealGameSync/UnrealGameSync/Controls/WorkspaceControl.cs
+++ b/Engine/Source/Programs/UnrealGameSync/UnrealGameSync/Controls/WorkspaceControl.cs
@@ -2147,7 +2147,7 @@ namespace UnrealGameSync
 			Color HoverBadgeColor = Color.FromArgb(BadgeColor.A, Math.Min(BadgeColor.R + 32, 255), Math.Min(BadgeColor.G + 32, 255), Math.Min(BadgeColor.B + 32, 255));
 
 			Action ClickHandler;
-			if(BadgeData == null || BadgeData.Url == null)
+			if(BadgeData == null || String.IsNullOrEmpty(BadgeData.Url))
 			{
 				ClickHandler = null;
 			}
