UGS: Add another automation request type that just checks if a given stream/project is open; similarly to the SyncProject request but without any modal UI.

diff --git a/Engine/Source/Programs/UnrealGameSync/UnrealGameSync/AutomationServer.cs b/Engine/Source/Programs/UnrealGameSync/UnrealGameSync/AutomationServer.cs
--- a/Engine/Source/Programs/UnrealGameSync/UnrealGameSync/AutomationServer.cs
+++ b/Engine/Source/Programs/UnrealGameSync/UnrealGameSync/AutomationServer.cs
@@ -15,7 +15,8 @@ namespace UnrealGameSync
 {
 	enum AutomationRequestType
 	{
-		SyncProject
+		SyncProject,
+		FindProject,
 	}
 
 	class AutomationRequestInput
@@ -56,7 +57,8 @@ namespace UnrealGameSync
 		Invalid,
 		Busy,
 		Canceled,
-		Error
+		Error,
+		NotFound
 	}
 
 	class AutomationRequestOutput
diff --git a/Engine/Source/Programs/UnrealGameSync/UnrealGameSync/Forms/MainWindow.cs b/Engine/Source/Programs/UnrealGameSync/UnrealGameSync/Forms/MainWindow.cs
--- a/Engine/Source/Programs/UnrealGameSync/UnrealGameSync/Forms/MainWindow.cs
+++ b/Engine/Source/Programs/UnrealGameSync/UnrealGameSync/Forms/MainWindow.cs
@@ -174,6 +174,11 @@ namespace UnrealGameSync
 						Request.SetOutput(new AutomationRequestOutput(Result));
 					}
 				}
+				else if(Request.Input.Type == AutomationRequestType.FindProject)
+				{
+					AutomationRequestOutput Output = FindProject(Request);
+					Request.SetOutput(Output);
+				}
 				else
 				{
 					Request.SetOutput(new AutomationRequestOutput(AutomationRequestResult.Invalid));
@@ -270,6 +275,36 @@ namespace UnrealGameSync
 			}
 		}
 
+		AutomationRequestOutput FindProject(AutomationRequest Request)
+		{
+			BinaryReader Reader = new BinaryReader(new MemoryStream(Request.Input.Data));
+			string StreamName = Reader.ReadString();
+			string ProjectPath = Reader.ReadString();
+
+			for(int ExistingTabIdx = 0; ExistingTabIdx < TabControl.GetTabCount(); ExistingTabIdx++)
+			{
+				WorkspaceControl ExistingWorkspace = TabControl.GetTabData(ExistingTabIdx) as WorkspaceControl;
+				if(ExistingWorkspace != null && String.Compare(ExistingWorkspace.StreamName, StreamName, StringComparison.OrdinalIgnoreCase) == 0 && ExistingWorkspace.SelectedProject != null)
+				{
+					string ClientPath = ExistingWorkspace.SelectedProject.ClientPath;
+					if(ClientPath != null && ClientPath.StartsWith("//"))
+					{
+						int SlashIdx = ClientPath.IndexOf('/', 2);
+						if(SlashIdx != -1)
+						{
+							string ExistingProjectPath = ClientPath.Substring(SlashIdx);
+							if(String.Compare(ExistingProjectPath, ProjectPath, StringComparison.OrdinalIgnoreCase) == 0)
+							{
+								return new AutomationRequestOutput(AutomationRequestResult.Ok, Encoding.UTF8.GetBytes(ExistingWorkspace.SelectedFileName));
+							}
+						}
+					}
+				}
+			}
+
+			return new AutomationRequestOutput(AutomationRequestResult.NotFound);
+		}
+
 		protected override void OnHandleCreated(EventArgs e)
 		{
 			base.OnHandleCreated(e);
