UGS: Add an OpenProject automation command, which will only sync the stream if you can't launch the editor.

diff --git a/Engine/Source/Programs/UnrealGameSync/UnrealGameSync/AutomationServer.cs b/Engine/Source/Programs/UnrealGameSync/UnrealGameSync/AutomationServer.cs
--- a/Engine/Source/Programs/UnrealGameSync/UnrealGameSync/AutomationServer.cs
+++ b/Engine/Source/Programs/UnrealGameSync/UnrealGameSync/AutomationServer.cs
@@ -17,6 +17,7 @@ namespace UnrealGameSync
 	{
 		SyncProject,
 		FindProject,
+		OpenProject,
 	}
 
 	class AutomationRequestInput
diff --git a/Engine/Source/Programs/UnrealGameSync/UnrealGameSync/Forms/MainWindow.cs b/Engine/Source/Programs/UnrealGameSync/UnrealGameSync/Forms/MainWindow.cs
--- a/Engine/Source/Programs/UnrealGameSync/UnrealGameSync/Forms/MainWindow.cs
+++ b/Engine/Source/Programs/UnrealGameSync/UnrealGameSync/Forms/MainWindow.cs
@@ -168,10 +168,10 @@ namespace UnrealGameSync
 				}
 				else if(Request.Input.Type == AutomationRequestType.SyncProject)
 				{
-					AutomationRequestResult Result = StartAutomatedSync(Request);
-					if(Result != AutomationRequestResult.Ok)
+					AutomationRequestOutput Output = StartAutomatedSync(Request, true);
+					if(Output != null)
 					{
-						Request.SetOutput(new AutomationRequestOutput(Result));
+						Request.SetOutput(Output);
 					}
 				}
 				else if(Request.Input.Type == AutomationRequestType.FindProject)
@@ -179,6 +179,14 @@ namespace UnrealGameSync
 					AutomationRequestOutput Output = FindProject(Request);
 					Request.SetOutput(Output);
 				}
+				else if(Request.Input.Type == AutomationRequestType.OpenProject)
+				{
+					AutomationRequestOutput Output = StartAutomatedSync(Request, false);
+					if(Output != null)
+					{
+						Request.SetOutput(Output);
+					}
+				}
 				else
 				{
 					Request.SetOutput(new AutomationRequestOutput(AutomationRequestResult.Invalid));
@@ -191,7 +199,7 @@ namespace UnrealGameSync
 			}
 		}
 
-		AutomationRequestResult StartAutomatedSync(AutomationRequest Request)
+		AutomationRequestOutput StartAutomatedSync(AutomationRequest Request, bool bForceSync)
 		{
 			ShowAndActivate();
 
@@ -202,7 +210,7 @@ namespace UnrealGameSync
 			AutomatedSyncWindow.WorkspaceInfo WorkspaceInfo;
 			if(!AutomatedSyncWindow.ShowModal(this, StreamName, ProjectPath, out WorkspaceInfo, Log))
 			{
-				return AutomationRequestResult.Canceled;
+				return new AutomationRequestOutput(AutomationRequestResult.Canceled);
 			}
 
 			if(WorkspaceInfo.bRequiresStreamSwitch)
@@ -223,7 +231,7 @@ namespace UnrealGameSync
 				if(!Perforce.SwitchStream(StreamName, Log))
 				{
 					Log.WriteLine("Unable to switch stream");
-					return AutomationRequestResult.Error;
+					return new AutomationRequestOutput(AutomationRequestResult.Error);
 				}
 			}
 
@@ -233,18 +241,23 @@ namespace UnrealGameSync
 			if(TabIdx == -1)
 			{
 				Log.WriteLine("Unable to open project");
-				return AutomationRequestResult.Error;
+				return new AutomationRequestOutput(AutomationRequestResult.Error);
 			}
 
 			WorkspaceControl Workspace = TabControl.GetTabData(TabIdx) as WorkspaceControl;
 			if(Workspace == null)
 			{
 				Log.WriteLine("Workspace was unable to open");
-				return AutomationRequestResult.Error;
+				return new AutomationRequestOutput(AutomationRequestResult.Error);
+			}
+
+			if(!bForceSync && Workspace.CanLaunchEditor())
+			{
+				return new AutomationRequestOutput(AutomationRequestResult.Ok, Encoding.UTF8.GetBytes(Workspace.SelectedFileName));
 			}
 
 			Workspace.AddStartupCallback((Control, bCancel) => StartAutomatedSyncAfterStartup(Control, bCancel, Request));
-			return AutomationRequestResult.Ok;
+			return null;
 		}
 
 		private void StartAutomatedSyncAfterStartup(WorkspaceControl Workspace, bool bCancel, AutomationRequest Request)
diff --git a/Engine/Source/Programs/UnrealGameSync/UnrealGameSync/Properties/AssemblyInfo.cs b/Engine/Source/Programs/UnrealGameSync/UnrealGameSync/Properties/AssemblyInfo.cs
--- a/Engine/Source/Programs/UnrealGameSync/UnrealGameSync/Properties/AssemblyInfo.cs
+++ b/Engine/Source/Programs/UnrealGameSync/UnrealGameSync/Properties/AssemblyInfo.cs
@@ -34,5 +34,5 @@ using System.Runtime.InteropServices;
 // You can specify all the values or you can default the Build and Revision Numbers 
 // by using the '*' as shown below:
 // [assembly: AssemblyVersion("1.0.*")]
-[assembly: AssemblyVersion("1.166.*")]
-[assembly: AssemblyFileVersion("1.166.0.0")]
+[assembly: AssemblyVersion("1.167.*")]
+[assembly: AssemblyFileVersion("1.167.0.0")]
