UGS: Add test app showing how to use UGS automation.

diff --git a/Engine/Source/Programs/UnrealGameSync/UnrealGameSyncAutomation/App.config b/Engine/Source/Programs/UnrealGameSync/UnrealGameSyncAutomation/App.config
new file mode 100644
--- /dev/null
+++ b/Engine/Source/Programs/UnrealGameSync/UnrealGameSyncAutomation/App.config
@@ -0,0 +1,6 @@
+﻿<?xml version="1.0" encoding="utf-8" ?>
+<configuration>
+    <startup> 
+        <supportedRuntime version="v4.0" sku=".NETFramework,Version=v4.6.1" />
+    </startup>
+</configuration>
\ No newline at end of file
diff --git a/Engine/Source/Programs/UnrealGameSync/UnrealGameSyncAutomation/Program.cs b/Engine/Source/Programs/UnrealGameSync/UnrealGameSyncAutomation/Program.cs
new file mode 100644
--- /dev/null
+++ b/Engine/Source/Programs/UnrealGameSync/UnrealGameSyncAutomation/Program.cs
@@ -0,0 +1,75 @@
+using Microsoft.Win32;
+using System;
+using System.Collections.Generic;
+using System.IO;
+using System.Linq;
+using System.Net;
+using System.Net.Sockets;
+using System.Text;
+using System.Threading.Tasks;
+
+namespace UnrealGameSyncAutomation
+{
+	class Program
+	{
+		enum AutomationRequestType
+		{
+			SyncProject,
+			FindProject,
+			OpenProject
+		}
+
+		enum AutomationRequestResult
+		{
+			Ok,
+			Invalid,
+			Busy,
+			Canceled,
+			Error,
+			NotFound
+		}
+
+		static void Main(string[] Arguments)
+		{
+			// Create the request data
+			MemoryStream InputDataStream = new MemoryStream();
+			using(BinaryWriter Writer = new BinaryWriter(InputDataStream))
+			{
+				Writer.Write("//UE4/Main");
+				Writer.Write("/Samples/Games/ShooterGame/ShooterGame.uproject");
+			}
+
+			Tuple<AutomationRequestResult, byte[]> Response = SendRequest(AutomationRequestType.OpenProject, InputDataStream.ToArray());
+		
+			string ResponseString = Encoding.UTF8.GetString(Response.Item2);
+			Console.WriteLine("{0}: {1}", Response.Item1, ResponseString);
+		}
+
+		static Tuple<AutomationRequestResult, byte[]> SendRequest(AutomationRequestType Request, byte[] RequestData)
+		{
+			int PortNumber = (int)Registry.GetValue("HKEY_CURRENT_USER\\Software\\Epic Games\\UnrealGameSync", "AutomationPort", null);
+			using(TcpClient Client = new TcpClient())
+			{
+				Client.Connect(new IPEndPoint(IPAddress.Loopback, PortNumber));
+
+				using(NetworkStream Stream = Client.GetStream())
+				{
+					// Send the request
+					BinaryWriter Writer = new BinaryWriter(Stream);
+					Writer.Write((int)Request);
+					Writer.Write(RequestData.Length);
+					Writer.Write(RequestData);
+
+					// Read the response
+					BinaryReader Reader = new BinaryReader(Stream);
+					AutomationRequestResult Result = (AutomationRequestResult)Reader.ReadInt32();
+					int ResponseLength = Reader.ReadInt32();
+					byte[] ResponseData = Reader.ReadBytes(ResponseLength);
+
+					// Return the response
+					return Tuple.Create(Result, ResponseData);
+				}
+			}
+		}
+	}
+}
diff --git a/Engine/Source/Programs/UnrealGameSync/UnrealGameSyncAutomation/Properties/AssemblyInfo.cs b/Engine/Source/Programs/UnrealGameSync/UnrealGameSyncAutomation/Properties/AssemblyInfo.cs
new file mode 100644
--- /dev/null
+++ b/Engine/Source/Programs/UnrealGameSync/UnrealGameSyncAutomation/Properties/AssemblyInfo.cs
@@ -0,0 +1,36 @@
+﻿using System.Reflection;
+using System.Runtime.CompilerServices;
+using System.Runtime.InteropServices;
+
+// General Information about an assembly is controlled through the following
+// set of attributes. Change these attribute values to modify the information
+// associated with an assembly.
+[assembly: AssemblyTitle("UnrealGameSyncAutomation")]
+[assembly: AssemblyDescription("")]
+[assembly: AssemblyConfiguration("")]
+[assembly: AssemblyCompany("")]
+[assembly: AssemblyProduct("UnrealGameSyncAutomation")]
+[assembly: AssemblyCopyright("Copyright ©  2018")]
+[assembly: AssemblyTrademark("")]
+[assembly: AssemblyCulture("")]
+
+// Setting ComVisible to false makes the types in this assembly not visible
+// to COM components.  If you need to access a type in this assembly from
+// COM, set the ComVisible attribute to true on that type.
+[assembly: ComVisible(false)]
+
+// The following GUID is for the ID of the typelib if this project is exposed to COM
+[assembly: Guid("7c9d0a33-c2a8-4415-b9ed-9f40337dbec6")]
+
+// Version information for an assembly consists of the following four values:
+//
+//      Major Version
+//      Minor Version
+//      Build Number
+//      Revision
+//
+// You can specify all the values or you can default the Build and Revision Numbers
+// by using the '*' as shown below:
+// [assembly: AssemblyVersion("1.0.*")]
+[assembly: AssemblyVersion("1.0.0.0")]
+[assembly: AssemblyFileVersion("1.0.0.0")]
diff --git a/Engine/Source/Programs/UnrealGameSync/UnrealGameSyncAutomation/UnrealGameSyncAutomation.csproj b/Engine/Source/Programs/UnrealGameSync/UnrealGameSyncAutomation/UnrealGameSyncAutomation.csproj
new file mode 100644
--- /dev/null
+++ b/Engine/Source/Programs/UnrealGameSync/UnrealGameSyncAutomation/UnrealGameSyncAutomation.csproj
@@ -0,0 +1,53 @@
+﻿<?xml version="1.0" encoding="utf-8"?>
+<Project ToolsVersion="15.0" xmlns="http://schemas.microsoft.com/developer/msbuild/2003">
+  <Import Project="$(MSBuildExtensionsPath)\$(MSBuildToolsVersion)\Microsoft.Common.props" Condition="Exists('$(MSBuildExtensionsPath)\$(MSBuildToolsVersion)\Microsoft.Common.props')" />
+  <PropertyGroup>
+    <Configuration Condition=" '$(Configuration)' == '' ">Debug</Configuration>
+    <Platform Condition=" '$(Platform)' == '' ">AnyCPU</Platform>
+    <ProjectGuid>{7C9D0A33-C2A8-4415-B9ED-9F40337DBEC6}</ProjectGuid>
+    <OutputType>Exe</OutputType>
+    <RootNamespace>UnrealGameSyncAutomation</RootNamespace>
+    <AssemblyName>UnrealGameSyncAutomation</AssemblyName>
+    <TargetFrameworkVersion>v4.6.1</TargetFrameworkVersion>
+    <FileAlignment>512</FileAlignment>
+    <AutoGenerateBindingRedirects>true</AutoGenerateBindingRedirects>
+    <Deterministic>true</Deterministic>
+  </PropertyGroup>
+  <PropertyGroup Condition=" '$(Configuration)|$(Platform)' == 'Debug|AnyCPU' ">
+    <PlatformTarget>AnyCPU</PlatformTarget>
+    <DebugSymbols>true</DebugSymbols>
+    <DebugType>full</DebugType>
+    <Optimize>false</Optimize>
+    <OutputPath>bin\Debug\</OutputPath>
+    <DefineConstants>DEBUG;TRACE</DefineConstants>
+    <ErrorReport>prompt</ErrorReport>
+    <WarningLevel>4</WarningLevel>
+  </PropertyGroup>
+  <PropertyGroup Condition=" '$(Configuration)|$(Platform)' == 'Release|AnyCPU' ">
+    <PlatformTarget>AnyCPU</PlatformTarget>
+    <DebugType>pdbonly</DebugType>
+    <Optimize>true</Optimize>
+    <OutputPath>bin\Release\</OutputPath>
+    <DefineConstants>TRACE</DefineConstants>
+    <ErrorReport>prompt</ErrorReport>
+    <WarningLevel>4</WarningLevel>
+  </PropertyGroup>
+  <ItemGroup>
+    <Reference Include="System" />
+    <Reference Include="System.Core" />
+    <Reference Include="System.Xml.Linq" />
+    <Reference Include="System.Data.DataSetExtensions" />
+    <Reference Include="Microsoft.CSharp" />
+    <Reference Include="System.Data" />
+    <Reference Include="System.Net.Http" />
+    <Reference Include="System.Xml" />
+  </ItemGroup>
+  <ItemGroup>
+    <Compile Include="Program.cs" />
+    <Compile Include="Properties\AssemblyInfo.cs" />
+  </ItemGroup>
+  <ItemGroup>
+    <None Include="App.config" />
+  </ItemGroup>
+  <Import Project="$(MSBuildToolsPath)\Microsoft.CSharp.targets" />
+</Project>
\ No newline at end of file
