UGS: Fix config files being read twice when the current workspace had files left over from another stream.

diff --git a/Engine/Source/Programs/UnrealGameSync/UnrealGameSync/PerforceMonitor.cs b/Engine/Source/Programs/UnrealGameSync/UnrealGameSync/PerforceMonitor.cs
--- a/Engine/Source/Programs/UnrealGameSync/UnrealGameSync/PerforceMonitor.cs
+++ b/Engine/Source/Programs/UnrealGameSync/UnrealGameSync/PerforceMonitor.cs
@@ -530,6 +530,12 @@ namespace UnrealGameSync
 				{
 					List<string> Lines = null;
 
+					// Skip file records which are still in the workspace, but were synced from a different branch. For these files, the action seems to be empty, so filter against that.
+					if(FileRecord.Action == null)
+					{
+						continue;
+					}
+
 					// If this file is open for edit, read the local version
 					string LocalFileName = FileRecord.ClientPath;
 					if(LocalFileName != null && File.Exists(LocalFileName) && (File.GetAttributes(LocalFileName) & FileAttributes.ReadOnly) == 0)
