UGS: Fix status panel not repainting when the tint color changes.

diff --git a/Engine/Source/Programs/UnrealGameSync/UnrealGameSync/Controls/StatusPanel.cs b/Engine/Source/Programs/UnrealGameSync/UnrealGameSync/Controls/StatusPanel.cs
--- a/Engine/Source/Programs/UnrealGameSync/UnrealGameSync/Controls/StatusPanel.cs
+++ b/Engine/Source/Programs/UnrealGameSync/UnrealGameSync/Controls/StatusPanel.cs
@@ -461,6 +461,11 @@ namespace UnrealGameSync
 				Resources = new StatusElementResources(Font);
 			}
 
+			if(TintColor != NewTintColor)
+			{
+				Invalidate();
+			}
+
 			InvalidateElements();
 			Lines.Clear();
 			Lines.AddRange(NewLines);
