Fix CIS warnings.

#rnx

diff --git a/Engine/Source/Programs/UnrealGameSync/UnrealGameSyncAutomation/App.config b/Engine/Source/Programs/UnrealGameSync/UnrealGameSyncAutomation/App.config
--- a/Engine/Source/Programs/UnrealGameSync/UnrealGameSyncAutomation/App.config
+++ b/Engine/Source/Programs/UnrealGameSync/UnrealGameSyncAutomation/App.config
@@ -1,6 +1,6 @@
-﻿<?xml version="1.0" encoding="utf-8" ?>
+<?xml version="1.0" encoding="utf-8"?>
 <configuration>
     <startup> 
-        <supportedRuntime version="v4.0" sku=".NETFramework,Version=v4.6.1" />
+        <supportedRuntime version="v4.0" sku=".NETFramework,Version=v4.5"/>
     </startup>
-</configuration>
\ No newline at end of file
+</configuration>
diff --git a/Engine/Source/Programs/UnrealGameSync/UnrealGameSyncAutomation/Program.cs b/Engine/Source/Programs/UnrealGameSync/UnrealGameSyncAutomation/Program.cs
--- a/Engine/Source/Programs/UnrealGameSync/UnrealGameSyncAutomation/Program.cs
+++ b/Engine/Source/Programs/UnrealGameSync/UnrealGameSyncAutomation/Program.cs
@@ -1,3 +1,5 @@
+// Copyright 1998-2018 Epic Games, Inc. All Rights Reserved.
+
 using Microsoft.Win32;
 using System;
 using System.Collections.Generic;
diff --git a/Engine/Source/Programs/UnrealGameSync/UnrealGameSyncAutomation/Properties/AssemblyInfo.cs b/Engine/Source/Programs/UnrealGameSync/UnrealGameSyncAutomation/Properties/AssemblyInfo.cs
--- a/Engine/Source/Programs/UnrealGameSync/UnrealGameSyncAutomation/Properties/AssemblyInfo.cs
+++ b/Engine/Source/Programs/UnrealGameSync/UnrealGameSyncAutomation/Properties/AssemblyInfo.cs
@@ -1,4 +1,6 @@
-﻿using System.Reflection;
+// Copyright 1998-2018 Epic Games, Inc. All Rights Reserved.
+
+using System.Reflection;
 using System.Runtime.CompilerServices;
 using System.Runtime.InteropServices;
 
diff --git a/Engine/Source/Programs/UnrealGameSync/UnrealGameSyncAutomation/UnrealGameSyncAutomation.csproj b/Engine/Source/Programs/UnrealGameSync/UnrealGameSyncAutomation/UnrealGameSyncAutomation.csproj
--- a/Engine/Source/Programs/UnrealGameSync/UnrealGameSyncAutomation/UnrealGameSyncAutomation.csproj
+++ b/Engine/Source/Programs/UnrealGameSync/UnrealGameSyncAutomation/UnrealGameSyncAutomation.csproj
@@ -8,10 +8,11 @@
     <OutputType>Exe</OutputType>
     <RootNamespace>UnrealGameSyncAutomation</RootNamespace>
     <AssemblyName>UnrealGameSyncAutomation</AssemblyName>
-    <TargetFrameworkVersion>v4.6.1</TargetFrameworkVersion>
+    <TargetFrameworkVersion>v4.5</TargetFrameworkVersion>
     <FileAlignment>512</FileAlignment>
     <AutoGenerateBindingRedirects>true</AutoGenerateBindingRedirects>
     <Deterministic>true</Deterministic>
+    <TargetFrameworkProfile />
   </PropertyGroup>
   <PropertyGroup Condition=" '$(Configuration)|$(Platform)' == 'Debug|AnyCPU' ">
     <PlatformTarget>AnyCPU</PlatformTarget>
