UGS: Fix missing Perforce login prompt if an existing ticket has expired.

diff --git a/Engine/Source/Programs/UnrealGameSync/UnrealGameSync/Perforce.cs b/Engine/Source/Programs/UnrealGameSync/UnrealGameSync/Perforce.cs
--- a/Engine/Source/Programs/UnrealGameSync/UnrealGameSync/Perforce.cs
+++ b/Engine/Source/Programs/UnrealGameSync/UnrealGameSync/Perforce.cs
@@ -510,7 +510,7 @@ namespace UnrealGameSync
 				bIsLoggedIn = true;
 				return true;
 			}
-			else if(Lines[0].Channel == PerforceOutputChannel.Error && Lines[0].Text.Contains("P4PASSWD"))
+			else if(Lines[0].Channel == PerforceOutputChannel.Error && (Lines[0].Text.Contains("P4PASSWD") || Lines[0].Text.Contains("has expired")))
 			{
 				bIsLoggedIn = false;
 				return true;
