remove accidentally added publish profiles

diff --git a/Engine/Source/Programs/UnrealGameSync/UnrealGameSyncMetadataServer/UnrealGameSyncMetadataServer.csproj b/Engine/Source/Programs/UnrealGameSync/UnrealGameSyncMetadataServer/UnrealGameSyncMetadataServer.csproj
--- a/Engine/Source/Programs/UnrealGameSync/UnrealGameSyncMetadataServer/UnrealGameSyncMetadataServer.csproj
+++ b/Engine/Source/Programs/UnrealGameSync/UnrealGameSyncMetadataServer/UnrealGameSyncMetadataServer.csproj
@@ -124,8 +124,6 @@
     <Content Include="NotForLicensees\Web.config">
       <SubType>Designer</SubType>
     </Content>
-    <None Include="Properties\PublishProfiles\FilePublish.pubxml" />
-    <None Include="Properties\PublishProfiles\Profile.pubxml" />
   </ItemGroup>
   <ItemGroup>
     <Content Include="packages.config" />
