UGS: Move DB setup script into a resource.

diff --git a/Engine/Source/Programs/UnrealGameSync/UnrealGameSyncMetadataServer/Global.asax.cs b/Engine/Source/Programs/UnrealGameSync/UnrealGameSyncMetadataServer/Global.asax.cs
--- a/Engine/Source/Programs/UnrealGameSync/UnrealGameSyncMetadataServer/Global.asax.cs
+++ b/Engine/Source/Programs/UnrealGameSync/UnrealGameSyncMetadataServer/Global.asax.cs
@@ -25,60 +25,7 @@ namespace UnrealGameSyncMetadataServer
 				using (SQLiteConnection Connection = new SQLiteConnection(ConnectionString, true))
 				{
 					Connection.Open();
-					string CreateTables = @"CREATE TABLE IF NOT EXISTS
-										[CIS] (
-											[Id]           INTEGER       NOT NULL PRIMARY KEY AUTOINCREMENT,
-											[ChangeNumber] INTEGER       NOT NULL,
-											[BuildType]    NCHAR (32)    NOT NULL,
-											[Result]       NCHAR (10)    NOT NULL,
-											[Url]          VARCHAR (512) NOT NULL,
-											[Project]      VARCHAR (256) NULL,
-											[ArchivePath]  VARCHAR (512) NULL
-										);
-										
-										CREATE TABLE IF NOT EXISTS
-										[Comments] (
-											[Id]           INTEGER    	  NOT NULL PRIMARY KEY AUTOINCREMENT,
-											[ChangeNumber] INTEGER        NOT NULL,
-											[UserName]     VARCHAR (128)  NOT NULL,
-											[Text]         VARCHAR (140)  NOT NULL,
-											[Project]      NVARCHAR (256) NOT NULL
-										);
-
-										CREATE TABLE IF NOT EXISTS
-										[Errors] (
-											[Id]        INTEGER		   NOT NULL PRIMARY KEY AUTOINCREMENT,
-											[Type]      VARCHAR (50)   NOT NULL,
-											[Text]      VARCHAR (1024) NOT NULL,
-											[UserName]  NVARCHAR (128) NOT NULL,
-											[Project]   VARCHAR (128)  NULL,
-											[Timestamp] DATETIME       NOT NULL,
-											[Version]   VARCHAR (64)   NOT NULL,
-											[IpAddress] VARCHAR (64)   NOT NULL
-										);
-										
-										CREATE TABLE IF NOT EXISTS
-										[Telemetry.v2] (
-											[Id]        INTEGER		   NOT NULL PRIMARY KEY AUTOINCREMENT,
-											[Action]    VARCHAR (128)  NOT NULL,
-											[Result]    VARCHAR (128)  NOT NULL,
-											[UserName]  NVARCHAR (128) NOT NULL,
-											[Project]   VARCHAR (128)  NOT NULL,
-											[Timestamp] DATETIME       NOT NULL,
-											[Duration]  REAL           NOT NULL,
-											[Version]   VARCHAR (64)   NULL,
-											[IpAddress] VARCHAR (64)   NULL
-										);
-
-										CREATE TABLE IF NOT EXISTS
-										[UserVotes] (
-											[Id]         INTEGER		NOT NULL PRIMARY KEY AUTOINCREMENT,
-											[Changelist] INTEGER        NOT NULL,
-											[UserName]   NVARCHAR (128) NOT NULL,
-											[Verdict]    NCHAR (32)     NOT NULL,
-											[Project]    NVARCHAR (256) NULL
-										);";
-					using (SQLiteCommand Command = new SQLiteCommand(CreateTables, Connection))
+					using (SQLiteCommand Command = new SQLiteCommand(Properties.Resources.Setup, Connection))
 					{
 						Command.ExecuteNonQuery();
 					}
diff --git a/Engine/Source/Programs/UnrealGameSync/UnrealGameSyncMetadataServer/UnrealGameSyncMetadataServer.csproj b/Engine/Source/Programs/UnrealGameSync/UnrealGameSyncMetadataServer/UnrealGameSyncMetadataServer.csproj
--- a/Engine/Source/Programs/UnrealGameSync/UnrealGameSyncMetadataServer/UnrealGameSyncMetadataServer.csproj
+++ b/Engine/Source/Programs/UnrealGameSync/UnrealGameSyncMetadataServer/UnrealGameSyncMetadataServer.csproj
@@ -24,6 +24,7 @@
     <NuGetPackageImportStamp>
     </NuGetPackageImportStamp>
     <TargetFrameworkProfile />
+    <Use64BitIISExpress />
   </PropertyGroup>
   <PropertyGroup Condition=" '$(Configuration)|$(Platform)' == 'Debug|AnyCPU' ">
     <TransformWebConfigEnabled>False</TransformWebConfigEnabled>
@@ -89,6 +90,12 @@
   </ItemGroup>
   <ItemGroup>
     <Compile Include="App_Start\WebApiConfig.cs" />
+    <Compile Include="Properties\Resources.Designer.cs">
+      <AutoGen>True</AutoGen>
+      <DesignTime>True</DesignTime>
+      <DependentUpon>Resources.resx</DependentUpon>
+    </Compile>
+    <Content Include="Setup.sql" />
     <Compile Include="Controllers\BuildController.cs" />
     <Compile Include="Controllers\CISController.cs" />
     <Compile Include="Controllers\CommentController.cs" />
@@ -128,6 +135,13 @@
   <ItemGroup>
     <Content Include="packages.config" />
   </ItemGroup>
+  <ItemGroup>
+    <EmbeddedResource Include="Properties\Resources.resx">
+      <Generator>ResXFileCodeGenerator</Generator>
+      <LastGenOutput>Resources.Designer.cs</LastGenOutput>
+      <SubType>Designer</SubType>
+    </EmbeddedResource>
+  </ItemGroup>
   <PropertyGroup>
     <VisualStudioVersion Condition="'$(VisualStudioVersion)' == ''">10.0</VisualStudioVersion>
     <VSToolsPath Condition="'$(VSToolsPath)' == ''">$(MSBuildExtensionsPath32)\Microsoft\VisualStudio\v$(VisualStudioVersion)</VSToolsPath>
