UGS: Fix tiny hand cursor when running on high DPI monitor.

diff --git a/Engine/Source/Programs/UnrealGameSync/UnrealGameSync/Controls/StatusPanel.cs b/Engine/Source/Programs/UnrealGameSync/UnrealGameSync/Controls/StatusPanel.cs
--- a/Engine/Source/Programs/UnrealGameSync/UnrealGameSync/Controls/StatusPanel.cs
+++ b/Engine/Source/Programs/UnrealGameSync/UnrealGameSync/Controls/StatusPanel.cs
@@ -6,12 +6,24 @@ using System.ComponentModel;
 using System.Drawing;
 using System.Drawing.Drawing2D;
 using System.Linq;
+using System.Runtime.InteropServices;
 using System.Text;
 using System.Threading.Tasks;
 using System.Windows.Forms;
 
 namespace UnrealGameSync
 {
+	// For some reason, WinForms uses embedded resources for certain cursors, which don't scale correctly for high DPI modes
+	static class NativeCursors
+	{
+		[DllImport("user32.dll")]
+		public static extern IntPtr LoadCursor(IntPtr hInstance, IntPtr lpCursorName);
+
+		const int IDC_HAND = 32649;
+
+		public static readonly Cursor Hand = new Cursor(LoadCursor(IntPtr.Zero, new IntPtr(IDC_HAND)));
+	}
+
 	class StatusElementResources
 	{
 		Dictionary<FontStyle, Font> FontCache = new Dictionary<FontStyle, Font>();
@@ -148,7 +160,7 @@ namespace UnrealGameSync
 			Text = InText;
 			Style = InStyle;
 			LinkAction = InLinkAction;
-			Cursor = Cursors.Hand;
+			Cursor = NativeCursors.Hand;
 		}
 
 		public override void OnClick(Point Location)
@@ -198,7 +210,7 @@ namespace UnrealGameSync
 			ClickAction = InClickAction;
 			if(ClickAction != null)
 			{
-				Cursor = Cursors.Hand;
+				Cursor = NativeCursors.Hand;
 			}
 		}
 
