UGS: Fix PerforceConnection.FileExists returning true for files which have been deleted.

#rnx

diff --git a/Engine/Source/Programs/UnrealGameSync/UnrealGameSync/Perforce.cs b/Engine/Source/Programs/UnrealGameSync/UnrealGameSync/Perforce.cs
--- a/Engine/Source/Programs/UnrealGameSync/UnrealGameSync/Perforce.cs
+++ b/Engine/Source/Programs/UnrealGameSync/UnrealGameSync/Perforce.cs
@@ -737,7 +737,7 @@ namespace UnrealGameSync
 			List<PerforceFileRecord> FileRecords;
 			if(RunCommand(String.Format("fstat \"{0}\"", Filter), out FileRecords, CommandOptions.IgnoreNoSuchFilesError | CommandOptions.IgnoreFilesNotInClientViewError | CommandOptions.IgnoreProtectedNamespaceError, Log))
 			{
-				bExists = (FileRecords.Count > 0);
+				bExists = (FileRecords.Exists(x => x.Action == null || !x.Action.Contains("delete")));
 				return true;
 			}
 			else
