Fix missing copyright notice.

#rnx

diff --git a/Engine/Source/Programs/UnrealGameSync/UnrealGameSyncAutomation/Properties/AssemblyInfo.cs b/Engine/Source/Programs/UnrealGameSync/UnrealGameSyncAutomation/Properties/AssemblyInfo.cs
--- a/Engine/Source/Programs/UnrealGameSync/UnrealGameSyncAutomation/Properties/AssemblyInfo.cs
+++ b/Engine/Source/Programs/UnrealGameSync/UnrealGameSyncAutomation/Properties/AssemblyInfo.cs
@@ -12,7 +12,7 @@ using System.Runtime.InteropServices;
 [assembly: AssemblyConfiguration("")]
 [assembly: AssemblyCompany("")]
 [assembly: AssemblyProduct("UnrealGameSyncAutomation")]
-[assembly: AssemblyCopyright("Copyright ©  2018")]
+[assembly: AssemblyCopyright("Copyright 1998-2019 Epic Games, Inc. All Rights Reserved.")]
 [assembly: AssemblyTrademark("")]
 [assembly: AssemblyCulture("")]
 
