UGS: Fix deleted files being enumerated when searching for files in the depot.

#rnx

diff --git a/Engine/Source/Programs/UnrealGameSync/UnrealGameSync/DetectProjectSettingsTask.cs b/Engine/Source/Programs/UnrealGameSync/UnrealGameSync/DetectProjectSettingsTask.cs
--- a/Engine/Source/Programs/UnrealGameSync/UnrealGameSync/DetectProjectSettingsTask.cs
+++ b/Engine/Source/Programs/UnrealGameSync/UnrealGameSync/DetectProjectSettingsTask.cs
@@ -235,13 +235,9 @@ namespace UnrealGameSync
 				List<PerforceFileRecord> Files;
 				if(PerforceClient.FindFiles(PerforceUtils.GetClientOrDepotDirectoryName(NewSelectedClientFileName) + "/Source/*Editor.Target.cs", out Files, Log) && Files.Count >= 1)
 				{
-					PerforceFileRecord File = Files.FirstOrDefault(x => x.Action == null || !x.Action.Contains("delete"));
-					if(File != null)
-					{
-						string DepotPath = File.DepotPath;
-						NewProjectEditorTarget = Path.GetFileNameWithoutExtension(Path.GetFileNameWithoutExtension(DepotPath.Substring(DepotPath.LastIndexOf('/') + 1)));
-						Log.WriteLine("Using {0} as editor target name (from {1})", NewProjectEditorTarget, Files[0]);
-					}
+					string DepotPath = Files[0].DepotPath;
+					NewProjectEditorTarget = Path.GetFileNameWithoutExtension(Path.GetFileNameWithoutExtension(DepotPath.Substring(DepotPath.LastIndexOf('/') + 1)));
+					Log.WriteLine("Using {0} as editor target name (from {1})", NewProjectEditorTarget, Files[0]);
 				}
 				if (NewProjectEditorTarget == null)
 				{
diff --git a/Engine/Source/Programs/UnrealGameSync/UnrealGameSync/Perforce.cs b/Engine/Source/Programs/UnrealGameSync/UnrealGameSync/Perforce.cs
--- a/Engine/Source/Programs/UnrealGameSync/UnrealGameSync/Perforce.cs
+++ b/Engine/Source/Programs/UnrealGameSync/UnrealGameSync/Perforce.cs
@@ -689,9 +689,14 @@ namespace UnrealGameSync
 
 		public bool FindFiles(string Filter, out List<PerforceFileRecord> FileRecords, TextWriter Log)
 		{
-			return RunCommand(String.Format("fstat \"{0}\"", Filter), out FileRecords, CommandOptions.None, Log);
+			bool bResult = RunCommand(String.Format("fstat \"{0}\"", Filter), out FileRecords, CommandOptions.None, Log);
+			if(bResult)
+			{
+                 FileRecords.RemoveAll(x => x.Action != null && x.Action.Contains("delete"));
+			}
+			return bResult;
 		}
-		
+
 		public bool Print(string DepotPath, out List<string> Lines, TextWriter Log)
 		{
 			string TempFileName = Path.GetTempFileName();
