UGS: Add an option to the Sync Filter dialog controlling whether to generate a solution containing all projects, or just the one being synced.

diff --git a/Engine/Source/Programs/UnrealGameSync/UnrealGameSync/Controls/WorkspaceControl.cs b/Engine/Source/Programs/UnrealGameSync/UnrealGameSync/Controls/WorkspaceControl.cs
--- a/Engine/Source/Programs/UnrealGameSync/UnrealGameSync/Controls/WorkspaceControl.cs
+++ b/Engine/Source/Programs/UnrealGameSync/UnrealGameSync/Controls/WorkspaceControl.cs
@@ -761,14 +761,6 @@ namespace UnrealGameSync
 		void StartSync(int ChangeNumber, WorkspaceUpdateCallback Callback)
 		{
 			WorkspaceUpdateOptions Options = WorkspaceUpdateOptions.Sync | WorkspaceUpdateOptions.SyncArchives | WorkspaceUpdateOptions.GenerateProjectFiles;
-			if(Settings.bAutoResolveConflicts)
-			{
-				Options |= WorkspaceUpdateOptions.AutoResolveChanges;
-			}
-			if(Settings.bUseIncrementalBuilds)
-			{
-				Options |= WorkspaceUpdateOptions.UseIncrementalBuilds;
-			}
 			if(Settings.bBuildAfterSync)
 			{
 				Options |= WorkspaceUpdateOptions.Build;
@@ -781,10 +773,6 @@ namespace UnrealGameSync
 			{
 				Options |= WorkspaceUpdateOptions.OpenSolutionAfterSync;
 			}
-			if(WorkspaceSettings.bSyncAllProjects ?? Settings.bSyncAllProjects)
-			{
-				Options |= WorkspaceUpdateOptions.SyncAllProjects;
-			}
 			StartWorkspaceUpdate(ChangeNumber, Options, Callback);
 		}
 
@@ -835,6 +823,23 @@ namespace UnrealGameSync
 
 		void StartWorkspaceUpdate(WorkspaceUpdateContext Context, WorkspaceUpdateCallback Callback)
 		{
+			if(Settings.bAutoResolveConflicts)
+			{
+				Context.Options |= WorkspaceUpdateOptions.AutoResolveChanges;
+			}
+			if(Settings.bUseIncrementalBuilds)
+			{
+				Context.Options |= WorkspaceUpdateOptions.UseIncrementalBuilds;
+			}
+			if(WorkspaceSettings.bSyncAllProjects ?? Settings.bSyncAllProjects)
+			{
+				Context.Options |= WorkspaceUpdateOptions.SyncAllProjects | WorkspaceUpdateOptions.IncludeAllProjectsInSolution;
+			}
+			if(WorkspaceSettings.bIncludeAllProjectsInSolution ?? Settings.bIncludeAllProjectsInSolution)
+			{
+				Context.Options |= WorkspaceUpdateOptions.IncludeAllProjectsInSolution;
+			}
+
 			UpdateCallback = Callback;
 
 			Context.StartTime = DateTime.UtcNow;
@@ -4254,16 +4259,18 @@ namespace UnrealGameSync
 
 		private void OptionsContextMenu_SyncFilter_Click(object sender, EventArgs e)
 		{
-			SyncFilter Filter = new SyncFilter(Workspace.GetSyncCategories(), Settings.SyncView, Settings.SyncExcludedCategories, Settings.bSyncAllProjects, WorkspaceSettings.SyncView, WorkspaceSettings.SyncIncludedCategories, WorkspaceSettings.SyncExcludedCategories, WorkspaceSettings.bSyncAllProjects);
+			SyncFilter Filter = new SyncFilter(Workspace.GetSyncCategories(), Settings.SyncView, Settings.SyncExcludedCategories, Settings.bSyncAllProjects, Settings.bIncludeAllProjectsInSolution, WorkspaceSettings.SyncView, WorkspaceSettings.SyncIncludedCategories, WorkspaceSettings.SyncExcludedCategories, WorkspaceSettings.bSyncAllProjects, WorkspaceSettings.bIncludeAllProjectsInSolution);
 			if(Filter.ShowDialog() == DialogResult.OK)
 			{
 				Settings.SyncExcludedCategories = Filter.GlobalExcludedCategories;
 				Settings.SyncView = Filter.GlobalView;
 				Settings.bSyncAllProjects = Filter.bGlobalSyncAllProjects;
+				Settings.bIncludeAllProjectsInSolution = Filter.bGlobalIncludeAllProjectsInSolution;
 				WorkspaceSettings.SyncIncludedCategories = Filter.WorkspaceIncludedCategories;
 				WorkspaceSettings.SyncExcludedCategories = Filter.WorkspaceExcludedCategories;
 				WorkspaceSettings.SyncView = Filter.WorkspaceView;
 				WorkspaceSettings.bSyncAllProjects = Filter.bWorkspaceSyncAllProjects;
+				WorkspaceSettings.bIncludeAllProjectsInSolution = Filter.bWorkspaceIncludeAllProjectsInSolution;
                 Settings.Save();
 			}
 		}
diff --git a/Engine/Source/Programs/UnrealGameSync/UnrealGameSync/Forms/SyncFilter.cs b/Engine/Source/Programs/UnrealGameSync/UnrealGameSync/Forms/SyncFilter.cs
--- a/Engine/Source/Programs/UnrealGameSync/UnrealGameSync/Forms/SyncFilter.cs
+++ b/Engine/Source/Programs/UnrealGameSync/UnrealGameSync/Forms/SyncFilter.cs
@@ -1,4 +1,4 @@
-﻿// Copyright 1998-2019 Epic Games, Inc. All Rights Reserved.
+// Copyright 1998-2019 Epic Games, Inc. All Rights Reserved.
 
 using System;
 using System.Collections.Generic;
@@ -18,12 +18,14 @@ namespace UnrealGameSync
 		public string[] GlobalView;
 		public Guid[] GlobalExcludedCategories;
 		public bool bGlobalSyncAllProjects;
+		public bool bGlobalIncludeAllProjectsInSolution;
 		public string[] WorkspaceView;
 		public Guid[] WorkspaceIncludedCategories;
 		public Guid[] WorkspaceExcludedCategories;
 		public bool? bWorkspaceSyncAllProjects;
+		public bool? bWorkspaceIncludeAllProjectsInSolution;
 
-		public SyncFilter(Dictionary<Guid, WorkspaceSyncCategory> InUniqueIdToCategory, string[] InGlobalView, Guid[] InGlobalExcludedCategories, bool bInGlobalProjectOnly, string[] InWorkspaceView, Guid[] InWorkspaceIncludedCategories, Guid[] InWorkspaceExcludedCategories, bool? bInWorkspaceProjectOnly)
+		public SyncFilter(Dictionary<Guid, WorkspaceSyncCategory> InUniqueIdToCategory, string[] InGlobalView, Guid[] InGlobalExcludedCategories, bool bInGlobalProjectOnly, bool bInGlobalIncludeAllProjectsInSolution, string[] InWorkspaceView, Guid[] InWorkspaceIncludedCategories, Guid[] InWorkspaceExcludedCategories, bool? bInWorkspaceProjectOnly, bool? bInWorkspaceIncludeAllProjectsInSolution)
 		{
 			InitializeComponent();
 
@@ -31,21 +33,26 @@ namespace UnrealGameSync
 			GlobalExcludedCategories = InGlobalExcludedCategories;
 			GlobalView = InGlobalView;
 			bGlobalSyncAllProjects = bInGlobalProjectOnly;
+			bGlobalIncludeAllProjectsInSolution = bInGlobalIncludeAllProjectsInSolution;
 			WorkspaceIncludedCategories = InWorkspaceIncludedCategories;
 			WorkspaceExcludedCategories = InWorkspaceExcludedCategories;
 			WorkspaceView = InWorkspaceView;
-			bWorkspaceSyncAllProjects = bInWorkspaceProjectOnly ?? bGlobalSyncAllProjects;
+			bWorkspaceSyncAllProjects = bInWorkspaceProjectOnly;
+			bWorkspaceIncludeAllProjectsInSolution = bInWorkspaceIncludeAllProjectsInSolution;
 
 			GlobalControl.SetView(GlobalView);
 			SetExcludedCategories(GlobalControl.CategoriesCheckList, UniqueIdToCategory, GlobalExcludedCategories);
 			GlobalControl.SyncAllProjects.Checked = bGlobalSyncAllProjects;
+			GlobalControl.IncludeAllProjectsInSolution.Checked = bGlobalIncludeAllProjectsInSolution;
 
 			WorkspaceControl.SetView(WorkspaceView);
 			SetExcludedCategories(WorkspaceControl.CategoriesCheckList, UniqueIdToCategory, UserSettings.GetEffectiveExcludedCategories(GlobalExcludedCategories, WorkspaceIncludedCategories, WorkspaceExcludedCategories));
 			WorkspaceControl.SyncAllProjects.Checked = bWorkspaceSyncAllProjects ?? bGlobalSyncAllProjects;
+			WorkspaceControl.IncludeAllProjectsInSolution.Checked = bWorkspaceIncludeAllProjectsInSolution ?? bGlobalIncludeAllProjectsInSolution;
 
 			GlobalControl.CategoriesCheckList.ItemCheck += GlobalControl_CategoriesCheckList_ItemCheck;
 			GlobalControl.SyncAllProjects.CheckStateChanged += GlobalControl_SyncAllProjects_CheckStateChanged;
+			GlobalControl.IncludeAllProjectsInSolution.CheckStateChanged += GlobalControl_IncludeAllProjectsInSolution_CheckStateChanged;
 		}
 
 		private void GlobalControl_CategoriesCheckList_ItemCheck(object sender, ItemCheckEventArgs e)
@@ -58,6 +65,11 @@ namespace UnrealGameSync
 			WorkspaceControl.SyncAllProjects.Checked = GlobalControl.SyncAllProjects.Checked;
 		}
 
+		private void GlobalControl_IncludeAllProjectsInSolution_CheckStateChanged(object sender, EventArgs e)
+		{
+			WorkspaceControl.IncludeAllProjectsInSolution.Checked = GlobalControl.IncludeAllProjectsInSolution.Checked;
+		}
+
 		private static void SetExcludedCategories(CheckedListBox ListBox, Dictionary<Guid, WorkspaceSyncCategory> UniqueIdToFilter, Guid[] ExcludedCategories)
 		{
 			ListBox.Items.Clear();
@@ -124,9 +136,11 @@ namespace UnrealGameSync
 		
 			GlobalView = NewGlobalView;
 			bGlobalSyncAllProjects = GlobalControl.SyncAllProjects.Checked;
+			bGlobalIncludeAllProjectsInSolution = GlobalControl.IncludeAllProjectsInSolution.Checked;
 
 			WorkspaceView = NewWorkspaceView;
 			bWorkspaceSyncAllProjects = (WorkspaceControl.SyncAllProjects.Checked == bGlobalSyncAllProjects)? (bool?)null : WorkspaceControl.SyncAllProjects.Checked;
+			bWorkspaceIncludeAllProjectsInSolution = (WorkspaceControl.IncludeAllProjectsInSolution.Checked == bGlobalIncludeAllProjectsInSolution)? (bool?)null : WorkspaceControl.IncludeAllProjectsInSolution.Checked;
 
 			GetExcludedCategories(out GlobalExcludedCategories, out WorkspaceIncludedCategories, out WorkspaceExcludedCategories);
 
diff --git a/Engine/Source/Programs/UnrealGameSync/UnrealGameSync/UserSettings.cs b/Engine/Source/Programs/UnrealGameSync/UnrealGameSync/UserSettings.cs
--- a/Engine/Source/Programs/UnrealGameSync/UnrealGameSync/UserSettings.cs
+++ b/Engine/Source/Programs/UnrealGameSync/UnrealGameSync/UserSettings.cs
@@ -174,6 +174,7 @@ namespace UnrealGameSync
 		public Guid[] SyncIncludedCategories;
 		public Guid[] SyncExcludedCategories;
 		public bool? bSyncAllProjects;
+		public bool? bIncludeAllProjectsInSolution;
 	}
 
 	class UserProjectSettings
@@ -207,6 +208,7 @@ namespace UnrealGameSync
 		public string[] SyncView;
 		public Guid[] SyncExcludedCategories;
 		public bool bSyncAllProjects;
+		public bool bIncludeAllProjectsInSolution;
 		public LatestChangeType SyncType;
 		public BuildConfig CompiledEditorBuildConfig; // NB: This assumes not using precompiled editor. See CurrentBuildConfig.
 		public TabLabels TabLabels;
@@ -310,6 +312,7 @@ namespace UnrealGameSync
 			SyncView = ConfigFile.GetValues("General.SyncFilter", new string[0]);
 			SyncExcludedCategories = ConfigFile.GetGuidValues("General.SyncExcludedCategories", new Guid[0]);
 			bSyncAllProjects = ConfigFile.GetValue("General.SyncAllProjects", false);
+			bIncludeAllProjectsInSolution = ConfigFile.GetValue("General.IncludeAllProjectsInSolution", false);
 			if(!Enum.TryParse(ConfigFile.GetValue("General.SyncType", ""), out SyncType))
 			{
 				SyncType = LatestChangeType.Good;
@@ -477,6 +480,7 @@ namespace UnrealGameSync
 					CurrentWorkspace.SyncIncludedCategories = new Guid[0];
 					CurrentWorkspace.SyncExcludedCategories = new Guid[0];
 					CurrentWorkspace.bSyncAllProjects = null;
+					CurrentWorkspace.bIncludeAllProjectsInSolution = null;
 				}
 				else
 				{
@@ -512,6 +516,9 @@ namespace UnrealGameSync
 					int SyncAllProjects = WorkspaceSection.GetValue("SyncAllProjects", -1);
 					CurrentWorkspace.bSyncAllProjects = (SyncAllProjects == 0)? (bool?)false : (SyncAllProjects == 1)? (bool?)true : (bool?)null;
 
+					int IncludeAllProjectsInSolution = WorkspaceSection.GetValue("IncludeAllProjectsInSolution", -1);
+					CurrentWorkspace.bIncludeAllProjectsInSolution = (IncludeAllProjectsInSolution == 0)? (bool?)false : (IncludeAllProjectsInSolution == 1)? (bool?)true : (bool?)null;
+
 					string[] BisectEntries = WorkspaceSection.GetValues("Bisect", new string[0]);
 					foreach(string BisectEntry in BisectEntries)
 					{
@@ -578,6 +585,7 @@ namespace UnrealGameSync
 			GeneralSection.SetValues("SyncFilter", SyncView);
 			GeneralSection.SetValues("SyncExcludedCategories", SyncExcludedCategories);
 			GeneralSection.SetValue("SyncAllProjects", bSyncAllProjects);
+			GeneralSection.SetValue("IncludeAllProjectsInSolution", bIncludeAllProjectsInSolution);
 			GeneralSection.SetValue("SyncType", SyncType.ToString());
 
 			// Build configuration
@@ -649,6 +657,10 @@ namespace UnrealGameSync
 				{
 					WorkspaceSection.SetValue("SyncAllProjects", CurrentWorkspace.bSyncAllProjects.Value);
 				}
+				if(CurrentWorkspace.bIncludeAllProjectsInSolution.HasValue)
+				{
+					WorkspaceSection.SetValue("IncludeAllProjectsInSolution", CurrentWorkspace.bIncludeAllProjectsInSolution.Value);
+				}
 
 				List<ConfigObject> BisectEntryObjects = new List<ConfigObject>();
 				foreach(KeyValuePair<int, BisectState> BisectPair in CurrentWorkspace.ChangeNumberToBisectState)
diff --git a/Engine/Source/Programs/UnrealGameSync/UnrealGameSync/Workspace.cs b/Engine/Source/Programs/UnrealGameSync/UnrealGameSync/Workspace.cs
--- a/Engine/Source/Programs/UnrealGameSync/UnrealGameSync/Workspace.cs
+++ b/Engine/Source/Programs/UnrealGameSync/UnrealGameSync/Workspace.cs
@@ -30,6 +30,7 @@ namespace UnrealGameSync
 		ContentOnly = 0x400,
 		UpdateFilter = 0x800,
 		SyncAllProjects = 0x1000,
+		IncludeAllProjectsInSolution = 0x2000,
 	}
 
 	enum WorkspaceUpdateResult
@@ -814,18 +815,21 @@ namespace UnrealGameSync
 				{
 					Progress.Set("Generating project files...", 0.0f);
 
-					string ProjectFileArgument = "";
-					if(SelectedLocalFileName.EndsWith(".uproject", StringComparison.InvariantCultureIgnoreCase) && (Context.Options & WorkspaceUpdateOptions.SyncAllProjects) == 0)
+					StringBuilder CommandLine = new StringBuilder();
+					CommandLine.AppendFormat("/C \"\"{0}\"", Path.Combine(LocalRootPath, "GenerateProjectFiles.bat"));
+					if((Context.Options & WorkspaceUpdateOptions.SyncAllProjects) == 0 && (Context.Options & WorkspaceUpdateOptions.IncludeAllProjectsInSolution) == 0)
 					{
-						ProjectFileArgument = String.Format("\"{0}\" ", SelectedLocalFileName);
+						if(SelectedLocalFileName.EndsWith(".uproject", StringComparison.InvariantCultureIgnoreCase))
+						{
+							CommandLine.AppendFormat(" \"{0}\"", SelectedLocalFileName);
+						}
 					}
-
-					string CommandLine = String.Format("/C \"\"{0}\" {1}-progress\"", Path.Combine(LocalRootPath, "GenerateProjectFiles.bat"), ProjectFileArgument);
+					CommandLine.Append(" -progress\"");
 
 					Log.WriteLine("Generating project files...");
 					Log.WriteLine("gpf> Running {0} {1}", CmdExe, CommandLine);
 
-					int GenerateProjectFilesResult = Utility.ExecuteProcess(CmdExe, null, CommandLine, null, new ProgressTextWriter(Progress, new PrefixedTextWriter("gpf> ", Log)));
+					int GenerateProjectFilesResult = Utility.ExecuteProcess(CmdExe, null, CommandLine.ToString(), null, new ProgressTextWriter(Progress, new PrefixedTextWriter("gpf> ", Log)));
 					if(GenerateProjectFilesResult != 0)
 					{
 						StatusMessage = String.Format("Failed to generate project files (exit code {0}).", GenerateProjectFilesResult);
