UGS: Prevent files which are open for add from showing up in the "Clean Workspace" list.

diff --git a/Engine/Source/Programs/UnrealGameSync/UnrealGameSync/FindFoldersToCleanTask.cs b/Engine/Source/Programs/UnrealGameSync/UnrealGameSync/FindFoldersToCleanTask.cs
--- a/Engine/Source/Programs/UnrealGameSync/UnrealGameSync/FindFoldersToCleanTask.cs
+++ b/Engine/Source/Programs/UnrealGameSync/UnrealGameSync/FindFoldersToCleanTask.cs
@@ -152,7 +152,13 @@ namespace UnrealGameSync
 				}
 
 				// Find all the files that should be deleted
-				LocalFolder.FilesToDelete.AddRange(LocalFolder.NameToFile.Values.Where(x => !PerforceFolder.NameToFile.ContainsKey(x.Name)));
+				foreach(FileInfo LocalFileInfo in LocalFolder.NameToFile.Values)
+				{
+					if(!PerforceFolder.NameToFile.ContainsKey(LocalFileInfo.Name) && !OpenClientPaths.Contains(LocalFileInfo.FullName))
+					{
+						LocalFolder.FilesToDelete.Add(LocalFileInfo);
+					}
+				}
 			}
 
 			// Figure out if this folder is just an empty directory that needs to be removed
