UGS: Require that the user manually clicks on a link to fetch more changes from the server. Doing so automatically is not a bounded operation, since not all changes may match the chosen filter. Searching for a previously synced change with older tabs can cause massive slowdowns and take a huge amount of memory.

diff --git a/Engine/Source/Programs/UnrealGameSync/UnrealGameSync/Controls/WorkspaceControl.cs b/Engine/Source/Programs/UnrealGameSync/UnrealGameSync/Controls/WorkspaceControl.cs
--- a/Engine/Source/Programs/UnrealGameSync/UnrealGameSync/Controls/WorkspaceControl.cs
+++ b/Engine/Source/Programs/UnrealGameSync/UnrealGameSync/Controls/WorkspaceControl.cs
@@ -136,6 +136,8 @@ namespace UnrealGameSync
 		[DllImport("uxtheme.dll", CharSet = CharSet.Unicode)]
 		static extern int SetWindowTheme(IntPtr hWnd, string pszSubAppName, string pszSubIdList);
 
+		const int BuildListExpandCount = 250;
+
 		const string EditorArchiveType = "Editor";
 
 		IWorkspaceControlOwner Owner;
@@ -201,6 +203,8 @@ namespace UnrealGameSync
 
 		Dictionary<string, int> NotifiedBuildTypeToChangeNumber = new Dictionary<string,int>();
 
+		bool bMouseOverExpandLink;
+
 		string HoverBadgeUniqueId = null;
 		bool bHoverSync;
 		PerforceChangeSummary ContextMenuChange;
@@ -229,6 +233,7 @@ namespace UnrealGameSync
 		string LastColumnSettings;
 		List<ColumnHeader> CustomColumns;
 		int MaxBuildBadgeChars;
+		ListViewItem ExpandItem;
 
 		bool bUpdateBuildListPosted;
 		bool bUpdateBuildMetadataPosted;
@@ -332,11 +337,6 @@ namespace UnrealGameSync
 			UpdateStatusPanel();
 			UpdateServiceBadges();
 
-			if(CurrentChangeNumber != -1)
-			{
-				SelectChange(CurrentChangeNumber);
-			}
-
 			PerforceMonitor.Start();
 			EventMonitor.Start();
 
@@ -705,51 +705,43 @@ namespace UnrealGameSync
 		private void BuildList_OnScroll()
 		{
 			PendingSelectedChangeNumber = -1;
-			UpdateNumRequestedBuilds(false);
 		}
 
-		void UpdateNumRequestedBuilds(bool bAllowShrink)
+		void ShrinkNumRequestedBuilds()
 		{
-			if(PerforceMonitor != null)
+			if(PerforceMonitor != null && BuildList.Items.Count > 0 && PendingSelectedChangeNumber == -1)
 			{
-				if(!Settings.bShowUnreviewedChanges)
+				// Find the number of visible items using a (slightly wasteful) binary search
+				int VisibleItemCount = 1;
+				for(int StepSize = BuildList.Items.Count / 2; StepSize >= 1; )
 				{
-					PerforceMonitor.PendingMaxChanges = 1000;
-				}
-				else
-				{
-					int NumItemsPerPage = Math.Max(BuildList.GetVisibleItemsPerPage(), 10);
-
-					// Find the number of visible items using a (slightly wasteful) binary search
-					int VisibleItemCount = 1;
-					for(int StepSize = BuildList.Items.Count / 2; StepSize >= 1; )
-					{
-						int TestIndex = VisibleItemCount + StepSize;
-						if(TestIndex < BuildList.Items.Count && BuildList.GetItemRect(TestIndex).Top < BuildList.Height)
-						{
-							VisibleItemCount += StepSize;
-						}
-						else
-						{
-							StepSize /= 2;
-						}
-					}
-
-					// Increase or decrease the number of builds we want, with a bit of rubber-banding
-					const int IncreaseStep = 50;
-					if(VisibleItemCount > BuildList.Items.Count - 20)
+					int TestIndex = VisibleItemCount + StepSize;
+					if(TestIndex < BuildList.Items.Count && BuildList.GetItemRect(TestIndex).Top < BuildList.Height)
 					{
-						PerforceMonitor.PendingMaxChanges = NumChanges + IncreaseStep;
+						VisibleItemCount += StepSize;
 					}
-					else if(bAllowShrink)
+					else
 					{
-						int NewNumChanges = ListIndexToChangeIndex[VisibleItemCount - 1] + IncreaseStep;
-						if(NewNumChanges < NumChanges)
-						{
-							PerforceMonitor.PendingMaxChanges = NewNumChanges;
-						}
+						StepSize /= 2;
 					}
 				}
+
+				// Figure out the last index to ensure is visible
+				int LastVisibleIndex = VisibleItemCount;
+				if(LastVisibleIndex >= ListIndexToChangeIndex.Count)
+				{
+					LastVisibleIndex = ListIndexToChangeIndex.Count - 1;
+				}
+
+				// Get the max number of changes to ensure this
+				int NewPendingMaxChanges = ListIndexToChangeIndex[LastVisibleIndex];
+				NewPendingMaxChanges = PerforceMonitor.InitialMaxChangesValue + ((Math.Max(NewPendingMaxChanges - PerforceMonitor.InitialMaxChangesValue, 0) + BuildListExpandCount - 1) / BuildListExpandCount) * BuildListExpandCount;
+
+				// Shrink the number of changes retained by the PerforceMonitor class
+				if(PerforceMonitor.PendingMaxChanges > NewPendingMaxChanges)
+				{
+					PerforceMonitor.PendingMaxChanges = NewPendingMaxChanges;
+				}
 			}
 		}
 
@@ -1033,11 +1025,11 @@ namespace UnrealGameSync
 					}
 				}
 
-				int SelectedChange = (BuildList.SelectedItems.Count > 0)? ((PerforceChangeSummary)BuildList.SelectedItems[0].Tag).Number : -1;
-
 				ChangeNumberToArchivePath.Clear();
 				ChangeNumberToLayoutInfo.Clear();
 
+				ExpandItem = null;
+
 				BuildList.BeginUpdate();
 
 				foreach(ListViewGroup Group in BuildList.Groups)
@@ -1148,6 +1140,19 @@ namespace UnrealGameSync
 					BuildList.Groups.RemoveAt(0);
 				}
 
+				if(Changes.Count > 0)
+				{
+					ExpandItem = new ListViewItem((BuildList.Groups.Count > 0)? BuildList.Groups[BuildList.Groups.Count - 1] : null);
+					ExpandItem.Tag = null;
+					ExpandItem.Selected = false;
+					ExpandItem.Text = "";
+					for(int ColumnIdx = 1; ColumnIdx < BuildList.Columns.Count; ColumnIdx++)
+					{
+						ExpandItem.SubItems.Add(new ListViewItem.ListViewSubItem(ExpandItem, ""));
+					}
+					BuildList.Items.Add(ExpandItem);
+				}
+
 				BuildList.EndUpdate();
 
 				if(PendingSelectedChangeNumber != -1)
@@ -1253,15 +1258,18 @@ namespace UnrealGameSync
 			// Add a dummy group for any other badges we have
 			foreach(ListViewItem Item in BuildList.Items)
 			{
-				EventSummary Summary = EventMonitor.GetSummaryForChange(((PerforceChangeSummary)Item.Tag).Number);
-				if(Summary != null)
+				if(Item.Tag != null)
 				{
-					foreach(BadgeData Badge in Summary.Badges)
+					EventSummary Summary = EventMonitor.GetSummaryForChange(((PerforceChangeSummary)Item.Tag).Number);
+					if(Summary != null)
 					{
-						string BadgeSlot = Badge.BadgeName;
-						if(!BadgeNameToGroup.ContainsKey(BadgeSlot))
+						foreach(BadgeData Badge in Summary.Badges)
 						{
-							BadgeNameToGroup.Add(BadgeSlot, "XXXX");
+							string BadgeSlot = Badge.BadgeName;
+							if(!BadgeNameToGroup.ContainsKey(BadgeSlot))
+							{
+								BadgeNameToGroup.Add(BadgeSlot, "XXXX");
+							}
 						}
 					}
 				}
@@ -1299,7 +1307,7 @@ namespace UnrealGameSync
 			CheckForStartupComplete();
 
 			// If we are filtering by badges, we may also need to update the build list
-			if(ProjectSettings.FilterBadges.Count > 0)
+			if(ProjectSettings.FilterType != FilterType.None || ProjectSettings.FilterBadges.Count > 0)
 			{
 				UpdateBuildList();
 			}
@@ -1515,23 +1523,106 @@ namespace UnrealGameSync
 			e.DrawDefault = true;
 		}
 
-		private void BuildList_DrawItem(object sender, DrawListViewItemEventArgs e)
+		class ExpandRowLayout
 		{
-			if(e.State.HasFlag(ListViewItemStates.Selected))
+			public string MainText;
+			public Rectangle MainRect;
+			public string LinkText;
+			public Rectangle LinkRect;
+		}
+
+		private ExpandRowLayout LayoutExpandRow(Graphics Graphics, Rectangle Bounds)
+		{
+			ExpandRowLayout Layout = new ExpandRowLayout();
+
+			int CurrentMaxChanges = PerforceMonitor.CurrentMaxChanges;
+
+			string ShowingChanges;
+			if(ListIndexToChangeIndex.Count == CurrentMaxChanges)
 			{
-				BuildList.DrawSelectedBackground(e.Graphics, e.Bounds);
+				ShowingChanges = String.Format("Showing {0} changes.", ListIndexToChangeIndex.Count);
 			}
-			else if(e.ItemIndex == BuildList.HoverItem)
+			else
 			{
-				BuildList.DrawTrackedBackground(e.Graphics, e.Bounds);
+				ShowingChanges = String.Format("Showing {0}/{1} changes.", ListIndexToChangeIndex.Count, CurrentMaxChanges);
 			}
-			else if(((PerforceChangeSummary)e.Item.Tag).Number == Workspace.PendingChangeNumber)
+
+			if(PerforceMonitor.PendingMaxChanges > CurrentMaxChanges)
 			{
-				BuildList.DrawTrackedBackground(e.Graphics, e.Bounds);
+				Layout.MainText = String.Format("{0}. Fetching {1} more from server...  ", ShowingChanges, PerforceMonitor.PendingMaxChanges - CurrentMaxChanges);
+				Layout.LinkText = "Cancel";
 			}
 			else
+			{
+				Layout.MainText = String.Format("{0}  ", ShowingChanges);
+				Layout.LinkText = "Show more...";
+			}
+
+			Size MainTextSize = TextRenderer.MeasureText(Graphics, Layout.MainText, BuildFont, new Size(1000, 1000), TextFormatFlags.NoPadding);
+			Size LinkTextSize = TextRenderer.MeasureText(Graphics, Layout.LinkText, BuildFont, new Size(1000, 1000), TextFormatFlags.NoPadding);
+
+			int MinX = Bounds.Left + (Bounds.Width - MainTextSize.Width - LinkTextSize.Width) / 2;
+			int MinY = Bounds.Bottom - MainTextSize.Height - 1;
+
+			Layout.MainRect = new Rectangle(new Point(MinX, MinY), MainTextSize);
+			Layout.LinkRect = new Rectangle(new Point(MinX + MainTextSize.Width, MinY), LinkTextSize);
+
+			return Layout;
+		}
+
+		private void DrawExpandRow(Graphics Graphics, ExpandRowLayout Layout)
+		{
+			TextRenderer.DrawText(Graphics, Layout.MainText, BuildFont, Layout.MainRect, SystemColors.WindowText, TextFormatFlags.SingleLine | TextFormatFlags.VerticalCenter | TextFormatFlags.NoPrefix | TextFormatFlags.NoPadding);
+
+			Color LinkColor = SystemColors.HotTrack;
+			if(bMouseOverExpandLink)
+			{
+				LinkColor = Color.FromArgb(LinkColor.B, LinkColor.G, LinkColor.R);
+			}
+
+			TextRenderer.DrawText(Graphics, Layout.LinkText, BuildFont, Layout.LinkRect, LinkColor, TextFormatFlags.SingleLine | TextFormatFlags.VerticalCenter | TextFormatFlags.NoPrefix | TextFormatFlags.NoPadding);
+		}
+
+		private bool HitTestExpandLink(Point Location)
+		{
+			if(ExpandItem == null)
+			{
+				return false;
+			}
+			using(Graphics Graphics = Graphics.FromHwnd(IntPtr.Zero))
+			{
+				ExpandRowLayout Layout = LayoutExpandRow(Graphics, ExpandItem.Bounds);
+				return Layout.LinkRect.Contains(Location);
+			}
+		}
+
+		private void BuildList_DrawItem(object sender, DrawListViewItemEventArgs e)
+		{
+			if(e.Item == ExpandItem)
 			{
 				BuildList.DrawDefaultBackground(e.Graphics, e.Bounds);
+
+				ExpandRowLayout ExpandLayout = LayoutExpandRow(e.Graphics, e.Bounds);
+				DrawExpandRow(e.Graphics, ExpandLayout);
+			}
+			else
+			{
+				if(e.State.HasFlag(ListViewItemStates.Selected))
+				{
+					BuildList.DrawSelectedBackground(e.Graphics, e.Bounds);
+				}
+				else if(e.ItemIndex == BuildList.HoverItem)
+				{
+					BuildList.DrawTrackedBackground(e.Graphics, e.Bounds);
+				}
+				else if(((PerforceChangeSummary)e.Item.Tag).Number == Workspace.PendingChangeNumber)
+				{
+					BuildList.DrawTrackedBackground(e.Graphics, e.Bounds);
+				}
+				else
+				{
+					BuildList.DrawDefaultBackground(e.Graphics, e.Bounds);
+				}
 			}
 		}
 
@@ -1615,6 +1706,10 @@ namespace UnrealGameSync
 		private void BuildList_DrawSubItem(object sender, DrawListViewSubItemEventArgs e)
 		{
 			PerforceChangeSummary Change = (PerforceChangeSummary)e.Item.Tag;
+			if(Change == null)
+			{
+				return;
+			}
 
 			float DpiScaleX = e.Graphics.DpiX / 96.0f;
 			float DpiScaleY = e.Graphics.DpiY / 96.0f;
@@ -1632,6 +1727,14 @@ namespace UnrealGameSync
 			int BadgeAlpha = bAllowSync? 255 : 128;
 			Color TextColor = (bAllowSync || Change.Number == Workspace.PendingChangeNumber || Change.Number == Workspace.CurrentChangeNumber || (WorkspaceSettings != null && WorkspaceSettings.AdditionalChangeNumbers.Contains(Change.Number)))? SystemColors.WindowText : Blend(SystemColors.Window, SystemColors.WindowText, 0.25f);
 
+			const int FadeRange = 6;
+			if(e.ItemIndex >= BuildList.Items.Count - FadeRange)
+			{
+				float Opacity = (float)(BuildList.Items.Count - e.ItemIndex - 0.9f) / FadeRange;
+				BadgeAlpha = (int)(BadgeAlpha * Opacity);
+				TextColor = Blend(SystemColors.Window, TextColor, Opacity);
+			}
+
 			if(e.ColumnIndex == IconColumn.Index)
 			{
 				EventSummary Summary = EventMonitor.GetSummaryForChange(Change.Number);
@@ -2381,13 +2484,16 @@ namespace UnrealGameSync
 				if(HitTest.Item != null)
 				{
 					PerforceChangeSummary Change = (PerforceChangeSummary)HitTest.Item.Tag;
-					if(Change.Number == Workspace.CurrentChangeNumber)
-					{
-						LaunchEditor();
-					}
-					else
+					if(Change != null)
 					{
-						StartSync(Change.Number, null);
+						if(Change.Number == Workspace.CurrentChangeNumber)
+						{
+							LaunchEditor();
+						}
+						else
+						{
+							StartSync(Change.Number, null);
+						}
 					}
 				}
 			}
@@ -3076,12 +3182,15 @@ namespace UnrealGameSync
 
 			PendingSelectedChangeNumber = ChangeNumber;
 
-			if(BuildList.Items.Count > 0)
+			int CurrentMaxChanges = PerforceMonitor.CurrentMaxChanges;
+			if(PerforceMonitor.PendingMaxChanges <= CurrentMaxChanges && BuildList.SelectedItems.Count == 0)
 			{
-				BuildList.Items[BuildList.Items.Count - 1].EnsureVisible();
+				PerforceMonitor.PendingMaxChanges = CurrentMaxChanges + BuildListExpandCount;
+				if(ExpandItem != null)
+				{
+					ExpandItem.EnsureVisible();
+				}
 			}
-
-			UpdateNumRequestedBuilds(false);
 		}
 
 		private void BuildList_MouseClick(object Sender, MouseEventArgs Args)
@@ -3091,100 +3200,119 @@ namespace UnrealGameSync
 				ListViewHitTestInfo HitTest = BuildList.HitTest(Args.Location);
 				if(HitTest.Item != null)
 				{
-					PerforceChangeSummary Change = (PerforceChangeSummary)HitTest.Item.Tag;
-					if(Workspace.PendingChangeNumber == Change.Number)
+					if(HitTest.Item == ExpandItem)
 					{
-						Rectangle SubItemRect = HitTest.Item.SubItems[StatusColumn.Index].Bounds;
-
-						if(Workspace.IsBusy())
+						if(HitTestExpandLink(Args.Location))
 						{
-							Rectangle CancelRect = new Rectangle(SubItemRect.Right - 16, SubItemRect.Top, 16, SubItemRect.Height);
-							Rectangle InfoRect = new Rectangle(SubItemRect.Right - 32, SubItemRect.Top, 16, SubItemRect.Height);
-							if(CancelRect.Contains(Args.Location))
+							int CurrentMaxChanges = PerforceMonitor.CurrentMaxChanges;
+							if(PerforceMonitor.PendingMaxChanges > CurrentMaxChanges)
 							{
-								CancelWorkspaceUpdate();
+								PerforceMonitor.PendingMaxChanges = CurrentMaxChanges;
 							}
-							else if(InfoRect.Contains(Args.Location) && !Splitter.IsLogVisible())
+							else
 							{
-								ToggleLogVisibility();
+								PerforceMonitor.PendingMaxChanges = CurrentMaxChanges + BuildListExpandCount;
 							}
+							BuildList.Invalidate();
 						}
-						else
+					}
+					else
+					{
+						PerforceChangeSummary Change = (PerforceChangeSummary)HitTest.Item.Tag;
+						if(Workspace.PendingChangeNumber == Change.Number)
 						{
-							Rectangle HappyRect = new Rectangle(SubItemRect.Right - 32, SubItemRect.Top, 16, SubItemRect.Height);
-							Rectangle FrownRect = new Rectangle(SubItemRect.Right - 16, SubItemRect.Top, 16, SubItemRect.Height);
-							if(HappyRect.Contains(Args.Location))
+							Rectangle SubItemRect = HitTest.Item.SubItems[StatusColumn.Index].Bounds;
+
+							if(Workspace.IsBusy())
 							{
-								EventMonitor.PostEvent(Change.Number, EventType.Good);
-								BuildList.Invalidate();
+								Rectangle CancelRect = new Rectangle(SubItemRect.Right - 16, SubItemRect.Top, 16, SubItemRect.Height);
+								Rectangle InfoRect = new Rectangle(SubItemRect.Right - 32, SubItemRect.Top, 16, SubItemRect.Height);
+								if(CancelRect.Contains(Args.Location))
+								{
+									CancelWorkspaceUpdate();
+								}
+								else if(InfoRect.Contains(Args.Location) && !Splitter.IsLogVisible())
+								{
+									ToggleLogVisibility();
+								}
 							}
-							else if(FrownRect.Contains(Args.Location))
+							else
 							{
-								EventMonitor.PostEvent(Change.Number, EventType.Bad);
-								BuildList.Invalidate();
+								Rectangle HappyRect = new Rectangle(SubItemRect.Right - 32, SubItemRect.Top, 16, SubItemRect.Height);
+								Rectangle FrownRect = new Rectangle(SubItemRect.Right - 16, SubItemRect.Top, 16, SubItemRect.Height);
+								if(HappyRect.Contains(Args.Location))
+								{
+									EventMonitor.PostEvent(Change.Number, EventType.Good);
+									BuildList.Invalidate();
+								}
+								else if(FrownRect.Contains(Args.Location))
+								{
+									EventMonitor.PostEvent(Change.Number, EventType.Bad);
+									BuildList.Invalidate();
+								}
 							}
 						}
-					}
-					else
-					{
-						Rectangle SyncBadgeRectangle = GetSyncBadgeRectangle(HitTest.Item.SubItems[StatusColumn.Index].Bounds);
-						if(SyncBadgeRectangle.Contains(Args.Location) && CanSyncChange(Change.Number))
+						else
 						{
-							StartSync(Change.Number, null);
+							Rectangle SyncBadgeRectangle = GetSyncBadgeRectangle(HitTest.Item.SubItems[StatusColumn.Index].Bounds);
+							if(SyncBadgeRectangle.Contains(Args.Location) && CanSyncChange(Change.Number))
+							{
+								StartSync(Change.Number, null);
+							}
 						}
-					}
 
-					if(DescriptionColumn.Index < HitTest.Item.SubItems.Count && HitTest.Item.SubItems[DescriptionColumn.Index] == HitTest.SubItem)
-					{
-						ChangeLayoutInfo LayoutInfo = GetChangeLayoutInfo(Change);
-						if(LayoutInfo.DescriptionBadges.Count > 0)
+						if(DescriptionColumn.Index < HitTest.Item.SubItems.Count && HitTest.Item.SubItems[DescriptionColumn.Index] == HitTest.SubItem)
 						{
-							Point BuildListLocation = GetBadgeListLocation(LayoutInfo.DescriptionBadges, HitTest.SubItem.Bounds, HorizontalAlign.Right, VerticalAlignment.Middle);
-							BuildListLocation.Offset(-2, 0);
-
-							foreach (BadgeInfo Badge in LayoutInfo.DescriptionBadges)
+							ChangeLayoutInfo LayoutInfo = GetChangeLayoutInfo(Change);
+							if(LayoutInfo.DescriptionBadges.Count > 0)
 							{
-								Rectangle BadgeBounds = Badge.GetBounds(BuildListLocation);
-								if(BadgeBounds.Contains(Args.Location) && Badge.ClickHandler != null)
+								Point BuildListLocation = GetBadgeListLocation(LayoutInfo.DescriptionBadges, HitTest.SubItem.Bounds, HorizontalAlign.Right, VerticalAlignment.Middle);
+								BuildListLocation.Offset(-2, 0);
+
+								foreach (BadgeInfo Badge in LayoutInfo.DescriptionBadges)
 								{
-									Badge.ClickHandler();
-									break;
+									Rectangle BadgeBounds = Badge.GetBounds(BuildListLocation);
+									if(BadgeBounds.Contains(Args.Location) && Badge.ClickHandler != null)
+									{
+										Badge.ClickHandler();
+										break;
+									}
 								}
 							}
 						}
-					}
 
-					if(CISColumn.Index < HitTest.Item.SubItems.Count && HitTest.Item.SubItems[CISColumn.Index] == HitTest.SubItem)
-					{
-						ChangeLayoutInfo LayoutInfo = GetChangeLayoutInfo(Change);
-						if(LayoutInfo.BuildBadges.Count > 0)
+						if(CISColumn.Index < HitTest.Item.SubItems.Count && HitTest.Item.SubItems[CISColumn.Index] == HitTest.SubItem)
 						{
-							Point BuildListLocation = GetBadgeListLocation(LayoutInfo.BuildBadges, HitTest.SubItem.Bounds, HorizontalAlign.Center, VerticalAlignment.Middle);
-							BuildListLocation.X = Math.Max(BuildListLocation.X, HitTest.SubItem.Bounds.Left);
-
-							BadgeInfo BadgeInfo = HitTestBadge(Args.Location, LayoutInfo.BuildBadges, BuildListLocation);
-							if(BadgeInfo != null && BadgeInfo.ClickHandler != null)
+							ChangeLayoutInfo LayoutInfo = GetChangeLayoutInfo(Change);
+							if(LayoutInfo.BuildBadges.Count > 0)
 							{
-								BadgeInfo.ClickHandler();
+								Point BuildListLocation = GetBadgeListLocation(LayoutInfo.BuildBadges, HitTest.SubItem.Bounds, HorizontalAlign.Center, VerticalAlignment.Middle);
+								BuildListLocation.X = Math.Max(BuildListLocation.X, HitTest.SubItem.Bounds.Left);
+
+								BadgeInfo BadgeInfo = HitTestBadge(Args.Location, LayoutInfo.BuildBadges, BuildListLocation);
+								if(BadgeInfo != null && BadgeInfo.ClickHandler != null)
+								{
+									BadgeInfo.ClickHandler();
+								}
 							}
 						}
-					}
 
-					foreach(ColumnHeader CustomColumn in CustomColumns)
-					{
-						if(CustomColumn.Index < HitTest.Item.SubItems.Count && HitTest.Item.SubItems[CustomColumn.Index] == HitTest.SubItem)
+						foreach(ColumnHeader CustomColumn in CustomColumns)
 						{
-							ChangeLayoutInfo LayoutInfo = GetChangeLayoutInfo((PerforceChangeSummary)HitTest.Item.Tag);
-
-							List<BadgeInfo> Badges;
-							if(LayoutInfo.CustomBadges.TryGetValue(CustomColumn.Text, out Badges) && Badges.Count > 0)
+							if(CustomColumn.Index < HitTest.Item.SubItems.Count && HitTest.Item.SubItems[CustomColumn.Index] == HitTest.SubItem)
 							{
-								Point ListLocation = GetBadgeListLocation(Badges, HitTest.SubItem.Bounds, HorizontalAlign.Center, VerticalAlignment.Middle);
+								ChangeLayoutInfo LayoutInfo = GetChangeLayoutInfo((PerforceChangeSummary)HitTest.Item.Tag);
 
-								BadgeInfo BadgeInfo = HitTestBadge(Args.Location, Badges, ListLocation);
-								if(BadgeInfo != null && BadgeInfo.ClickHandler != null)
+								List<BadgeInfo> Badges;
+								if(LayoutInfo.CustomBadges.TryGetValue(CustomColumn.Text, out Badges) && Badges.Count > 0)
 								{
-									BadgeInfo.ClickHandler();
+									Point ListLocation = GetBadgeListLocation(Badges, HitTest.SubItem.Bounds, HorizontalAlign.Center, VerticalAlignment.Middle);
+
+									BadgeInfo BadgeInfo = HitTestBadge(Args.Location, Badges, ListLocation);
+									if(BadgeInfo != null && BadgeInfo.ClickHandler != null)
+									{
+										BadgeInfo.ClickHandler();
+									}
 								}
 							}
 						}
@@ -3307,8 +3435,16 @@ namespace UnrealGameSync
 		{
 			ListViewHitTestInfo HitTest = BuildList.HitTest(e.Location);
 
+			bool bNewMouseOverExpandLink = HitTest.Item != null && HitTestExpandLink(e.Location);
+			if(bMouseOverExpandLink != bNewMouseOverExpandLink)
+			{
+				bMouseOverExpandLink = bNewMouseOverExpandLink;
+				Cursor = bMouseOverExpandLink? NativeCursors.Hand : Cursors.Arrow;
+				BuildList.Invalidate();
+			}
+
 			string NewHoverBadgeUniqueId = null;
-			if(HitTest.Item != null)
+			if(HitTest.Item != null && HitTest.Item.Tag is PerforceChangeSummary)
 			{
 				int ColumnIndex = HitTest.Item.SubItems.IndexOf(HitTest.SubItem);
 				if(ColumnIndex == DescriptionColumn.Index)
@@ -3626,7 +3762,7 @@ namespace UnrealGameSync
 		private void OnlyShowReviewedCheckBox_CheckedChanged(object sender, EventArgs e)
 		{
 			UpdateBuildList();
-			UpdateNumRequestedBuilds(true);
+			ShrinkNumRequestedBuilds();
 		}
 
 		public void Activate()
@@ -3647,7 +3783,7 @@ namespace UnrealGameSync
 
 		public void Deactivate()
 		{
-			UpdateNumRequestedBuilds(true);
+			ShrinkNumRequestedBuilds();
 		}
 
 		private void BuildList_ItemMouseHover(object sender, ListViewItemMouseHoverEventArgs Args)
@@ -3658,6 +3794,10 @@ namespace UnrealGameSync
 				if(Args.Item.Bounds.Contains(ClientPoint))
 				{
 					PerforceChangeSummary Change = (PerforceChangeSummary)Args.Item.Tag;
+					if(Change == null)
+					{
+						return;
+					}
 
 					Rectangle CISBounds = Args.Item.SubItems[CISColumn.Index].Bounds;
 					if(CISBounds.Contains(ClientPoint))
@@ -3724,6 +3864,13 @@ namespace UnrealGameSync
 				HoverBadgeUniqueId = null;
 				BuildList.Invalidate();
 			}
+
+			if(bMouseOverExpandLink)
+			{
+				Cursor = Cursors.Arrow;
+				bMouseOverExpandLink = false;
+				BuildList.Invalidate();
+			}
 		}
 
 		private void OptionsContextMenu_ShowLog_Click(object sender, EventArgs e)
@@ -3884,30 +4031,33 @@ namespace UnrealGameSync
 			for(int Idx = 0; Idx < BuildList.Items.Count; Idx++)
 			{
 				PerforceChangeSummary Change = (PerforceChangeSummary)BuildList.Items[Idx].Tag;
-				if(ChangeType == LatestChangeType.Any)
+				if(Change != null)
 				{
-					if(CanSyncChange(Change.Number))
+					if(ChangeType == LatestChangeType.Any)
 					{
-						ChangeNumber = FindNewestGoodContentChange(Change.Number);
-						return true;
+						if(CanSyncChange(Change.Number))
+						{
+							ChangeNumber = FindNewestGoodContentChange(Change.Number);
+							return true;
+						}
 					}
-				}
-				else if(ChangeType == LatestChangeType.Good)
-				{
-					EventSummary Summary = EventMonitor.GetSummaryForChange(Change.Number);
-					if(Summary != null && Summary.Verdict == ReviewVerdict.Good && CanSyncChange(Change.Number))
+					else if(ChangeType == LatestChangeType.Good)
 					{
-						ChangeNumber = FindNewestGoodContentChange(Change.Number);
-						return true;
+						EventSummary Summary = EventMonitor.GetSummaryForChange(Change.Number);
+						if(Summary != null && Summary.Verdict == ReviewVerdict.Good && CanSyncChange(Change.Number))
+						{
+							ChangeNumber = FindNewestGoodContentChange(Change.Number);
+							return true;
+						}
 					}
-				}
-				else if(ChangeType == LatestChangeType.Starred)
-				{
-					EventSummary Summary = EventMonitor.GetSummaryForChange(Change.Number);
-					if(((Summary != null && Summary.LastStarReview != null) || PromotedChangeNumbers.Contains(Change.Number)) && CanSyncChange(Change.Number))
+					else if(ChangeType == LatestChangeType.Starred)
 					{
-						ChangeNumber = FindNewestGoodContentChange(Change.Number);
-						return true;
+						EventSummary Summary = EventMonitor.GetSummaryForChange(Change.Number);
+						if(((Summary != null && Summary.LastStarReview != null) || PromotedChangeNumbers.Contains(Change.Number)) && CanSyncChange(Change.Number))
+						{
+							ChangeNumber = FindNewestGoodContentChange(Change.Number);
+							return true;
+						}
 					}
 				}
 			}
@@ -4389,7 +4539,7 @@ namespace UnrealGameSync
 			Settings.Save();
 
 			UpdateBuildList();
-			UpdateNumRequestedBuilds(true);
+			ShrinkNumRequestedBuilds();
 		}
 
 		private void OptionsContextMenu_ShowChanges_ShowAutomated_Click(object sender, EventArgs e)
@@ -4796,7 +4946,7 @@ namespace UnrealGameSync
 		private void UpdateBuildListFilter()
 		{
 			UpdateBuildList();
-			UpdateNumRequestedBuilds(true);
+			ShrinkNumRequestedBuilds();
 		}
 	}
 }
diff --git a/Engine/Source/Programs/UnrealGameSync/UnrealGameSync/PerforceMonitor.cs b/Engine/Source/Programs/UnrealGameSync/UnrealGameSync/PerforceMonitor.cs
--- a/Engine/Source/Programs/UnrealGameSync/UnrealGameSync/PerforceMonitor.cs
+++ b/Engine/Source/Programs/UnrealGameSync/UnrealGameSync/PerforceMonitor.cs
@@ -29,6 +29,8 @@ namespace UnrealGameSync
 			}
 		}
 
+		public int InitialMaxChangesValue = 100;
+
 		PerforceConnection Perforce;
 		readonly string BranchClientPath;
 		readonly string SelectedClientFileName;
@@ -63,7 +65,7 @@ namespace UnrealGameSync
 			BranchClientPath = InBranchClientPath;
 			SelectedClientFileName = InSelectedClientFileName;
 			SelectedProjectIdentifier = InSelectedProjectIdentifier;
-			PendingMaxChangesValue = 100;
+			PendingMaxChangesValue = InitialMaxChangesValue;
 			LastChangeByCurrentUser = -1;
 			LastCodeChangeByCurrentUser = -1;
 			bIsEnterpriseProject = bInIsEnterpriseProject;
@@ -362,6 +364,12 @@ namespace UnrealGameSync
 			{
 				string[] CodeExtensions = { ".cs", ".h", ".cpp", ".inl", ".usf", ".ush", ".uproject", ".uplugin" };
 
+				// Skip this stuff if the user wants us to query for more changes
+				if(PendingMaxChanges > CurrentMaxChanges)
+				{
+					break;
+				}
+
 				// If there's something to check for, find all the content changes after this changelist
 				PerforceDescribeRecord DescribeRecord;
 				if(Perforce.Describe(QueryChangeNumber, out DescribeRecord, LogWriter))
