UGS: When browsing for the project to sync in a new workspace, hide the *.uprojectdirs files by default.

diff --git a/Engine/Source/Programs/UnrealGameSync/UnrealGameSync/Forms/SelectProjectFromWorkspaceWindow.cs b/Engine/Source/Programs/UnrealGameSync/UnrealGameSync/Forms/SelectProjectFromWorkspaceWindow.cs
--- a/Engine/Source/Programs/UnrealGameSync/UnrealGameSync/Forms/SelectProjectFromWorkspaceWindow.cs
+++ b/Engine/Source/Programs/UnrealGameSync/UnrealGameSync/Forms/SelectProjectFromWorkspaceWindow.cs
@@ -84,6 +84,7 @@ namespace UnrealGameSync
 		[DllImport("Shell32.dll", EntryPoint = "ExtractIconExW", CharSet = CharSet.Unicode, ExactSpelling = true, CallingConvention = CallingConvention.StdCall)]
 		private static extern int ExtractIconEx(string sFile, int iIndex, IntPtr piLargeVersion, out IntPtr piSmallVersion, int amountIcons);
 
+		List<string> ProjectFiles;
 		string SelectedProjectFile;
 
 		class ProjectNode
@@ -106,6 +107,7 @@ namespace UnrealGameSync
 		{
 			InitializeComponent();
 			
+			this.ProjectFiles = ProjectFiles;
 			this.SelectedProjectFile = SelectedProjectFile;
 
 			// Make the image strip containing icons for nodes in the tree
@@ -135,8 +137,25 @@ namespace UnrealGameSync
 			RootNode.Expand();
 			ProjectTreeView.Nodes.Add(RootNode);
 
+			// Populate the tree
+			Populate();
+		}
+
+		private void Populate()
+		{
+			// Clear out the existing nodes
+			TreeNode RootNode = ProjectTreeView.Nodes[0];
+			RootNode.Nodes.Clear();
+
+			// Filter the project files
+			List<string> FilteredProjectFiles = new List<string>(ProjectFiles); 
+			if(!ShowProjectDirsFiles.Checked)
+			{
+				FilteredProjectFiles.RemoveAll(x => x.EndsWith(".uprojectdirs", StringComparison.OrdinalIgnoreCase));
+			}
+
 			// Sort by paths, then files
-			List<ProjectNode> ProjectNodes = ProjectFiles.Select(x => new ProjectNode(x)).OrderBy(x => x.Folder).ThenBy(x => x.Name).ToList();
+			List<ProjectNode> ProjectNodes = FilteredProjectFiles.Select(x => new ProjectNode(x)).OrderBy(x => x.Folder).ThenBy(x => x.Name).ToList();
 
 			// Add the folders for each project
 			TreeNode[] ProjectParentNodes = new TreeNode[ProjectNodes.Count];
@@ -236,5 +255,10 @@ namespace UnrealGameSync
 			DialogResult = DialogResult.Cancel;
 			Close();
 		}
+
+		private void ShowProjectDirsFiles_CheckedChanged(object sender, EventArgs e)
+		{
+			Populate();
+		}
 	}
 }
