UGS: Make errors in the log display in red, and warnings in yellow.

diff --git a/Engine/Source/Programs/UnrealGameSync/UnrealGameSync/Controls/LogControl.cs b/Engine/Source/Programs/UnrealGameSync/UnrealGameSync/Controls/LogControl.cs
--- a/Engine/Source/Programs/UnrealGameSync/UnrealGameSync/Controls/LogControl.cs
+++ b/Engine/Source/Programs/UnrealGameSync/UnrealGameSync/Controls/LogControl.cs
@@ -1,4 +1,4 @@
-﻿// Copyright 1998-2019 Epic Games, Inc. All Rights Reserved.
+// Copyright 1998-2019 Epic Games, Inc. All Rights Reserved.
 
 using System;
 using System.Collections.Generic;
@@ -12,6 +12,7 @@ using System.Windows.Forms;
 using System.Runtime.InteropServices;
 using System.Collections.Concurrent;
 using System.IO;
+using System.Text.RegularExpressions;
 
 namespace UnrealGameSync
 {
@@ -749,10 +750,24 @@ namespace UnrealGameSync
 				int SelectMaxIdx;
 				ClipSelectionToLine(Idx, out SelectMinIdx, out SelectMaxIdx);
 
+				Color TextColor;
+				if(Regex.IsMatch(Lines[Idx], "(?<!\\w)error[: ]", RegexOptions.IgnoreCase))
+				{
+					TextColor = Color.FromArgb(189, 54, 47);
+				}
+				else if(Regex.IsMatch(Lines[Idx], "(?<!\\w)warning[: ]", RegexOptions.IgnoreCase))
+				{
+					TextColor = Color.FromArgb(128, 128, 0);
+				}
+				else
+				{
+					TextColor = Color.Black;
+				}
+
 				int TextY = (Idx - ScrollLine) * FontSize.Height;
 				if(SelectMinIdx > 0)
 				{
-					TextRenderer.DrawText(e.Graphics, Lines[Idx].Substring(0, SelectMinIdx), Font, new Point(TextX, TextY), Color.Black, TextFormatFlags.NoPadding);
+					TextRenderer.DrawText(e.Graphics, Lines[Idx].Substring(0, SelectMinIdx), Font, new Point(TextX, TextY), TextColor, TextFormatFlags.NoPadding);
 				}
 				if(SelectMinIdx < SelectMaxIdx)
 				{
@@ -761,7 +776,7 @@ namespace UnrealGameSync
 				}
 				if(SelectMaxIdx < Lines[Idx].Length)
 				{
-					TextRenderer.DrawText(e.Graphics, Lines[Idx].Substring(SelectMaxIdx), Font, new Point(TextX + (SelectMaxIdx * FontSize.Width), TextY), Color.Black, TextFormatFlags.NoPadding);
+					TextRenderer.DrawText(e.Graphics, Lines[Idx].Substring(SelectMaxIdx), Font, new Point(TextX + (SelectMaxIdx * FontSize.Width), TextY), TextColor, TextFormatFlags.NoPadding);
 				}
 			}
 		}
