UGS: Add a check that the window bounds are on one of the displays before restoring them.

diff --git a/Engine/Source/Programs/UnrealGameSync/UnrealGameSync/Forms/MainWindow.cs b/Engine/Source/Programs/UnrealGameSync/UnrealGameSync/Forms/MainWindow.cs
--- a/Engine/Source/Programs/UnrealGameSync/UnrealGameSync/Forms/MainWindow.cs
+++ b/Engine/Source/Programs/UnrealGameSync/UnrealGameSync/Forms/MainWindow.cs
@@ -1169,8 +1169,19 @@ namespace UnrealGameSync
 		{
 			if(Settings.WindowBounds != null)
 			{
-				Location = Settings.WindowBounds.Value.Location;
-				Size = Settings.WindowBounds.Value.Size;
+				Rectangle WindowBounds = Settings.WindowBounds.Value;
+				if(WindowBounds.Width > MinimumSize.Width && WindowBounds.Height > MinimumSize.Height)
+				{
+					foreach (Screen Screen in Screen.AllScreens)
+					{
+						if(WindowBounds.IntersectsWith(Screen.Bounds))
+						{
+							Location = Settings.WindowBounds.Value.Location;
+							Size = Settings.WindowBounds.Value.Size;
+							break;
+						}
+					}
+				}
 			}
 			WindowState = Settings.WindowState;
 		}
