UGS: Add a deployment setting for whether to send telemetry, and default it to off. We haven't been using this data for some time, and it's causing a lot of slow DB locks.

diff --git a/Engine/Source/Programs/UnrealGameSync/UnrealGameSync/DeploymentSettings.cs b/Engine/Source/Programs/UnrealGameSync/UnrealGameSync/DeploymentSettings.cs
--- a/Engine/Source/Programs/UnrealGameSync/UnrealGameSync/DeploymentSettings.cs
+++ b/Engine/Source/Programs/UnrealGameSync/UnrealGameSync/DeploymentSettings.cs
@@ -24,5 +24,10 @@ namespace UnrealGameSync
 		/// The UnrealGameSync executable should be located at Release/UnrealGameSync.exe under this path, with any dependent DLLs.
 		/// </summary>
 		public static readonly string DefaultDepotPath = null;
+
+		/// <summary>
+		/// Whether to send telemetry data by default. Can be useful for finding users in need of a hardware upgrade, but can bloat the DB for large teams.
+		/// </summary>
+		public static bool bSendTelemetry = false;
 	}
 }
diff --git a/Engine/Source/Programs/UnrealGameSync/UnrealGameSync/Telemetry.cs b/Engine/Source/Programs/UnrealGameSync/UnrealGameSync/Telemetry.cs
--- a/Engine/Source/Programs/UnrealGameSync/UnrealGameSync/Telemetry.cs
+++ b/Engine/Source/Programs/UnrealGameSync/UnrealGameSync/Telemetry.cs
@@ -176,6 +176,11 @@ namespace UnrealGameSync
 
 		bool SendTimingData(TelemetryTimingData Data, string Version, string IpAddress)
 		{
+			if(!DeploymentSettings.bSendTelemetry)
+			{
+				return true;
+			}
+
 			try
 			{
 				Stopwatch Timer = Stopwatch.StartNew();
