UGS: Prevent messages like "Showing X/Y changes." where X > Y in list of changes.

#rnx

diff --git a/Engine/Source/Programs/UnrealGameSync/UnrealGameSync/Controls/WorkspaceControl.cs b/Engine/Source/Programs/UnrealGameSync/UnrealGameSync/Controls/WorkspaceControl.cs
--- a/Engine/Source/Programs/UnrealGameSync/UnrealGameSync/Controls/WorkspaceControl.cs
+++ b/Engine/Source/Programs/UnrealGameSync/UnrealGameSync/Controls/WorkspaceControl.cs
@@ -1544,7 +1544,7 @@ namespace UnrealGameSync
 			}
 			else
 			{
-				ShowingChanges = String.Format("Showing {0}/{1} changes.", ListIndexToChangeIndex.Count, CurrentMaxChanges);
+				ShowingChanges = String.Format("Showing {0}/{1} changes.", ListIndexToChangeIndex.Count, Math.Max(ListIndexToChangeIndex.Count, CurrentMaxChanges));
 			}
 
 			if(PerforceMonitor.PendingMaxChanges > CurrentMaxChanges)
