UGS: Add the Build/*/PipelineCaches folders to be excluded for Android/Mac/IOS via the standard platform sync filter options.

#rnx

diff --git a/Engine/Source/Programs/UnrealGameSync/UnrealGameSync/Workspace.cs b/Engine/Source/Programs/UnrealGameSync/UnrealGameSync/Workspace.cs
--- a/Engine/Source/Programs/UnrealGameSync/UnrealGameSync/Workspace.cs
+++ b/Engine/Source/Programs/UnrealGameSync/UnrealGameSync/Workspace.cs
@@ -104,11 +104,11 @@ namespace UnrealGameSync
 		{
 			new WorkspaceSyncCategory(new Guid("{6703E989-D912-451D-93AD-B48DE748D282}"), "Content", "*.uasset", "*.umap"),
 			new WorkspaceSyncCategory(new Guid("{6507C2FB-19DD-403A-AFA3-BBF898248D5A}"), "Documentation", "/Engine/Documentation/..."),
-			new WorkspaceSyncCategory(new Guid("{FD7C716E-4BAD-43AE-8FAE-8748EF9EE44D}"), "Platform Support: Android", "/Engine/Source/ThirdParty/.../Android/..."),
+			new WorkspaceSyncCategory(new Guid("{FD7C716E-4BAD-43AE-8FAE-8748EF9EE44D}"), "Platform Support: Android", "/Engine/Source/ThirdParty/.../Android/...", ".../Build/Android/PipelineCaches/..."),
 			new WorkspaceSyncCategory(new Guid("{3299A73D-2176-4C0F-BC99-C1C6631AF6C4}"), "Platform Support: HTML5", "/Engine/Source/ThirdParty/.../HTML5/...", "/Engine/Extras/ThirdPartyNotUE/emsdk/..."),
-			new WorkspaceSyncCategory(new Guid("{176B2EB2-35F7-4E8E-B131-5F1C5F0959AF}"), "Platform Support: iOS", "/Engine/Source/ThirdParty/.../IOS/..."),
+			new WorkspaceSyncCategory(new Guid("{176B2EB2-35F7-4E8E-B131-5F1C5F0959AF}"), "Platform Support: iOS", "/Engine/Source/ThirdParty/.../IOS/...", ".../Build/IOS/PipelineCaches/..."),
 			new WorkspaceSyncCategory(new Guid("{F44B2D25-CBC0-4A8F-B6B3-E4A8125533DD}"), "Platform Support: Linux", "/Engine/Source/ThirdParty/.../Linux/..."),
-			new WorkspaceSyncCategory(new Guid("{2AF45231-0D75-463B-BF9F-ABB3231091BB}"), "Platform Support: Mac", "/Engine/Source/ThirdParty/.../Mac/..."),
+			new WorkspaceSyncCategory(new Guid("{2AF45231-0D75-463B-BF9F-ABB3231091BB}"), "Platform Support: Mac", "/Engine/Source/ThirdParty/.../Mac/...", ".../Build/Mac/PipelineCaches/..."),
 			new WorkspaceSyncCategory(new Guid("{C8CB4934-ADE9-46C9-B6E3-61A659E1FAF5}"), "Platform Support: PS4", ".../PS4/..."),
 			new WorkspaceSyncCategory(new Guid("{F8AE5AC3-DA2D-4719-BABF-8A90D878379E}"), "Platform Support: Switch", ".../Switch/..."),
 			new WorkspaceSyncCategory(new Guid("{3788A0BC-188C-4A0D-950A-D68175F0D110}"), "Platform Support: tvOS", "/Engine/Source/ThirdParty/.../TVOS/..."),
