UGS: Fix incorrect handling of return code when fetching the project file for a new client.

diff --git a/Engine/Source/Programs/UnrealGameSync/UnrealGameSync/DetectProjectSettingsTask.cs b/Engine/Source/Programs/UnrealGameSync/UnrealGameSync/DetectProjectSettingsTask.cs
--- a/Engine/Source/Programs/UnrealGameSync/UnrealGameSync/DetectProjectSettingsTask.cs
+++ b/Engine/Source/Programs/UnrealGameSync/UnrealGameSync/DetectProjectSettingsTask.cs
@@ -303,7 +303,7 @@ namespace UnrealGameSync
 				else
 				{
 					List<string> ProjectLines;
-					if(PerforceClient.Print(NewSelectedClientFileName, out ProjectLines, Log))
+					if(!PerforceClient.Print(NewSelectedClientFileName, out ProjectLines, Log))
 					{
 						ErrorMessage = String.Format("Unable to get contents of {0}", NewSelectedClientFileName);
 						return false;
