UGS: Fix incorrect display of total number of known changes when a branch has fewer changes that the default batch size.

diff --git a/Engine/Source/Programs/UnrealGameSync/UnrealGameSync/Controls/WorkspaceControl.cs b/Engine/Source/Programs/UnrealGameSync/UnrealGameSync/Controls/WorkspaceControl.cs
--- a/Engine/Source/Programs/UnrealGameSync/UnrealGameSync/Controls/WorkspaceControl.cs
+++ b/Engine/Source/Programs/UnrealGameSync/UnrealGameSync/Controls/WorkspaceControl.cs
@@ -1535,18 +1535,9 @@ namespace UnrealGameSync
 		{
 			ExpandRowLayout Layout = new ExpandRowLayout();
 
-			int CurrentMaxChanges = PerforceMonitor.CurrentMaxChanges;
-
-			string ShowingChanges;
-			if(ListIndexToChangeIndex.Count == CurrentMaxChanges)
-			{
-				ShowingChanges = String.Format("Showing {0} changes.", ListIndexToChangeIndex.Count);
-			}
-			else
-			{
-				ShowingChanges = String.Format("Showing {0}/{1} changes.", ListIndexToChangeIndex.Count, Math.Max(ListIndexToChangeIndex.Count, CurrentMaxChanges));
-			}
+			string ShowingChanges = String.Format("Showing {0}/{1} changes.", ListIndexToChangeIndex.Count, NumChanges);
 
+			int CurrentMaxChanges = PerforceMonitor.CurrentMaxChanges;
 			if(PerforceMonitor.PendingMaxChanges > CurrentMaxChanges)
 			{
 				Layout.MainText = String.Format("{0}. Fetching {1} more from server...  ", ShowingChanges, PerforceMonitor.PendingMaxChanges - CurrentMaxChanges);
