UGS: Suppress the fade-out on the list of changes if the branch does not have at least the requested number of changes.

diff --git a/Engine/Source/Programs/UnrealGameSync/UnrealGameSync/Controls/WorkspaceControl.cs b/Engine/Source/Programs/UnrealGameSync/UnrealGameSync/Controls/WorkspaceControl.cs
--- a/Engine/Source/Programs/UnrealGameSync/UnrealGameSync/Controls/WorkspaceControl.cs
+++ b/Engine/Source/Programs/UnrealGameSync/UnrealGameSync/Controls/WorkspaceControl.cs
@@ -1543,6 +1543,11 @@ namespace UnrealGameSync
 				Layout.MainText = String.Format("{0}. Fetching {1} more from server...  ", ShowingChanges, PerforceMonitor.PendingMaxChanges - CurrentMaxChanges);
 				Layout.LinkText = "Cancel";
 			}
+			else if(PerforceMonitor.CurrentMaxChanges >= NumChanges)
+			{
+				Layout.MainText = ShowingChanges;
+				Layout.LinkText = "";
+			}
 			else
 			{
 				Layout.MainText = String.Format("{0}  ", ShowingChanges);
@@ -1719,7 +1724,7 @@ namespace UnrealGameSync
 			Color TextColor = (bAllowSync || Change.Number == Workspace.PendingChangeNumber || Change.Number == Workspace.CurrentChangeNumber || (WorkspaceSettings != null && WorkspaceSettings.AdditionalChangeNumbers.Contains(Change.Number)))? SystemColors.WindowText : Blend(SystemColors.Window, SystemColors.WindowText, 0.25f);
 
 			const int FadeRange = 6;
-			if(e.ItemIndex >= BuildList.Items.Count - FadeRange)
+			if(e.ItemIndex >= BuildList.Items.Count - FadeRange && NumChanges >= PerforceMonitor.CurrentMaxChanges)
 			{
 				float Opacity = (float)(BuildList.Items.Count - e.ItemIndex - 0.9f) / FadeRange;
 				BadgeAlpha = (int)(BadgeAlpha * Opacity);
