UGS: Exclude any P4CONFIG file from the clean workspace dialog.

#rnx

diff --git a/Engine/Source/Programs/UnrealGameSync/UnrealGameSync/FindFoldersToCleanTask.cs b/Engine/Source/Programs/UnrealGameSync/UnrealGameSync/FindFoldersToCleanTask.cs
--- a/Engine/Source/Programs/UnrealGameSync/UnrealGameSync/FindFoldersToCleanTask.cs
+++ b/Engine/Source/Programs/UnrealGameSync/UnrealGameSync/FindFoldersToCleanTask.cs
@@ -102,14 +102,14 @@ namespace UnrealGameSync
 			}
 		}
 
-		void MergeTrees(FolderToClean LocalFolder, PerforceHaveFolder PerforceFolder, HashSet<string> OpenClientPaths)
+		void MergeTrees(FolderToClean LocalFolder, PerforceHaveFolder PerforceFolder, HashSet<string> OpenClientPaths, string PerforceConfigFile)
 		{
 			if(PerforceFolder == null)
 			{
 				// Loop through all the local sub-folders
 				foreach(FolderToClean LocalSubFolder in LocalFolder.NameToSubFolder.Values)
 				{
-					MergeTrees(LocalSubFolder, null, OpenClientPaths);
+					MergeTrees(LocalSubFolder, null, OpenClientPaths, PerforceConfigFile);
 				}
 
 				// Delete everything
@@ -122,7 +122,7 @@ namespace UnrealGameSync
 				{
 					PerforceHaveFolder PerforceSubFolder;
 					PerforceFolder.NameToSubFolder.TryGetValue(LocalSubFolder.Name, out PerforceSubFolder);
-					MergeTrees(LocalSubFolder, PerforceSubFolder, OpenClientPaths);
+					MergeTrees(LocalSubFolder, PerforceSubFolder, OpenClientPaths, PerforceConfigFile);
 				}
 
 				// Also merge all the Perforce folders that no longer exist
@@ -132,7 +132,7 @@ namespace UnrealGameSync
 					if(!LocalFolder.NameToSubFolder.TryGetValue(PerforceSubFolderPair.Key, out LocalSubFolder))
 					{
 						LocalSubFolder = new FolderToClean(new DirectoryInfo(Path.Combine(LocalFolder.Directory.FullName, PerforceSubFolderPair.Key)));
-						MergeTrees(LocalSubFolder, PerforceSubFolderPair.Value, OpenClientPaths);
+						MergeTrees(LocalSubFolder, PerforceSubFolderPair.Value, OpenClientPaths, PerforceConfigFile);
 						LocalFolder.NameToSubFolder.Add(LocalSubFolder.Name, LocalSubFolder);
 					}
 				}
@@ -161,6 +161,12 @@ namespace UnrealGameSync
 				}
 			}
 
+			// Remove any config files
+			if(PerforceConfigFile != null)
+			{
+				LocalFolder.FilesToDelete.RemoveAll(x => String.Compare(x.Name, PerforceConfigFile, StringComparison.OrdinalIgnoreCase) == 0);
+			}
+
 			// Figure out if this folder is just an empty directory that needs to be removed
 			LocalFolder.bEmptyLeaf = LocalFolder.NameToFile.Count == 0 && LocalFolder.NameToSubFolder.Count == 0 && LocalFolder.FilesToSync.Count == 0;
 
@@ -294,8 +300,12 @@ namespace UnrealGameSync
 			// Wait to finish scanning the directory
 			FinishedScan.WaitOne();
 
+			// Find the value of the P4CONFIG variable
+			string PerforceConfigFile;
+			PerforceClient.GetSetting("P4CONFIG", out PerforceConfigFile, Log);
+
 			// Merge the trees
-			MergeTrees(RootFolderToClean, RootHaveFolder, OpenLocalFiles);
+			MergeTrees(RootFolderToClean, RootHaveFolder, OpenLocalFiles, PerforceConfigFile);
 
 			// Remove all the empty folders
 			RemoveEmptyFolders(RootFolderToClean);
