UGS: Fix error message when unable to retrieve a project from the server.

#rnx

diff --git a/Engine/Source/Programs/UnrealGameSync/UnrealGameSync/DetectProjectSettingsTask.cs b/Engine/Source/Programs/UnrealGameSync/UnrealGameSync/DetectProjectSettingsTask.cs
--- a/Engine/Source/Programs/UnrealGameSync/UnrealGameSync/DetectProjectSettingsTask.cs
+++ b/Engine/Source/Programs/UnrealGameSync/UnrealGameSync/DetectProjectSettingsTask.cs
@@ -113,7 +113,7 @@ namespace UnrealGameSync
 				{
 					if(!PerforceClient.ConvertToLocalPath(NewSelectedClientFileName, out NewSelectedFileName, Log))
 					{
-						ErrorMessage = String.Format("Couldn't get client path for {0}", NewSelectedFileName);
+						ErrorMessage = String.Format("Couldn't get client path for {0}", NewSelectedClientFileName);
 						return false;
 					}
 				}
