UGS: Fix the "Show More" link not being visible at the bottom of the list of changes.

diff --git a/Engine/Source/Programs/UnrealGameSync/UnrealGameSync/Controls/WorkspaceControl.cs b/Engine/Source/Programs/UnrealGameSync/UnrealGameSync/Controls/WorkspaceControl.cs
--- a/Engine/Source/Programs/UnrealGameSync/UnrealGameSync/Controls/WorkspaceControl.cs
+++ b/Engine/Source/Programs/UnrealGameSync/UnrealGameSync/Controls/WorkspaceControl.cs
@@ -1543,7 +1543,7 @@ namespace UnrealGameSync
 				Layout.MainText = String.Format("{0}. Fetching {1} more from server...  ", ShowingChanges, PerforceMonitor.PendingMaxChanges - CurrentMaxChanges);
 				Layout.LinkText = "Cancel";
 			}
-			else if(PerforceMonitor.CurrentMaxChanges >= NumChanges)
+			else if(PerforceMonitor.CurrentMaxChanges > NumChanges)
 			{
 				Layout.MainText = ShowingChanges;
 				Layout.LinkText = "";
