UGS: Fix close button in tab control not scaling correctly on high DPI monitors.

diff --git a/Engine/Source/Programs/UnrealGameSync/UnrealGameSync/Controls/TabControl.cs b/Engine/Source/Programs/UnrealGameSync/UnrealGameSync/Controls/TabControl.cs
--- a/Engine/Source/Programs/UnrealGameSync/UnrealGameSync/Controls/TabControl.cs
+++ b/Engine/Source/Programs/UnrealGameSync/UnrealGameSync/Controls/TabControl.cs
@@ -22,6 +22,7 @@ namespace UnrealGameSync
 			public object Data;
 			public int MinX;
 			public int Width;
+			public int CloseButtonWidth;
 			public Size TextSize;
 			public Tuple<Color, float> Highlight;
 		}
@@ -253,7 +254,8 @@ namespace UnrealGameSync
 					Tab.Width = TabPadding + Tab.TextSize.Width + (int)(TabPadding * DpiScaleX);
 					if(Idx == SelectedTabIdx)
 					{
-						Tab.Width += (int)(TabCloseButtonWidth * DpiScaleX);
+						Tab.CloseButtonWidth = (int)(TabCloseButtonWidth * DpiScaleX);
+						Tab.Width += Tab.CloseButtonWidth;
 					}
 				}
 
@@ -324,7 +326,7 @@ namespace UnrealGameSync
 				{
 					if(HoverTabIdx == SelectedTabIdx)
 					{
-						if(e.Location.X > Tabs[HoverTabIdx].MinX + Tabs[HoverTabIdx].Width - TabCloseButtonWidth)
+						if(e.Location.X > Tabs[HoverTabIdx].MinX + Tabs[HoverTabIdx].Width - Tabs[HoverTabIdx].CloseButtonWidth)
 						{
 							RemoveTab(HoverTabIdx);
 						}
