UGS: Always use the saved global P4 settings rather than relying on them being passed in from UnrealGameSyncLauncher. Ensures they are correct when debugging.

#rnx

diff --git a/Engine/Source/Programs/UnrealGameSync/UnrealGameSync/Program.cs b/Engine/Source/Programs/UnrealGameSync/UnrealGameSync/Program.cs
--- a/Engine/Source/Programs/UnrealGameSync/UnrealGameSync/Program.cs
+++ b/Engine/Source/Programs/UnrealGameSync/UnrealGameSync/Program.cs
@@ -44,20 +44,16 @@ namespace UnrealGameSync
 
 		static void InnerMain(Mutex InstanceMutex, EventWaitHandle ActivateEvent, string[] Args)
 		{
-			List<string> RemainingArgs = new List<string>(Args);
+			string ServerAndPort = null;
+			string UserName = null;
+			string UpdatePath = null;
+			Utility.ReadGlobalPerforceSettings(ref ServerAndPort, ref UserName, ref UpdatePath);
 
-			string UpdatePath;
-			ParseArgument(RemainingArgs, "-updatepath=", out UpdatePath);
+			List<string> RemainingArgs = new List<string>(Args);
 
 			string UpdateSpawn;
 			ParseArgument(RemainingArgs, "-updatespawn=", out UpdateSpawn);
 
-			string ServerAndPort;
-			ParseArgument(RemainingArgs, "-p4port=", out ServerAndPort);
-
-			string UserName;
-			ParseArgument(RemainingArgs, "-p4user=", out UserName);
-
 			bool bRestoreState;
 			ParseOption(RemainingArgs, "-restorestate", out bRestoreState);
 
