UGS: Show counter for number of files fetched from the server when syncing. Useful for slow connections, where UGS can appear to hang.

diff --git a/Engine/Source/Programs/UnrealGameSync/UnrealGameSync/Perforce.cs b/Engine/Source/Programs/UnrealGameSync/UnrealGameSync/Perforce.cs
--- a/Engine/Source/Programs/UnrealGameSync/UnrealGameSync/Perforce.cs
+++ b/Engine/Source/Programs/UnrealGameSync/UnrealGameSync/Perforce.cs
@@ -1084,6 +1084,15 @@ namespace UnrealGameSync
 			return RunCommand(String.Format("have \"{0}\"", Filter), out FileRecords, CommandOptions.IgnoreFilesNotOnClientError, Log);
 		}
 
+		public bool Have(string Filter, Action<PerforceFileRecord> HaveOutput, TextWriter Log)
+		{
+			using(PerforceTagRecordParser Parser = new PerforceTagRecordParser(x => HaveOutput(new PerforceFileRecord(x))))
+			{
+				string CommandLine = String.Format("-ztag have \"{0}\"", Filter);
+				return RunCommand(CommandLine.ToString(), null, Line => FilterTaggedOutput(Line, Parser, Log), CommandOptions.IgnoreFilesNotOnClientError, Log);
+			}
+		}
+
 		public bool Stat(string Filter, out List<PerforceFileRecord> FileRecords, TextWriter Log)
 		{
 			return RunCommand(String.Format("fstat \"{0}\"", Filter), out FileRecords, CommandOptions.IgnoreFilesNotOnClientError | CommandOptions.IgnoreNoSuchFilesError | CommandOptions.IgnoreProtectedNamespaceError, Log);
@@ -1157,6 +1166,20 @@ namespace UnrealGameSync
 			}
 		}
 
+		private static bool FilterTaggedOutput(PerforceOutputLine Line, PerforceTagRecordParser Parser, TextWriter Log)
+		{
+			if(Line.Channel == PerforceOutputChannel.TaggedInfo)
+			{
+				Parser.OutputLine(Line.Text);
+				return true;
+			}
+			else
+			{
+				Log.WriteLine(Line.Text);
+				return Line.Channel != PerforceOutputChannel.Error;
+			}
+		}
+
 		private static bool FilterSyncOutput(PerforceOutputLine Line, PerforceTagRecordParser Parser, List<string> TamperedFiles, TextWriter Log)
 		{
 			if(Line.Channel == PerforceOutputChannel.TaggedInfo)
@@ -1191,6 +1214,15 @@ namespace UnrealGameSync
 			return RunCommand(String.Format("sync -n {0}@{1}{2}", Filter, bOnlyFilesInThisChange? "=" : "", ChangeNumber), out FileRecords, CommandOptions.IgnoreFilesUpToDateError | CommandOptions.IgnoreNoSuchFilesError | CommandOptions.IgnoreFilesNotInClientViewError, Log);
 		}
 
+		public bool SyncPreview(string Filter, int ChangeNumber, bool bOnlyFilesInThisChange, Action<PerforceFileRecord> SyncOutput, TextWriter Log)
+		{
+			using(PerforceTagRecordParser Parser = new PerforceTagRecordParser(x => SyncOutput(new PerforceFileRecord(x))))
+			{
+				string CommandLine = String.Format("-ztag sync -n {0}@{1}{2}", Filter, bOnlyFilesInThisChange? "=" : "", ChangeNumber);
+				return RunCommand(CommandLine.ToString(), null, Line => FilterTaggedOutput(Line, Parser, Log), CommandOptions.NoFailOnErrors | CommandOptions.IgnoreFilesUpToDateError | CommandOptions.IgnoreExitCode, Log);
+			}
+		}
+
 		public bool ForceSync(string Filter, TextWriter Log)
 		{
 			return RunCommand(String.Format("sync -f \"{0}\"", Filter), CommandOptions.IgnoreFilesUpToDateError, Log);
diff --git a/Engine/Source/Programs/UnrealGameSync/UnrealGameSync/Workspace.cs b/Engine/Source/Programs/UnrealGameSync/UnrealGameSync/Workspace.cs
--- a/Engine/Source/Programs/UnrealGameSync/UnrealGameSync/Workspace.cs
+++ b/Engine/Source/Programs/UnrealGameSync/UnrealGameSync/Workspace.cs
@@ -138,6 +138,42 @@ namespace UnrealGameSync
 
 		public event Action<WorkspaceUpdateContext, WorkspaceUpdateResult, string> OnUpdateComplete;
 
+		class RecordCounter : IDisposable
+		{
+			ProgressValue Progress;
+			string Message;
+			int Count;
+			Stopwatch Timer = Stopwatch.StartNew();
+
+			public RecordCounter(ProgressValue Progress, string Message)
+			{
+				this.Progress = Progress;
+				this.Message = Message;
+
+				Progress.Set(Message);
+			}
+
+			public void Dispose()
+			{
+				UpdateMessage();
+			}
+
+			public void Increment()
+			{
+				Count++;
+				if(Timer.ElapsedMilliseconds > 250)
+				{
+					UpdateMessage();
+				}
+			}
+
+			public void UpdateMessage()
+			{
+				Progress.Set(String.Format("{0} ({1:N0})", Message, Count));
+				Timer.Restart();
+			}
+		}
+
 		public Workspace(PerforceConnection InPerforce, string InLocalRootPath, string InSelectedLocalFileName, string InClientRootPath, string InSelectedClientFileName, int InInitialChangeNumber, string InInitialSyncFilterHash, int InLastBuiltChangeNumber, string InTelemetryProjectPath, bool bInIsEnterpriseProject, TextWriter InLog)
 		{
 			Perforce = InPerforce;
@@ -318,9 +354,6 @@ namespace UnrealGameSync
 						return WorkspaceUpdateResult.FailedToSyncLoginExpired;
 					}
 
-					// Find all the files that are out of date
-					Progress.Set("Finding files to sync...");
-
 					// Figure out which paths to sync
 					List<string> SyncPaths = GetSyncPaths((Context.Options & WorkspaceUpdateOptions.SyncAllProjects) != 0, Context.SyncFilter);
 
@@ -352,11 +385,14 @@ namespace UnrealGameSync
 						Log.WriteLine("Filter has changed ({0} -> {1}); finding files in workspace that need to be removed.", (String.IsNullOrEmpty(CurrentSyncFilterHash))? "None" : CurrentSyncFilterHash, NextSyncFilterHash);
 
 						// Find all the files that are in this workspace
-						List<PerforceFileRecord> HaveFiles;
-						if(!Perforce.Have("//...", out HaveFiles, Log))
+						List<PerforceFileRecord> HaveFiles = new List<PerforceFileRecord>();
+						using(RecordCounter HaveCounter = new RecordCounter(Progress, "Querying workspace..."))
 						{
-							StatusMessage = "Unable to query files.";
-							return WorkspaceUpdateResult.FailedToSync;
+							if(!Perforce.Have("//...", Record => { HaveFiles.Add(Record); HaveCounter.Increment(); }, Log))
+							{
+								StatusMessage = "Unable to query files.";
+								return WorkspaceUpdateResult.FailedToSync;
+							}
 						}
 
 						// Build a filter for the current sync paths
@@ -426,48 +462,51 @@ namespace UnrealGameSync
 
 					// Find all the server changes, and anything that's opened for edit locally. We need to sync files we have open to schedule a resolve.
 					List<string> SyncDepotPaths = new List<string>();
-					foreach(string SyncPath in SyncPaths)
+					using(RecordCounter Counter = new RecordCounter(Progress, "Filtering files..."))
 					{
-						List<PerforceFileRecord> SyncRecords;
-						if(!Perforce.SyncPreview(SyncPath, PendingChangeNumber, !Context.Options.HasFlag(WorkspaceUpdateOptions.Sync), out SyncRecords, Log))
+						foreach(string SyncPath in SyncPaths)
 						{
-							StatusMessage = String.Format("Couldn't enumerate changes matching {0}.", SyncPath);
-							return WorkspaceUpdateResult.FailedToSync;
-						}
+							List<PerforceFileRecord> SyncRecords = new List<PerforceFileRecord>();
+							if(!Perforce.SyncPreview(SyncPath, PendingChangeNumber, !Context.Options.HasFlag(WorkspaceUpdateOptions.Sync), Record => { SyncRecords.Add(Record); Counter.Increment(); }, Log))
+							{
+								StatusMessage = String.Format("Couldn't enumerate changes matching {0}.", SyncPath);
+								return WorkspaceUpdateResult.FailedToSync;
+							}
 
-						foreach(PerforceFileRecord SyncRecord in SyncRecords)
-						{
-							try
+							foreach(PerforceFileRecord SyncRecord in SyncRecords)
 							{
-								if(!String.IsNullOrEmpty(SyncRecord.ClientPath))
+								try
+								{
+									if(!String.IsNullOrEmpty(SyncRecord.ClientPath))
+									{
+										Path.GetFullPath(SyncRecord.ClientPath);
+									}
+								}
+								catch(PathTooLongException)
 								{
-									Path.GetFullPath(SyncRecord.ClientPath);
+									Log.WriteLine("The local path for {0} exceeds the maximum allowed by Windows. Re-sync your workspace to a directory with a shorter name, or delete the file from the server.", SyncRecord.ClientPath);
+									StatusMessage = "File exceeds maximum path length allowed by Windows.";
+									return WorkspaceUpdateResult.FailedToSync;
 								}
 							}
-							catch(PathTooLongException)
+
+							if(UserFilter != null)
 							{
-								Log.WriteLine("The local path for {0} exceeds the maximum allowed by Windows. Re-sync your workspace to a directory with a shorter name, or delete the file from the server.", SyncRecord.ClientPath);
-								StatusMessage = "File exceeds maximum path length allowed by Windows.";
-								return WorkspaceUpdateResult.FailedToSync;
+								SyncRecords.RemoveAll(x => !String.IsNullOrEmpty(x.ClientPath) && !MatchFilter(Path.GetFullPath(x.ClientPath), UserFilter));
 							}
-						}
 
-						if(UserFilter != null)
-						{
-							SyncRecords.RemoveAll(x => !String.IsNullOrEmpty(x.ClientPath) && !MatchFilter(Path.GetFullPath(x.ClientPath), UserFilter));
-						}
+							SyncDepotPaths.AddRange(SyncRecords.Select(x => x.DepotPath));
 
-						SyncDepotPaths.AddRange(SyncRecords.Select(x => x.DepotPath));
+							List<PerforceFileRecord> OpenRecords;
+							if(!Perforce.GetOpenFiles(SyncPath, out OpenRecords, Log))
+							{
+								StatusMessage = String.Format("Couldn't find open files matching {0}.", SyncPath);
+								return WorkspaceUpdateResult.FailedToSync;
+							}
 
-						List<PerforceFileRecord> OpenRecords;
-						if(!Perforce.GetOpenFiles(SyncPath, out OpenRecords, Log))
-						{
-							StatusMessage = String.Format("Couldn't find open files matching {0}.", SyncPath);
-							return WorkspaceUpdateResult.FailedToSync;
+							// don't force a sync on added files
+							SyncDepotPaths.AddRange(OpenRecords.Where(x => x.Action != "add" && x.Action != "branch" && x.Action != "move/add").Select(x => x.DepotPath));
 						}
-
-						// don't force a sync on added files
-						SyncDepotPaths.AddRange(OpenRecords.Where(x => x.Action != "add" && x.Action != "branch" && x.Action != "move/add").Select(x => x.DepotPath));
 					}
 
 					// Filter out all the binaries that we don't want
