UGS: Add a "select user" dialog, which allows picking the user name to connect to P4 with from a list.

diff --git a/Engine/Source/Programs/UnrealGameSync/UnrealGameSync/Forms/SelectUserWindow.cs b/Engine/Source/Programs/UnrealGameSync/UnrealGameSync/Forms/SelectUserWindow.cs
new file mode 100644
--- /dev/null
+++ b/Engine/Source/Programs/UnrealGameSync/UnrealGameSync/Forms/SelectUserWindow.cs
@@ -0,0 +1,215 @@
+// Copyright 1998-2018 Epic Games, Inc. All Rights Reserved.
+
+using System;
+using System.Collections.Generic;
+using System.ComponentModel;
+using System.Data;
+using System.Drawing;
+using System.IO;
+using System.Linq;
+using System.Text;
+using System.Threading.Tasks;
+using System.Windows.Forms;
+
+namespace UnrealGameSync
+{
+	partial class SelectUserWindow : Form
+	{
+		class EnumerateUsersTask : IPerforceModalTask
+		{
+			public List<PerforceUserRecord> Users;
+
+			public bool Run(PerforceConnection Perforce, TextWriter Log, out string ErrorMessage)
+			{
+				if(!Perforce.FindUsers(out Users, Log))
+				{
+					ErrorMessage = "Unable to enumerate users.";
+					return false;
+				}
+
+				ErrorMessage = null;
+				return true;
+			}
+		}
+
+		private int SelectedUserIndex;
+		private List<PerforceUserRecord> Users;
+		
+		private SelectUserWindow(List<PerforceUserRecord> Users, int SelectedUserIndex)
+		{
+			InitializeComponent();
+
+			this.SelectedUserIndex = SelectedUserIndex;
+			this.Users = Users;
+
+			PopulateList();
+			UpdateOkButton();
+		}
+
+		private void MoveSelection(int Delta)
+		{
+			if(UserListView.Items.Count > 0)
+			{
+				int CurrentIndex = -1;
+				if(UserListView.SelectedIndices.Count > 0)
+				{
+					CurrentIndex = UserListView.SelectedIndices[0];
+				}
+
+				int NextIndex = CurrentIndex + Delta;
+				if(NextIndex < 0)
+				{
+					NextIndex = 0;
+				}
+				else if(NextIndex >= UserListView.Items.Count)
+				{
+					NextIndex = UserListView.Items.Count - 1;
+				}
+
+				if(CurrentIndex != NextIndex)
+				{
+					if(CurrentIndex != -1)
+					{
+						UserListView.Items[CurrentIndex].Selected = false;
+					}
+
+					UserListView.Items[NextIndex].Selected = true;
+					SelectedUserIndex = (int)UserListView.Items[NextIndex].Tag;
+				}
+			}
+		}
+
+		protected override bool ProcessCmdKey(ref Message Msg, Keys KeyData)
+		{
+			if(KeyData == Keys.Up)
+			{
+				MoveSelection(-1);
+				return true;
+			}
+			else if(KeyData == Keys.Down)
+			{
+				MoveSelection(+1);
+				return true;
+			}
+			return base.ProcessCmdKey(ref Msg, KeyData);
+		}
+
+		private bool IncludeInFilter(PerforceUserRecord User, string[] FilterWords)
+		{
+			foreach(string FilterWord in FilterWords)
+			{
+				if(User.Name.IndexOf(FilterWord, StringComparison.OrdinalIgnoreCase) == -1 
+					&& User.FullName.IndexOf(FilterWord, StringComparison.OrdinalIgnoreCase) == -1
+					&& User.Email.IndexOf(FilterWord, StringComparison.OrdinalIgnoreCase) == -1)
+				{
+					return false;
+				}
+			}
+			return true;
+		}
+
+		private void PopulateList()
+		{
+			UserListView.BeginUpdate();
+			UserListView.Items.Clear();
+
+			int SelectedItemIndex = -1;
+
+			string[] Filter = FilterTextBox.Text.Split(new char[]{' '}, StringSplitOptions.RemoveEmptyEntries);
+			for(int Idx = 0; Idx < Users.Count; Idx++)
+			{
+				PerforceUserRecord User = Users[Idx];
+				if(IncludeInFilter(User, Filter))
+				{
+					ListViewItem Item = new ListViewItem(User.Name);
+					Item.SubItems.Add(new ListViewItem.ListViewSubItem(Item, User.FullName));
+					Item.SubItems.Add(new ListViewItem.ListViewSubItem(Item, User.Email));
+					Item.Tag = (int)Idx;
+					UserListView.Items.Add(Item);
+
+					if(SelectedItemIndex == -1 && Idx >= SelectedUserIndex)
+					{
+						SelectedItemIndex = UserListView.Items.Count - 1;
+						Item.Selected = true;
+					}
+				}
+			}
+
+			if(SelectedItemIndex == -1 && UserListView.Items.Count > 0)
+			{
+				SelectedItemIndex = UserListView.Items.Count - 1;
+				UserListView.Items[SelectedItemIndex].Selected = true;
+			}
+
+			if(SelectedItemIndex != -1)
+			{
+				UserListView.EnsureVisible(SelectedItemIndex);
+			}
+
+			UserListView.EndUpdate();
+
+			UpdateOkButton();
+		}
+
+		public static bool ShowModal(IWin32Window Owner, string ServerAndPort, string UserName, TextWriter Log, out string SelectedUserName)
+		{
+			EnumerateUsersTask Task = new EnumerateUsersTask();
+
+			string ErrorMessage;
+			ModalTaskResult Result = PerforceModalTask.Execute(Owner, null, ServerAndPort, UserName, Task, "Finding users", "Finding users, please wait...", Log, out ErrorMessage);
+			if(Result != ModalTaskResult.Succeeded)
+			{
+				if(!String.IsNullOrEmpty(ErrorMessage))
+				{
+					MessageBox.Show(Owner, ErrorMessage);
+				}
+
+				SelectedUserName = null;
+				return false;
+			}
+
+			SelectUserWindow SelectUser = new SelectUserWindow(Task.Users, 0);
+			if(SelectUser.ShowDialog(Owner) == DialogResult.OK)
+			{
+				SelectedUserName = Task.Users[SelectUser.SelectedUserIndex].Name;
+				return true;
+			}
+			else
+			{
+				SelectedUserName = null;
+				return false;
+			}
+		}
+
+		private void FilterTextBox_TextChanged(object sender, EventArgs e)
+		{
+			PopulateList();
+		}
+
+		private void UpdateOkButton()
+		{
+			OkBtn.Enabled = UserListView.SelectedItems.Count > 0;
+		}
+
+		private void OkBtn_Click(object sender, EventArgs e)
+		{
+			if(UserListView.SelectedItems.Count > 0)
+			{
+				SelectedUserIndex = (int)UserListView.SelectedItems[0].Tag;
+				DialogResult = DialogResult.OK;
+				Close();
+			}
+		}
+
+		private void CancelBtn_Click(object sender, EventArgs e)
+		{
+			DialogResult = DialogResult.Cancel;
+			Close();
+		}
+
+		private void UserListView_SelectedIndexChanged(object sender, EventArgs e)
+		{
+			UpdateOkButton();
+		}
+	}
+}
diff --git a/Engine/Source/Programs/UnrealGameSync/UnrealGameSync/Perforce.cs b/Engine/Source/Programs/UnrealGameSync/UnrealGameSync/Perforce.cs
--- a/Engine/Source/Programs/UnrealGameSync/UnrealGameSync/Perforce.cs
+++ b/Engine/Source/Programs/UnrealGameSync/UnrealGameSync/Perforce.cs
@@ -72,6 +72,14 @@ namespace UnrealGameSync
 		}
 	}
 
+	[DebuggerDisplay("{Name}")]
+	class PerforceUserRecord
+	{
+		public string Name;
+		public string FullName;
+		public string Email;
+	}
+
 	[DebuggerDisplay("{DepotFile}")]
 	class PerforceDescribeFileRecord
 	{
@@ -696,6 +704,26 @@ namespace UnrealGameSync
 			}
 			return bResult;
 		}
+		public bool FindUsers(out List<PerforceUserRecord> UserRecords, TextWriter Log)
+		{
+			List<Dictionary<string, string>> TagRecords;
+			if(!RunCommand("users", out TagRecords, CommandOptions.None, Log))
+			{
+				UserRecords = null;
+				return false;
+			}
+
+			UserRecords = new List<PerforceUserRecord>();
+			foreach(Dictionary<string, string> TagRecord in TagRecords)
+			{
+				PerforceUserRecord UserRecord = new PerforceUserRecord();
+				if(TagRecord.TryGetValue("User", out UserRecord.Name) && TagRecord.TryGetValue("FullName", out UserRecord.FullName) && TagRecord.TryGetValue("Email", out UserRecord.Email))
+				{
+					UserRecords.Add(UserRecord);
+				}
+			}
+			return true;
+		}
 
 		public bool Print(string DepotPath, out List<string> Lines, TextWriter Log)
 		{
diff --git a/Engine/Source/Programs/UnrealGameSync/UnrealGameSync/UnrealGameSync.csproj b/Engine/Source/Programs/UnrealGameSync/UnrealGameSync/UnrealGameSync.csproj
--- a/Engine/Source/Programs/UnrealGameSync/UnrealGameSync/UnrealGameSync.csproj
+++ b/Engine/Source/Programs/UnrealGameSync/UnrealGameSync/UnrealGameSync.csproj
@@ -166,6 +166,12 @@
     <Compile Include="Forms\AutomatedSyncWindow.Designer.cs">
       <DependentUpon>AutomatedSyncWindow.cs</DependentUpon>
     </Compile>
+    <Compile Include="Forms\SelectUserWindow.cs">
+      <SubType>Form</SubType>
+    </Compile>
+    <Compile Include="Forms\SelectUserWindow.Designer.cs">
+      <DependentUpon>SelectUserWindow.cs</DependentUpon>
+    </Compile>
     <Compile Include="Forms\CleanWorkspaceWindow.cs">
       <SubType>Form</SubType>
     </Compile>
@@ -362,6 +368,9 @@
     <EmbeddedResource Include="Forms\AutomatedSyncWindow.resx">
       <DependentUpon>AutomatedSyncWindow.cs</DependentUpon>
     </EmbeddedResource>
+    <EmbeddedResource Include="Forms\SelectUserWindow.resx">
+      <DependentUpon>SelectUserWindow.cs</DependentUpon>
+    </EmbeddedResource>
     <EmbeddedResource Include="Forms\CleanWorkspaceWindow.resx">
       <DependentUpon>CleanWorkspaceWindow.cs</DependentUpon>
     </EmbeddedResource>
