UGS: Missing files for supporting "select user" window.

#rnx

diff --git a/Engine/Source/Programs/UnrealGameSync/UnrealGameSync/Forms/AutomatedSyncWindow.cs b/Engine/Source/Programs/UnrealGameSync/UnrealGameSync/Forms/AutomatedSyncWindow.cs
--- a/Engine/Source/Programs/UnrealGameSync/UnrealGameSync/Forms/AutomatedSyncWindow.cs
+++ b/Engine/Source/Programs/UnrealGameSync/UnrealGameSync/Forms/AutomatedSyncWindow.cs
@@ -165,7 +165,7 @@ namespace UnrealGameSync
 		{
 			string NewServerAndPort;
 			string NewUserName;
-			if(ConnectWindow.ShowModal(this, ServerAndPort, UserName, out NewServerAndPort, out NewUserName))
+			if(ConnectWindow.ShowModal(this, ServerAndPort, UserName, Log, out NewServerAndPort, out NewUserName))
 			{
 				ServerAndPort = NewServerAndPort;
 				UserName = NewUserName;
diff --git a/Engine/Source/Programs/UnrealGameSync/UnrealGameSync/Forms/ConnectWindow.cs b/Engine/Source/Programs/UnrealGameSync/UnrealGameSync/Forms/ConnectWindow.cs
--- a/Engine/Source/Programs/UnrealGameSync/UnrealGameSync/Forms/ConnectWindow.cs
+++ b/Engine/Source/Programs/UnrealGameSync/UnrealGameSync/Forms/ConnectWindow.cs
@@ -17,11 +17,14 @@ namespace UnrealGameSync
 	{
 		string ServerAndPort;
 		string UserName;
+		TextWriter Log;
 
-		public ConnectWindow(string ServerAndPort, string UserName)
+		public ConnectWindow(string ServerAndPort, string UserName, TextWriter Log)
 		{
 			InitializeComponent();
 
+			this.Log = Log;
+
 			if(String.IsNullOrEmpty(ServerAndPort))
 			{
 				ServerAndPort = null;
@@ -44,9 +47,9 @@ namespace UnrealGameSync
 			UpdateEnabledControls();
 		}
 
-		public static bool ShowModal(IWin32Window Owner, string ServerAndPort, string UserName, out string NewServerAndPort, out string NewUserName)
+		public static bool ShowModal(IWin32Window Owner, string ServerAndPort, string UserName, TextWriter Log, out string NewServerAndPort, out string NewUserName)
 		{
-			ConnectWindow Connect = new ConnectWindow(ServerAndPort, UserName);
+			ConnectWindow Connect = new ConnectWindow(ServerAndPort, UserName, Log);
 			if(Connect.ShowDialog() == DialogResult.OK)
 			{
 				NewServerAndPort = Connect.ServerAndPort;
@@ -108,5 +111,14 @@ namespace UnrealGameSync
 			UserNameLabel.Enabled = !bUseDefaultSettings;
 			UserNameTextBox.Enabled = !bUseDefaultSettings;
 		}
+
+		private void BrowseUserBtn_Click(object sender, EventArgs e)
+		{
+			string NewUserName;
+			if(SelectUserWindow.ShowModal(this, ServerAndPortTextBox.Text, UserNameTextBox.Text, Log, out NewUserName))
+			{
+				UserNameTextBox.Text = NewUserName;
+			}
+		}
 	}
 }
diff --git a/Engine/Source/Programs/UnrealGameSync/UnrealGameSync/Forms/OpenProjectWindow.cs b/Engine/Source/Programs/UnrealGameSync/UnrealGameSync/Forms/OpenProjectWindow.cs
--- a/Engine/Source/Programs/UnrealGameSync/UnrealGameSync/Forms/OpenProjectWindow.cs
+++ b/Engine/Source/Programs/UnrealGameSync/UnrealGameSync/Forms/OpenProjectWindow.cs
@@ -315,7 +315,7 @@ namespace UnrealGameSync
 		{
 			string NewServerAndPort;
 			string NewUserName;
-			if(ConnectWindow.ShowModal(this, ServerAndPort, UserName, out NewServerAndPort, out NewUserName))
+			if(ConnectWindow.ShowModal(this, ServerAndPort, UserName, Log, out NewServerAndPort, out NewUserName))
 			{
 				ServerAndPort = NewServerAndPort;
 				UserName = NewUserName;
