UGS: Fix "Clean Workspace" not working on workspaces with forward slashes for paths.

diff --git a/Engine/Source/Programs/UnrealGameSync/UnrealGameSync/FindFoldersToCleanTask.cs b/Engine/Source/Programs/UnrealGameSync/UnrealGameSync/FindFoldersToCleanTask.cs
--- a/Engine/Source/Programs/UnrealGameSync/UnrealGameSync/FindFoldersToCleanTask.cs
+++ b/Engine/Source/Programs/UnrealGameSync/UnrealGameSync/FindFoldersToCleanTask.cs
@@ -240,7 +240,7 @@ namespace UnrealGameSync
 			}
 
 			// Get the prefix for any local file
-			string LocalRootPrefix = RootFolderToClean.Directory.FullName + Path.DirectorySeparatorChar;
+			string LocalRootPrefix = RootFolderToClean.Directory.FullName.Replace(Path.AltDirectorySeparatorChar, Path.DirectorySeparatorChar) + Path.DirectorySeparatorChar;
 
 			// Query the have table and build a separate tree from it
 			PerforceHaveFolder RootHaveFolder = new PerforceHaveFolder();
