Fix copyright notices.

diff --git a/Engine/Source/Programs/UnrealGameSync/UnrealGameSync/Forms/SelectUserWindow.cs b/Engine/Source/Programs/UnrealGameSync/UnrealGameSync/Forms/SelectUserWindow.cs
--- a/Engine/Source/Programs/UnrealGameSync/UnrealGameSync/Forms/SelectUserWindow.cs
+++ b/Engine/Source/Programs/UnrealGameSync/UnrealGameSync/Forms/SelectUserWindow.cs
@@ -1,4 +1,4 @@
-// Copyright 1998-2018 Epic Games, Inc. All Rights Reserved.
+// Copyright 1998-2019 Epic Games, Inc. All Rights Reserved.
 
 using System;
 using System.Collections.Generic;
