UGS: Stop writing out ObjectVersion.cpp and Version.h for branches that don't need it.

diff --git a/Engine/Programs/UnrealGameSync/UnrealGameSync.ini b/Engine/Programs/UnrealGameSync/UnrealGameSync.ini
--- a/Engine/Programs/UnrealGameSync/UnrealGameSync.ini
+++ b/Engine/Programs/UnrealGameSync/UnrealGameSync.ini
@@ -1,5 +1,6 @@
 [Options]
 UseFastModularVersioning=1
+UseFastModularVersioningV2=1
 DebugGameHasSeparateExecutable=1
 +ExcludeChanges=^Rebuild lightmaps\.$
 +ExcludeChanges=^Rebuild lightmaps & texture streaming\.$
diff --git a/Engine/Source/Programs/UnrealGameSync/UnrealGameSync/Workspace.cs b/Engine/Source/Programs/UnrealGameSync/UnrealGameSync/Workspace.cs
--- a/Engine/Source/Programs/UnrealGameSync/UnrealGameSync/Workspace.cs
+++ b/Engine/Source/Programs/UnrealGameSync/UnrealGameSync/Workspace.cs
@@ -632,7 +632,24 @@ namespace UnrealGameSync
 						}
 
 						// Update the version files
-						if(ProjectConfigFile.GetValue("Options.UseFastModularVersioning", false))
+						if(ProjectConfigFile.GetValue("Options.UseFastModularVersioningV2", false))
+						{
+							bool bIsEpicInternal;
+							Perforce.FileExists(ClientRootPath + "/Engine/Build/NotForLicensees/EpicInternal.txt", out bIsEpicInternal, Log);
+
+							Dictionary<string, string> BuildVersionStrings = new Dictionary<string,string>();
+							BuildVersionStrings["\"Changelist\":"] = String.Format(" {0},", PendingChangeNumber);
+							BuildVersionStrings["\"CompatibleChangelist\":"] = String.Format(" {0},", VersionChangeNumber);
+							BuildVersionStrings["\"BranchName\":"] = String.Format(" \"{0}\"", BranchOrStreamName.Replace('/', '+'));
+							BuildVersionStrings["\"IsPromotedBuild\":"] = " 0,";
+							BuildVersionStrings["\"IsLicenseeVersion\":"] = bIsEpicInternal? "0," : "1,";
+							if(!UpdateVersionFile(ClientRootPath + BuildVersionFileName, BuildVersionStrings, PendingChangeNumber))
+							{
+								StatusMessage = String.Format("Failed to update {0}.", BuildVersionFileName);
+								return WorkspaceUpdateResult.FailedToSync;
+							}
+						}
+						else if(ProjectConfigFile.GetValue("Options.UseFastModularVersioning", false))
 						{
 							bool bIsEpicInternal;
 							Perforce.FileExists(ClientRootPath + "/Engine/Build/NotForLicensees/EpicInternal.txt", out bIsEpicInternal, Log);
