UGS: Always allow version.h and objectversion.cpp to be clobbered, since we don't update them any more.

diff --git a/Engine/Source/Programs/UnrealGameSync/UnrealGameSync/Workspace.cs b/Engine/Source/Programs/UnrealGameSync/UnrealGameSync/Workspace.cs
--- a/Engine/Source/Programs/UnrealGameSync/UnrealGameSync/Workspace.cs
+++ b/Engine/Source/Programs/UnrealGameSync/UnrealGameSync/Workspace.cs
@@ -122,6 +122,9 @@ namespace UnrealGameSync
 		const string VersionHeaderFileName = "/Engine/Source/Runtime/Launch/Resources/Version.h";
 		const string ObjectVersionFileName = "/Engine/Source/Runtime/Core/Private/UObject/ObjectVersion.cpp";
 
+		readonly string LocalVersionHeaderFileName = VersionHeaderFileName.Replace('/', '\\');
+		readonly string LocalObjectVersionFileName = ObjectVersionFileName.Replace('/', '\\');
+
 		public readonly PerforceConnection Perforce;
 		public readonly string LocalRootPath;
 		public readonly string SelectedLocalFileName;
@@ -512,8 +515,11 @@ namespace UnrealGameSync
 					// Filter out all the binaries that we don't want
 					FileFilter Filter = new FileFilter(FileFilterType.Include);
 					Filter.Exclude("..." + BuildVersionFileName);
-					Filter.Exclude("..." + VersionHeaderFileName);
-					Filter.Exclude("..." + ObjectVersionFileName);
+					if(!ProjectConfigFile.GetValue("Options.UseFastModularVersioningV2", false))
+					{
+						Filter.Exclude("..." + VersionHeaderFileName);
+						Filter.Exclude("..." + ObjectVersionFileName);
+					}
 					if(Context.Options.HasFlag(WorkspaceUpdateOptions.ContentOnly))
 					{
 						Filter.Exclude("*.usf");
@@ -548,6 +554,11 @@ namespace UnrealGameSync
 							if(!Context.ClobberFiles.ContainsKey(TamperedFile))
 							{
 								Context.ClobberFiles[TamperedFile] = true;
+								if(TamperedFile.EndsWith(LocalObjectVersionFileName, StringComparison.OrdinalIgnoreCase) || TamperedFile.EndsWith(LocalVersionHeaderFileName, StringComparison.OrdinalIgnoreCase))
+								{
+									// Hack for UseFastModularVersioningV2; we don't need to update these files any more.
+									continue;
+								}
 								NumNewFilesToClobber++;
 							}
 						}
