UGS: Fix changes being greyed out when using bisect mode.

diff --git a/Engine/Source/Programs/UnrealGameSync/UnrealGameSync/Controls/WorkspaceControl.cs b/Engine/Source/Programs/UnrealGameSync/UnrealGameSync/Controls/WorkspaceControl.cs
--- a/Engine/Source/Programs/UnrealGameSync/UnrealGameSync/Controls/WorkspaceControl.cs
+++ b/Engine/Source/Programs/UnrealGameSync/UnrealGameSync/Controls/WorkspaceControl.cs
@@ -1724,7 +1724,7 @@ namespace UnrealGameSync
 			Color TextColor = (bAllowSync || Change.Number == Workspace.PendingChangeNumber || Change.Number == Workspace.CurrentChangeNumber || (WorkspaceSettings != null && WorkspaceSettings.AdditionalChangeNumbers.Contains(Change.Number)))? SystemColors.WindowText : Blend(SystemColors.Window, SystemColors.WindowText, 0.25f);
 
 			const int FadeRange = 6;
-			if(e.ItemIndex >= BuildList.Items.Count - FadeRange && NumChanges >= PerforceMonitor.CurrentMaxChanges)
+			if(e.ItemIndex >= BuildList.Items.Count - FadeRange && NumChanges >= PerforceMonitor.CurrentMaxChanges && !IsBisectModeEnabled())
 			{
 				float Opacity = (float)(BuildList.Items.Count - e.ItemIndex - 0.9f) / FadeRange;
 				BadgeAlpha = (int)(BadgeAlpha * Opacity);
