UGS: Make the recononciling of a workspace after changing the sync filter a separate operation to syncing any new files, and remember the workspace contents between runs to avoid querying it twice. This saves a lot of time on slow connections.

diff --git a/Engine/Source/Programs/UnrealGameSync/UnrealGameSync/Workspace.cs b/Engine/Source/Programs/UnrealGameSync/UnrealGameSync/Workspace.cs
--- a/Engine/Source/Programs/UnrealGameSync/UnrealGameSync/Workspace.cs
+++ b/Engine/Source/Programs/UnrealGameSync/UnrealGameSync/Workspace.cs
@@ -60,6 +60,7 @@ namespace UnrealGameSync
 		public HashSet<Guid> CustomBuildSteps;
 		public Dictionary<string, string> Variables;
 		public PerforceSyncOptions PerforceSyncOptions;
+		public List<PerforceFileRecord> HaveFiles; // Cached when sync filter has changed
 
 		public WorkspaceUpdateContext(int InChangeNumber, WorkspaceUpdateOptions InOptions, string[] InSyncFilter, Dictionary<Guid, ConfigObject> InDefaultBuildSteps, List<ConfigObject> InUserBuildSteps, HashSet<Guid> InCustomBuildSteps, Dictionary<string, string> InVariables)
 		{
@@ -382,20 +383,24 @@ namespace UnrealGameSync
 					}
 
 					// If the hash differs, enumerate everything in the workspace to find what needs to be removed
-					List<string> RemoveDepotPaths = new List<string>();
 					if (NextSyncFilterHash != CurrentSyncFilterHash)
 					{
 						Log.WriteLine("Filter has changed ({0} -> {1}); finding files in workspace that need to be removed.", (String.IsNullOrEmpty(CurrentSyncFilterHash))? "None" : CurrentSyncFilterHash, NextSyncFilterHash);
 
 						// Find all the files that are in this workspace
-						List<PerforceFileRecord> HaveFiles = new List<PerforceFileRecord>();
-						using(RecordCounter HaveCounter = new RecordCounter(Progress, "Querying workspace..."))
+						List<PerforceFileRecord> HaveFiles = Context.HaveFiles;
+						if(HaveFiles == null)
 						{
-							if(!Perforce.Have("//...", Record => { HaveFiles.Add(Record); HaveCounter.Increment(); }, Log))
+							HaveFiles = new List<PerforceFileRecord>();
+							using(RecordCounter HaveCounter = new RecordCounter(Progress, "Sync filter changed; checking workspace..."))
 							{
-								StatusMessage = "Unable to query files.";
-								return WorkspaceUpdateResult.FailedToSync;
+								if(!Perforce.Have("//...", Record => { HaveFiles.Add(Record); HaveCounter.Increment(); }, Log))
+								{
+									StatusMessage = "Unable to query files.";
+									return WorkspaceUpdateResult.FailedToSync;
+								}
 							}
+							Context.HaveFiles = HaveFiles;
 						}
 
 						// Build a filter for the current sync paths
@@ -412,6 +417,7 @@ namespace UnrealGameSync
 						}
 
 						// Remove all the files that are not included by the filter
+						List<string> RemoveDepotPaths = new List<string>();
 						foreach(PerforceFileRecord HaveFile in HaveFiles)
 						{
 							try
@@ -459,8 +465,29 @@ namespace UnrealGameSync
 							RemoveDepotPaths.RemoveAll(x => !Context.DeleteFiles[x]);
 						}
 
-						// Clear the current sync filter hash. If the sync is canceled, we'll be in an indeterminate state, and we should always clean next time round.
-						CurrentSyncFilterHash = "INVALID";
+						// Actually delete any files that we don't want
+						if(RemoveDepotPaths.Count > 0)
+						{
+							// Clear the current sync filter hash. If the sync is canceled, we'll be in an indeterminate state, and we should always clean next time round.
+							CurrentSyncFilterHash = "INVALID";
+
+							// Find all the depot paths that will be synced
+							HashSet<string> RemainingDepotPathsToRemove = new HashSet<string>(StringComparer.InvariantCultureIgnoreCase);
+							RemainingDepotPathsToRemove.UnionWith(RemoveDepotPaths);
+
+							// Build the list of revisions to sync
+							List<string> RevisionsToRemove = new List<string>();
+							RevisionsToRemove.AddRange(RemoveDepotPaths.Select(x => String.Format("{0}#0", x)));
+
+							WorkspaceUpdateResult RemoveResult = SyncFileRevisions("Removing files...", Context, RevisionsToRemove, RemainingDepotPathsToRemove, out StatusMessage);
+							if(RemoveResult != WorkspaceUpdateResult.Success)
+							{
+								return RemoveResult;
+							}
+						}
+
+						// Update the sync filter hash. We've removed any files we need to at this point.
+						CurrentSyncFilterHash = NextSyncFilterHash;
 					}
 
 					// Find all the server changes, and anything that's opened for edit locally. We need to sync files we have open to schedule a resolve.
@@ -529,57 +556,18 @@ namespace UnrealGameSync
 
 					// Find all the depot paths that will be synced
 					HashSet<string> RemainingDepotPaths = new HashSet<string>(StringComparer.InvariantCultureIgnoreCase);
-					RemainingDepotPaths.UnionWith(RemoveDepotPaths);
 					RemainingDepotPaths.UnionWith(SyncDepotPaths);
 
 					// Build the list of revisions to sync
 					List<string> SyncRevisions = new List<string>();
-					SyncRevisions.AddRange(RemoveDepotPaths.Select(x => String.Format("{0}#0", x)));
 					SyncRevisions.AddRange(SyncDepotPaths.Select(x => String.Format("{0}@{1}", x, PendingChangeNumber)));
 
-					// Sync them all
-					List<string> TamperedFiles = new List<string>();
-					if(!Perforce.Sync(SyncRevisions, Record => UpdateSyncProgress(Record, RemainingDepotPaths, SyncRevisions.Count), TamperedFiles, false, Context.PerforceSyncOptions, Log))
+					WorkspaceUpdateResult SyncResult = SyncFileRevisions("Syncing files...", Context, SyncRevisions, RemainingDepotPaths, out StatusMessage);
+					if(SyncResult != WorkspaceUpdateResult.Success)
 					{
-						StatusMessage = "Aborted sync due to errors.";
-						return WorkspaceUpdateResult.FailedToSync;
+						return SyncResult;
 					}
 
-					// If any files need to be clobbered, defer to the main thread to figure out which ones
-					if(TamperedFiles.Count > 0)
-					{
-						int NumNewFilesToClobber = 0;
-						foreach(string TamperedFile in TamperedFiles)
-						{
-							if(!Context.ClobberFiles.ContainsKey(TamperedFile))
-							{
-								Context.ClobberFiles[TamperedFile] = true;
-								if(TamperedFile.EndsWith(LocalObjectVersionFileName, StringComparison.OrdinalIgnoreCase) || TamperedFile.EndsWith(LocalVersionHeaderFileName, StringComparison.OrdinalIgnoreCase))
-								{
-									// Hack for UseFastModularVersioningV2; we don't need to update these files any more.
-									continue;
-								}
-								NumNewFilesToClobber++;
-							}
-						}
-						if(NumNewFilesToClobber > 0)
-						{
-							StatusMessage = String.Format("Cancelled sync after checking files to clobber ({0} new files).", NumNewFilesToClobber);
-							return WorkspaceUpdateResult.FilesToClobber;
-						}
-						foreach(string TamperedFile in TamperedFiles)
-						{
-							if(Context.ClobberFiles[TamperedFile] && !Perforce.ForceSync(TamperedFile, PendingChangeNumber, Log))
-							{
-								StatusMessage = String.Format("Couldn't sync {0}.", TamperedFile);
-								return WorkspaceUpdateResult.FailedToSync;
-							}
-						}
-					}
-
-					// Update the sync filter hash. We've removed any files we need to at this point.
-					CurrentSyncFilterHash = NextSyncFilterHash;
-
 					int VersionChangeNumber = -1;
 					if(Context.Options.HasFlag(WorkspaceUpdateOptions.Sync) && !Context.Options.HasFlag(WorkspaceUpdateOptions.UpdateFilter))
 					{
@@ -1157,6 +1145,53 @@ namespace UnrealGameSync
 			return bMatch;
 		}
 
+		WorkspaceUpdateResult SyncFileRevisions(string Prefix, WorkspaceUpdateContext Context, List<string> SyncRevisions, HashSet<string> RemainingDepotPaths, out string StatusMessage)
+		{
+			// Sync them all
+			List<string> TamperedFiles = new List<string>();
+			if(!Perforce.Sync(SyncRevisions, Record => UpdateSyncProgress(Prefix, Record, RemainingDepotPaths, SyncRevisions.Count), TamperedFiles, false, Context.PerforceSyncOptions, Log))
+			{
+				StatusMessage = "Aborted sync due to errors.";
+				return WorkspaceUpdateResult.FailedToSync;
+			}
+
+			// If any files need to be clobbered, defer to the main thread to figure out which ones
+			if(TamperedFiles.Count > 0)
+			{
+				int NumNewFilesToClobber = 0;
+				foreach(string TamperedFile in TamperedFiles)
+				{
+					if(!Context.ClobberFiles.ContainsKey(TamperedFile))
+					{
+						Context.ClobberFiles[TamperedFile] = true;
+						if(TamperedFile.EndsWith(LocalObjectVersionFileName, StringComparison.OrdinalIgnoreCase) || TamperedFile.EndsWith(LocalVersionHeaderFileName, StringComparison.OrdinalIgnoreCase))
+						{
+							// Hack for UseFastModularVersioningV2; we don't need to update these files any more.
+							continue;
+						}
+						NumNewFilesToClobber++;
+					}
+				}
+				if(NumNewFilesToClobber > 0)
+				{
+					StatusMessage = String.Format("Cancelled sync after checking files to clobber ({0} new files).", NumNewFilesToClobber);
+					return WorkspaceUpdateResult.FilesToClobber;
+				}
+				foreach(string TamperedFile in TamperedFiles)
+				{
+					if(Context.ClobberFiles[TamperedFile] && !Perforce.ForceSync(TamperedFile, PendingChangeNumber, Log))
+					{
+						StatusMessage = String.Format("Couldn't sync {0}.", TamperedFile);
+						return WorkspaceUpdateResult.FailedToSync;
+					}
+				}
+			}
+
+			// All succeeded
+			StatusMessage = null;
+			return WorkspaceUpdateResult.Success;
+		}
+
 		static ConfigFile ReadProjectConfigFile(string LocalRootPath, string SelectedLocalFileName, TextWriter Log)
 		{
 			// Find the valid config file paths
@@ -1242,11 +1277,11 @@ namespace UnrealGameSync
 			return true;
 		}
 
-		void UpdateSyncProgress(PerforceFileRecord Record, HashSet<string> RemainingFiles, int NumFiles)
+		void UpdateSyncProgress(string Prefix, PerforceFileRecord Record, HashSet<string> RemainingFiles, int NumFiles)
 		{
 			RemainingFiles.Remove(Record.DepotPath);
 
-			string Message = String.Format("Syncing files... ({0}/{1})", NumFiles - RemainingFiles.Count, NumFiles);
+			string Message = String.Format("{0} ({1}/{2})", Prefix, NumFiles - RemainingFiles.Count, NumFiles);
 			float Fraction = Math.Min((float)(NumFiles - RemainingFiles.Count) / (float)NumFiles, 1.0f);
 			Progress.Set(Message, Fraction);
 
