UGS: Fix DPI scaling in "Clean Workspace" window.

diff --git a/Engine/Source/Programs/UnrealGameSync/UnrealGameSync/Forms/CleanWorkspaceWindow.cs b/Engine/Source/Programs/UnrealGameSync/UnrealGameSync/Forms/CleanWorkspaceWindow.cs
--- a/Engine/Source/Programs/UnrealGameSync/UnrealGameSync/Forms/CleanWorkspaceWindow.cs
+++ b/Engine/Source/Programs/UnrealGameSync/UnrealGameSync/Forms/CleanWorkspaceWindow.cs
@@ -153,7 +153,25 @@ namespace UnrealGameSync
 
 			Icon[] Icons = new Icon[]{ Icon.FromHandle(FolderIconPtr), Icon.FromHandle(FileIconPtr) };
 
-			Bitmap TypeImageListBitmap = new Bitmap(Icons.Length * 3 * 32, 16, System.Drawing.Imaging.PixelFormat.Format32bppArgb);
+			Size LargestIconSize = Size.Empty;
+			foreach(Icon Icon in Icons)
+			{
+				LargestIconSize = new Size(Math.Max(LargestIconSize.Width, Icon.Width), Math.Max(LargestIconSize.Height, Icon.Height));
+			}
+
+			Size LargestCheckBoxSize = Size.Empty;
+			using(Graphics Graphics = Graphics.FromHwnd(IntPtr.Zero))
+			{
+				foreach(CheckBoxState State in CheckBoxStates)
+				{
+					Size CheckBoxSize = CheckBoxRenderer.GetGlyphSize(Graphics, State);
+					LargestCheckBoxSize = new Size(Math.Max(LargestCheckBoxSize.Width, CheckBoxSize.Width), Math.Max(LargestCheckBoxSize.Height, CheckBoxSize.Height));
+				}
+			}
+
+			Size ImageSize = new Size(LargestCheckBoxSize.Width + LargestIconSize.Width, Math.Max(LargestIconSize.Height, LargestCheckBoxSize.Height));
+
+			Bitmap TypeImageListBitmap = new Bitmap(Icons.Length * 3 * ImageSize.Width, ImageSize.Height, System.Drawing.Imaging.PixelFormat.Format32bppArgb);
 			using(Graphics Graphics = Graphics.FromImage(TypeImageListBitmap))
 			{
 				int MinX = 0;
@@ -162,15 +180,18 @@ namespace UnrealGameSync
 					for(int StateIdx = 0; StateIdx < 3; StateIdx++)
 					{
 						Size CheckBoxSize = CheckBoxRenderer.GetGlyphSize(Graphics, CheckBoxStates[StateIdx]);
-						CheckBoxRenderer.DrawCheckBox(Graphics, new Point(MinX + (16 - CheckBoxSize.Width) / 2, (16 - CheckBoxSize.Height) / 2), CheckBoxStates[StateIdx]);
-						Graphics.DrawIcon(Icons[IconIdx], MinX + 16, 0);
-						MinX += 32;
+						CheckBoxRenderer.DrawCheckBox(Graphics, new Point(MinX + (LargestCheckBoxSize.Width - CheckBoxSize.Width) / 2, (LargestCheckBoxSize.Height - CheckBoxSize.Height) / 2), CheckBoxStates[StateIdx]);
+
+						Size IconSize = Icons[IconIdx].Size;
+						Graphics.DrawIcon(Icons[IconIdx], MinX + LargestCheckBoxSize.Width + (LargestIconSize.Width - IconSize.Width) / 2, (LargestIconSize.Height - IconSize.Height) / 2);
+
+						MinX += ImageSize.Width;
 					}
 				}
 			}
 
 			ImageList TypeImageList = new ImageList();
-			TypeImageList.ImageSize = new Size(32, 16);
+			TypeImageList.ImageSize = ImageSize;
 			TypeImageList.ColorDepth = ColorDepth.Depth32Bit;
 			TypeImageList.Images.AddStrip(TypeImageListBitmap);
 			TreeView.ImageList = TypeImageList;
