Fix up LatestData model to have the same number of M's in "Comment" as the one in EventMonitor

diff --git a/Engine/Source/Programs/UnrealGameSync/UnrealGameSyncMetadataServer/Connectors/SqlConnector.cs b/Engine/Source/Programs/UnrealGameSync/UnrealGameSyncMetadataServer/Connectors/SqlConnector.cs
--- a/Engine/Source/Programs/UnrealGameSync/UnrealGameSyncMetadataServer/Connectors/SqlConnector.cs
+++ b/Engine/Source/Programs/UnrealGameSync/UnrealGameSyncMetadataServer/Connectors/SqlConnector.cs
@@ -61,7 +61,7 @@ namespace UnrealGameSyncMetadataServer.Connectors
 					}
 				}
 			}
-			return new LatestData { LastBuildId = LastBuildId, LastCommmentId = LastCommentId, LastEventId = LastEventId };
+			return new LatestData { LastBuildId = LastBuildId, LastCommentId = LastCommentId, LastEventId = LastEventId };
 		}
 		public static List<EventData> GetUserVotes(string Project, long LastEventId)
 		{
diff --git a/Engine/Source/Programs/UnrealGameSync/UnrealGameSyncMetadataServer/Controllers/CISController.cs b/Engine/Source/Programs/UnrealGameSync/UnrealGameSyncMetadataServer/Controllers/CISController.cs
--- a/Engine/Source/Programs/UnrealGameSync/UnrealGameSyncMetadataServer/Controllers/CISController.cs
+++ b/Engine/Source/Programs/UnrealGameSync/UnrealGameSyncMetadataServer/Controllers/CISController.cs
@@ -14,7 +14,7 @@ namespace UnrealGameSyncMetadataServer.Controllers
 		public long[] Get(string Project = null)
 		{
 			LatestData LatestIDs = SqlConnector.GetLastIds(Project);
-			return new long[] { LatestIDs.LastEventId, LatestIDs.LastCommmentId, LatestIDs.LastBuildId };
+			return new long[] { LatestIDs.LastEventId, LatestIDs.LastCommentId, LatestIDs.LastBuildId };
 		}
 		public List<BuildData> Get(string Project, long LastBuildId)
 		{
diff --git a/Engine/Source/Programs/UnrealGameSync/UnrealGameSyncMetadataServer/Models/LatestData.cs b/Engine/Source/Programs/UnrealGameSync/UnrealGameSyncMetadataServer/Models/LatestData.cs
--- a/Engine/Source/Programs/UnrealGameSync/UnrealGameSyncMetadataServer/Models/LatestData.cs
+++ b/Engine/Source/Programs/UnrealGameSync/UnrealGameSyncMetadataServer/Models/LatestData.cs
@@ -10,7 +10,7 @@ namespace UnrealGameSyncMetadataServer.Models
 	public class LatestData
 	{
 		public long LastEventId { get; set; }
-		public long LastCommmentId { get; set; }
+		public long LastCommentId { get; set; }
 		public long LastBuildId { get; set; }
 	}
 }
\ No newline at end of file
