UGS: Initial support for build health notification feature.

* Build health issues can be posted to the metadata server, and have multiple build jobs in different branches associated with them.
* Users can be added as watchers to particular build health issues via the metadata server (intended to be done by some automated process), which causes a notification to be shown on the client allowing them to acknowledge or dismiss it.
* Users can be nominated as the person that needs to fix a particular issue, which suppresses the notification for any watchers.
* Users can close a build health issue by clicking on the red "Close Build Issue" button in the UGS status panel and entering the CL that contains the fix.
* Anyone can subscribe to notifications that are not acknowledged after a certain period of time.
* A "Build Health" button in the status panel allows reviewing any open build health issues at a time.

To enable for a project, set the "ShowBuildHealth" option under the [Options] section to 1 in the UGS config file, and post build issues using the newly added MetadataTool.

diff --git a/Engine/Source/Programs/UnrealGameSync/MetadataTool/App.config b/Engine/Source/Programs/UnrealGameSync/MetadataTool/App.config
new file mode 100644
--- /dev/null
+++ b/Engine/Source/Programs/UnrealGameSync/MetadataTool/App.config
@@ -0,0 +1,6 @@
+<?xml version="1.0" encoding="utf-8"?>
+<configuration>
+    <startup> 
+        <supportedRuntime version="v4.0" sku=".NETFramework,Version=v4.6.2"/>
+    </startup>
+</configuration>
diff --git a/Engine/Source/Programs/UnrealGameSync/MetadataTool/MetadataTool.csproj b/Engine/Source/Programs/UnrealGameSync/MetadataTool/MetadataTool.csproj
new file mode 100644
--- /dev/null
+++ b/Engine/Source/Programs/UnrealGameSync/MetadataTool/MetadataTool.csproj
@@ -0,0 +1,61 @@
+﻿<?xml version="1.0" encoding="utf-8"?>
+<Project ToolsVersion="15.0" xmlns="http://schemas.microsoft.com/developer/msbuild/2003">
+  <Import Project="$(MSBuildExtensionsPath)\$(MSBuildToolsVersion)\Microsoft.Common.props" Condition="Exists('$(MSBuildExtensionsPath)\$(MSBuildToolsVersion)\Microsoft.Common.props')" />
+  <PropertyGroup>
+    <Configuration Condition=" '$(Configuration)' == '' ">Debug</Configuration>
+    <Platform Condition=" '$(Platform)' == '' ">AnyCPU</Platform>
+    <ProjectGuid>{1B562929-0D0C-4CBE-8C43-3683F54D0414}</ProjectGuid>
+    <OutputType>Exe</OutputType>
+    <RootNamespace>MetadataTool</RootNamespace>
+    <AssemblyName>MetadataTool</AssemblyName>
+    <TargetFrameworkVersion>v4.6.2</TargetFrameworkVersion>
+    <FileAlignment>512</FileAlignment>
+    <AutoGenerateBindingRedirects>true</AutoGenerateBindingRedirects>
+    <Deterministic>true</Deterministic>
+    <TargetFrameworkProfile />
+  </PropertyGroup>
+  <PropertyGroup Condition=" '$(Configuration)|$(Platform)' == 'Debug|AnyCPU' ">
+    <PlatformTarget>AnyCPU</PlatformTarget>
+    <DebugSymbols>true</DebugSymbols>
+    <DebugType>full</DebugType>
+    <Optimize>false</Optimize>
+    <OutputPath>bin\Debug\</OutputPath>
+    <DefineConstants>DEBUG;TRACE</DefineConstants>
+    <ErrorReport>prompt</ErrorReport>
+    <WarningLevel>4</WarningLevel>
+  </PropertyGroup>
+  <PropertyGroup Condition=" '$(Configuration)|$(Platform)' == 'Release|AnyCPU' ">
+    <PlatformTarget>AnyCPU</PlatformTarget>
+    <DebugType>pdbonly</DebugType>
+    <Optimize>true</Optimize>
+    <OutputPath>bin\Release\</OutputPath>
+    <DefineConstants>TRACE</DefineConstants>
+    <ErrorReport>prompt</ErrorReport>
+    <WarningLevel>4</WarningLevel>
+  </PropertyGroup>
+  <ItemGroup>
+    <Reference Include="System" />
+    <Reference Include="System.Core" />
+    <Reference Include="System.Web.Extensions" />
+    <Reference Include="System.Xml.Linq" />
+    <Reference Include="System.Data.DataSetExtensions" />
+    <Reference Include="Microsoft.CSharp" />
+    <Reference Include="System.Data" />
+    <Reference Include="System.Net.Http" />
+    <Reference Include="System.Xml" />
+  </ItemGroup>
+  <ItemGroup>
+    <Compile Include="Program.cs" />
+    <Compile Include="Properties\AssemblyInfo.cs" />
+  </ItemGroup>
+  <ItemGroup>
+    <None Include="App.config" />
+  </ItemGroup>
+  <ItemGroup>
+    <ProjectReference Include="..\..\DotNETCommon\DotNETUtilities\DotNETUtilities.csproj">
+      <Project>{5d7d66e8-8c76-4af9-b3ec-2ef03421d730}</Project>
+      <Name>DotNETUtilities</Name>
+    </ProjectReference>
+  </ItemGroup>
+  <Import Project="$(MSBuildToolsPath)\Microsoft.CSharp.targets" />
+</Project>
\ No newline at end of file
diff --git a/Engine/Source/Programs/UnrealGameSync/MetadataTool/Program.cs b/Engine/Source/Programs/UnrealGameSync/MetadataTool/Program.cs
new file mode 100644
--- /dev/null
+++ b/Engine/Source/Programs/UnrealGameSync/MetadataTool/Program.cs
@@ -0,0 +1,456 @@
+// Copyright 1998-2019 Epic Games, Inc. All Rights Reserved.
+
+using System;
+using System.Collections.Generic;
+using System.Linq;
+using System.Text;
+using System.Reflection;
+using System.Threading.Tasks;
+using Tools.DotNETCommon;
+using System.Net;
+using System.IO;
+using System.Web.Script.Serialization;
+
+namespace MetadataTool
+{
+	[AttributeUsage(AttributeTargets.Field)]
+	class OptionalAttribute : Attribute
+	{
+	}
+
+	class Program
+	{
+		class UserErrorException : Exception
+		{
+			public UserErrorException(string Message)
+				: base(Message)
+			{
+			}
+		}
+
+		class MissingArgumentException : UserErrorException
+		{
+			public MissingArgumentException(string Message)
+				: base(Message)
+			{
+			}
+		}
+
+		class Command
+		{
+			public string Name;
+			public string Verb;
+			public string Resource;
+			public Type BodyType;
+			public string[] OptionalParams;
+
+			public Command(string Name, string Verb, string Resource, Type BodyType = null, string[] OptionalParams = null)
+			{
+				this.Name = Name;
+				this.Verb = Verb;
+				this.Resource = Resource;
+				this.BodyType = BodyType;
+				this.OptionalParams = OptionalParams;
+			}
+		}
+
+		#pragma warning disable CS0649
+
+		enum Outcome
+		{
+			Success = 1,
+			Error = 2,
+			Warning = 3,
+		}
+
+		class AddIssueBody
+		{
+			public string Project;
+			public string Summary;
+		}
+
+		class UpdateBuildBody
+		{
+			[Optional] public int? Outcome;
+		}
+
+		class UpdateIssueBody
+		{
+			[Optional] public bool? Acknowledged;
+			[Optional] public int? FixChange;
+			[Optional] public string Owner;
+			[Optional] public string NominatedBy;
+			[Optional] public bool? Resolved;
+			[Optional] public string Summary;
+			[Optional] public string Url;
+		}
+
+		class AddBuildBody
+		{
+			public string Name;
+			public string Url;
+			public string Stream;
+			public int Change;
+			public Outcome Outcome;
+		}
+
+		class WatcherBody
+		{
+			public string UserName;
+		}
+
+		#pragma warning restore CS0649
+
+		static Command[] Commands =
+		{
+			// Issues
+			new Command("add-issue", "POST", "/api/issues", typeof(AddIssueBody)),
+			new Command("get-issue", "GET", "/api/issues/{id}"),
+			new Command("get-issues", "GET", "/api/issues", OptionalParams: new string[]{ "user" }),
+			new Command("update-issue", "PUT", "/api/issues/{Id}", typeof(UpdateIssueBody)),
+
+			// Builds
+			new Command("add-build", "POST", "/api/issues/{Issue}/builds", typeof(AddBuildBody)),
+			new Command("get-builds", "GET", "/api/issues/{Issue}/builds"),
+
+			// Individual builds
+			new Command("get-build", "GET", "/api/issuebuilds/{Build}"),
+			new Command("update-build", "PUT", "/api/issuebuilds/{Build}", typeof(UpdateBuildBody)),
+
+			// Watchers
+			new Command("get-watchers", "GET", "/api/issues/{Issue}/watchers"),
+			new Command("add-watcher", "POST", "/api/issues/{Issue}/watchers", typeof(WatcherBody)),
+			new Command("remove-watcher", "DELETE", "/api/issues/{Issue}/watchers", typeof(WatcherBody))
+		};
+
+		static int Main(string[] Args)
+		{
+			try
+			{
+				bool bResult = InnerMain(Args);
+				return bResult? 0 : 1;
+			}
+			catch(UserErrorException Ex)
+			{
+				Console.WriteLine("{0}", Ex.Message);
+				return 1;
+			}
+		}
+
+		static bool InnerMain(string[] Args)
+		{
+			List<string> RemainingArgs = new List<string>(Args);
+
+			// Get the server URL
+			string ServerUrl = ParseOptionalParam(RemainingArgs, "Server") ?? Environment.GetEnvironmentVariable("UGS_METADATA_SERVER_URL");
+			if(RemainingArgs.Count == 0)
+			{
+				PrintCommands();
+				return true;
+			}
+
+			// Register all the commands
+			Command Command = Commands.FirstOrDefault(x => String.Compare(x.Name, RemainingArgs[0], StringComparison.OrdinalIgnoreCase) == 0);
+			if(Command == null)
+			{
+				Console.WriteLine("Unknown command '{0}'", RemainingArgs[0]);
+				return false;
+			}
+			RemainingArgs.RemoveAt(0);
+
+			// Handle help for this command
+			if(ParseOption(RemainingArgs, "-Help"))
+			{
+				PrintCommandHelp(Command);
+				return false;
+			}
+
+			// Make sure we have the server parameter
+			if(String.IsNullOrEmpty(ServerUrl))
+			{
+				Console.WriteLine("Missing -Server=... argument.");
+				Console.WriteLine();
+				PrintCommandHelp(Command);
+				return false;
+			}
+
+			// Try to execute the command
+			try
+			{
+				// Create the command url
+				StringBuilder CommandUrl = new StringBuilder();
+				CommandUrl.AppendFormat("{0}{1}", ServerUrl, Command.Resource);
+
+				// Replace all the arguments that are embedded into the resource name
+				foreach(string ResourceArg in ParseArgumentNamesFromResource(Command.Resource))
+				{
+					string Value = ParseRequiredParam(RemainingArgs, ResourceArg);
+					CommandUrl.Replace("{" + ResourceArg + "}", Value);
+				}
+
+				// Add all the required and optional parameters
+				List<string> QueryParams = new List<string>();
+				if(Command.OptionalParams != null)
+				{
+					foreach(string OptionalParam in Command.OptionalParams)
+					{
+						string Value = ParseOptionalParam(RemainingArgs, OptionalParam);
+						if(Value != null)
+						{
+							QueryParams.Add(String.Format("{0}={1}", OptionalParam, Value));
+						}
+					}
+				}
+
+				// Append the parameters to the URL
+				for(int Idx = 0; Idx < QueryParams.Count; Idx++)
+				{
+					if(Idx == 0)
+					{
+						CommandUrl.Append("?");
+					}
+					else
+					{
+						CommandUrl.Append("&");
+					}
+					CommandUrl.Append(QueryParams[Idx]);
+				}
+
+				// Parse additional options for the message body
+				string BodyText = null;
+				if(Command.BodyType != null)
+				{
+					object BodyObject = ParseObject(RemainingArgs, Command.BodyType);
+					BodyText = new JavaScriptSerializer().Serialize(BodyObject);
+				}
+
+				// Make sure there are no unused arguments
+				if(RemainingArgs.Count > 0)
+				{
+					foreach(string RemainingArg in RemainingArgs)
+					{
+						Console.WriteLine("Unused argument: {0}", RemainingArg);
+					}
+					return false;
+				}
+
+				// Create the request
+				HttpWebRequest Request = (HttpWebRequest)WebRequest.Create(CommandUrl.ToString());
+				Request.ContentType = "application/json";
+				Request.Method = Command.Verb;
+				if (BodyText != null)
+				{
+					byte[] BodyData = Encoding.UTF8.GetBytes(BodyText);
+					using (Stream RequestStream = Request.GetRequestStream())
+					{
+						RequestStream.Write(BodyData, 0, BodyData.Length);
+					}
+				}
+
+				// Read the response
+				HttpWebResponse Response = (HttpWebResponse)Request.GetResponse();
+				Console.WriteLine("Response: {0} ({1})", (int)Response.StatusCode, Response.StatusDescription);
+				using (StreamReader ResponseReader = new StreamReader(Response.GetResponseStream(), Encoding.Default))
+				{
+					string ResponseContent = ResponseReader.ReadToEnd();
+					if(!String.IsNullOrEmpty(ResponseContent))
+					{
+						Console.WriteLine(Json.Format(ResponseContent));
+					}
+				}
+				return true;
+			}
+			catch(CommandLineArgumentException Ex)
+			{
+				Console.WriteLine("{0}", Ex.Message);
+				Console.WriteLine();
+				PrintCommandHelp(Command);
+				return false;
+			}
+		}
+
+		static void PrintCommands()
+		{
+			Console.WriteLine("SYNTAX:");
+			Console.WriteLine("  MetadataTool.exe -Server=http://HostName:Port <Command> <Arguments...>");
+			Console.WriteLine();
+			Console.WriteLine("Global Options:");
+			Console.WriteLine("  -Server=<url>     Specifies the endpoint for the metadata server (as http://hostname:port)");
+			Console.WriteLine();
+			Console.WriteLine("Available commands:");
+			foreach(Command Command in Commands)
+			{
+				Console.WriteLine("  {0}", Command.Name);
+			}
+		}
+
+		static void PrintCommandHelp(Command Command)
+		{
+			Console.WriteLine("Arguments for command '{0}':", Command.Name);
+
+			foreach(string ArgumentName in ParseArgumentNamesFromResource(Command.Resource))
+			{
+				Console.WriteLine("  -{0}=...", ArgumentName);
+			}
+			if(Command.OptionalParams != null)
+			{
+				foreach(string ArgumentName in Command.OptionalParams)
+				{
+					Console.WriteLine("  -{0}=... (optional)", ArgumentName);
+				}
+			}
+			if(Command.BodyType != null)
+			{
+				foreach(FieldInfo Field in Command.BodyType.GetFields(BindingFlags.Instance | BindingFlags.Public))
+				{
+					if(IsOptionalObjectField(Field))
+					{
+						Console.WriteLine("  -{0}={1} (optional)", Field.Name, GetFieldValueDescription(Field.FieldType));
+					}
+					else
+					{
+						Console.WriteLine("  -{0}={1}", Field.Name, GetFieldValueDescription(Field.FieldType));
+					}
+				}
+			}
+		}
+
+		static string GetFieldValueDescription(Type FieldType)
+		{
+			if(FieldType.IsEnum)
+			{
+				return String.Join("|", Enum.GetNames(FieldType));
+			}
+			else
+			{
+				return "...";
+			}
+		}
+
+		static IEnumerable<string> ParseArgumentNamesFromResource(string Resource)
+		{
+			for(int ArgIdx = Resource.IndexOf('{'); ArgIdx != -1; ArgIdx = Resource.IndexOf('{', ArgIdx))
+			{
+				int ArgEndIdx = Resource.IndexOf('}', ArgIdx + 1);
+				yield return Resource.Substring(ArgIdx + 1, ArgEndIdx - (ArgIdx + 1));
+				ArgIdx = ArgEndIdx + 1;
+			}
+		}
+
+		static object ParseObject(List<string> CommandArgs, Type ObjectType)
+		{
+			object Instance = Activator.CreateInstance(ObjectType);
+			foreach(FieldInfo Field in Instance.GetType().GetFields(BindingFlags.Instance | BindingFlags.Public))
+			{
+				string Value = ParseOptionalParam(CommandArgs, Field.Name);
+				if(Value == null)
+				{
+					if(!IsOptionalObjectField(Field))
+					{
+						throw new CommandLineArgumentException(String.Format("Missing -{0}=... argument", Field.Name));
+					}
+				}
+				else
+				{
+					Type FieldType = Field.FieldType;
+					if(FieldType.IsGenericType && FieldType.GetGenericTypeDefinition() == typeof(Nullable<>))
+					{
+						FieldType = FieldType.GetGenericArguments()[0];
+					}
+
+					if(FieldType == typeof(string))
+					{
+						Field.SetValue(Instance, Value);
+					}
+					else if(FieldType == typeof(int))
+					{
+						int IntValue;
+						if(!int.TryParse(Value, out IntValue))
+						{
+							throw new CommandLineArgumentException(String.Format("Invalid value for '{0}' - {1} is not an integer.", Field.Name, Value));
+						}
+						Field.SetValue(Instance, IntValue);
+					}
+					else if(FieldType == typeof(bool))
+					{
+						bool BoolValue;
+						if(!bool.TryParse(Value, out BoolValue))
+						{
+							throw new CommandLineArgumentException(String.Format("Invalid value for '{0}' - {1} is not a boolean.", Field.Name, Value));
+						}
+						Field.SetValue(Instance, BoolValue);
+					}
+					else if(FieldType.IsEnum)
+					{
+						object EnumValue;
+						try
+						{
+							EnumValue = Enum.Parse(FieldType, Value, true);
+						}
+						catch(ArgumentException)
+						{
+							throw new CommandLineArgumentException(String.Format("Invalid value for '{0}' - should be {1}", Field.Name, String.Join("/", Enum.GetNames(FieldType))));
+						}
+						Field.SetValue(Instance, EnumValue);
+					}
+					else
+					{
+						throw new NotImplementedException(String.Format("Unsupported type '{0}'", FieldType));
+					}
+				}
+			}
+			return Instance;
+		}
+
+		static bool IsOptionalObjectField(FieldInfo Field)
+		{
+			if(Field.GetCustomAttribute<OptionalAttribute>() != null)
+			{
+				return true;
+			}
+			if(Field.FieldType.IsGenericType && Field.FieldType.GetGenericTypeDefinition() == typeof(Nullable<>))
+			{
+				return true;
+			}
+			return false;
+		}
+
+		static bool ParseOption(List<string> RemainingArgs, string Option)
+		{
+			for(int Idx = 0; Idx < RemainingArgs.Count; Idx++)
+			{
+				if(RemainingArgs[Idx].Equals(Option, StringComparison.OrdinalIgnoreCase))
+				{
+					RemainingArgs.RemoveAt(Idx);
+					return true;
+				}
+			}
+			return false;
+		}
+
+		static string ParseRequiredParam(List<string> RemainingArgs, string Name)
+		{
+			string Result = ParseOptionalParam(RemainingArgs, Name);
+			if(Result == null)
+			{
+				throw new CommandLineArgumentException(String.Format("Missing -{0}=... argument", Name));
+			}
+			return Result;
+		}
+
+		static string ParseOptionalParam(List<string> RemainingArgs, string Name)
+		{
+			for(int Idx = 0; Idx < RemainingArgs.Count; Idx++)
+			{
+				string RemainingArg = RemainingArgs[Idx];
+				if(RemainingArg.Length >= Name.Length + 2 && RemainingArg[0] == '-' && RemainingArg[Name.Length + 1] == '=' && String.Compare(RemainingArg, 1, Name, 0, Name.Length, StringComparison.OrdinalIgnoreCase) == 0)
+				{
+					string Value = RemainingArgs[Idx].Substring(Name.Length + 2);
+					RemainingArgs.RemoveAt(Idx);
+					return Value;
+				}
+			}
+			return null;
+		}
+	}
+}
diff --git a/Engine/Source/Programs/UnrealGameSync/MetadataTool/Properties/AssemblyInfo.cs b/Engine/Source/Programs/UnrealGameSync/MetadataTool/Properties/AssemblyInfo.cs
new file mode 100644
--- /dev/null
+++ b/Engine/Source/Programs/UnrealGameSync/MetadataTool/Properties/AssemblyInfo.cs
@@ -0,0 +1,38 @@
+// Copyright 1998-2019 Epic Games, Inc. All Rights Reserved.
+
+using System.Reflection;
+using System.Runtime.CompilerServices;
+using System.Runtime.InteropServices;
+
+// General Information about an assembly is controlled through the following
+// set of attributes. Change these attribute values to modify the information
+// associated with an assembly.
+[assembly: AssemblyTitle("MetadataTool")]
+[assembly: AssemblyDescription("")]
+[assembly: AssemblyConfiguration("")]
+[assembly: AssemblyCompany("")]
+[assembly: AssemblyProduct("MetadataTool")]
+[assembly: AssemblyCopyright("Copyright 1998-2019 Epic Games, Inc. All Rights Reserved.")]
+[assembly: AssemblyTrademark("")]
+[assembly: AssemblyCulture("")]
+
+// Setting ComVisible to false makes the types in this assembly not visible
+// to COM components.  If you need to access a type in this assembly from
+// COM, set the ComVisible attribute to true on that type.
+[assembly: ComVisible(false)]
+
+// The following GUID is for the ID of the typelib if this project is exposed to COM
+[assembly: Guid("1b562929-0d0c-4cbe-8c43-3683f54d0414")]
+
+// Version information for an assembly consists of the following four values:
+//
+//      Major Version
+//      Minor Version
+//      Build Number
+//      Revision
+//
+// You can specify all the values or you can default the Build and Revision Numbers
+// by using the '*' as shown below:
+// [assembly: AssemblyVersion("1.0.*")]
+[assembly: AssemblyVersion("1.0.0.0")]
+[assembly: AssemblyFileVersion("1.0.0.0")]
diff --git a/Engine/Source/Programs/UnrealGameSync/UnrealGameSync/Controls/AlertButtonControl.cs b/Engine/Source/Programs/UnrealGameSync/UnrealGameSync/Controls/AlertButtonControl.cs
new file mode 100644
--- /dev/null
+++ b/Engine/Source/Programs/UnrealGameSync/UnrealGameSync/Controls/AlertButtonControl.cs
@@ -0,0 +1,181 @@
+// Copyright 1998-2019 Epic Games, Inc. All Rights Reserved.
+
+using System;
+using System.Collections.Generic;
+using System.ComponentModel;
+using System.Drawing;
+using System.Drawing.Drawing2D;
+using System.Data;
+using System.Linq;
+using System.Text;
+using System.Threading.Tasks;
+using System.Windows.Forms;
+using System.Windows.Forms.VisualStyles;
+
+namespace UnrealGameSync
+{
+	partial class AlertButtonControl : Button
+	{
+		public enum AlertButtonTheme
+		{
+			Normal,
+			Green,
+			Red,
+			Custom
+		}
+
+		[TypeConverter(typeof(ExpandableObjectConverter))]
+		public struct AlertButtonColors
+		{
+			public Color ForeColor
+			{
+				get; set;
+			}
+
+			public Color BorderColor
+			{
+				get; set;
+			}
+
+			public Color BackgroundColor1
+			{
+				get; set;
+			}
+
+			public Color BackgroundColor2
+			{
+				get; set;
+			}
+
+			public Color BackgroundColorDown
+			{
+				get; set;
+			}
+		}
+
+		AlertButtonTheme ThemeValue;
+		AlertButtonColors Colors;
+		AlertButtonColors CustomColorsValue;
+		bool bMouseOver;
+		bool bMouseDown;
+
+		public AlertButtonTheme Theme
+		{
+			get { return ThemeValue; }
+			set { ThemeValue = value; UpdateThemeColors(); }
+		}
+
+		public AlertButtonColors CustomColors
+		{
+			get { return CustomColorsValue; }
+			set { CustomColorsValue = value; UpdateThemeColors(); }
+		}
+
+		public new Color ForeColor
+		{
+			get { return base.ForeColor; }
+		}
+
+		public AlertButtonControl()
+		{
+			Theme = AlertButtonTheme.Normal;
+			UpdateThemeColors();
+		}
+
+		private void UpdateThemeColors()
+		{
+			switch(Theme)
+			{
+				case AlertButtonTheme.Normal:
+					Colors.ForeColor = Color.FromArgb(64, 86, 106);
+					Colors.BorderColor = Color.FromArgb(230, 232, 235);
+					Colors.BackgroundColor1 = Color.FromArgb(255, 255, 255);
+					Colors.BackgroundColor2 = Color.FromArgb(244, 245, 247);
+					Colors.BackgroundColorDown = Color.FromArgb(234, 235, 237);
+					break;
+				case AlertButtonTheme.Green:
+					Colors.ForeColor = Color.FromArgb(255, 255, 255);
+					Colors.BorderColor = Color.FromArgb(143, 199, 156);
+					Colors.BackgroundColor1 = Color.FromArgb(116, 192, 134);
+					Colors.BackgroundColor2 = Color.FromArgb(99, 175, 117);
+					Colors.BackgroundColorDown = Color.FromArgb(90, 165, 107);
+					break;
+				case AlertButtonTheme.Red:
+					Colors.ForeColor = Color.FromArgb(255, 255, 255);
+					Colors.BorderColor = Color.FromArgb(230, 232, 235);
+					Colors.BackgroundColor1 = Color.FromArgb(222, 108, 86);
+					Colors.BackgroundColor2 = Color.FromArgb(214, 69, 64);
+					Colors.BackgroundColorDown = Color.FromArgb(204, 59, 54);
+					break;
+				case AlertButtonTheme.Custom:
+					Colors = CustomColorsValue;
+					break;
+			}
+
+			base.ForeColor = Colors.ForeColor;
+			Invalidate();
+		}
+
+		protected override void OnMouseDown(MouseEventArgs mevent)
+		{
+			base.OnMouseDown(mevent);
+
+			bMouseDown = true;
+			Invalidate();
+		}
+
+		protected override void OnMouseUp(MouseEventArgs mevent)
+		{
+			base.OnMouseUp(mevent);
+
+			bMouseDown = false;
+			Invalidate();
+		}
+
+		protected override void OnMouseEnter(EventArgs e)
+		{
+			base.OnMouseHover(e);
+
+			bMouseOver = true;
+			Invalidate();
+		}
+
+		protected override void OnMouseLeave(EventArgs e)
+		{
+			base.OnMouseLeave(e);
+
+			bMouseOver = false;
+			Invalidate();
+		}
+
+		protected override void OnPaint(PaintEventArgs Event)
+		{
+			Event.Graphics.FillRectangle(SystemBrushes.Window, 0, 0, Width, Height);
+			using(GraphicsPath Path = new GraphicsPath())
+			{
+				const int Diameter = 4;
+
+				Path.StartFigure();
+				Path.AddArc(Width - 1 - Diameter, Height - 1 - Diameter, Diameter, Diameter, 0, 90);
+				Path.AddArc(0, Height - 1 - Diameter, Diameter, Diameter, 90, 90);
+				Path.AddArc(0, 0, Diameter, Diameter, 180, 90);
+				Path.AddArc(Width - 1 - Diameter, 0, Diameter, Diameter, 270, 90);
+				Path.CloseFigure();
+
+				Color BackgroundColorMin = (bMouseDown && bMouseOver)? Colors.BackgroundColorDown : bMouseOver? Colors.BackgroundColor2 : Colors.BackgroundColor1;
+				Color BackgroundColorMax = (bMouseDown && bMouseOver)? Colors.BackgroundColorDown : Colors.BackgroundColor2;
+
+				using(LinearGradientBrush Brush = new LinearGradientBrush(new Point(0, 0), new Point(0, Height), BackgroundColorMin, BackgroundColorMax))
+				{
+					Event.Graphics.FillPath(Brush, Path);
+				}
+				using(Pen SolidPen = new Pen(Colors.BorderColor))
+				{
+					Event.Graphics.DrawPath(SolidPen, Path);
+				}
+			}
+
+			TextRenderer.DrawText(Event.Graphics, Text, Font, new Rectangle(0, 0, Width, Height), ForeColor, TextFormatFlags.HorizontalCenter | TextFormatFlags.VerticalCenter | TextFormatFlags.SingleLine);
+		}
+	}
+}
diff --git a/Engine/Source/Programs/UnrealGameSync/UnrealGameSync/Controls/AlertImageControl.cs b/Engine/Source/Programs/UnrealGameSync/UnrealGameSync/Controls/AlertImageControl.cs
new file mode 100644
--- /dev/null
+++ b/Engine/Source/Programs/UnrealGameSync/UnrealGameSync/Controls/AlertImageControl.cs
@@ -0,0 +1,49 @@
+// Copyright 1998-2019 Epic Games, Inc. All Rights Reserved.
+
+using System;
+using System.Collections.Generic;
+using System.ComponentModel;
+using System.Drawing;
+using System.Drawing.Drawing2D;
+using System.Data;
+using System.Linq;
+using System.Text;
+using System.Threading.Tasks;
+using System.Windows.Forms;
+using System.Windows.Forms.VisualStyles;
+
+namespace UnrealGameSync
+{
+	partial class AlertImageControl : UserControl
+	{
+		Image ImageValue;
+
+		public Image Image
+		{
+			get { return ImageValue; }
+			set { ImageValue = value; Invalidate(); }
+		}
+
+		public AlertImageControl()
+		{
+		}
+
+		protected override void OnPaint(PaintEventArgs Event)
+		{
+			Event.Graphics.FillRectangle(SystemBrushes.Window, 0, 0, Width, Height);
+
+			if(Image != null)
+			{
+				float Scale = Math.Min((float)Width / Image.Width, (float)Height / Image.Height);
+
+				int ImageW = (int)(Image.Width * Scale);
+				int ImageH = (int)(Image.Height * Scale);
+
+				Event.Graphics.InterpolationMode = InterpolationMode.HighQualityBicubic;
+				Event.Graphics.SmoothingMode = SmoothingMode.HighQuality;
+				Event.Graphics.PixelOffsetMode = PixelOffsetMode.HighQuality;
+				Event.Graphics.DrawImage(Image, (Width - ImageW) / 2, (Height - ImageH) / 2, ImageW, ImageH); 
+			}
+		}
+	}
+}
diff --git a/Engine/Source/Programs/UnrealGameSync/UnrealGameSync/Controls/BuildListControl.cs b/Engine/Source/Programs/UnrealGameSync/UnrealGameSync/Controls/BuildListControl.cs
--- a/Engine/Source/Programs/UnrealGameSync/UnrealGameSync/Controls/BuildListControl.cs
+++ b/Engine/Source/Programs/UnrealGameSync/UnrealGameSync/Controls/BuildListControl.cs
@@ -12,7 +12,7 @@ using System.Windows.Forms;
 
 namespace UnrealGameSync
 {
-    public class BuildListControl : CustomListViewControl
+    class BuildListControl : CustomListViewControl
 	{
 		[StructLayout(LayoutKind.Sequential)]
 		public class SCROLLINFO
diff --git a/Engine/Source/Programs/UnrealGameSync/UnrealGameSync/Controls/CustomListViewControl.cs b/Engine/Source/Programs/UnrealGameSync/UnrealGameSync/Controls/CustomListViewControl.cs
--- a/Engine/Source/Programs/UnrealGameSync/UnrealGameSync/Controls/CustomListViewControl.cs
+++ b/Engine/Source/Programs/UnrealGameSync/UnrealGameSync/Controls/CustomListViewControl.cs
@@ -13,12 +13,198 @@ using System.Windows.Forms.VisualStyles;
 
 namespace UnrealGameSync
 {
-	public partial class CustomListViewControl : ListView
+	abstract class CustomListViewWidget
+	{
+		public ListViewItem Item
+		{
+			get;
+			private set;
+		}
+
+		Rectangle? PreviousBounds;
+
+		public Cursor Cursor
+		{
+			get;
+			set;
+		}
+
+		public CustomListViewWidget(ListViewItem Item)
+		{
+			this.Item = Item;
+			this.Cursor = Cursors.Arrow;
+		}
+
+		public void Invalidate()
+		{
+			PreviousBounds = null;
+		}
+
+		public virtual bool RequiresLayout()
+		{
+			return !PreviousBounds.HasValue;
+		}
+
+		public void ConditionalLayout(Control Owner, Rectangle Bounds)
+		{
+			if(RequiresLayout() || PreviousBounds.Value != Bounds)
+			{
+				using(Graphics Graphics = Owner.CreateGraphics())
+				{
+					Layout(Graphics, Bounds);
+				}
+				PreviousBounds = Bounds;
+			}
+		}
+
+		public virtual void OnMouseMove(Point Location)
+		{
+		}
+
+		public virtual void OnMouseLeave()
+		{
+		}
+
+		public virtual void OnMouseDown(Point Location)
+		{
+		}
+
+		public virtual void OnMouseUp(Point Location)
+		{
+		}
+
+		public abstract void Layout(Graphics Graphics, Rectangle Bounds);
+		public abstract void Render(Graphics Graphics);
+	}
+
+	class StatusLineListViewWidget : CustomListViewWidget
+	{
+		public StatusElementResources Resources;
+		public StatusLine Line = new StatusLine();
+
+		public HorizontalAlignment HorizontalAlignment
+		{
+			get; set;
+		}
+
+		Rectangle Bounds;
+		StatusElement MouseDownElement;
+		StatusElement MouseOverElement;
+
+		public StatusLineListViewWidget(ListViewItem Item, StatusElementResources Resources)
+			: base(Item)
+		{
+			this.Resources = Resources;
+			this.HorizontalAlignment = HorizontalAlignment.Center;
+		}
+
+		public override bool RequiresLayout()
+		{
+			return base.RequiresLayout() || Line.RequiresLayout();
+		}
+
+		public override void OnMouseDown(Point Location)
+		{
+			StatusElement Element;
+			if(Line.HitTest(Location, out Element))
+			{
+				MouseDownElement = Element;
+				MouseDownElement.bMouseDown = true;
+
+				Invalidate();
+			}
+		}
+
+		public override void OnMouseUp(Point Location)
+		{
+			if(MouseDownElement != null)
+			{
+				StatusElement ClickElement = MouseDownElement;
+
+				MouseDownElement.bMouseDown = false;
+				MouseDownElement = null;
+
+				if(ClickElement.Bounds.Contains(Location))
+				{
+					ClickElement.OnClick(Location);
+				}
+
+				Invalidate();
+			}
+		}
+
+		public override void OnMouseMove(Point Location)
+		{
+			StatusElement NewMouseOverElement;
+			Line.HitTest(Location, out NewMouseOverElement);
+
+			if(MouseOverElement != NewMouseOverElement)
+			{
+				if(MouseOverElement != null)
+				{
+					Cursor = Cursors.Arrow;
+					MouseOverElement.bMouseOver = false;
+				}
+
+				MouseOverElement = NewMouseOverElement;
+
+				if(MouseOverElement != null)
+				{
+					Cursor = MouseOverElement.Cursor;
+					MouseOverElement.bMouseOver = true;
+				}
+
+				Invalidate();
+			}
+		}
+
+		public override void OnMouseLeave()
+		{
+			if(MouseOverElement != null)
+			{
+				MouseOverElement.bMouseOver = false;
+				MouseOverElement = null;
+
+				Invalidate();
+			}
+		}
+
+		public override void Layout(Graphics Graphics, Rectangle Bounds)
+		{
+			this.Bounds = Bounds;
+
+			int OffsetY = Bounds.Y + Bounds.Height / 2;
+			Line.Layout(Graphics, new Point(Bounds.X, OffsetY), Resources);
+
+			if(HorizontalAlignment == HorizontalAlignment.Center)
+			{
+				int OffsetX = Bounds.X + (Bounds.Width - Line.Bounds.Width) / 2;
+				Line.Layout(Graphics, new Point(OffsetX, OffsetY), Resources);
+			}
+			else if(HorizontalAlignment == HorizontalAlignment.Right)
+			{
+				int OffsetX = Bounds.Right - Line.Bounds.Width;
+				Line.Layout(Graphics, new Point(OffsetX, OffsetY), Resources);
+			}
+		}
+
+		public override void Render(Graphics Graphics)
+		{
+			Graphics.IntersectClip(Bounds);
+
+			Line.Draw(Graphics, Resources);
+		}
+	}
+
+	partial class CustomListViewControl : ListView
 	{
 		VisualStyleRenderer SelectedItemRenderer;
 		VisualStyleRenderer TrackedItemRenderer;
 		public int HoverItem = -1;
 
+		CustomListViewWidget MouseOverWidget;
+		CustomListViewWidget MouseDownWidget;
+
 		public CustomListViewControl()
 		{
             if (Application.RenderWithVisualStyles) 
@@ -28,10 +214,79 @@ namespace UnrealGameSync
 			}
 		}
 
+		protected void ConditionalLayoutWidget(CustomListViewWidget Widget)
+		{
+			if(Widget.Item.Tag == Widget)
+			{
+				Widget.ConditionalLayout(this, Widget.Item.Bounds);
+			}
+			else
+			{
+				for(int Idx = 0; Idx < Widget.Item.SubItems.Count; Idx++)
+				{
+					ListViewItem.ListViewSubItem SubItem = Widget.Item.SubItems[Idx];
+					if(SubItem.Tag == Widget)
+					{
+						Rectangle Bounds = SubItem.Bounds;
+						for(int EndIdx = Idx + 1; EndIdx < Widget.Item.SubItems.Count && Widget.Item.SubItems[EndIdx].Tag == Widget; EndIdx++)
+						{
+							Rectangle EndBounds = Widget.Item.SubItems[EndIdx].Bounds;
+							Bounds = new Rectangle(Bounds.X, Bounds.Y, EndBounds.Right - Bounds.X, EndBounds.Bottom - Bounds.Y);
+						}
+						Widget.ConditionalLayout(this, Bounds);
+						break;
+					}
+				}
+			}
+		}
+
+		protected void ConditionalRedrawWidget(CustomListViewWidget Widget)
+		{
+			if(Widget.Item.Index != -1 && Widget.RequiresLayout())
+			{
+				ConditionalLayoutWidget(Widget);
+				RedrawItems(Widget.Item.Index, Widget.Item.Index, true);
+			}
+		}
+
+		protected CustomListViewWidget FindWidget(Point Location)
+		{
+			return FindWidget(HitTest(Location));
+		}
+
+		protected CustomListViewWidget FindWidget(ListViewHitTestInfo HitTest)
+		{
+			if(HitTest.Item != null)
+			{
+				CustomListViewWidget Widget = HitTest.Item.Tag as CustomListViewWidget;
+				if(Widget != null)
+				{
+					return Widget;
+				}
+			}
+			if(HitTest.SubItem != null)
+			{
+				CustomListViewWidget Widget = HitTest.SubItem.Tag as CustomListViewWidget;
+				if(Widget != null)
+				{
+					return Widget;
+				}
+			}
+			return null;
+		}
+
 		protected override void OnMouseLeave(EventArgs e)
 		{
 			base.OnMouseLeave(e);
 
+			if(MouseOverWidget != null)
+			{
+				MouseOverWidget.OnMouseLeave();
+				ConditionalRedrawWidget(MouseOverWidget);
+				MouseOverWidget = null;
+				Invalidate();
+			}
+
 			if(HoverItem != -1)
 			{
 				HoverItem = -1;
@@ -39,10 +294,58 @@ namespace UnrealGameSync
 			}
 		}
 
+		protected override void OnMouseDown(MouseEventArgs e)
+		{
+			base.OnMouseDown(e);
+
+			if(MouseDownWidget != null)
+			{
+				MouseDownWidget.OnMouseUp(e.Location);
+				ConditionalRedrawWidget(MouseDownWidget);
+			}
+
+			MouseDownWidget = FindWidget(e.Location);
+
+			if(MouseDownWidget != null)
+			{
+				MouseDownWidget.OnMouseDown(e.Location);
+				ConditionalRedrawWidget(MouseDownWidget);
+			}
+		}
+
+		protected override void OnMouseUp(MouseEventArgs e)
+		{
+			base.OnMouseUp(e);
+
+			if(MouseDownWidget != null)
+			{
+				CustomListViewWidget Widget = MouseDownWidget;
+				MouseDownWidget = null;
+
+				Widget.OnMouseUp(e.Location);
+				ConditionalRedrawWidget(Widget);
+			}
+		}
+
 		protected override void OnMouseMove(MouseEventArgs e)
 		{
 			ListViewHitTestInfo HitTest = this.HitTest(e.Location);
 
+			CustomListViewWidget NewMouseOverWidget = FindWidget(HitTest);
+			if(MouseOverWidget != null && MouseOverWidget != NewMouseOverWidget)
+			{
+				Cursor = Cursors.Arrow;
+				MouseOverWidget.OnMouseLeave();
+				ConditionalRedrawWidget(MouseOverWidget);
+			}
+			MouseOverWidget = NewMouseOverWidget;
+			if(MouseOverWidget != null)
+			{
+				Cursor = MouseOverWidget.Cursor;
+				MouseOverWidget.OnMouseMove(e.Location);
+				ConditionalRedrawWidget(MouseOverWidget);
+			}
+
 			int PrevHoverItem = HoverItem;
 			HoverItem = (HitTest.Item == null)? -1 : HitTest.Item.Index;
 			if(HoverItem != PrevHoverItem)
@@ -60,6 +363,26 @@ namespace UnrealGameSync
 			base.OnMouseMove(e);
 		}
 
+		public void DrawCustomSubItem(Graphics Graphics, ListViewItem.ListViewSubItem SubItem)
+		{
+			CustomListViewWidget Widget = SubItem.Tag as CustomListViewWidget;
+			if(Widget != null)
+			{
+				foreach(ListViewItem.ListViewSubItem OtherSubItem in Widget.Item.SubItems)
+				{
+					if(OtherSubItem.Tag == Widget)
+					{
+						if(OtherSubItem == SubItem)
+						{
+							ConditionalLayoutWidget(Widget);
+							Widget.Render(Graphics);
+						}
+						break;
+					}
+				}
+			}
+		}
+
 		public void DrawBackground(Graphics Graphics, ListViewItem Item)
 		{
 			if(Item.Selected)
diff --git a/Engine/Source/Programs/UnrealGameSync/UnrealGameSync/Controls/StatusPanel.cs b/Engine/Source/Programs/UnrealGameSync/UnrealGameSync/Controls/StatusPanel.cs
--- a/Engine/Source/Programs/UnrealGameSync/UnrealGameSync/Controls/StatusPanel.cs
+++ b/Engine/Source/Programs/UnrealGameSync/UnrealGameSync/Controls/StatusPanel.cs
@@ -33,7 +33,10 @@ namespace UnrealGameSync
 		{
 			FontCache.Add(FontStyle.Regular, BaseFont);
 
-			BadgeFont = new Font(BaseFont.FontFamily, BaseFont.Size - 1.25f, FontStyle.Bold);
+			using(Graphics Graphics = Graphics.FromHwnd(IntPtr.Zero))
+			{
+				BadgeFont = new Font(BaseFont.FontFamily, 7.0f/* * Graphics.DpiY / 96.0f*/, FontStyle.Bold);
+			}
 		}
 
 		public void Dispose()
@@ -130,11 +133,13 @@ namespace UnrealGameSync
 	class TextStatusElement : StatusElement
 	{
 		string Text;
+		Color Color;
 		FontStyle Style;
 
-		public TextStatusElement(string InText, FontStyle InStyle)
+		public TextStatusElement(string InText, Color InColor, FontStyle InStyle)
 		{
 			Text = InText;
+			Color = InColor;
 			Style = InStyle;
 		}
 
@@ -145,7 +150,7 @@ namespace UnrealGameSync
 
 		public override void Draw(Graphics Graphics, StatusElementResources Resources)
 		{
-			TextRenderer.DrawText(Graphics, Text, Resources.FindOrAddFont(Style), Bounds.Location, SystemColors.ControlText, TextFormatFlags.NoPadding);
+			TextRenderer.DrawText(Graphics, Text, Resources.FindOrAddFont(Style), Bounds.Location, Color, TextFormatFlags.NoPadding);
 		}
 	}
 
@@ -193,9 +198,19 @@ namespace UnrealGameSync
 		string Name;
 		Color BackgroundColor;
 		Color HoverBackgroundColor;
-		Action ClickAction;
+		Action<Point, Rectangle> ClickAction;
+
+		public bool MergeLeft
+		{
+			get; set;
+		}
+
+		public bool MergeRight
+		{
+			get; set;
+		}
 
-		public BadgeStatusElement(string InName, Color InBackgroundColor, Action InClickAction)
+		public BadgeStatusElement(string InName, Color InBackgroundColor, Action<Point, Rectangle> InClickAction)
 		{
 			Name = InName;
 			BackgroundColor = InBackgroundColor;
@@ -218,7 +233,7 @@ namespace UnrealGameSync
 		{
 			if(ClickAction != null)
 			{
-				ClickAction();
+				ClickAction(Location, Bounds);
 			}
 		}
 
@@ -226,15 +241,30 @@ namespace UnrealGameSync
 		{
 			Size LabelSize = TextRenderer.MeasureText(Name, Resources.BadgeFont);
 			int BadgeHeight = Resources.BadgeFont.Height + 1;
-			return new Size(LabelSize.Width + BadgeHeight - 4, BadgeHeight);
+			return new Size(LabelSize.Width + 1 + BadgeHeight - 4, BadgeHeight);
 		}
 
 		public override void Draw(Graphics Graphics, StatusElementResources Resources)
 		{
-			using(SolidBrush Brush = new SolidBrush(bMouseOver? HoverBackgroundColor : BackgroundColor))
+			SmoothingMode PrevSmoothingMode = Graphics.SmoothingMode;
+			Graphics.SmoothingMode = SmoothingMode.HighQuality;
+
+			using (GraphicsPath Path = new GraphicsPath())
 			{
-				Graphics.FillRectangle(Brush, Bounds);
+				Path.StartFigure();
+				Path.AddLine(Bounds.Left + (MergeLeft? 1 : 0), Bounds.Top, Bounds.Left - (MergeLeft? 1 : 0), Bounds.Bottom);
+				Path.AddLine(Bounds.Left - (MergeLeft? 1 : 0), Bounds.Bottom, Bounds.Right - 2 - (MergeRight? 1 : 0), Bounds.Bottom);
+				Path.AddLine(Bounds.Right - 2 - (MergeRight? 1 : 0), Bounds.Bottom, Bounds.Right - 2 + (MergeRight? 1 : 0), Bounds.Top);
+				Path.AddLine(Bounds.Right - 2 + (MergeRight? 1 : 0), Bounds.Top, Bounds.Left + (MergeLeft? 1 : 0), Bounds.Top);
+				Path.CloseFigure();
+
+				using(SolidBrush Brush = new SolidBrush(bMouseOver? HoverBackgroundColor : BackgroundColor))
+				{
+					Graphics.FillPath(Brush, Path);
+				}
 			}
+
+			Graphics.SmoothingMode = PrevSmoothingMode;
 			TextRenderer.DrawText(Graphics, Name, Resources.BadgeFont, Bounds, Color.White, TextFormatFlags.HorizontalCenter | TextFormatFlags.VerticalCenter | TextFormatFlags.SingleLine | TextFormatFlags.NoPrefix | TextFormatFlags.PreserveGraphicsClipping);
 		}
 	}
@@ -271,6 +301,7 @@ namespace UnrealGameSync
 
 	class StatusLine
 	{
+		bool bModified;
 		List<StatusElement> Elements = new List<StatusElement>();
 
 		public StatusLine()
@@ -290,39 +321,79 @@ namespace UnrealGameSync
 			private set;
 		}
 
+		public bool RequiresLayout()
+		{
+			return bModified;
+		}
+
+		public void Clear()
+		{
+			Elements.Clear();
+			bModified = true;
+		}
+
+		public void Add(StatusElement Element)
+		{
+			Elements.Add(Element);
+			bModified = true;
+		}
+
 		public void AddIcon(Image InIcon)
 		{
 			Elements.Add(new IconStatusElement(InIcon));
+			bModified = true;
 		}
 
 		public void AddIcon(Image InStrip, Size InIconSize, int InIndex)
 		{
 			Elements.Add(new IconStripStatusElement(InStrip, InIconSize, InIndex));
+			bModified = true;
 		}
 
 		public void AddText(string InText, FontStyle InStyle = FontStyle.Regular)
 		{
-			Elements.Add(new TextStatusElement(InText, InStyle));
+			Elements.Add(new TextStatusElement(InText, SystemColors.ControlText, InStyle));
+			bModified = true;
+		}
+
+		public void AddText(string InText, Color InColor, FontStyle InStyle = FontStyle.Regular)
+		{
+			Elements.Add(new TextStatusElement(InText, InColor, InStyle));
+			bModified = true;
 		}
 
 		public void AddLink(string InText, FontStyle InStyle, Action InLinkAction)
 		{
 			Elements.Add(new LinkStatusElement(InText, InStyle, (P, R) => { InLinkAction(); }));
+			bModified = true;
 		}
 
 		public void AddLink(string InText, FontStyle InStyle, Action<Point, Rectangle> InLinkAction)
 		{
 			Elements.Add(new LinkStatusElement(InText, InStyle, InLinkAction));
+			bModified = true;
 		}
 
-		public void AddBadge(string InText, Color InBackgroundColor, Action InClickAction)
+		public void AddBadge(string InText, Color InBackgroundColor, Action<Point, Rectangle> InClickAction)
 		{
 			Elements.Add(new BadgeStatusElement(InText, InBackgroundColor, InClickAction));
+			if(Elements.Count >= 2)
+			{
+				BadgeStatusElement PrevBadge = Elements[Elements.Count - 2] as BadgeStatusElement;
+				BadgeStatusElement NextBadge = Elements[Elements.Count - 1] as BadgeStatusElement;
+				if(PrevBadge != null && NextBadge != null)
+				{
+					PrevBadge.MergeRight = true;
+					NextBadge.MergeLeft = true;
+				}
+			}
+			bModified = true;
 		}
 
 		public void AddProgressBar(float Progress)
 		{
 			Elements.Add(new ProgressBarStatusElement(Progress));
+			bModified = true;
 		}
 
 		public bool HitTest(Point Location, out StatusElement OutElement)
@@ -352,6 +423,8 @@ namespace UnrealGameSync
 				NextLocation = Element.Layout(Graphics, NextLocation, Resources);
 				Bounds = Rectangle.Union(Bounds, Element.Bounds);
 			}
+
+			bModified = false;
 		}
 
 		public void Draw(Graphics Graphics, StatusElementResources Resources)
diff --git a/Engine/Source/Programs/UnrealGameSync/UnrealGameSync/Controls/WorkspaceControl.cs b/Engine/Source/Programs/UnrealGameSync/UnrealGameSync/Controls/WorkspaceControl.cs
--- a/Engine/Source/Programs/UnrealGameSync/UnrealGameSync/Controls/WorkspaceControl.cs
+++ b/Engine/Source/Programs/UnrealGameSync/UnrealGameSync/Controls/WorkspaceControl.cs
@@ -42,6 +42,7 @@ namespace UnrealGameSync
 		void SetupScheduledSync();
 		void UpdateProgress();
 		void ModifyApplicationSettings();
+		void UpdateAlertWindows();
 	}
 
 	delegate void WorkspaceStartupCallback(WorkspaceControl Workspace, bool bCancel);
@@ -113,10 +114,10 @@ namespace UnrealGameSync
 			public Dictionary<string, List<BadgeInfo>> CustomBadges;
 		}
 
-		static Rectangle GoodBuildIcon = new Rectangle(0, 0, 16, 16);
-		static Rectangle MixedBuildIcon = new Rectangle(16, 0, 16, 16);
-		static Rectangle BadBuildIcon = new Rectangle(32, 0, 16, 16);
-		static Rectangle DefaultBuildIcon = new Rectangle(48, 0, 16, 16);
+		public static Rectangle GoodBuildIcon = new Rectangle(0, 0, 16, 16);
+		public static Rectangle MixedBuildIcon = new Rectangle(16, 0, 16, 16);
+		public static Rectangle BadBuildIcon = new Rectangle(32, 0, 16, 16);
+		public static Rectangle DefaultBuildIcon = new Rectangle(48, 0, 16, 16);
 		static Rectangle PromotedBuildIcon = new Rectangle(64, 0, 16, 16);
 		static Rectangle DetailsIcon = new Rectangle(80, 0, 16, 16);
 		static Rectangle InfoIcon = new Rectangle(96, 0, 16, 16);
@@ -149,6 +150,8 @@ namespace UnrealGameSync
 		UserWorkspaceSettings WorkspaceSettings;
 		UserProjectSettings ProjectSettings;
 
+		bool bUserHasOpenIssues = false;
+
 		public UserSelectedProjectSettings SelectedProject
 		{
 			get;
@@ -189,6 +192,7 @@ namespace UnrealGameSync
 		bool bIsEnterpriseProject;
 		PerforceMonitor PerforceMonitor;
 		Workspace Workspace;
+		IssueMonitor IssueMonitor;
 		EventMonitor EventMonitor;
 		System.Windows.Forms.Timer UpdateTimer;
 		HashSet<int> PromotedChangeNumbers = new HashSet<int>();
@@ -244,7 +248,7 @@ namespace UnrealGameSync
 		System.Threading.Timer StartupTimer;
 		List<WorkspaceStartupCallback> StartupCallbacks;
 
-		public WorkspaceControl(IWorkspaceControlOwner InOwner, string InApiUrl, string InOriginalExecutableFileName, bool bInUnstable, DetectProjectSettingsTask DetectSettings, LineBasedTextWriter InLog, UserSettings InSettings)
+		public WorkspaceControl(IWorkspaceControlOwner InOwner, string InApiUrl, string InOriginalExecutableFileName, bool bInUnstable, DetectProjectSettingsTask DetectSettings, IssueMonitor InIssueMonitor, LineBasedTextWriter InLog, UserSettings InSettings)
 		{
 			InitializeComponent();
 
@@ -255,6 +259,7 @@ namespace UnrealGameSync
 			DataFolder = DetectSettings.DataFolder;
 			OriginalExecutableFileName = InOriginalExecutableFileName;
 			bUnstable = bInUnstable;
+			IssueMonitor = InIssueMonitor;
 			Log = InLog;
 			Settings = InSettings;
 			WorkspaceSettings = InSettings.FindOrAddWorkspace(DetectSettings.BranchClientPath);
@@ -342,6 +347,34 @@ namespace UnrealGameSync
 
 			StartupTimer = new System.Threading.Timer(x => MainThreadSynchronizationContext.Post((o) => { if(!IsDisposed){ StartupTimerElapsed(false); } }, null), null, TimeSpan.FromSeconds(20.0), TimeSpan.FromMilliseconds(-1.0));
 			StartupCallbacks = new List<WorkspaceStartupCallback>();
+
+			IssueMonitor.OnIssuesChanged += IssueMonitor_OnIssuesChangedAsync;
+		}
+
+		public void IssueMonitor_OnIssuesChangedAsync()
+		{
+			MainThreadSynchronizationContext.Post((o) => { if(IssueMonitor != null){ IssueMonitor_OnIssuesChanged(); } }, null);
+		}
+
+		public List<IssueData> GetOpenIssuesForUser()
+		{
+			return IssueMonitor.GetIssues().Where(x => x.FixChange <= 0 && String.Compare(x.Owner, Workspace.Perforce.UserName, StringComparison.OrdinalIgnoreCase) == 0).ToList();
+		}
+
+		public void IssueMonitor_OnIssuesChanged()
+		{
+			bool bPrevUserHasOpenIssues = bUserHasOpenIssues;
+			bUserHasOpenIssues = GetOpenIssuesForUser().Count > 0;
+			if(bUserHasOpenIssues != bPrevUserHasOpenIssues)
+			{
+				UpdateStatusPanel();
+				StatusPanel.Invalidate();
+			}
+		}
+
+		public IssueMonitor GetIssueMonitor()
+		{
+			return IssueMonitor;
 		}
 
 		private void CheckForStartupComplete()
@@ -579,6 +612,11 @@ namespace UnrealGameSync
 				StartupCallbacks = null;
 			}
 
+			if(IssueMonitor != null)
+			{
+				IssueMonitor.OnIssuesChanged -= IssueMonitor_OnIssuesChangedAsync;
+				IssueMonitor = null;
+			}
 			if(StartupTimer != null)
 			{
 				StartupTimer.Dispose();
@@ -734,7 +772,11 @@ namespace UnrealGameSync
 				}
 
 				// Get the max number of changes to ensure this
-				int NewPendingMaxChanges = ListIndexToChangeIndex[LastVisibleIndex];
+				int NewPendingMaxChanges = 0;
+				if(LastVisibleIndex >= 0)
+				{
+					NewPendingMaxChanges = ListIndexToChangeIndex[LastVisibleIndex];
+				}
 				NewPendingMaxChanges = PerforceMonitor.InitialMaxChangesValue + ((Math.Max(NewPendingMaxChanges - PerforceMonitor.InitialMaxChangesValue, 0) + BuildListExpandCount - 1) / BuildListExpandCount) * BuildListExpandCount;
 
 				// Shrink the number of changes retained by the PerforceMonitor class
@@ -1531,7 +1573,7 @@ namespace UnrealGameSync
 			public Rectangle LinkRect;
 		}
 
-		private ExpandRowLayout LayoutExpandRow(Graphics Graphics, Rectangle Bounds)
+		ExpandRowLayout LayoutExpandRow(Graphics Graphics, Rectangle Bounds)
 		{
 			ExpandRowLayout Layout = new ExpandRowLayout();
 
@@ -1601,7 +1643,7 @@ namespace UnrealGameSync
 				ExpandRowLayout ExpandLayout = LayoutExpandRow(e.Graphics, e.Bounds);
 				DrawExpandRow(e.Graphics, ExpandLayout);
 			}
-			else
+			else if(Workspace != null)
 			{
 				if(e.State.HasFlag(ListViewItemStates.Selected))
 				{
@@ -2411,25 +2453,9 @@ namespace UnrealGameSync
 			}
 		}
 
-		private static string FormatUserName(string UserName)
+		public static string FormatUserName(string UserName)
 		{
-			StringBuilder NormalUserName = new StringBuilder();
-			for(int Idx = 0; Idx < UserName.Length; Idx++)
-			{
-				if(Idx == 0 || UserName[Idx - 1] == '.')
-				{
-					NormalUserName.Append(Char.ToUpper(UserName[Idx]));
-				}
-				else if(UserName[Idx] == '.')
-				{
-					NormalUserName.Append(' ');
-				}
-				else
-				{
-					NormalUserName.Append(UserName[Idx]);
-				}
-			}
-			return NormalUserName.ToString();
+			return Utility.FormatUserName(UserName);
 		}
 
 		private static void AppendUserList(StringBuilder Builder, string Separator, string Format, IEnumerable<EventData> Reviews)
@@ -2845,6 +2871,18 @@ namespace UnrealGameSync
 				ProgramsLine.AddText("  |  ");
 				ProgramsLine.AddLink("Windows Explorer", FontStyle.Regular, () => { Process.Start("explorer.exe", String.Format("\"{0}\"", Path.GetDirectoryName(SelectedFileName))); });
 
+				if(Workspace.ProjectConfigFile.GetValue("Options.ShowBuildHealth", false))
+				{
+					ProgramsLine.AddText("  |  ");
+					ProgramsLine.AddLink("Build Health", FontStyle.Regular, (P, R) => { ShowBuildHealthMenu(R); });
+
+					if(bUserHasOpenIssues)
+					{
+						ProgramsLine.AddText("  |  ");
+						ProgramsLine.AddBadge("Close Build Issue", GetBuildBadgeColor(BadgeResult.Failure), (P, R) => ShowCloseBuildIssueMenu(R));
+					}
+				}
+
 				foreach(KeyValuePair<string, BadgeData> ServiceBadge in ServiceBadges)
 				{
 					ProgramsLine.AddText("  |  ");
@@ -2854,7 +2892,7 @@ namespace UnrealGameSync
 					}
 					else
 					{
-						ProgramsLine.AddBadge(ServiceBadge.Key, GetBuildBadgeColor(ServiceBadge.Value.Result), () => { Process.Start(ServiceBadge.Value.Url); });
+						ProgramsLine.AddBadge(ServiceBadge.Key, GetBuildBadgeColor(ServiceBadge.Value.Result), (P, R) => { Process.Start(ServiceBadge.Value.Url); });
 					}
 				}
 				ProgramsLine.AddText("  |  ");
@@ -2946,6 +2984,116 @@ namespace UnrealGameSync
 			StatusPanel.Set(Lines, Caption, Alert, TintColor);
 		}
 
+		private void ShowBuildHealthMenu(Rectangle Bounds)
+		{
+			while(BuildHealthContextMenu.Items[0] != BuildHealthContextMenu_Separator)
+			{
+				BuildHealthContextMenu.Items.RemoveAt(0);
+			}
+
+			List<IssueData> Issues = IssueMonitor.GetIssues();
+			if(Issues.Count == 0)
+			{
+				BuildHealthContextMenu_Separator.Visible = false;
+			}
+			else
+			{
+				BuildHealthContextMenu_Separator.Visible = true;
+				for(int Idx = 0; Idx < Issues.Count; Idx++)
+				{
+					IssueData Issue = Issues[Idx];
+
+					StringBuilder Description = new StringBuilder();
+					Description.AppendFormat("{0}: {1}", Issue.Id, Issue.Summary);
+					if(Issue.Owner == null)
+					{
+						Description.AppendFormat(" - Unassigned");
+					}
+					else
+					{
+						Description.AppendFormat(" - {0}", FormatUserName(Issue.Owner));
+						if(!Issue.AcknowledgedAt.HasValue)
+						{
+							Description.Append(" (?)");
+						}
+					}
+
+					Description.AppendFormat(" ({0})", Utility.FormatDurationMinutes((int)((Issue.RetrievedAt - Issue.CreatedAt).TotalMinutes + 1)));
+
+					ToolStripMenuItem Item = new ToolStripMenuItem(Description.ToString());
+					if(Issue.FixChange > 0)
+					{
+						Item.ForeColor = SystemColors.GrayText;
+					}
+					else if(Issue.Owner != null && String.Compare(Issue.Owner, Workspace.Perforce.UserName, StringComparison.OrdinalIgnoreCase) == 0)
+					{
+						Item.Font = new Font(Item.Font, FontStyle.Bold);
+					}
+					Item.Click += (s, e) => { BuildHealthContextMenu_Issue_Click(Issue); };
+					BuildHealthContextMenu.Items.Insert(Idx, Item);
+				}
+			}
+
+			BuildHealthContextMenu.Show(StatusPanel, new Point(Bounds.Right, Bounds.Bottom), ToolStripDropDownDirection.BelowLeft);
+		}
+
+		private void ShowCloseBuildIssueMenu(Rectangle Bounds)
+		{
+			List<IssueData> Issues = GetOpenIssuesForUser();
+			if(Issues.Count == 1)
+			{
+				CloseBuildIssueContextMenu_Issue_Click(Issues[0]);
+			}
+			else if(Issues.Count > 1)
+			{
+				CloseBuildIssueMenu.Items.Clear();
+
+				foreach(IssueData Issue in Issues)
+				{
+					StringBuilder Description = new StringBuilder();
+					Description.AppendFormat("{0}: {1}", Issue.Id, Issue.Summary);
+					Description.AppendFormat(" ({0})", Utility.FormatDurationMinutes((int)((Issue.RetrievedAt - Issue.CreatedAt).TotalMinutes + 1)));
+
+					ToolStripMenuItem Item = new ToolStripMenuItem(Description.ToString());
+					Item.Click += (s, e) => { CloseBuildIssueContextMenu_Issue_Click(Issue); };
+					CloseBuildIssueMenu.Items.Add(Item);
+				}
+
+				CloseBuildIssueMenu.Show(StatusPanel, new Point(Bounds.Right, Bounds.Bottom), ToolStripDropDownDirection.BelowLeft);
+			}
+		}
+
+		private void BuildHealthContextMenu_Issue_Click(IssueData Issue)
+		{
+			ShowIssueDetails(Issue);
+		}
+
+		private void CloseBuildIssueContextMenu_Issue_Click(IssueData Issue)
+		{
+			int FixChangeNumber = 0;
+			if(IssueFixedWindow.ShowModal(this, Workspace.Perforce, ref FixChangeNumber))
+			{
+				IssueUpdateData Update = new IssueUpdateData();
+				Update.Id = Issue.Id;
+				Update.FixChange = FixChangeNumber;
+				IssueMonitor.PostUpdate(Update);
+			}
+		}
+
+		public void ShowIssueDetails(IssueData Issue)
+		{
+			IssueDetailsWindow.ShowModal(this, IssueMonitor, Workspace.Perforce.ServerAndPort, Workspace.Perforce.UserName, Issue, Log, StreamName);
+		}
+
+		private void BuildHealthContextMenu_Settings_Click(object sender, EventArgs e)
+		{
+			IssueSettingsWindow IssueSettings = new IssueSettingsWindow(Settings);
+			if(IssueSettings.ShowDialog(this) == DialogResult.OK)
+			{
+				Owner.UpdateAlertWindows();
+			}
+		}
+
 		private void ShowRequiredSdkInfo()
 		{
 			string[] SdkInfoEntries;
diff --git a/Engine/Source/Programs/UnrealGameSync/UnrealGameSync/Forms/IssueAlertWindow.cs b/Engine/Source/Programs/UnrealGameSync/UnrealGameSync/Forms/IssueAlertWindow.cs
new file mode 100644
--- /dev/null
+++ b/Engine/Source/Programs/UnrealGameSync/UnrealGameSync/Forms/IssueAlertWindow.cs
@@ -0,0 +1,212 @@
+// Copyright 1998-2019 Epic Games, Inc. All Rights Reserved.
+
+using System;
+using System.Drawing;
+using System.Runtime.InteropServices;
+using System.Linq;
+using System.Text;
+using System.Windows.Forms;
+using System.Collections.Generic;
+
+namespace UnrealGameSync
+{
+	partial class IssueAlertWindow : Form
+	{
+		static readonly IntPtr HWND_TOPMOST = new IntPtr(-1);
+
+		const uint SWP_NOMOVE = 0x0002;
+		const uint SWP_NOSIZE = 0x0001;
+		const uint SWP_NOACTIVATE = 0x0010;
+
+		[DllImport("user32.dll", SetLastError=true)]
+		static extern bool SetWindowPos(IntPtr hWnd, IntPtr hWndInsertAfter, int X, int Y, int Width, int Height, uint Flags);
+
+		public readonly IssueMonitor IssueMonitor;
+
+		public bool bIsWarning;
+
+		public bool bStrongAlert = false;
+
+		public IssueAlertWindow(IssueMonitor IssueMonitor, IssueData Issue, IssueAlertReason Reason)
+		{
+			this.IssueMonitor = IssueMonitor;
+			this.Issue = Issue;
+
+			InitializeComponent();
+
+			SetIssue(Issue, Reason);
+		}
+
+		public IssueData Issue
+		{
+			get;
+			private set;
+		}
+
+		public IssueAlertReason Reason
+		{
+			get;
+			private set;
+		}
+
+		protected override void OnPaintBackground(PaintEventArgs e)
+		{
+			base.OnPaintBackground(e);
+
+			if(bStrongAlert)
+			{
+				Color StripeColor = bIsWarning? Color.FromArgb(216, 167, 64) : Color.FromArgb(214, 69, 64);
+				using(Brush StripeBrush = new SolidBrush(StripeColor))
+				{
+					e.Graphics.FillRectangle(StripeBrush, 0, 0, Bounds.Width - 1, Bounds.Height - 1);
+				}
+			}
+			else
+			{
+				Color StripeColor = bIsWarning? Color.FromArgb(216, 167, 64) : Color.FromArgb(214, 69, 64);
+
+				Color BackgroundColor = Color.FromArgb(241, 236, 236);
+				using(Brush SolidBrush = new SolidBrush(BackgroundColor))
+				{
+	//				e.Graphics.FillRectangle(SolidBrush, 0, 0, Bounds.Width - 1, Bounds.Height - 1);
+				}
+
+				Color BorderColor = StripeColor; 
+				using(Pen Pen = new Pen(Color.FromArgb(128, 128, 128)))
+				{
+					e.Graphics.DrawRectangle(Pen, 0, 0, Bounds.Width - 1, Bounds.Height - 1);
+				}
+
+				using(Brush StripeBrush = new SolidBrush(StripeColor))
+				{
+					e.Graphics.FillRectangle(StripeBrush, 0, 0, /*6*/10 * e.Graphics.DpiX / 96.0f, Bounds.Height);
+				}
+			}
+		}
+
+		public void SetIssue(IssueData NewIssue, IssueAlertReason NewReason)
+		{
+			StringBuilder OwnerTextBuilder = new StringBuilder();
+			if(NewIssue.Owner == null)
+			{
+				OwnerTextBuilder.Append("Currently unassigned.");
+			}
+			else
+			{
+				if(String.Compare(NewIssue.Owner, IssueMonitor.UserName, StringComparison.OrdinalIgnoreCase) == 0 && NewIssue.NominatedBy != null)
+				{
+					OwnerTextBuilder.AppendFormat("You have been nominated to fix this issue by {0}.", Utility.FormatUserName(NewIssue.NominatedBy));
+				}
+				else
+				{
+					OwnerTextBuilder.AppendFormat("Assigned to {0}", Utility.FormatUserName(NewIssue.Owner));
+					if(NewIssue.NominatedBy != null)
+					{
+						OwnerTextBuilder.AppendFormat(" by {0}", Utility.FormatUserName(NewIssue.NominatedBy));
+					}
+					OwnerTextBuilder.Append(".");
+				}
+			}
+			OwnerTextBuilder.AppendFormat(" Open for {0}.", Utility.FormatDurationMinutes((int)(NewIssue.RetrievedAt - NewIssue.CreatedAt).TotalMinutes));
+			string OwnerText = OwnerTextBuilder.ToString();
+
+			bool bNewIsWarning = Issue.Builds.Count > 0 && !Issue.Builds.Any(x => x.Outcome != IssueBuildOutcome.Warning);
+
+			Issue = NewIssue;
+
+			if(NewIssue.Summary != SummaryLabel.Text || OwnerText != OwnerLabel.Text || Reason != NewReason || bIsWarning != bNewIsWarning)
+			{
+				Rectangle PrevBounds = Bounds;
+				SuspendLayout();
+
+				SummaryLabel.Text = NewIssue.Summary;
+				OwnerLabel.Text = OwnerText;
+
+				if(bStrongAlert)
+				{
+					SummaryLabel.ForeColor = Color.FromArgb(255, 255, 255);
+					SummaryLabel.LinkColor = Color.FromArgb(255, 255, 255);
+					OwnerLabel.ForeColor = Color.FromArgb(255, 255, 255);
+//					AcceptBtn.BackgroundColor1 = Color.FromArgb(214, 69, 64);//Color.FromArgb(255, 255, 255);
+//					AcceptBtn.BackgroundColor2 = Color.FromArgb(214, 69, 64);//Color.FromArgb(255, 255, 255);
+//					AcceptBtn.BorderColor = Color.FromArgb(255, 255, 255);
+//					AcceptBtn.ForeColor = Color.FromArgb(255, 255, 255);//DetailsBtn.ForeColor;
+					AcceptBtn.Margin = new Padding(0,0, 20, 0);//.Right = 20;
+					LatestBuildLinkLabel.LinkColor = Color.FromArgb(255, 255, 255);
+				}
+
+				if(bIsWarning != bNewIsWarning)
+				{
+					bIsWarning = bNewIsWarning;
+					IconPictureBox.Image = bIsWarning? Properties.Resources.AlertWarningIcon : Properties.Resources.AlertErrorIcon;
+				}
+
+				if(Reason != NewReason)
+				{
+					Reason = NewReason;
+
+					List<Button> Buttons = new List<Button>();
+					if(!bStrongAlert)
+					{
+						Buttons.Add(DetailsBtn);
+					}
+					if((NewReason & IssueAlertReason.Owner) != 0)
+					{
+						AcceptBtn.Text = "Acknowledge";
+						Buttons.Add(AcceptBtn);
+					}
+					else if((NewReason & IssueAlertReason.Normal) != 0)
+					{
+						AcceptBtn.Text = "Will Fix";
+						Buttons.Add(AcceptBtn);
+						DeclineBtn.Text = "Not Me";
+						Buttons.Add(DeclineBtn);
+					}
+
+					tableLayoutPanel3.ColumnCount = Buttons.Count;
+					tableLayoutPanel3.Controls.Clear();
+					for(int Idx = 0; Idx < Buttons.Count; Idx++)
+					{
+						tableLayoutPanel3.Controls.Add(Buttons[Idx], Idx, 0);
+					}
+				}
+
+				ResumeLayout(true);
+				Location = new Point(PrevBounds.Right - Bounds.Width, PrevBounds.Y);
+				Invalidate();
+			}
+		}
+
+		protected override void OnHandleCreated(EventArgs e)
+		{
+			base.OnHandleCreated(e);
+
+			// Manually make the window topmost, so that we can pass SWP_NOACTIVATE
+			SetWindowPos(Handle, HWND_TOPMOST, 0, 0, 0, 0, SWP_NOMOVE | SWP_NOSIZE | SWP_NOACTIVATE);
+		}
+
+		protected override bool ShowWithoutActivation 
+		{
+			get { return true; }
+		}
+
+		private void SummaryLabel_LinkClicked(object sender, System.Windows.Forms.LinkLabelLinkClickedEventArgs e)
+		{
+			LaunchUrl();
+		}
+
+		private void IssueAlertWindow_Click(object sender, EventArgs e)
+		{
+			LaunchUrl();
+		}
+
+		public void LaunchUrl()
+		{
+			IssueBuildData LastBuild = Issue.Builds.OrderByDescending(x => x.Change).FirstOrDefault();
+			if(LastBuild != null)
+			{
+				System.Diagnostics.Process.Start(LastBuild.Url);
+			}
+		}
+	}
+}
diff --git a/Engine/Source/Programs/UnrealGameSync/UnrealGameSync/Forms/IssueDetailsWindow.cs b/Engine/Source/Programs/UnrealGameSync/UnrealGameSync/Forms/IssueDetailsWindow.cs
new file mode 100644
--- /dev/null
+++ b/Engine/Source/Programs/UnrealGameSync/UnrealGameSync/Forms/IssueDetailsWindow.cs
@@ -0,0 +1,1125 @@
+// Copyright 1998-2019 Epic Games, Inc. All Rights Reserved.
+
+using System;
+using System.Collections.Generic;
+using System.Data;
+using System.Drawing;
+using System.Drawing.Drawing2D;
+using System.IO;
+using System.Linq;
+using System.Text;
+using System.Text.RegularExpressions;
+using System.Threading;
+using System.Windows.Forms;
+
+namespace UnrealGameSync
+{
+	partial class IssueDetailsWindow : Form
+	{
+		class PerforceChangeRange
+		{
+			public bool bExpanded;
+			public int MinChange;
+			public int MaxChange;
+			public List<PerforceChangeSummary> Changes;
+			public string ErrorMessage;
+			public IssueBuildData Build;
+		}
+
+		class PerforceChangeDetailsWithDescribeRecord : PerforceChangeDetails
+		{
+			public PerforceDescribeRecord DescribeRecord;
+
+			public PerforceChangeDetailsWithDescribeRecord(PerforceDescribeRecord DescribeRecord)
+				: base(DescribeRecord)
+			{
+				this.DescribeRecord = DescribeRecord;
+			}
+		}
+
+		class PerforceWorkerThread : IDisposable
+		{
+			readonly PerforceConnection Perforce;
+			readonly string Filter;
+			readonly Action<PerforceChangeRange> OnUpdateChanges;
+			readonly Action<PerforceChangeSummary> OnUpdateChangeMetadata;
+			readonly TextWriter Log;
+			readonly object LockObject = new object();
+			Thread WorkerThread;
+			bool bTerminate;
+			AutoResetEvent RefreshEvent;
+			List<PerforceChangeRange> Requests = new List<PerforceChangeRange>();
+			Dictionary<int, PerforceChangeDetailsWithDescribeRecord> ChangeNumberToDetails = new Dictionary<int, PerforceChangeDetailsWithDescribeRecord>();
+
+			public PerforceWorkerThread(PerforceConnection Perforce, string Filter, Action<PerforceChangeRange> OnUpdateChanges, Action<PerforceChangeSummary> OnUpdateChangeMetadata, TextWriter Log)
+			{
+				this.Perforce = Perforce;
+				this.Filter = Filter;
+				this.OnUpdateChanges = OnUpdateChanges;
+				this.OnUpdateChangeMetadata = OnUpdateChangeMetadata;
+				this.Log = Log;
+
+				RefreshEvent = new AutoResetEvent(true);
+				WorkerThread = new Thread(() => Run());
+				WorkerThread.Start();
+			}
+
+			public void AddRequest(PerforceChangeRange Range)
+			{
+				lock(LockObject)
+				{
+					Requests.Add(Range);
+				}
+				RefreshEvent.Set();
+			}
+
+			public void Dispose()
+			{
+				if(WorkerThread != null)
+				{
+					bTerminate = true;
+					if(!WorkerThread.Join(100))
+					{
+						WorkerThread.Abort();
+						WorkerThread.Join();
+					}
+					WorkerThread = null;
+				}
+				if(RefreshEvent != null)
+				{
+					RefreshEvent.Dispose();
+					RefreshEvent = null;
+				}
+			}
+
+			void Run()
+			{
+				List<PerforceChangeRange> CompletedRequests = new List<PerforceChangeRange>(); 
+				while(!bTerminate)
+				{
+					// Check if there's a request in the queue
+					PerforceChangeRange NextRequest = null;
+					lock(LockObject)
+					{
+						if(Requests.Count > 0)
+						{
+							NextRequest = Requests[0];
+							Requests.RemoveAt(0);
+						}
+					}
+
+					// Process the request
+					if(NextRequest != null)
+					{
+						string RangeFilter = String.Format("{0}@{1},{2}", Filter, NextRequest.MinChange, (NextRequest.MaxChange == -1)? "now" : NextRequest.MaxChange.ToString());
+
+						List<PerforceChangeSummary> NewChanges;
+						if(Perforce.FindChanges(RangeFilter, -1, out NewChanges, Log))
+						{
+							NextRequest.Changes = NewChanges;
+							CompletedRequests.Add(NextRequest);
+						}
+						else
+						{
+							NextRequest.ErrorMessage = "Unable to fetch changes. Check your connection settings and try again.";
+						}
+
+						OnUpdateChanges(NextRequest);
+						continue;
+					}
+
+					// Figure out which changes to fetch
+					List<PerforceChangeSummary> DescribeChanges;
+					lock(LockObject)
+					{
+						DescribeChanges = CompletedRequests.SelectMany(x => x.Changes).Where(x => !ChangeNumberToDetails.ContainsKey(x.Number)).ToList();
+					}
+
+					// Fetch info on each individual change
+					foreach(PerforceChangeSummary DescribeChange in DescribeChanges)
+					{
+						lock(LockObject)
+						{
+							if(bTerminate || Requests.Count > 0)
+							{
+								break;
+							}
+						}
+
+						PerforceDescribeRecord Record;
+						if(Perforce.Describe(DescribeChange.Number, out Record, Log))
+						{
+							lock(LockObject)
+							{
+								ChangeNumberToDetails[Record.ChangeNumber] = new PerforceChangeDetailsWithDescribeRecord(Record);
+								if((Record.ChangeNumber & 3) == 0)
+								{
+									ChangeNumberToDetails[Record.ChangeNumber].bContainsCode = true;
+									ChangeNumberToDetails[Record.ChangeNumber].bContainsContent = true;
+								}
+							}
+						}
+
+						OnUpdateChangeMetadata(DescribeChange);
+					}
+
+					// Wait for something to change
+					RefreshEvent.WaitOne();
+				}
+			}
+
+			public bool TryGetChangeDetails(int ChangeNumber, out PerforceChangeDetailsWithDescribeRecord Details)
+			{
+				lock(LockObject)
+				{
+					return ChangeNumberToDetails.TryGetValue(ChangeNumber, out Details);
+				}
+			}
+		}
+
+		class BadgeInfo
+		{
+			public string Label;
+			public Color Color;
+
+			public BadgeInfo(string Label, Color Color)
+			{
+				this.Label = Label;
+				this.Color = Color;
+			}
+		}
+
+		class ExpandRangeStatusElement : StatusElement
+		{
+			string Text;
+			Action LinkAction;
+
+			public ExpandRangeStatusElement(string Text, Action InLinkAction)
+			{
+				this.Text = Text;
+				this.LinkAction = InLinkAction;
+				Cursor = NativeCursors.Hand;
+			}
+
+			public override void OnClick(Point Location)
+			{
+				LinkAction();
+			}
+
+			public override Size Measure(Graphics Graphics, StatusElementResources Resources)
+			{
+				return TextRenderer.MeasureText(Graphics, Text, Resources.FindOrAddFont(FontStyle.Regular), new Size(int.MaxValue, int.MaxValue), TextFormatFlags.NoPadding);
+			}
+
+			public override void Draw(Graphics Graphics, StatusElementResources Resources)
+			{
+				Color TextColor = Color.Gray;
+				FontStyle Style = FontStyle.Italic;
+				if(bMouseDown)
+				{
+					TextColor = Color.FromArgb(TextColor.B / 2, TextColor.G / 2, TextColor.R);
+					Style |= FontStyle.Underline;
+				}
+				else if(bMouseOver)
+				{
+					TextColor = Color.FromArgb(TextColor.B, TextColor.G, TextColor.R);
+					Style |= FontStyle.Underline;
+				}
+				TextRenderer.DrawText(Graphics, Text, Resources.FindOrAddFont(Style), Bounds.Location, TextColor, TextFormatFlags.NoPadding);
+			}
+		}
+
+		IssueMonitor IssueMonitor;
+		IssueData Issue;
+		PerforceConnection Perforce;
+		PerforceWorkerThread PerforceWorker;
+		TextWriter Log;
+		SynchronizationContext MainThreadSynchronizationContext;
+		string SelectedStream;
+		List<PerforceChangeRange> SelectedStreamRanges = new List<PerforceChangeRange>();
+		bool bIsDisposing;
+		Font BoldFont;
+		PerforceChangeSummary ContextMenuChange;
+		string LastOwner;
+		System.Windows.Forms.Timer UpdateTimer;
+		StatusElementResources StatusElementResources;
+
+		IssueDetailsWindow(IssueMonitor IssueMonitor, IssueData Issue, string ServerAndPort, string UserName, TextWriter Log, string CurrentStream)
+		{
+			this.IssueMonitor = IssueMonitor;
+			this.Issue = Issue;
+			this.Perforce = new PerforceConnection(UserName, null, ServerAndPort);
+			this.Log = Log;
+
+			MainThreadSynchronizationContext = SynchronizationContext.Current;
+
+			InitializeComponent();
+
+			this.StatusElementResources = new StatusElementResources(BuildListView.Font);
+			base.Disposed += IssueDetailsWindow_Disposed;
+
+			IssueMonitor.OnIssuesChanged += OnUpdateIssuesAsync;
+			IssueMonitor.StartTracking(Issue.Id);
+
+			BuildListView.SmallImageList = new ImageList(){ ImageSize = new Size(1, 20) };
+
+			System.Reflection.PropertyInfo DoubleBufferedProperty = typeof(Control).GetProperty("DoubleBuffered", System.Reflection.BindingFlags.NonPublic | System.Reflection.BindingFlags.Instance);
+			DoubleBufferedProperty.SetValue(BuildListView, true, null); 
+
+			using(Graphics Graphics = Graphics.FromHwnd(IntPtr.Zero))
+			{
+				float DpiScaleX = Graphics.DpiX / 96.0f;
+				foreach(ColumnHeader Column in BuildListView.Columns)
+				{
+					Column.Width = (int)(Column.Width * DpiScaleX);
+				}
+			}
+
+			int SelectIdx = 0;
+			foreach(string Stream in Issue.Builds.Select(x => x.Stream).Distinct().OrderBy(x => x))
+			{
+				StreamComboBox.Items.Add(Stream);
+				if(Stream == CurrentStream)
+				{
+					SelectIdx = StreamComboBox.Items.Count - 1;
+				}
+			}
+			if(StreamComboBox.Items.Count > 0)
+			{
+				StreamComboBox.SelectedIndex = SelectIdx;
+			}
+
+			FilterTypeComboBox.SelectedIndex = 0;
+
+			CreateWorker();
+
+			UpdateCurrentIssue();
+			UpdateEnabledButtons();
+		}
+
+		private void StartUpdateTimer()
+		{
+			if(UpdateTimer == null)
+			{
+				UpdateTimer = new System.Windows.Forms.Timer();
+				UpdateTimer.Interval = 100;
+				UpdateTimer.Tick += UpdateTimer_Tick;
+				UpdateTimer.Start();
+
+				components.Add(UpdateTimer);
+			}
+		}
+
+		private void StopUpdateTimer()
+		{
+			if(UpdateTimer != null)
+			{
+				components.Remove(UpdateTimer);
+
+				UpdateTimer.Dispose();
+				UpdateTimer = null;
+			}
+		}
+
+		private void UpdateTimer_Tick(object sender, EventArgs e)
+		{
+			UpdateBuildList();
+			StopUpdateTimer();
+		}
+
+		void UpdateSummaryTextIfChanged(TextBox TextBox, string NewText)
+		{
+			if(TextBox.Text != NewText)
+			{
+				TextBox.SelectionLength = 0;
+				TextBox.Text = NewText;
+			}
+		}
+
+		void UpdateCurrentIssue()
+		{
+			DescriptionLabel.Text = Issue.Summary.ToString();
+
+			StringBuilder Status = new StringBuilder();
+			if(IssueMonitor.HasPendingUpdate())
+			{
+				Status.Append("Updating...");
+			}
+			else if(Issue.FixChange != 0)
+			{
+				if(Issue.FixChange < 0)
+				{
+					if(Issue.ResolvedAt.HasValue)
+					{
+						Status.AppendFormat("Closed as systemic issue.", Issue.FixChange);
+					}
+					else
+					{
+						Status.AppendFormat("Fixed as systemic issue (pending verification).", Issue.FixChange);
+					}
+				}
+				else
+				{
+					if(Issue.ResolvedAt.HasValue)
+					{
+						Status.AppendFormat("Closed. Fixed in CL {0}.", Issue.FixChange);
+					}
+					else
+					{
+						Status.AppendFormat("Fixed in CL {0} (pending verification)", Issue.FixChange);
+					}
+				}
+			}
+			else if(Issue.ResolvedAt.HasValue)
+			{
+				Status.Append("Resolved");
+			}
+			else if(Issue.Owner == null)
+			{
+				Status.Append("Currently unassigned");
+			}
+			else
+			{
+				Status.Append(Utility.FormatUserName(Issue.Owner));
+				if(Issue.NominatedBy != null)
+				{
+					Status.AppendFormat(" nominated by {0}", Utility.FormatUserName(Issue.NominatedBy));
+				}
+				if(Issue.AcknowledgedAt.HasValue)
+				{
+					Status.AppendFormat(" (acknowledged {0})", Utility.FormatRecentDateTime(Issue.AcknowledgedAt.Value.ToLocalTime()));
+				}
+				else
+				{
+					Status.Append(" (not acknowledged)");
+				}
+			}
+			UpdateSummaryTextIfChanged(StatusTextBox, Status.ToString());
+
+			StringBuilder OpenSince = new StringBuilder();
+			OpenSince.Append(Utility.FormatRecentDateTime(Issue.CreatedAt.ToLocalTime()));
+			if(OpenSince.Length > 0)
+			{
+				OpenSince[0] = Char.ToUpper(OpenSince[0]);
+			}
+			OpenSince.AppendFormat(" ({0})", Utility.FormatDurationMinutes(Issue.RetrievedAt - Issue.CreatedAt));
+			UpdateSummaryTextIfChanged(OpenSinceTextBox, OpenSince.ToString());
+
+			if(LastOwner != Issue.Owner)
+			{
+				LastOwner = Issue.Owner;
+				BuildListView.Invalidate();
+			}
+
+			ReopenBtn.Enabled = (Issue.FixChange != 0);
+		}
+
+		private void IssueDetailsWindow_Disposed(object sender, EventArgs e)
+		{
+			bIsDisposing = true;
+			DestroyWorker();
+
+			if(IssueMonitor != null)
+			{
+				IssueMonitor.StopTracking(Issue.Id);
+				IssueMonitor.OnIssuesChanged -= OnUpdateIssuesAsync;
+				IssueMonitor = null;
+			}
+		}
+
+		void FetchAllBuildChanges()
+		{
+			foreach(PerforceChangeRange Range in SelectedStreamRanges)
+			{
+				FetchBuildChanges(Range);
+			}
+		}
+
+		void FetchBuildChanges(PerforceChangeRange Range)
+		{
+			if(!Range.bExpanded)
+			{
+				Range.bExpanded = true;
+				PerforceWorker.AddRequest(Range);
+				UpdateBuildList();
+			}
+		}
+
+		bool FilterMatch(Regex FilterRegex, PerforceChangeSummary Summary)
+		{
+			if(FilterRegex.IsMatch(Summary.User))
+			{
+				return true;
+			}
+			if(FilterRegex.IsMatch(Summary.Description))
+			{
+				return true;
+			}
+			return false;
+		}
+
+		bool FilterMatch(Regex FilterRegex, PerforceDescribeRecord DescribeRecord)
+		{
+			if(FilterRegex.IsMatch(DescribeRecord.User))
+			{
+				return true;
+			}
+			if(FilterRegex.IsMatch(DescribeRecord.Description))
+			{
+				return true;
+			}
+			if(DescribeRecord.Files.Any(x => FilterRegex.IsMatch(x.DepotFile)))
+			{
+				return true;
+			}
+			return false;
+		}
+
+		void UpdateBuildList()
+		{
+			int NumNewItems = 0;
+			BuildListView.BeginUpdate();
+
+			// Capture the initial selection
+			object PrevSelection = null;
+			if(BuildListView.SelectedItems.Count > 0)
+			{
+				PrevSelection = BuildListView.SelectedItems[0].Tag;
+			}
+
+			// Get all the search terms
+			string[] FilterTerms = FilterTextBox.Text.Split(new char[]{ ' ' }, StringSplitOptions.RemoveEmptyEntries);
+
+			// Build a regex from each filter term
+			List<Regex> FilterRegexes = new List<Regex>();
+			foreach(string FilterTerm in FilterTerms)
+			{
+				string RegexText = Regex.Escape(FilterTerm);
+				RegexText = RegexText.Replace("\\?", ".");
+				RegexText = RegexText.Replace("\\*", "[^\\\\/]*");
+				RegexText = RegexText.Replace("\\.\\.\\.", ".*");
+				FilterRegexes.Add(new Regex(RegexText, RegexOptions.IgnoreCase));
+			}
+
+			// Get the filter type
+			bool bOnlyShowCodeChanges = (FilterTypeComboBox.SelectedIndex == 1);
+			bool bOnlyShowContentChanges = (FilterTypeComboBox.SelectedIndex == 2);
+
+			// Create rows for all the ranges
+			foreach(PerforceChangeRange Range in SelectedStreamRanges)
+			{
+				if(Range.bExpanded)
+				{
+					if(Range.Changes == null)
+					{
+						ListViewItem FetchingItem = new ListViewItem("");
+						FetchingItem.Tag = Range;
+						FetchingItem.SubItems.Add("");
+						FetchingItem.SubItems.Add("");
+
+						StatusLineListViewWidget FetchingWidget = new StatusLineListViewWidget(FetchingItem, StatusElementResources);
+						FetchingWidget.HorizontalAlignment = HorizontalAlignment.Left;
+						if(Range.ErrorMessage != null)
+						{
+							FetchingWidget.Line.AddText(Range.ErrorMessage, Color.Gray, FontStyle.Italic);
+						}
+						else
+						{
+							FetchingWidget.Line.AddText("Fetching changes, please wait...", Color.Gray, FontStyle.Italic);
+						}
+			
+						FetchingItem.SubItems.Add(new ListViewItem.ListViewSubItem(FetchingItem, ""){ Tag = FetchingWidget });
+						FetchingItem.SubItems.Add(new ListViewItem.ListViewSubItem(FetchingItem, ""){ Tag = FetchingWidget });
+
+						BuildListView.Items.Insert(NumNewItems++, FetchingItem);
+					}
+					else
+					{
+						foreach(PerforceChangeSummary Change in Range.Changes)
+						{
+							PerforceChangeDetailsWithDescribeRecord Details;
+							PerforceWorker.TryGetChangeDetails(Change.Number, out Details);
+
+							if(FilterRegexes.Count > 0 || bOnlyShowCodeChanges || bOnlyShowContentChanges)
+							{
+								if(Details == null)
+								{
+									if(bOnlyShowCodeChanges || bOnlyShowCodeChanges)
+									{
+										continue;
+									}
+									if(FilterRegexes.Any(x => !FilterMatch(x, Change)))
+									{
+										continue;
+									}
+								}
+								else
+								{
+									if(bOnlyShowCodeChanges && !Details.bContainsCode)
+									{
+										continue;
+									}
+									if(bOnlyShowContentChanges && !Details.bContainsContent)
+									{
+										continue;
+									}
+									if(FilterRegexes.Any(x => !FilterMatch(x, Details.DescribeRecord)))
+									{
+										continue;
+									}
+								}
+							}
+
+							ListViewItem Item = new ListViewItem("");
+							Item.Tag = Change;
+							Item.SubItems.Add(Change.Number.ToString());
+
+							StatusLineListViewWidget TypeWidget = new StatusLineListViewWidget(Item, StatusElementResources);
+							UpdateChangeTypeWidget(TypeWidget, Details);
+							Item.SubItems.Add(new ListViewItem.ListViewSubItem(Item, "") { Tag = TypeWidget });
+
+							Item.SubItems.Add(WorkspaceControl.FormatUserName(Change.User));
+							Item.SubItems.Add(Change.Description);
+							BuildListView.Items.Insert(NumNewItems++, Item);
+						}
+					}
+				}
+				/*
+				else
+				{
+					ListViewItem RangeItem = new ListViewItem("");
+					RangeItem.Tag = Range;
+
+					StatusLineListViewWidget RangeWidget = new StatusLineListViewWidget(RangeItem, StatusElementResources);
+					RangeWidget.Line.AddText("-", Color.Gray);//.AddBadge("...", Color.LightGray, () => FetchBuildChanges(RangeItem));
+					RangeItem.SubItems.Add(new ListViewItem.ListViewSubItem(RangeItem, ""){ Tag = RangeWidget });
+
+					StatusLineListViewWidget RangeWidget2 = new StatusLineListViewWidget(RangeItem, StatusElementResources);
+					RangeWidget2.Line.AddText("-", Color.Gray);//.AddBadge("...", Color.LightGray, () => FetchBuildChanges(RangeItem));
+					RangeItem.SubItems.Add(new ListViewItem.ListViewSubItem(RangeItem, "-"){ Tag = RangeWidget2 });
+
+					StatusLineListViewWidget TipWidget = new StatusLineListViewWidget(RangeItem, StatusElementResources);
+					TipWidget.HorizontalAlignment = HorizontalAlignment.Left;
+					TipWidget.Line.Add(new ExpandRangeStatusElement("Click to show changes...", () => FetchAllBuildChanges()));
+
+					RangeItem.SubItems.Add(new ListViewItem.ListViewSubItem(RangeItem, ""){ Tag = TipWidget });
+					RangeItem.SubItems.Add(new ListViewItem.ListViewSubItem(RangeItem, ""){ Tag = TipWidget });
+
+					BuildListView.Items.Insert(NumNewItems++, RangeItem);
+				}
+				*/
+				ListViewItem BuildItem = new ListViewItem("");
+				BuildItem.Tag = Range.Build;
+				BuildItem.SubItems.Add(Range.Build.Change.ToString());
+				BuildItem.SubItems.Add("");
+
+				StatusLineListViewWidget BuildWidget = new StatusLineListViewWidget(BuildItem, StatusElementResources);
+				BuildWidget.HorizontalAlignment = HorizontalAlignment.Left;
+				BuildWidget.Line.AddLink(Range.Build.Name, FontStyle.Underline, () => System.Diagnostics.Process.Start(Range.Build.Url));
+				BuildItem.SubItems.Add(new ListViewItem.ListViewSubItem(BuildItem, ""){ Tag = BuildWidget });
+				BuildItem.SubItems.Add(new ListViewItem.ListViewSubItem(BuildItem, ""){ Tag = BuildWidget });
+
+				BuildListView.Items.Insert(NumNewItems++, BuildItem);
+			}
+			
+			// Re-select the original item
+			for(int Idx = 0; Idx < BuildListView.Items.Count; Idx++)
+			{
+				if(BuildListView.Items[Idx].Tag == PrevSelection)
+				{
+					BuildListView.Items[Idx].Selected = true;
+				}
+			}
+
+			// Remove all the items we no longer need
+			while(BuildListView.Items.Count > NumNewItems)
+			{
+				BuildListView.Items.RemoveAt(BuildListView.Items.Count - 1);
+			}
+
+			// Redraw everything
+			BuildListView.EndUpdate();
+			BuildListView.Invalidate();
+		}
+		
+		void UpdateChangeTypeWidget(StatusLineListViewWidget TypeWidget, PerforceChangeDetails Details)
+		{
+			TypeWidget.Line.Clear();
+			if(Details == null)
+			{
+				TypeWidget.Line.AddBadge("Unknown", Color.FromArgb(192, 192, 192), null);
+			}
+			else
+			{
+				if(Details.bContainsCode)
+				{
+					TypeWidget.Line.AddBadge("Code", Color.FromArgb(116, 185, 255), null);
+				}
+				if(Details.bContainsContent)
+				{
+					TypeWidget.Line.AddBadge("Content", Color.FromArgb(162, 155, 255), null);
+				}
+			}
+		}
+
+		void CreateWorker()
+		{
+			string NewSelectedStream = StreamComboBox.SelectedItem as string;
+			if(SelectedStream != NewSelectedStream)
+			{
+				DestroyWorker();
+			
+				SelectedStream = NewSelectedStream;
+
+				BuildListView.BeginUpdate();
+				BuildListView.Items.Clear();
+
+				if(SelectedStream != null)
+				{
+					SelectedStreamRanges = new List<PerforceChangeRange>();
+
+					int MaxChange = -1;
+					foreach(IssueBuildData Build in Issue.Builds.Where(x => x.Stream == SelectedStream).OrderByDescending(x => x.Change).ThenByDescending(x => x.Url))
+					{
+						PerforceChangeRange Range = new PerforceChangeRange();
+						Range.MinChange = Build.Change + 1;
+						Range.MaxChange = MaxChange;
+						Range.Build = Build;
+						SelectedStreamRanges.Add(Range);
+
+						MaxChange = Build.Change;
+					}
+
+					PerforceWorker = new PerforceWorkerThread(Perforce, String.Format("{0}/...", SelectedStream), OnRequestCompleteAsync, UpdateChangeMetadataAsync, Log);
+
+					UpdateBuildList();
+
+					for(int Idx = 0; Idx + 2 < SelectedStreamRanges.Count; Idx++)
+					{
+						if(SelectedStreamRanges[Idx].Build.Outcome != SelectedStreamRanges[Idx + 1].Build.Outcome)
+						{
+							FetchBuildChanges(SelectedStreamRanges[Idx + 1]);
+						}
+					}
+					FetchBuildChanges(SelectedStreamRanges[SelectedStreamRanges.Count - 1]);
+				}
+
+				BuildListView.EndUpdate();
+				BuildListView.Invalidate();
+			}
+		}
+
+		void UpdateEnabledButtons()
+		{
+			AssignBtn.Enabled = BuildListView.SelectedItems.Count > 0 && (BuildListView.SelectedItems[0].Tag as PerforceChangeSummary) != null;
+		}
+
+		void OnUpdateIssues()
+		{
+			List<IssueData> NewIssues = IssueMonitor.GetIssues();
+			foreach(IssueData NewIssue in NewIssues)
+			{
+				if(NewIssue.Id == Issue.Id)
+				{
+					Issue = NewIssue;
+					UpdateCurrentIssue();
+					break;
+				}
+			}
+		}
+
+		private void OnUpdateIssuesAsync()
+		{
+			MainThreadSynchronizationContext.Post((o) => { if(!bIsDisposing){ OnUpdateIssues(); } }, null);
+		}
+
+		void OnUpdateError(Exception Ex)
+		{
+			MessageBox.Show(String.Format("Unable to update database.\n\n{0}", Ex));
+		}
+
+		void OnUpdateErrorAsync(Exception Ex)
+		{
+			MainThreadSynchronizationContext.Post((o) => { if(!bIsDisposing){ OnUpdateError(Ex); } }, null);
+		}
+
+		void DestroyWorker()
+		{
+			if(PerforceWorker != null)
+			{
+				PerforceWorker.Dispose();
+				PerforceWorker = null;
+			}
+		}
+
+		void OnRequestComplete(PerforceChangeRange Range)
+		{
+			UpdateBuildList();
+		}
+
+		void OnRequestCompleteAsync(PerforceChangeRange Range)
+		{
+			MainThreadSynchronizationContext.Post((o) => { if(!bIsDisposing) { OnRequestComplete(Range); } }, null);
+		}
+
+		void UpdateChangeMetadata(PerforceChangeSummary Change)
+		{
+			if(FilterTextBox.Text.Length > 0 || FilterTypeComboBox.SelectedIndex != 0)
+			{
+				StartUpdateTimer();
+			}
+			else
+			{
+				foreach(ListViewItem Item in BuildListView.Items)
+				{
+					if(Item.Tag == Change)
+					{
+						PerforceChangeDetailsWithDescribeRecord Details;
+						PerforceWorker.TryGetChangeDetails(Change.Number, out Details);
+
+						StatusLineListViewWidget TypeWidget = (StatusLineListViewWidget)Item.SubItems[2].Tag;
+						UpdateChangeTypeWidget(TypeWidget, Details);
+
+						BuildListView.RedrawItems(Item.Index, Item.Index, true);
+						break;
+					}
+				}
+			}
+		}
+
+		void UpdateChangeMetadataAsync(PerforceChangeSummary Change)
+		{
+			MainThreadSynchronizationContext.Post((o) => { if(!bIsDisposing) { UpdateChangeMetadata(Change); } }, null);
+		}
+
+		public static void ShowModal(IWin32Window Owner, IssueMonitor IssueMonitor, string ServerAndPort, string UserName, IssueData Issue, TextWriter Log, string CurrentStream)
+		{
+			using(IssueDetailsWindow Window = new IssueDetailsWindow(IssueMonitor, Issue, ServerAndPort, UserName, Log, CurrentStream))
+			{
+				Window.ShowDialog(Owner);
+			}
+		}
+
+		private void CloseBtn_Click(object sender, EventArgs e)
+		{
+			DialogResult = DialogResult.OK;
+			Close();
+		}
+
+		private void StreamComboBox_SelectedIndexChanged(object sender, EventArgs e)
+		{
+			DestroyWorker();
+			CreateWorker();
+		}
+
+		private void BuildListView_DrawColumnHeader(object sender, DrawListViewColumnHeaderEventArgs e)
+		{
+			e.DrawDefault = true;
+		}
+
+		private void BuildListView_DrawItem(object sender, DrawListViewItemEventArgs e)
+		{
+			if(e.Item.Selected)
+			{
+				BuildListView.DrawSelectedBackground(e.Graphics, e.Bounds);
+			}
+			else if(e.ItemIndex == BuildListView.HoverItem)
+			{
+				BuildListView.DrawTrackedBackground(e.Graphics, e.Bounds);
+			}
+			else if(e.Item.Tag is IssueBuildData)
+			{
+				IssueBuildData BuildData = (IssueBuildData)e.Item.Tag;
+
+				Color BackgroundColor;
+				if(BuildData.Outcome == IssueBuildOutcome.Error)
+				{
+					BackgroundColor = Color.FromArgb(254, 248, 246);
+				}
+				else if(BuildData.Outcome == IssueBuildOutcome.Warning)
+				{
+					BackgroundColor = Color.FromArgb(254, 254, 246);
+				}
+				else if(BuildData.Outcome == IssueBuildOutcome.Success)
+				{
+					BackgroundColor = Color.FromArgb(248, 254, 246);
+				}
+				else
+				{
+					BackgroundColor = Color.FromArgb(245, 245, 245);
+				}
+
+				using(SolidBrush Brush = new SolidBrush(BackgroundColor))
+				{
+					e.Graphics.FillRectangle(Brush, e.Bounds);
+				}
+			}
+			else
+			{
+				BuildListView.DrawDefaultBackground(e.Graphics, e.Bounds);
+			}
+		}
+
+		private void BuildListView_DrawSubItem(object sender, DrawListViewSubItemEventArgs e)
+		{
+			Font CurrentFont = this.Font;//BuildFont;(Change.Number == Workspace.PendingChangeNumber || Change.Number == Workspace.CurrentChangeNumber)? SelectedBuildFont : BuildFont;
+
+			Color TextColor = SystemColors.WindowText;
+
+			if(e.SubItem.Tag is CustomListViewWidget)
+			{
+				BuildListView.DrawCustomSubItem(e.Graphics, e.SubItem);
+			}
+			else if(e.Item.Tag is PerforceChangeSummary)
+			{
+				PerforceChangeSummary Change = (PerforceChangeSummary)e.Item.Tag;
+
+				Font ChangeFont = BuildListView.Font;
+				if(Issue.Owner != null && String.Compare(Change.User, Issue.Owner, StringComparison.OrdinalIgnoreCase) == 0)
+				{
+					ChangeFont = BoldFont;
+				}
+
+				if(e.ColumnIndex == ChangeHeader.Index)
+				{
+					TextRenderer.DrawText(e.Graphics, e.SubItem.Text, ChangeFont, e.Bounds, TextColor, TextFormatFlags.EndEllipsis | TextFormatFlags.SingleLine | TextFormatFlags.HorizontalCenter | TextFormatFlags.VerticalCenter | TextFormatFlags.NoPrefix);
+				}
+				else if(e.ColumnIndex == AuthorHeader.Index)
+				{
+					TextRenderer.DrawText(e.Graphics, e.SubItem.Text, ChangeFont, e.Bounds, TextColor, TextFormatFlags.EndEllipsis | TextFormatFlags.SingleLine | TextFormatFlags.Left | TextFormatFlags.VerticalCenter | TextFormatFlags.NoPrefix);
+				}
+				else if(e.ColumnIndex == DescriptionHeader.Index)
+				{
+					TextRenderer.DrawText(e.Graphics, e.SubItem.Text, ChangeFont, e.Bounds, TextColor, TextFormatFlags.EndEllipsis | TextFormatFlags.SingleLine | TextFormatFlags.Left | TextFormatFlags.VerticalCenter | TextFormatFlags.NoPrefix);
+				}
+			}
+			else if(e.Item.Tag is IssueBuildData)
+			{
+				Font BoldFont = BuildListView.Font;
+
+//				TextColor = SystemColors.Window;
+				IssueBuildData BuildData = (IssueBuildData)e.Item.Tag;
+				if(e.ColumnIndex == IconHeader.Index)
+				{
+					if(BuildData.Outcome == IssueBuildOutcome.Success)
+					{
+						DrawIcon(e.Graphics, e.Bounds, WorkspaceControl.GoodBuildIcon);
+					}
+					else if(BuildData.Outcome == IssueBuildOutcome.Warning)
+					{
+						DrawIcon(e.Graphics, e.Bounds, WorkspaceControl.MixedBuildIcon);
+					}
+					else if(BuildData.Outcome == IssueBuildOutcome.Error)
+					{
+						DrawIcon(e.Graphics, e.Bounds, WorkspaceControl.BadBuildIcon);
+					}
+					else
+					{
+						DrawIcon(e.Graphics, e.Bounds, WorkspaceControl.DefaultBuildIcon);
+					}
+				}
+				else if(e.ColumnIndex == ChangeHeader.Index)
+				{
+					TextRenderer.DrawText(e.Graphics, BuildData.Change.ToString(), BoldFont, e.Bounds, TextColor, TextFormatFlags.EndEllipsis | TextFormatFlags.SingleLine | TextFormatFlags.HorizontalCenter | TextFormatFlags.VerticalCenter | TextFormatFlags.NoPrefix);
+				}
+				else if(e.ColumnIndex == TypeHeader.Index)
+				{
+					string Status;
+					if(BuildData.Outcome == IssueBuildOutcome.Success)
+					{
+						Status = "Succeeded";
+					}
+					else if(BuildData.Outcome == IssueBuildOutcome.Warning)
+					{
+						Status = "Warning";
+					}
+					else if(BuildData.Outcome == IssueBuildOutcome.Error)
+					{
+						Status = "Failed";
+					}
+					else
+					{
+						Status = "Pending";
+					}
+					TextRenderer.DrawText(e.Graphics, Status, BoldFont, e.Bounds, TextColor, TextFormatFlags.EndEllipsis | TextFormatFlags.SingleLine | TextFormatFlags.HorizontalCenter | TextFormatFlags.VerticalCenter | TextFormatFlags.NoPrefix);
+				}
+				else if(e.ColumnIndex == AuthorHeader.Index)
+				{
+					Rectangle Bounds = new Rectangle(e.Bounds.X, e.Bounds.Y, e.Bounds.Width + e.Item.SubItems[e.ColumnIndex].Bounds.Width, e.Bounds.Height);
+					TextRenderer.DrawText(e.Graphics, BuildData.Name, Font, Bounds, TextColor, TextFormatFlags.EndEllipsis | TextFormatFlags.SingleLine | TextFormatFlags.Left | TextFormatFlags.VerticalCenter | TextFormatFlags.NoPrefix);
+				}
+			}
+			else
+			{
+				TextFormatFlags Flags = TextFormatFlags.EndEllipsis | TextFormatFlags.SingleLine | TextFormatFlags.VerticalCenter | TextFormatFlags.NoPrefix;
+				if(BuildListView.Columns[e.ColumnIndex].TextAlign == HorizontalAlignment.Left)
+				{
+					Flags |= TextFormatFlags.Left;
+				}
+				else if(BuildListView.Columns[e.ColumnIndex].TextAlign == HorizontalAlignment.Center)
+				{
+					Flags |= TextFormatFlags.HorizontalCenter;
+				}
+				else
+				{
+					Flags |= TextFormatFlags.Right;
+				}
+				TextRenderer.DrawText(e.Graphics, e.SubItem.Text, CurrentFont, e.Bounds, TextColor, Flags);
+			}
+		}
+
+		private void DrawIcon(Graphics Graphics, Rectangle Bounds, Rectangle Icon)
+		{
+			float DpiScaleX = Graphics.DpiX / 96.0f;
+			float DpiScaleY = Graphics.DpiY / 96.0f;
+
+			float IconX = Bounds.Left + (Bounds.Width - 16 * DpiScaleX) / 2;
+			float IconY = Bounds.Top + (Bounds.Height - 16 * DpiScaleY) / 2;
+
+			Graphics.DrawImage(Properties.Resources.Icons, IconX, IconY, Icon, GraphicsUnit.Pixel);
+		}
+
+		private void BuildListView_FontChanged(object sender, EventArgs e)
+		{
+			if(BoldFont != null)
+			{
+				BoldFont.Dispose();
+			}
+			BoldFont = new Font(BuildListView.Font.FontFamily, BuildListView.Font.SizeInPoints, FontStyle.Bold);
+
+			if(StatusElementResources != null)
+			{
+				StatusElementResources.Dispose();
+			}
+			StatusElementResources = new StatusElementResources(BuildListView.Font);
+		}
+
+		private void BuildListView_MouseClick(object Sender, MouseEventArgs Args)
+		{
+			if(Args.Button == MouseButtons.Right)
+			{
+				ListViewHitTestInfo HitTest = BuildListView.HitTest(Args.Location);
+				if(HitTest.Item != null && HitTest.Item.Tag != null)
+				{
+					ContextMenuChange = HitTest.Item.Tag as PerforceChangeSummary; 
+					if(ContextMenuChange != null)
+					{
+						BuildListContextMenu_Assign.Text = String.Format("Assign to {0}", ContextMenuChange.User);
+						BuildListContextMenu.Show(BuildListView, Args.Location);
+					}
+				}
+			}
+		}
+
+		private void BuildListContextMenu_MoreInfo_Click(object sender, EventArgs e)
+		{
+			if(!Utility.SpawnHiddenProcess("p4vc.exe", String.Format("-p\"{0}\" change {1}", Perforce.ServerAndPort, ContextMenuChange.Number)))
+			{
+				MessageBox.Show("Unable to spawn p4vc. Check you have P4V installed.");
+			}
+		}
+
+		private void BuildListView_SelectedIndexChanged(object sender, EventArgs e)
+		{
+			UpdateEnabledButtons();
+		}
+
+		private void BuildListContextMenu_Blame_Click(object sender, EventArgs e)
+		{
+			if(ContextMenuChange != null)
+			{
+				AssignToUser(ContextMenuChange.User);
+			}
+		}
+
+		private void AssignToUser(string User)
+		{
+			IssueUpdateData Update = new IssueUpdateData();
+			Update.Id = Issue.Id;
+			Update.Owner = User;
+			Update.FixChange = 0;
+			if(String.Compare(User, Perforce.UserName, StringComparison.OrdinalIgnoreCase) == 0)
+			{
+				Update.NominatedBy = "";
+				Update.Acknowledged = true;
+			}
+			else
+			{
+				Update.NominatedBy = Perforce.UserName;
+				Update.Acknowledged = false;
+			}
+			IssueMonitor.PostUpdate(Update);
+
+			UpdateCurrentIssue();
+		}
+
+		private void AssignBtn_Click(object sender, EventArgs e)
+		{
+			if(BuildListView.SelectedItems.Count > 0)
+			{
+				PerforceChangeSummary Change = BuildListView.SelectedItems[0].Tag as PerforceChangeSummary;
+				if(Change != null)
+				{
+					AssignToUser(Change.User);
+				}
+			}
+		}
+
+		private void AssignToMeBtn_Click(object sender, EventArgs e)
+		{
+			AssignToUser(IssueMonitor.UserName);
+		}
+
+		private void AssignToOtherBtn_Click(object sender, EventArgs e)
+		{
+			string SelectedUserName;
+			if(SelectUserWindow.ShowModal(this, Perforce.ServerAndPort, Perforce.UserName, Log, out SelectedUserName))
+			{
+				AssignToUser(SelectedUserName);
+			}
+		}
+
+		private void FilterTextBox_TextChanged(object sender, EventArgs e)
+		{
+			StopUpdateTimer();
+			StartUpdateTimer();
+		}
+
+		private void FilterTypeComboBox_SelectedIndexChanged(object sender, EventArgs e)
+		{
+			UpdateBuildList();
+		}
+
+		private void MarkFixedBtn_Click(object sender, EventArgs e)
+		{
+			int FixChangeNumber = Issue.FixChange;
+			if(IssueFixedWindow.ShowModal(this, Perforce, ref FixChangeNumber))
+			{
+				IssueUpdateData Update = new IssueUpdateData();
+				Update.Id = Issue.Id;
+				Update.FixChange = FixChangeNumber;
+				IssueMonitor.PostUpdate(Update);
+			}
+		}
+
+		private void DescriptionLabel_LinkClicked(object sender, LinkLabelLinkClickedEventArgs e)
+		{
+			IssueBuildData LastBuild = Issue.Builds.Where(x => x.Stream == SelectedStream).OrderByDescending(x => x.Change).ThenByDescending(x => x.Url).FirstOrDefault();
+			if(LastBuild != null)
+			{
+				System.Diagnostics.Process.Start(LastBuild.Url);
+			}
+		}
+
+		private void ReopenBtn_Click(object sender, EventArgs e)
+		{
+			if(Issue.FixChange != 0)
+			{
+				IssueUpdateData Update = new IssueUpdateData();
+				Update.Id = Issue.Id;
+				Update.FixChange = 0;
+				IssueMonitor.PostUpdate(Update);
+			}
+		}
+	}
+}
diff --git a/Engine/Source/Programs/UnrealGameSync/UnrealGameSync/Forms/MainWindow.cs b/Engine/Source/Programs/UnrealGameSync/UnrealGameSync/Forms/MainWindow.cs
--- a/Engine/Source/Programs/UnrealGameSync/UnrealGameSync/Forms/MainWindow.cs
+++ b/Engine/Source/Programs/UnrealGameSync/UnrealGameSync/Forms/MainWindow.cs
@@ -58,6 +58,7 @@ namespace UnrealGameSync
 
 		UpdateMonitor UpdateMonitor;
 		SynchronizationContext MainThreadSynchronizationContext;
+		List<IssueMonitor> IssueMonitors = new List<IssueMonitor>();
 
 		string ApiUrl;
 		string DataFolder;
@@ -84,6 +85,9 @@ namespace UnrealGameSync
 
 		bool bAllowCreatingHandle;
 
+		Rectangle PrimaryWorkArea;
+		List<IssueAlertWindow> AlertWindows = new List<IssueAlertWindow>();
+
 		public MainWindow(UpdateMonitor InUpdateMonitor, string InApiUrl, string InDataFolder, string InCacheFolder, bool bInRestoreStateOnLoad, string InOriginalExecutableFileName, bool bInUnstable, DetectProjectSettingsResult[] StartupProjects, LineBasedTextWriter InLog, UserSettings InSettings)
 		{
 			InitializeComponent();
@@ -156,6 +160,11 @@ namespace UnrealGameSync
 			ResumeLayout(false);
 
 			bAllowCreatingHandle = true;
+
+			foreach(IssueMonitor IssueMonitor in IssueMonitors)
+			{
+				IssueMonitor.Start();
+			}
 		}
 
 		void PostAutomationRequest(AutomationRequest Request)
@@ -389,6 +398,32 @@ namespace UnrealGameSync
 			}
 			TabPanel.Dispose();
 
+			HashSet<IssueMonitor> NewIssueMonitors = new HashSet<IssueMonitor>();
+			for(int Idx = 0; Idx < TabControl.GetTabCount(); Idx++)
+			{
+				WorkspaceControl Workspace = TabControl.GetTabData(Idx) as WorkspaceControl;
+				if(Workspace != null)
+				{
+					NewIssueMonitors.Add(Workspace.GetIssueMonitor());
+				}
+			}
+			foreach(IssueMonitor IssueMonitor in IssueMonitors)
+			{
+				if(!NewIssueMonitors.Contains(IssueMonitor))
+				{
+					for(int Idx = AlertWindows.Count - 1; Idx >= 0; Idx--)
+					{
+						IssueAlertWindow AlertWindow = AlertWindows[Idx];
+						if(AlertWindow.IssueMonitor == IssueMonitor)
+						{
+							CloseAlertWindow(AlertWindow);
+						}
+					}
+					IssueMonitor.Dispose();
+				}
+			}
+			IssueMonitors = NewIssueMonitors.ToList();
+
 			SaveTabSettings();
 		}
 
@@ -416,6 +451,12 @@ namespace UnrealGameSync
 
 			StopScheduleTimer();
 
+			foreach(IssueMonitor IssueMonitor in IssueMonitors)
+			{
+				IssueMonitor.Dispose();
+			}
+			IssueMonitors.Clear();
+
 			if(AutomationServer != null)
 			{
 				AutomationServer.Dispose();
@@ -477,6 +518,8 @@ namespace UnrealGameSync
 
 			StatusLine OpenLine = new StatusLine();
 			OpenLine.AddLink("Open project...", FontStyle.Bold | FontStyle.Underline, () => { OpenNewProject(); });
+			OpenLine.AddText("  |  ");
+			OpenLine.AddLink("Application settings...", FontStyle.Bold | FontStyle.Underline, () => { ModifyApplicationSettings(); });
 			Lines.Add(OpenLine);
 
 			DefaultControl.Set(Lines, null, null, null);
@@ -964,8 +1007,18 @@ namespace UnrealGameSync
 				}
 			}
 
+			// Find or add an issue monitor for this
+			IssueMonitor IssueMonitor = IssueMonitors.FirstOrDefault(x => String.Compare(x.UserName, ProjectSettings.PerforceClient.UserName, StringComparison.OrdinalIgnoreCase) == 0);
+			if(IssueMonitor == null)
+			{
+				string LogFileName = Path.Combine(DataFolder, String.Format("IssueMonitor-{0}.log", ProjectSettings.PerforceClient.UserName));
+				IssueMonitor = new IssueMonitor(ApiUrl, ProjectSettings.PerforceClient.UserName, TimeSpan.FromSeconds(5.0), LogFileName);
+				IssueMonitor.OnIssuesChanged += IssueMonitor_OnUpdateAsync;
+				IssueMonitors.Add(IssueMonitor);
+			}
+
 			// Now that we have the project settings, we can construct the tab
-			WorkspaceControl NewWorkspace = new WorkspaceControl(this, ApiUrl, OriginalExecutableFileName, bUnstable, ProjectSettings, Log, Settings);
+			WorkspaceControl NewWorkspace = new WorkspaceControl(this, ApiUrl, OriginalExecutableFileName, bUnstable, ProjectSettings, IssueMonitor, Log, Settings);
 			NewWorkspace.Parent = TabPanel;
 			NewWorkspace.Dock = DockStyle.Fill;
 			NewWorkspace.Hide();
@@ -1185,5 +1238,193 @@ namespace UnrealGameSync
 			}
 			WindowState = Settings.WindowState;
 		}
+
+		public void UpdateAlertWindows()
+		{
+			HashSet<IssueData> AllIssues = new HashSet<IssueData>();
+			foreach(IssueMonitor IssueMonitor in IssueMonitors)
+			{
+				List<IssueData> Issues = IssueMonitor.GetIssues();
+				foreach(IssueData Issue in Issues)
+				{
+					IssueAlertReason Reason = 0;
+					if(Issue.FixChange <= 0)
+					{
+						if(Issue.Owner == null)
+						{
+							if(Issue.bNotify)
+							{
+								Reason |= IssueAlertReason.Normal;
+							}
+							if(Settings.NotifyUnassignedMinutes >= 0 && Issue.RetrievedAt - Issue.CreatedAt >= TimeSpan.FromMinutes(Settings.NotifyUnassignedMinutes))
+							{
+								Reason |= IssueAlertReason.UnassignedTimer;
+							}
+						}
+						else
+						{
+							if(String.Compare(Issue.Owner, IssueMonitor.UserName, StringComparison.OrdinalIgnoreCase) == 0 && !Issue.AcknowledgedAt.HasValue)
+							{
+								Reason |= IssueAlertReason.Owner;
+							}
+						}
+						if(Settings.NotifyUnresolvedMinutes >= 0 && Issue.RetrievedAt - Issue.CreatedAt >= TimeSpan.FromMinutes(Settings.NotifyUnresolvedMinutes))
+						{
+							Reason |= IssueAlertReason.UnresolvedTimer;
+						}
+
+						IssueAlertReason PrevReason;
+						if(IssueMonitor.IssueIdToAlertReason.TryGetValue(Issue.Id, out PrevReason))
+						{
+							Reason &= ~PrevReason;
+						}
+					}
+
+					IssueAlertWindow AlertWindow = AlertWindows.FirstOrDefault(x => x.IssueMonitor == IssueMonitor && x.Issue.Id == Issue.Id);
+					if(AlertWindow == null)
+					{
+						if(Reason != 0)
+						{
+							ShowAlertWindow(IssueMonitor, Issue, Reason);
+						}
+					}
+					else
+					{
+						if(Reason != 0)
+						{
+							AlertWindow.SetIssue(Issue, Reason);
+						}
+						else
+						{
+							CloseAlertWindow(AlertWindow);
+						}
+					}
+				}
+				AllIssues.UnionWith(Issues);
+			}
+
+			// Close any alert windows which don't have an active issues
+			for(int Idx = 0; Idx < AlertWindows.Count; Idx++)
+			{
+				IssueAlertWindow AlertWindow = AlertWindows[Idx];
+				if(!AllIssues.Contains(AlertWindow.Issue))
+				{
+					AlertWindow.IssueMonitor.IssueIdToAlertReason.Remove(AlertWindow.Issue.Id);
+					CloseAlertWindow(AlertWindow);
+					Idx--;
+				}
+			}
+		}
+
+		void IssueMonitor_OnUpdateAsync()
+		{
+			MainThreadSynchronizationContext.Post((o) => IssueMonitor_OnUpdate(), null);
+		}
+
+		void IssueMonitor_OnUpdate()
+		{
+			UpdateAlertWindows();
+		}
+
+		void ShowAlertWindow(IssueMonitor IssueMonitor, IssueData Issue, IssueAlertReason Reason)
+		{
+			IssueAlertWindow Alert = new IssueAlertWindow(IssueMonitor, Issue, Reason);
+			Alert.AcceptBtn.Click += (s, e) => AcceptIssue(Alert);
+			Alert.DeclineBtn.Click += (s, e) => DeclineIssue(Alert);
+			Alert.DetailsBtn.Click += (s, e) => ShowIssueDetails(Alert);
+
+			SetAlertWindowPositions();
+			AlertWindows.Add(Alert);
+			SetAlertWindowPosition(AlertWindows.Count - 1);
+
+			Alert.Show(this);
+
+			UpdateAlertPositionsTimer.Enabled = true;
+		}
+
+		void AcceptIssue(IssueAlertWindow Alert)
+		{
+			IssueData Issue = Alert.Issue;
+
+			IssueAlertReason Reason;
+			Alert.IssueMonitor.IssueIdToAlertReason.TryGetValue(Issue.Id, out Reason);
+			Alert.IssueMonitor.IssueIdToAlertReason[Issue.Id] = Reason | Alert.Reason;
+
+			IssueUpdateData Update = new IssueUpdateData();
+			Update.Id = Issue.Id;
+			Update.Owner = Alert.IssueMonitor.UserName;
+			Update.NominatedBy = null;
+			Update.Acknowledged = true;
+			Alert.IssueMonitor.PostUpdate(Update);
+
+			CloseAlertWindow(Alert);
+		}
+
+		void DeclineIssue(IssueAlertWindow Alert)
+		{
+			IssueData Issue = Alert.Issue;
+
+			IssueAlertReason Reason;
+			Alert.IssueMonitor.IssueIdToAlertReason.TryGetValue(Issue.Id, out Reason);
+			Alert.IssueMonitor.IssueIdToAlertReason[Issue.Id] = Reason | Alert.Reason;
+
+			CloseAlertWindow(Alert);
+		}
+
+		void ShowIssueDetails(IssueAlertWindow Alert)
+		{
+			for(int Idx = 0; Idx < TabControl.GetTabCount(); Idx++)
+			{
+				WorkspaceControl Workspace = TabControl.GetTabData(Idx) as WorkspaceControl;
+				if(Workspace != null && Workspace.GetIssueMonitor() == Alert.IssueMonitor)
+				{
+					Workspace.ShowIssueDetails(Alert.Issue);
+					break;
+				}
+			}
+		}
+
+		void CloseAlertWindow(IssueAlertWindow Alert)
+		{
+			IssueData Issue = Alert.Issue;
+
+			Alert.Close();
+			Alert.Dispose();
+
+			AlertWindows.Remove(Alert);
+
+			for(int Idx = 0; Idx < AlertWindows.Count; Idx++)
+			{
+				SetAlertWindowPosition(Idx);
+			}
+
+			if(AlertWindows.Count == 0)
+			{
+				UpdateAlertPositionsTimer.Enabled = false;
+			}
+		}
+
+		private void SetAlertWindowPosition(int Idx)
+		{
+			AlertWindows[Idx].Location = new Point(PrimaryWorkArea.Right - 40 - AlertWindows[Idx].Size.Width, PrimaryWorkArea.Height - 40 - (Idx + 1) * (AlertWindows[Idx].Size.Height + 15));
+		}
+
+		private void SetAlertWindowPositions()
+		{
+			Rectangle NewPrimaryWorkArea = Screen.PrimaryScreen.WorkingArea;
+			if(NewPrimaryWorkArea != PrimaryWorkArea)
+			{
+				PrimaryWorkArea = NewPrimaryWorkArea;
+				for(int Idx = 0; Idx < AlertWindows.Count; Idx++)
+				{
+					SetAlertWindowPosition(Idx);
+				}
+			}
+		}
+
+		private void UpdateAlertPositionsTimer_Tick(object sender, EventArgs e)
+		{
+			SetAlertWindowPositions();
+		}
 	}
 }
diff --git a/Engine/Source/Programs/UnrealGameSync/UnrealGameSync/IssueMonitor.cs b/Engine/Source/Programs/UnrealGameSync/UnrealGameSync/IssueMonitor.cs
new file mode 100644
--- /dev/null
+++ b/Engine/Source/Programs/UnrealGameSync/UnrealGameSync/IssueMonitor.cs
@@ -0,0 +1,359 @@
+using System;
+using System.Collections.Concurrent;
+using System.Collections.Generic;
+using System.Diagnostics;
+using System.Linq;
+using System.Text;
+using System.Threading;
+using System.Threading.Tasks;
+
+namespace UnrealGameSync
+{
+	public enum IssueBuildOutcome
+	{
+		Unknown,
+		Success,
+		Error,
+		Warning,
+	}
+
+	public class IssueBuildData
+	{
+		public string Stream;
+		public int Change;
+		public string Name;
+		public string Url;
+		public IssueBuildOutcome Outcome;
+	}
+
+	public class IssueData
+	{
+		public long Id;
+		public DateTime CreatedAt;
+		public DateTime RetrievedAt;
+		public string Project;
+		public string Summary;
+		public string Owner;
+		public string NominatedBy;
+		public DateTime? AcknowledgedAt;
+		public int FixChange;
+		public DateTime? ResolvedAt;
+		public bool bNotify;
+		public List<IssueBuildData> Builds;
+	}
+
+	public class IssueUpdateData
+	{
+		public long Id;
+		public string Owner;
+		public string NominatedBy;
+		public bool? Acknowledged;
+		public int? FixChange;
+		public bool? Resolved;
+	}
+
+	[Flags]
+	public enum IssueAlertReason
+	{
+		Normal = 1,
+		Owner = 2,
+		UnassignedTimer = 4,
+		UnresolvedTimer = 8,
+	}
+
+	class IssueMonitor : IDisposable
+	{
+		public readonly string ApiUrl;
+		public readonly string UserName;
+		Thread WorkerThread;
+		BoundedLogWriter LogWriter;
+		bool bDisposing;
+		AutoResetEvent RefreshEvent;
+		int UpdateIntervalMs;
+		long TrackingIssueId = -1;
+		List<IssueData> Issues = new List<IssueData>();
+		object LockObject = new object();
+		List<IssueUpdateData> PendingUpdates = new List<IssueUpdateData>();
+
+		public event Action OnIssuesChanged;
+
+		// Only used by MainWindow, but easier to just store here
+		public Dictionary<long, IssueAlertReason> IssueIdToAlertReason = new Dictionary<long, IssueAlertReason>();
+
+		public IssueMonitor(string ApiUrl, string UserName, TimeSpan UpdateInterval, string LogFileName)
+		{
+			this.ApiUrl = ApiUrl;
+			this.UserName = UserName;
+			this.UpdateIntervalMs = (int)UpdateInterval.TotalMilliseconds;
+
+			LogWriter = new BoundedLogWriter(LogFileName);
+			if(ApiUrl == null)
+			{
+				LastStatusMessage = "Database functionality disabled due to empty ApiUrl.";
+			}
+			else
+			{
+				LogWriter.WriteLine("Using connection string: {0}", this.ApiUrl);
+			}
+
+			RefreshEvent = new AutoResetEvent(false);
+		}
+
+		public string LastStatusMessage
+		{
+			get;
+			private set;
+		}
+
+		public List<IssueData> GetIssues()
+		{
+			return Issues;
+		}
+
+		public TimeSpan GetUpdateInterval()
+		{
+			return TimeSpan.FromMilliseconds(UpdateIntervalMs);
+		}
+
+		public void SetUpdateInterval(TimeSpan UpdateInterval)
+		{
+			UpdateIntervalMs = (int)UpdateInterval.TotalMilliseconds;
+			RefreshEvent.Set();
+		}
+
+		public void StartTracking(long IssueId)
+		{
+			lock(LockObject)
+			{
+				TrackingIssueId = IssueId;
+			}
+			RefreshEvent.Set();
+		}
+
+		public void StopTracking(long IssueId)
+		{
+			lock(LockObject)
+			{
+				TrackingIssueId = -1;
+			}
+		}
+
+		public bool HasPendingUpdate()
+		{
+			return PendingUpdates.Count > 0;
+		}
+
+		public void PostUpdate(IssueUpdateData Update)
+		{
+			bool bUpdatedIssues;
+			lock(LockObject)
+			{
+				PendingUpdates.Add(Update);
+				bUpdatedIssues = ApplyPendingUpdate(Issues, Update);
+			}
+
+			RefreshEvent.Set();
+
+			if(bUpdatedIssues)
+			{
+				OnIssuesChanged?.Invoke();
+			}
+		}
+
+		static bool ApplyPendingUpdate(List<IssueData> Issues, IssueUpdateData Update)
+		{
+			bool bUpdated = false;
+			for(int Idx = 0; Idx < Issues.Count; Idx++)
+			{
+				IssueData Issue = Issues[Idx];
+				if(Update.Id == Issue.Id)
+				{
+					if(Update.Owner != null && Update.Owner != Issue.Owner)
+					{
+						Issue.Owner = Update.Owner;
+						bUpdated = true;
+					}
+					if(Update.NominatedBy != null && Update.NominatedBy != Issue.NominatedBy)
+					{
+						Issue.NominatedBy = Update.NominatedBy;
+						bUpdated = true;
+					}
+					if(Update.Acknowledged.HasValue && Update.Acknowledged.Value != Issue.AcknowledgedAt.HasValue)
+					{
+						Issue.AcknowledgedAt = Update.Acknowledged.Value? (DateTime?)DateTime.UtcNow : null;
+						bUpdated = true;
+					}
+					if(Update.FixChange.HasValue)
+					{
+						Issue.FixChange = Update.FixChange.Value;
+						if(Issue.FixChange != 0)
+						{
+							Issues.RemoveAt(Idx);
+						}
+						bUpdated = true;
+					}
+					break;
+				}
+			}
+			return bUpdated;
+		}
+
+		public void Start()
+		{
+			if(ApiUrl != null)
+			{
+				WorkerThread = new Thread(() => PollForUpdates());
+				WorkerThread.Start();
+			}
+		}
+
+		public void Dispose()
+		{
+			bDisposing = true;
+			if(WorkerThread != null)
+			{
+				RefreshEvent.Set();
+				if(!WorkerThread.Join(100))
+				{
+					WorkerThread.Abort();
+					WorkerThread.Join();
+				}
+				WorkerThread = null;
+			}
+			if(LogWriter != null)
+			{
+				LogWriter.Dispose();
+				LogWriter = null;
+			}
+		}
+
+		void PollForUpdates()
+		{
+			while (!bDisposing)
+			{
+				// Check if there's any pending update
+				IssueUpdateData PendingUpdate;
+				lock(LockObject)
+				{
+					if(PendingUpdates.Count > 0)
+					{
+						PendingUpdate = PendingUpdates[0];
+					}
+					else
+					{
+						PendingUpdate = null;
+					}
+				}
+
+				// If we have an update, try to post it to the backend and check for another
+				if(PendingUpdate != null)
+				{
+					if(SendUpdate(PendingUpdate))
+					{
+						lock (LockObject) { PendingUpdates.RemoveAt(0); }
+					}
+					else
+					{
+						RefreshEvent.WaitOne(5 * 1000);
+					}
+					continue;
+				}
+
+				// Read all the current issues
+				ReadCurrentIssues();
+
+				// Wait for something else to do
+				RefreshEvent.WaitOne(UpdateIntervalMs);
+			}
+		}
+
+		bool SendUpdate(IssueUpdateData Update)
+		{
+			try
+			{
+				RESTApi.PUT<IssueUpdateData>(ApiUrl, String.Format("issues/{0}", Update.Id), Update);
+				return true;
+			}
+			catch(Exception Ex)
+			{
+				LogWriter.WriteException(Ex, "Failed with exception.");
+				LastStatusMessage = String.Format("Failed to send update: ({0})", Ex.ToString());
+				return false;
+			}
+		}
+
+		bool ReadCurrentIssues()
+		{
+			try
+			{
+				Stopwatch Timer = Stopwatch.StartNew();
+				LogWriter.WriteLine();
+				LogWriter.WriteLine("Polling for notifications at {0}...", DateTime.Now.ToString());
+
+				// Get the initial number of issues. We won't post updates if this stays at zero.
+				int InitialNumIssues = Issues.Count;
+
+				// Fetch the new issues
+				List<IssueData> NewIssues = RESTApi.GET<List<IssueData>>(ApiUrl, String.Format("issues?user={0}", UserName));
+
+				// Check if we're tracking a particular issue. If so, we want updates even when it's resolved.
+				long LocalTrackingIssueId;
+				lock(LockObject)
+				{
+					LocalTrackingIssueId = TrackingIssueId;
+				}
+				if(LocalTrackingIssueId != -1 && !NewIssues.Any(x => x.Id == LocalTrackingIssueId))
+				{
+					try
+					{
+						IssueData Issue = RESTApi.GET<IssueData>(ApiUrl, String.Format("issues/{0}", LocalTrackingIssueId));
+						if(Issue != null)
+						{
+							NewIssues.Add(Issue);
+						}
+					}
+					catch(Exception Ex)
+					{
+						LogWriter.WriteException(Ex, "Exception while fetching tracked issue");
+					}
+				}
+
+				// Update all the builds for each issue
+				foreach(IssueData NewIssue in NewIssues)
+				{
+					NewIssue.Builds = RESTApi.GET<List<IssueBuildData>>(ApiUrl, String.Format("issues/{0}/builds", NewIssue.Id));
+				}
+
+				// Apply any pending updates to this issue list, and update it
+				lock(LockObject)
+				{
+					foreach(IssueUpdateData PendingUpdate in PendingUpdates)
+					{
+						ApplyPendingUpdate(NewIssues, PendingUpdate);
+					}
+					Issues = NewIssues;
+				}
+
+				// Update the main thread
+				if(InitialNumIssues > 0 || Issues.Count > 0)
+				{
+					if(OnIssuesChanged != null)
+					{
+						OnIssuesChanged();
+					}
+				}
+
+				// Update the stats
+				LastStatusMessage = String.Format("Last update took {0}ms", Timer.ElapsedMilliseconds);
+				LogWriter.WriteLine("Done in {0}ms.", Timer.ElapsedMilliseconds);
+				return true;
+			}
+			catch(Exception Ex)
+			{
+				LogWriter.WriteException(Ex, "Failed with exception.");
+				LastStatusMessage = String.Format("Last update failed: ({0})", Ex.ToString());
+				return false;
+			}
+		}
+	}
+}
diff --git a/Engine/Source/Programs/UnrealGameSync/UnrealGameSync/Perforce.cs b/Engine/Source/Programs/UnrealGameSync/UnrealGameSync/Perforce.cs
--- a/Engine/Source/Programs/UnrealGameSync/UnrealGameSync/Perforce.cs
+++ b/Engine/Source/Programs/UnrealGameSync/UnrealGameSync/Perforce.cs
@@ -786,12 +786,21 @@ namespace UnrealGameSync
 		}
 
 		public bool FindChanges(IEnumerable<string> Filters, int MaxResults, out List<PerforceChangeSummary> Changes, TextWriter Log)
+		{
+			return FindChanges(Filters, null, MaxResults, out Changes, Log);
+		}
+
+		public bool FindChanges(IEnumerable<string> Filters, string ByUser, int MaxResults, out List<PerforceChangeSummary> Changes, TextWriter Log)
 		{
 			string Arguments = "changes -s submitted -t -L";
 			if(MaxResults > 0)
 			{
 				Arguments += String.Format(" -m {0}", MaxResults);
 			}
+			if(ByUser != null)
+			{
+				Arguments += String.Format(" -u \"{0}\"", ByUser);
+			}
 			foreach(string Filter in Filters)
 			{
 				Arguments += String.Format(" \"{0}\"", Filter);
@@ -1076,6 +1085,34 @@ namespace UnrealGameSync
 			return true;
 		}
 
+		public bool DescribeMultiple(IEnumerable<int> ChangeNumbers, out List<PerforceDescribeRecord> OutRecords, TextWriter Log)
+		{
+			string CommandLine = String.Format("describe -s {0}", String.Join(" ", ChangeNumbers.Select(x => x.ToString())));
+
+			List<Dictionary<string, string>> RawRecords = new List<Dictionary<string,string>>();
+			if(!RunCommandWithBinaryOutput(CommandLine, RawRecords, CommandOptions.None, Log))
+			{
+				OutRecords = null;
+				return false;
+			}
+
+			List<PerforceDescribeRecord> Records = new List<PerforceDescribeRecord>();
+			foreach(Dictionary<string, string> RawRecord in RawRecords)
+			{
+				string Code;
+				if(!RawRecords[0].TryGetValue("code", out Code) || Code != "stat")
+				{
+					Log.WriteLine("Unexpected response from p4 {0}: {1}", CommandLine, String.Join(", ", RawRecords[0].Select(x => String.Format("( \"{0}\", \"{1}\" )", x.Key, x.Value))));
+					OutRecords = null;
+					return false;
+				}
+				Records.Add(new PerforceDescribeRecord(RawRecord));
+			}
+
+			OutRecords = Records;
+			return true;
+		}
+
 		public bool Where(string Filter, out PerforceWhereRecord WhereRecord, TextWriter Log)
 		{
 			List<PerforceFileRecord> FileRecords;
diff --git a/Engine/Source/Programs/UnrealGameSync/UnrealGameSync/PerforceMonitor.cs b/Engine/Source/Programs/UnrealGameSync/UnrealGameSync/PerforceMonitor.cs
--- a/Engine/Source/Programs/UnrealGameSync/UnrealGameSync/PerforceMonitor.cs
+++ b/Engine/Source/Programs/UnrealGameSync/UnrealGameSync/PerforceMonitor.cs
@@ -17,6 +17,25 @@ namespace UnrealGameSync
 		public string Description;
 		public bool bContainsCode;
 		public bool bContainsContent;
+
+		public PerforceChangeDetails(PerforceDescribeRecord DescribeRecord)
+		{
+			Description = DescribeRecord.Description;
+
+			// Check whether the files are code or content
+			string[] CodeExtensions = { ".cs", ".h", ".cpp", ".inl", ".usf", ".ush", ".uproject", ".uplugin" };
+			foreach(PerforceDescribeFileRecord File in DescribeRecord.Files)
+			{
+				if(CodeExtensions.Any(Extension => File.DepotFile.EndsWith(Extension, StringComparison.InvariantCultureIgnoreCase)))
+				{
+					bContainsCode = true;
+				}
+				else
+				{
+					bContainsContent = true;
+				}
+			}
+		}
 	}
 
 	class PerforceMonitor : IDisposable
@@ -362,8 +381,6 @@ namespace UnrealGameSync
 			// Update them in batches
 			foreach(int QueryChangeNumber in QueryChangeNumbers)
 			{
-				string[] CodeExtensions = { ".cs", ".h", ".cpp", ".inl", ".usf", ".ush", ".uproject", ".uplugin" };
-
 				// Skip this stuff if the user wants us to query for more changes
 				if(PendingMaxChanges > CurrentMaxChanges)
 				{
@@ -375,23 +392,7 @@ namespace UnrealGameSync
 				if(Perforce.Describe(QueryChangeNumber, out DescribeRecord, LogWriter))
 				{
 					// Create the details object
-					PerforceChangeDetails Details = new PerforceChangeDetails();
-					Details.Description = DescribeRecord.Description;
-
-					// Check whether the files are code or content
-					foreach(PerforceDescribeFileRecord File in DescribeRecord.Files)
-					{
-						if(CodeExtensions.Any(Extension => File.DepotFile.EndsWith(Extension, StringComparison.InvariantCultureIgnoreCase)))
-						{
-							Details.bContainsCode = true;
-						}
-						else
-						{
-							Details.bContainsContent = true;
-						}
-					}
-
-					// Update the type of this change
+					PerforceChangeDetails Details = new PerforceChangeDetails(DescribeRecord);
 					lock(this)
 					{
 						if(!ChangeDetails.ContainsKey(QueryChangeNumber))
diff --git a/Engine/Source/Programs/UnrealGameSync/UnrealGameSync/RESTApi.cs b/Engine/Source/Programs/UnrealGameSync/UnrealGameSync/RESTApi.cs
--- a/Engine/Source/Programs/UnrealGameSync/UnrealGameSync/RESTApi.cs
+++ b/Engine/Source/Programs/UnrealGameSync/UnrealGameSync/RESTApi.cs
@@ -38,7 +38,7 @@ namespace UnrealGameSync
 			// Add json to request body
 			if (!string.IsNullOrEmpty(RequestBody))
 			{
-				if (Method == "POST")
+				if (Method == "POST" || Method == "PUT")
 				{
 					byte[] bytes = Encoding.ASCII.GetBytes(RequestBody);
 					using (Stream RequestStream = Request.GetRequestStream())
@@ -78,9 +78,20 @@ namespace UnrealGameSync
 		{
 			return SendRequestInternal(URI, Resource, "POST", RequestBody, QueryParams);
 		}
+
+		public static string GET(string URI, string Resource, params string[] QueryParams)
+		{
+			return SendRequestInternal(URI, Resource, "GET", null, QueryParams);
+		}
+
 		public static T GET<T>(string URI, string Resource, params string[] QueryParams)
 		{
 			return new JavaScriptSerializer().Deserialize<T>(SendRequestInternal(URI, Resource, "GET", null, QueryParams));
 		}
+
+		public static string PUT<T>(string URI, string Resource, T Object, params string[] QueryParams)
+		{
+			return SendRequestInternal(URI, Resource, "PUT", new JavaScriptSerializer().Serialize(Object), QueryParams);
+		}
 	}
 }
diff --git a/Engine/Source/Programs/UnrealGameSync/UnrealGameSync/UnrealGameSync.csproj b/Engine/Source/Programs/UnrealGameSync/UnrealGameSync/UnrealGameSync.csproj
--- a/Engine/Source/Programs/UnrealGameSync/UnrealGameSync/UnrealGameSync.csproj
+++ b/Engine/Source/Programs/UnrealGameSync/UnrealGameSync/UnrealGameSync.csproj
@@ -115,9 +115,13 @@
     <Compile Include="AutomationServer.cs" />
     <Compile Include="BuildStep.cs" />
     <Compile Include="ChildProcess.cs" />
-    <Compile Include="Controls\CustomListViewControl.cs">
+    <Compile Include="Controls\AlertButtonControl.cs">
       <SubType>Component</SubType>
     </Compile>
+    <Compile Include="Controls\AlertImageControl.cs">
+      <SubType>UserControl</SubType>
+    </Compile>
+    <Compile Include="Controls\CustomListViewControl.cs" />
     <Compile Include="Controls\ErrorPanel.cs">
       <SubType>Component</SubType>
     </Compile>
@@ -166,6 +170,24 @@
     <Compile Include="Forms\AutomatedSyncWindow.Designer.cs">
       <DependentUpon>AutomatedSyncWindow.cs</DependentUpon>
     </Compile>
+    <Compile Include="Forms\IssueAlertWindow.cs">
+      <SubType>Form</SubType>
+    </Compile>
+    <Compile Include="Forms\IssueAlertWindow.Designer.cs">
+      <DependentUpon>IssueAlertWindow.cs</DependentUpon>
+    </Compile>
+    <Compile Include="Forms\IssueFixedWindow.cs">
+      <SubType>Form</SubType>
+    </Compile>
+    <Compile Include="Forms\IssueFixedWindow.Designer.cs">
+      <DependentUpon>IssueFixedWindow.cs</DependentUpon>
+    </Compile>
+    <Compile Include="Forms\IssueSettingsWindow.cs">
+      <SubType>Form</SubType>
+    </Compile>
+    <Compile Include="Forms\IssueSettingsWindow.Designer.cs">
+      <DependentUpon>IssueSettingsWindow.cs</DependentUpon>
+    </Compile>
     <Compile Include="Forms\SelectUserWindow.cs">
       <SubType>Form</SubType>
     </Compile>
@@ -190,6 +212,12 @@
     <Compile Include="Forms\ConnectWindow.Designer.cs">
       <DependentUpon>ConnectWindow.cs</DependentUpon>
     </Compile>
+    <Compile Include="Forms\IssueDetailsWindow.cs">
+      <SubType>Form</SubType>
+    </Compile>
+    <Compile Include="Forms\IssueDetailsWindow.Designer.cs">
+      <DependentUpon>IssueDetailsWindow.cs</DependentUpon>
+    </Compile>
     <Compile Include="Forms\NewWorkspaceWindow.cs">
       <SubType>Form</SubType>
     </Compile>
@@ -336,6 +364,7 @@
     <Compile Include="Forms\SyncFilter.Designer.cs">
       <DependentUpon>SyncFilter.cs</DependentUpon>
     </Compile>
+    <Compile Include="IssueMonitor.cs" />
     <Compile Include="ModalTask.cs" />
     <Compile Include="PerforceModalTask.cs" />
     <Compile Include="PerforceMonitor.cs" />
@@ -368,6 +397,15 @@
     <EmbeddedResource Include="Forms\AutomatedSyncWindow.resx">
       <DependentUpon>AutomatedSyncWindow.cs</DependentUpon>
     </EmbeddedResource>
+    <EmbeddedResource Include="Forms\IssueAlertWindow.resx">
+      <DependentUpon>IssueAlertWindow.cs</DependentUpon>
+    </EmbeddedResource>
+    <EmbeddedResource Include="Forms\IssueFixedWindow.resx">
+      <DependentUpon>IssueFixedWindow.cs</DependentUpon>
+    </EmbeddedResource>
+    <EmbeddedResource Include="Forms\IssueSettingsWindow.resx">
+      <DependentUpon>IssueSettingsWindow.cs</DependentUpon>
+    </EmbeddedResource>
     <EmbeddedResource Include="Forms\SelectUserWindow.resx">
       <DependentUpon>SelectUserWindow.cs</DependentUpon>
     </EmbeddedResource>
@@ -380,6 +418,9 @@
     <EmbeddedResource Include="Forms\ConnectWindow.resx">
       <DependentUpon>ConnectWindow.cs</DependentUpon>
     </EmbeddedResource>
+    <EmbeddedResource Include="Forms\IssueDetailsWindow.resx">
+      <DependentUpon>IssueDetailsWindow.cs</DependentUpon>
+    </EmbeddedResource>
     <EmbeddedResource Include="Forms\NewWorkspaceWindow.resx">
       <DependentUpon>NewWorkspaceWindow.cs</DependentUpon>
     </EmbeddedResource>
@@ -506,6 +547,8 @@
     </BootstrapperPackage>
   </ItemGroup>
   <ItemGroup>
+    <None Include="Resources\AlertErrorIcon.png" />
+    <None Include="Resources\AlertWarningIcon.png" />
     <Content Include="Resources\Perforce.png" />
     <None Include="Resources\DefaultErrorLogo.png" />
     <Content Include="Resources\DefaultProjectLogo.png" />
diff --git a/Engine/Source/Programs/UnrealGameSync/UnrealGameSync/UserSettings.cs b/Engine/Source/Programs/UnrealGameSync/UnrealGameSync/UserSettings.cs
--- a/Engine/Source/Programs/UnrealGameSync/UnrealGameSync/UserSettings.cs
+++ b/Engine/Source/Programs/UnrealGameSync/UnrealGameSync/UserSettings.cs
@@ -229,6 +229,10 @@ namespace UnrealGameSync
 		public List<Tuple<string, bool>> EditorArguments = new List<Tuple<string,bool>>();
 		public bool bEditorArgumentsPrompt;
 
+		// Notification settings
+		public int NotifyUnassignedMinutes;
+		public int NotifyUnresolvedMinutes;
+
 		// Project settings
 		Dictionary<string, UserWorkspaceSettings> WorkspaceKeyToSettings = new Dictionary<string,UserWorkspaceSettings>();
 		Dictionary<string, UserProjectSettings> ProjectKeyToSettings = new Dictionary<string,UserProjectSettings>();
@@ -372,6 +376,10 @@ namespace UnrealGameSync
 			ScheduleAnyOpenProject = ConfigFile.GetValue("Schedule.AnyOpenProject", true);
 			ScheduleProjects = ReadProjectList("Schedule.Projects", "Schedule.ProjectFileNames");
 
+			// Notification settings
+			NotifyUnassignedMinutes = ConfigFile.GetValue("Notifications.NotifyUnassignedMinutes", -1);
+			NotifyUnresolvedMinutes = ConfigFile.GetValue("Notifications.NotifyUnresolvedMinutes", -1);
+
 			// Perforce settings
 			if(!int.TryParse(ConfigFile.GetValue("Perforce.NumRetries", "0"), out SyncOptions.NumRetries))
 			{
@@ -622,6 +630,18 @@ namespace UnrealGameSync
 				WindowSection.SetValue("Bounds", FormatRectangleValue(WindowBounds.Value));
 			}
 
+			// Notification settings
+			ConfigSection NotificationSection = ConfigFile.FindOrAddSection("Notifications");
+			NotificationSection.Clear();
+			if(NotifyUnassignedMinutes != -1)
+			{
+				NotificationSection.SetValue("NotifyUnassignedMinutes", NotifyUnassignedMinutes);
+			}
+			if(NotifyUnresolvedMinutes != -1)
+			{
+				NotificationSection.SetValue("NotifyUnresolvedMinutes", NotifyUnresolvedMinutes);
+			}
+
 			// Current workspace settings
 			foreach(KeyValuePair<string, UserWorkspaceSettings> Pair in WorkspaceKeyToSettings)
 			{
diff --git a/Engine/Source/Programs/UnrealGameSync/UnrealGameSync/Utility.cs b/Engine/Source/Programs/UnrealGameSync/UnrealGameSync/Utility.cs
--- a/Engine/Source/Programs/UnrealGameSync/UnrealGameSync/Utility.cs
+++ b/Engine/Source/Programs/UnrealGameSync/UnrealGameSync/Utility.cs
@@ -6,6 +6,7 @@ using Microsoft.Win32;
 using System;
 using System.Collections.Generic;
 using System.Diagnostics;
+using System.Drawing;
 using System.IO;
 using System.Linq;
 using System.Runtime.InteropServices;
@@ -424,5 +425,81 @@ namespace UnrealGameSync
 				}
 			}
 		}
+
+		public static Color Blend(Color First, Color Second, float T)
+		{
+			return Color.FromArgb((int)(First.R + (Second.R - First.R) * T), (int)(First.G + (Second.G - First.G) * T), (int)(First.B + (Second.B - First.B) * T));
+		}
+
+		public static string FormatRecentDateTime(DateTime Date)
+		{
+			DateTime Now = DateTime.Now;
+
+			DateTime Midnight = new DateTime(Now.Year, Now.Month, Now.Day);
+			DateTime MidnightTonight = Midnight + TimeSpan.FromDays(1.0);
+
+			if(Date > MidnightTonight)
+			{
+				return String.Format("{0} at {1}", Date.ToLongDateString(), Date.ToShortTimeString());
+			}
+			else if(Date >= Midnight)
+			{
+				return String.Format("today at {0}", Date.ToShortTimeString());
+			}
+			else if(Date >= Midnight - TimeSpan.FromDays(1.0))
+			{
+				return String.Format("yesterday at {0}", Date.ToShortTimeString());
+			}
+			else if(Date >= Midnight - TimeSpan.FromDays(5.0))
+			{
+				return String.Format("{0:dddd} at {1}", Date, Date.ToShortTimeString());
+			}
+			else
+			{
+				return String.Format("{0} at {1}", Date.ToLongDateString(), Date.ToShortTimeString());
+			}
+		}
+
+		public static string FormatDurationMinutes(TimeSpan Duration)
+		{
+			return FormatDurationMinutes((int)(Duration.TotalMinutes + 1));
+		}
+
+		public static string FormatDurationMinutes(int TotalMinutes)
+		{
+			if(TotalMinutes > 24 * 60)
+			{
+				return String.Format("{0}d {1}h", TotalMinutes / (24 * 60), (TotalMinutes / 60) % 24);
+			}
+			else if(TotalMinutes > 60)
+			{
+				return String.Format("{0}h {1}m", TotalMinutes / 60, TotalMinutes % 60);
+			}
+			else
+			{
+				return String.Format("{0}m", TotalMinutes);
+			}
+		}
+
+		public static string FormatUserName(string UserName)
+		{
+			StringBuilder NormalUserName = new StringBuilder();
+			for(int Idx = 0; Idx < UserName.Length; Idx++)
+			{
+				if(Idx == 0 || UserName[Idx - 1] == '.')
+				{
+					NormalUserName.Append(Char.ToUpper(UserName[Idx]));
+				}
+				else if(UserName[Idx] == '.')
+				{
+					NormalUserName.Append(' ');
+				}
+				else
+				{
+					NormalUserName.Append(Char.ToLower(UserName[Idx]));
+				}
+			}
+			return NormalUserName.ToString();
+		}
 	}
 }
diff --git a/Engine/Source/Programs/UnrealGameSync/UnrealGameSyncMetadataServer/App_Start/WebApiConfig.cs b/Engine/Source/Programs/UnrealGameSync/UnrealGameSyncMetadataServer/App_Start/WebApiConfig.cs
--- a/Engine/Source/Programs/UnrealGameSync/UnrealGameSyncMetadataServer/App_Start/WebApiConfig.cs
+++ b/Engine/Source/Programs/UnrealGameSync/UnrealGameSyncMetadataServer/App_Start/WebApiConfig.cs
@@ -11,6 +11,30 @@ namespace UnrealGameSyncMetadataServer
 			// Web API routes
 			config.MapHttpAttributeRoutes();
 
+			config.Routes.MapHttpRoute(
+				name: "IssueBuildsSubApi",
+				routeTemplate: "api/issues/{IssueId}/builds",
+				defaults: new { controller = "IssueBuildsSub" }
+			);
+
+			config.Routes.MapHttpRoute(
+				name: "IssueBuildsApi",
+				routeTemplate: "api/issuebuilds/{BuildId}",
+				defaults: new { controller = "IssueBuilds" }
+			);
+
+			config.Routes.MapHttpRoute(
+				name: "IssueEventsApi",
+				routeTemplate: "api/issues/{IssueId}/events",
+				defaults: new { controller = "IssueEvents" }
+			);
+
+			config.Routes.MapHttpRoute(
+				name: "IssueWatchersApi",
+				routeTemplate: "api/issues/{IssueId}/watchers",
+				defaults: new { controller = "IssueWatchers" }
+			);
+
 			config.Routes.MapHttpRoute(
 				name: "DefaultApi",
 				routeTemplate: "api/{controller}/{id}",
diff --git a/Engine/Source/Programs/UnrealGameSync/UnrealGameSyncMetadataServer/Connectors/SqlConnector.cs b/Engine/Source/Programs/UnrealGameSync/UnrealGameSyncMetadataServer/Connectors/SqlConnector.cs
--- a/Engine/Source/Programs/UnrealGameSync/UnrealGameSyncMetadataServer/Connectors/SqlConnector.cs
+++ b/Engine/Source/Programs/UnrealGameSync/UnrealGameSyncMetadataServer/Connectors/SqlConnector.cs
@@ -1,10 +1,14 @@
 ﻿// Copyright 1998-2019 Epic Games, Inc. All Rights Reserved.
 
+using Newtonsoft.Json;
+using Newtonsoft.Json.Linq;
 using System;
 using System.Collections.Generic;
 using System.Data.SqlClient;
 using System.Data.SQLite;
 using System.IO;
+using System.Linq;
+using System.Text;
 using System.Text.RegularExpressions;
 using UnrealGameSyncMetadataServer.Models;
 
@@ -305,5 +309,368 @@ namespace UnrealGameSyncMetadataServer.Connectors
 		{
 			return Wildcard.EndsWith("...") && Project.StartsWith(Wildcard.Substring(0, Wildcard.Length - 4), StringComparison.InvariantCultureIgnoreCase);
 		}
+
+		private static string NormalizeUserName(string UserName)
+		{
+			return UserName.ToUpperInvariant();
+		}
+
+		public static long FindOrAddUserId(string Name)
+		{
+			using (SQLiteConnection Connection = new SQLiteConnection(SqlConnector.ConnectionString))
+			{
+				Connection.Open();
+				return FindOrAddUserId(Name, Connection);
+			}
+		}
+
+		private static long FindOrAddUserId(string Name, SQLiteConnection Connection)
+		{
+			if(Name.Length == 0)
+			{
+				return -1;
+			}
+
+			string NormalizedName = NormalizeUserName(Name);
+
+			using (SQLiteCommand Command = new SQLiteCommand("SELECT [Id] FROM [Users] WHERE Name = @Name", Connection))
+			{
+				Command.Parameters.AddWithValue("@Name", NormalizedName);
+				object UserId = Command.ExecuteScalar();
+				if(UserId != null)
+				{
+					return (long)UserId;
+				}
+			}
+
+			using (SQLiteCommand Command = new SQLiteCommand("INSERT OR IGNORE INTO [Users] (Name) VALUES (@Name); SELECT [Id] FROM [Users] WHERE Name = @Name", Connection))
+			{
+				Command.Parameters.AddWithValue("@Name", NormalizedName);
+				object UserId = Command.ExecuteScalar();
+				return (long)UserId;
+			}
+		}
+
+		public static long AddIssue(IssueData Issue)
+		{
+			long IssueId;
+			using (SQLiteConnection Connection = new SQLiteConnection(SqlConnector.ConnectionString))
+			{
+				Connection.Open();
+
+				using (SQLiteCommand Command = new SQLiteCommand("INSERT INTO [Issues] (Project, Summary, CreatedAt, FixChange) VALUES (@Project, @Summary, DATETIME('now'), 0)", Connection))
+				{
+					Command.Parameters.AddWithValue("@Project", Issue.Project);
+					Command.Parameters.AddWithValue("@Summary", Issue.Summary);
+					Command.ExecuteNonQuery();
+
+					IssueId = Connection.LastInsertRowId;
+				}
+			}
+			return IssueId;
+		}
+
+		public static IssueData GetIssue(long IssueId)
+		{
+			List<IssueData> Issues = GetIssuesInternal(IssueId, null);
+			if(Issues.Count == 0)
+			{
+				return null;
+			}
+			else
+			{
+				return Issues[0];
+			}
+		}
+
+		public static List<IssueData> GetIssues()
+		{
+			return GetIssuesInternal(-1, null);
+		}
+
+		public static List<IssueData> GetIssues(string UserName)
+		{
+			return GetIssuesInternal(-1, UserName);
+		}
+
+		private static List<IssueData> GetIssuesInternal(long IssueId, string UserName)
+		{
+			List<IssueData> Issues = new List<IssueData>();
+			using (SQLiteConnection Connection = new SQLiteConnection(SqlConnector.ConnectionString))
+			{
+				Connection.Open();
+
+				long UserId = -1;
+				if(UserName != null)
+				{
+					UserId = FindOrAddUserId(UserName);
+				}
+
+				StringBuilder CommandBuilder = new StringBuilder();
+				CommandBuilder.Append("SELECT");
+				CommandBuilder.Append(" Issues.Id, Issues.CreatedAt, DATETIME('now'), Issues.Project, Issues.Summary, OwnerUsers.Name, NominatedByUsers.Name, Issues.AcknowledgedAt, Issues.FixChange, Issues.ResolvedAt");
+				if(UserName != null)
+				{
+					CommandBuilder.Append(", IssueWatchers.UserId");
+				}
+				CommandBuilder.Append(" FROM [Issues]");
+				CommandBuilder.Append(" LEFT JOIN [Users] AS [OwnerUsers] ON OwnerUsers.Id = Issues.OwnerId");
+				CommandBuilder.Append(" LEFT JOIN [Users] AS [NominatedByUsers] ON NominatedByUsers.Id = Issues.NominatedById");
+				if(UserName != null)
+				{
+					CommandBuilder.Append(" LEFT JOIN [IssueWatchers] ON IssueWatchers.IssueId = Issues.Id AND IssueWatchers.UserId = @UserId");
+				}
+				if(IssueId != -1)
+				{
+					CommandBuilder.Append(" WHERE Issues.Id = @IssueId");
+				}
+				else
+				{
+					CommandBuilder.Append(" WHERE Issues.ResolvedAt IS NULL");
+				}
+
+				using (SQLiteCommand Command = new SQLiteCommand(CommandBuilder.ToString(), Connection))
+				{
+					if(IssueId != -1)
+					{
+						Command.Parameters.AddWithValue("@IssueId", IssueId);
+					}
+					if(UserName != null)
+					{
+						Command.Parameters.AddWithValue("@UserId", UserId);
+					}
+
+					using(SQLiteDataReader Reader = Command.ExecuteReader())
+					{
+						while(Reader.Read())
+						{
+							IssueData Issue = new IssueData();
+							Issue.Id = Reader.GetInt64(0);
+							Issue.CreatedAt = Reader.GetDateTime(1);
+							Issue.RetrievedAt = Reader.GetDateTime(2);
+							Issue.Project = Reader.GetString(3);
+							Issue.Summary = Reader.GetString(4);
+							Issue.Owner = Reader.IsDBNull(5)? null : Reader.GetString(5);
+							Issue.NominatedBy = Reader.IsDBNull(6)? null : Reader.GetString(6);
+							Issue.AcknowledgedAt = Reader.IsDBNull(7)? (DateTime?)null : Reader.GetDateTime(7);
+							Issue.FixChange = Reader.GetInt32(8);
+							Issue.ResolvedAt = Reader.IsDBNull(9)? (DateTime?)null : Reader.GetDateTime(9);
+							if(UserName != null)
+							{
+								Issue.bNotify = !Reader.IsDBNull(10);
+							}
+							Issues.Add(Issue);
+						}
+					}
+				}
+			}
+			return Issues;
+		}
+
+		public static void UpdateIssue(long IssueId, IssueUpdateData Issue)
+		{
+			using (SQLiteConnection Connection = new SQLiteConnection(SqlConnector.ConnectionString))
+			{
+				Connection.Open();
+
+				using (SQLiteCommand Command = new SQLiteCommand(Connection))
+				{
+					List<string> Columns = new List<string>();
+					List<string> Values = new List<string>();
+					if(Issue.Summary != null)
+					{
+						Columns.Add("Summary");
+						Values.Add("@Summary");
+						Command.Parameters.AddWithValue("@Summary", Issue.Summary);
+					}
+					if(Issue.Owner != null)
+					{
+						Columns.Add("OwnerId");
+						Values.Add("@OwnerId");
+						Command.Parameters.AddWithValue("OwnerId", FindOrAddUserId(Issue.Owner, Connection));
+					}
+					if(Issue.NominatedBy != null)
+					{
+						Columns.Add("NominatedById");
+						Values.Add("@NominatedById");
+						Command.Parameters.AddWithValue("NominatedById", FindOrAddUserId(Issue.NominatedBy, Connection));
+					}
+					if(Issue.Acknowledged.HasValue)
+					{
+						Columns.Add("AcknowledgedAt");
+						Values.Add(Issue.Acknowledged.Value? "DATETIME('now')" : "NULL");
+					}
+					if(Issue.FixChange.HasValue)
+					{
+						Columns.Add("FixChange");
+						Values.Add("@FixChange");
+						Command.Parameters.AddWithValue("FixChange", Issue.FixChange.Value);
+					}
+					if(Issue.Resolved.HasValue)
+					{
+						Columns.Add("ResolvedAt");
+						Values.Add(Issue.Resolved.Value? "DATETIME('now')" : "NULL");
+					}
+
+					Command.CommandText = String.Format("UPDATE [Issues] SET ({0}) = ({1}) WHERE Id = @IssueId", String.Join(", ", Columns), String.Join(", ", Values));
+					Command.Parameters.AddWithValue("@IssueId", IssueId);
+					Command.ExecuteNonQuery();
+				}
+			}
+		}
+
+		public static void AddWatcher(long IssueId, string UserName)
+		{
+			using (SQLiteConnection Connection = new SQLiteConnection(SqlConnector.ConnectionString))
+			{
+				Connection.Open();
+
+				long UserId = FindOrAddUserId(UserName, Connection);
+
+				using(SQLiteCommand Command = new SQLiteCommand("INSERT OR IGNORE INTO [IssueWatchers] (IssueId, UserId) VALUES (@IssueId, @UserId)", Connection))
+				{
+					Command.Parameters.AddWithValue("@IssueId", IssueId);
+					Command.Parameters.AddWithValue("@UserId", UserId);
+					Command.ExecuteNonQuery();
+				}
+			}
+		}
+
+		public static List<string> GetWatchers(long IssueId)
+		{
+			List<string> Watchers = new List<string>();
+			using (SQLiteConnection Connection = new SQLiteConnection(SqlConnector.ConnectionString))
+			{
+				Connection.Open();
+
+				StringBuilder CommandBuilder = new StringBuilder();
+				CommandBuilder.Append("SELECT Users.Name FROM [IssueWatchers]");
+				CommandBuilder.Append(" LEFT JOIN [Users] ON IssueWatchers.UserId = Users.Id");
+				CommandBuilder.Append(" WHERE IssueWatchers.IssueId = @IssueId");
+
+				using(SQLiteCommand Command = new SQLiteCommand(CommandBuilder.ToString(), Connection))
+				{
+					Command.Parameters.AddWithValue("@IssueId", IssueId);
+					using(SQLiteDataReader Reader = Command.ExecuteReader())
+					{
+						while(Reader.Read())
+						{
+							Watchers.Add(Reader.GetString(0));
+						}
+					}
+				}
+			}
+			return Watchers;
+		}
+
+		public static void RemoveWatcher(long IssueId, string UserName)
+		{
+			using (SQLiteConnection Connection = new SQLiteConnection(SqlConnector.ConnectionString))
+			{
+				Connection.Open();
+
+				long UserId = FindOrAddUserId(UserName, Connection);
+
+				using(SQLiteCommand Command = new SQLiteCommand("DELETE FROM [IssueWatchers] WHERE IssueId = @IssueId AND UserId = @UserId", Connection))
+				{
+					Command.Parameters.AddWithValue("@IssueId", IssueId);
+					Command.Parameters.AddWithValue("@UserId", UserId);
+					Command.ExecuteNonQuery();
+				}
+			}
+		}
+
+		public static long AddBuild(long IssueId, string Stream, int Change, string Name, string Url, int Outcome)
+		{
+			using (SQLiteConnection Connection = new SQLiteConnection(SqlConnector.ConnectionString))
+			{
+				Connection.Open();
+
+				using (SQLiteCommand Command = new SQLiteCommand("INSERT INTO [IssueBuilds] (IssueId, Stream, Change, Name, Url, Outcome) VALUES (@IssueId, @Stream, @Change, @Name, @Url, @Outcome)", Connection))
+				{
+					Command.Parameters.AddWithValue("@IssueId", IssueId);
+					Command.Parameters.AddWithValue("@Stream", Stream);
+					Command.Parameters.AddWithValue("@Change", Change);
+					Command.Parameters.AddWithValue("@Name", Name);
+					Command.Parameters.AddWithValue("@Url", Url);
+					Command.Parameters.AddWithValue("@Outcome", Outcome);
+					Command.ExecuteNonQuery();
+
+					return Connection.LastInsertRowId;
+				}
+			}
+		}
+
+		public static List<IssueBuildData> GetBuilds(long IssueId)
+		{
+			List<IssueBuildData> Builds = new List<IssueBuildData>();
+			using (SQLiteConnection Connection = new SQLiteConnection(SqlConnector.ConnectionString))
+			{
+				Connection.Open();
+
+				using(SQLiteCommand Command = new SQLiteCommand("SELECT IssueBuilds.Id, IssueBuilds.Stream, IssueBuilds.Change, IssueBuilds.Name, IssueBuilds.Url, IssueBuilds.Outcome FROM [IssueBuilds] WHERE IssueBuilds.IssueId = @IssueId", Connection))
+				{
+					Command.Parameters.AddWithValue("@IssueId", IssueId);
+					using(SQLiteDataReader Reader = Command.ExecuteReader())
+					{
+						while(Reader.Read())
+						{
+							long Id = Reader.GetInt64(0);
+							string Stream = Reader.GetString(1);
+							int Change = Reader.GetInt32(2);
+							string Name = Reader.GetString(3);
+							string Url = Reader.GetString(4);
+							int Outcome = Reader.GetInt32(5);
+							Builds.Add(new IssueBuildData { Id = Id, Stream = Stream, Change = Change, Name = Name, Url = Url, Outcome = Outcome });
+						}
+					}
+				}
+			}
+			return Builds;
+		}
+
+		public static IssueBuildData GetBuild(long BuildId)
+		{
+			IssueBuildData Build = null;
+			using (SQLiteConnection Connection = new SQLiteConnection(SqlConnector.ConnectionString))
+			{
+				Connection.Open();
+
+				using(SQLiteCommand Command = new SQLiteCommand("SELECT IssueBuilds.Id, IssueBuilds.Stream, IssueBuilds.Change, IssueBuilds.Name, IssueBuilds.Url, IssueBuilds.Outcome FROM [IssueBuilds] WHERE IssueBuilds.Id = @BuildId", Connection))
+				{
+					Command.Parameters.AddWithValue("@BuildId", BuildId);
+					using(SQLiteDataReader Reader = Command.ExecuteReader())
+					{
+						while(Reader.Read())
+						{
+							long Id = Reader.GetInt64(0);
+							string Stream = Reader.GetString(1);
+							int Change = Reader.GetInt32(2);
+							string Name = Reader.GetString(3);
+							string Url = Reader.GetString(4);
+							int Outcome = Reader.GetInt32(5);
+
+							Build = new IssueBuildData { Id = Id, Stream = Stream, Change = Change, Name = Name, Url = Url, Outcome = Outcome };
+						}
+					}
+				}
+			}
+			return Build;
+		}
+
+		public static void UpdateBuild(long BuildId, int Outcome)
+		{
+			using (SQLiteConnection Connection = new SQLiteConnection(SqlConnector.ConnectionString))
+			{
+				Connection.Open();
+
+				using (SQLiteCommand Command = new SQLiteCommand("UPDATE [IssueBuilds] SET (Outcome) = (@Outcome) WHERE Id = @BuildId", Connection))
+				{
+					Command.Parameters.AddWithValue("@BuildId", BuildId);
+					Command.Parameters.AddWithValue("@Outcome", Outcome);
+					Command.ExecuteNonQuery();
+				}
+			}
+		}
 	}
 }
\ No newline at end of file
diff --git a/Engine/Source/Programs/UnrealGameSync/UnrealGameSyncMetadataServer/Controllers/IssuesController.cs b/Engine/Source/Programs/UnrealGameSync/UnrealGameSyncMetadataServer/Controllers/IssuesController.cs
new file mode 100644
--- /dev/null
+++ b/Engine/Source/Programs/UnrealGameSync/UnrealGameSyncMetadataServer/Controllers/IssuesController.cs
@@ -0,0 +1,100 @@
+// Copyright 1998-2018 Epic Games, Inc. All Rights Reserved.
+
+using Newtonsoft.Json.Linq;
+using System;
+using System.Collections.Generic;
+using System.Dynamic;
+using System.Linq;
+using System.Net.Http;
+using System.Web.Http;
+using System.Web.Http.Results;
+using UnrealGameSyncMetadataServer.Connectors;
+using UnrealGameSyncMetadataServer.Models;
+
+namespace UnrealGameSyncMetadataServer.Controllers
+{
+	public class IssuesController : ApiController
+	{
+		[HttpGet]
+		public List<IssueData> Get()
+		{
+			return SqlConnector.GetIssues();
+		}
+
+		[HttpGet]
+		public List<IssueData> Get(string User)
+		{
+			return SqlConnector.GetIssues(User);
+		}
+
+		[HttpGet]
+		public IssueData Get(long id)
+		{
+			return SqlConnector.GetIssue(id);
+		}
+
+		[HttpPut]
+		public void Put(long id, IssueUpdateData Issue)
+		{
+			SqlConnector.UpdateIssue(id, Issue);
+		}
+
+		[HttpPost]
+		public object Post(IssueData Issue)
+		{
+			long IssueId = SqlConnector.AddIssue(Issue);
+			return new { Id = IssueId };
+		}
+	}
+
+	public class IssueBuildsSubController : ApiController
+	{
+		[HttpGet]
+		public List<IssueBuildData> Get(long IssueId)
+		{
+			return SqlConnector.GetBuilds(IssueId);
+		}
+
+		[HttpPost]
+		public void Post(long IssueId, [FromBody] IssueBuildData Data)
+		{
+			SqlConnector.AddBuild(IssueId, Data.Stream, Data.Change, Data.Name, Data.Url, Data.Outcome);
+		}
+	}
+
+	public class IssueBuildsController : ApiController
+	{
+		[HttpGet]
+		public IssueBuildData Get(long BuildId)
+		{
+			return SqlConnector.GetBuild(BuildId);
+		}
+
+		[HttpPut]
+		public void Put(long BuildId, [FromBody] IssueBuildUpdateData Data)
+		{
+			SqlConnector.UpdateBuild(BuildId, Data.Outcome);
+		}
+	}
+
+	public class IssueWatchersController : ApiController
+	{
+		[HttpGet]
+		public List<string> Get(long IssueId)
+		{
+			return SqlConnector.GetWatchers(IssueId);
+		}
+
+		[HttpPost]
+		public void Post(long IssueId, [FromBody] IssueWatcherData Data)
+		{
+			SqlConnector.AddWatcher(IssueId, Data.UserName);
+		}
+
+		[HttpDelete]
+		public void Delete(long IssueId, [FromBody] IssueWatcherData Data)
+		{
+			SqlConnector.RemoveWatcher(IssueId, Data.UserName);
+		}
+	}
+}
diff --git a/Engine/Source/Programs/UnrealGameSync/UnrealGameSyncMetadataServer/Models/IssueData.cs b/Engine/Source/Programs/UnrealGameSync/UnrealGameSyncMetadataServer/Models/IssueData.cs
new file mode 100644
--- /dev/null
+++ b/Engine/Source/Programs/UnrealGameSync/UnrealGameSyncMetadataServer/Models/IssueData.cs
@@ -0,0 +1,55 @@
+// Copyright 1998-2019 Epic Games, Inc. All Rights Reserved.
+
+using Newtonsoft.Json.Linq;
+using System;
+using System.Collections.Generic;
+using System.Linq;
+using System.Web;
+
+namespace UnrealGameSyncMetadataServer.Models
+{
+	public class IssueWatcherData
+	{
+		public string UserName;
+	}
+
+	public class IssueBuildData
+	{
+		public long Id;
+		public string Stream;
+		public int Change;
+		public string Name;
+		public string Url;
+		public int Outcome;
+	}
+
+	public class IssueBuildUpdateData
+	{
+		public int Outcome;
+	}
+
+	public class IssueData
+	{
+		public long Id;
+		public DateTime CreatedAt;
+		public DateTime RetrievedAt;
+		public string Project;
+		public string Summary;
+		public string Owner;
+		public string NominatedBy;
+		public DateTime? AcknowledgedAt;
+		public int FixChange;
+		public DateTime? ResolvedAt;
+		public bool bNotify;
+	}
+
+	public class IssueUpdateData
+	{
+		public string Summary;
+		public string Owner;
+		public string NominatedBy;
+		public bool? Acknowledged;
+		public int? FixChange;
+		public bool? Resolved;
+	}
+}
diff --git a/Engine/Source/Programs/UnrealGameSync/UnrealGameSyncMetadataServer/UnrealGameSyncMetadataServer.csproj b/Engine/Source/Programs/UnrealGameSync/UnrealGameSyncMetadataServer/UnrealGameSyncMetadataServer.csproj
--- a/Engine/Source/Programs/UnrealGameSync/UnrealGameSyncMetadataServer/UnrealGameSyncMetadataServer.csproj
+++ b/Engine/Source/Programs/UnrealGameSync/UnrealGameSyncMetadataServer/UnrealGameSyncMetadataServer.csproj
@@ -90,6 +90,9 @@
   </ItemGroup>
   <ItemGroup>
     <Compile Include="App_Start\WebApiConfig.cs" />
+    <Compile Include="Controllers\IssuesController.cs" />
+    <Compile Include="Controllers\UserController.cs" />
+    <Compile Include="Models\IssueData.cs" />
     <Compile Include="Properties\Resources.Designer.cs">
       <AutoGen>True</AutoGen>
       <DesignTime>True</DesignTime>
@@ -118,7 +121,9 @@
   </ItemGroup>
   <ItemGroup>
     <Content Include="Global.asax" />
-    <Content Include="Web.config" />
+    <Content Include="Web.config">
+      <SubType>Designer</SubType>
+    </Content>
     <None Include="Web.template.config" />
     <None Include="Web.Debug.config">
       <DependentUpon>Web.template.config</DependentUpon>
