Rename the UnrealGameSyncMetadataServer project to MetadataServer.

diff --git a/Engine/Source/Programs/UnrealGameSync/MetadataServer/App_Start/WebApiConfig.cs b/Engine/Source/Programs/UnrealGameSync/MetadataServer/App_Start/WebApiConfig.cs
new file mode 100644
--- /dev/null
+++ b/Engine/Source/Programs/UnrealGameSync/MetadataServer/App_Start/WebApiConfig.cs
@@ -0,0 +1,45 @@
+﻿// Copyright 1998-2019 Epic Games, Inc. All Rights Reserved.
+
+using System.Web.Http;
+
+namespace UnrealGameSyncMetadataServer
+{
+	public static class WebApiConfig
+	{
+		public static void Register(HttpConfiguration config)
+		{
+			// Web API routes
+			config.MapHttpAttributeRoutes();
+
+			config.Routes.MapHttpRoute(
+				name: "IssueBuildsSubApi",
+				routeTemplate: "api/issues/{IssueId}/builds",
+				defaults: new { controller = "IssueBuildsSub" }
+			);
+
+			config.Routes.MapHttpRoute(
+				name: "IssueBuildsApi",
+				routeTemplate: "api/issuebuilds/{BuildId}",
+				defaults: new { controller = "IssueBuilds" }
+			);
+
+			config.Routes.MapHttpRoute(
+				name: "IssueEventsApi",
+				routeTemplate: "api/issues/{IssueId}/events",
+				defaults: new { controller = "IssueEvents" }
+			);
+
+			config.Routes.MapHttpRoute(
+				name: "IssueWatchersApi",
+				routeTemplate: "api/issues/{IssueId}/watchers",
+				defaults: new { controller = "IssueWatchers" }
+			);
+
+			config.Routes.MapHttpRoute(
+				name: "DefaultApi",
+				routeTemplate: "api/{controller}/{id}",
+				defaults: new { id = RouteParameter.Optional }
+			);
+		}
+	}
+}
diff --git a/Engine/Source/Programs/UnrealGameSync/MetadataServer/Connectors/SqlConnector.cs b/Engine/Source/Programs/UnrealGameSync/MetadataServer/Connectors/SqlConnector.cs
new file mode 100644
--- /dev/null
+++ b/Engine/Source/Programs/UnrealGameSync/MetadataServer/Connectors/SqlConnector.cs
@@ -0,0 +1,676 @@
+﻿// Copyright 1998-2019 Epic Games, Inc. All Rights Reserved.
+
+using Newtonsoft.Json;
+using Newtonsoft.Json.Linq;
+using System;
+using System.Collections.Generic;
+using System.Data.SqlClient;
+using System.Data.SQLite;
+using System.IO;
+using System.Linq;
+using System.Text;
+using System.Text.RegularExpressions;
+using UnrealGameSyncMetadataServer.Models;
+
+namespace UnrealGameSyncMetadataServer.Connectors
+{
+	public static class SqlConnector
+	{
+		private static string ConnectionString = System.Configuration.ConfigurationManager.ConnectionStrings["ConnectionString"].ConnectionString;
+		public static LatestData GetLastIds(string Project = null)
+		{
+			// Get ids going back 432 builds for the project being asked for
+			// Do this by grouping by ChangeNumber to get unique entries, then take the 432nd id
+			long LastEventId = 0;
+			long LastCommentId = 0;
+			long LastBuildId = 0;
+			string ProjectLikeString = "%" + (Project == null ? String.Empty : GetProjectStream(Project)) + "%";
+			using (SQLiteConnection Connection = new SQLiteConnection(ConnectionString))
+			{
+				Connection.Open();
+				using (SQLiteCommand Command = new SQLiteCommand("SELECT Id FROM [UserVotes] WHERE Project LIKE @param1 GROUP BY Changelist ORDER BY Changelist DESC LIMIT 1 OFFSET 432;", Connection))
+				{
+					Command.Parameters.AddWithValue("@param1", ProjectLikeString);
+					using (SQLiteDataReader Reader = Command.ExecuteReader())
+					{
+						while (Reader.Read())
+						{
+							LastEventId = Reader.GetInt64(0);
+							break;
+						}
+					}
+				}
+				using (SQLiteCommand Command = new SQLiteCommand("SELECT Id FROM [Comments] WHERE Project LIKE @param1 GROUP BY ChangeNumber ORDER BY ChangeNumber DESC LIMIT 1 OFFSET 432;", Connection))
+				{
+					Command.Parameters.AddWithValue("@param1", ProjectLikeString);
+					using (SQLiteDataReader Reader = Command.ExecuteReader())
+					{
+						while (Reader.Read())
+						{
+							LastCommentId = Reader.GetInt32(0);
+							break;
+						}
+					}
+				}
+				using (SQLiteCommand Command = new SQLiteCommand("SELECT Id FROM [CIS] WHERE Project LIKE @param1 GROUP BY ChangeNumber ORDER BY ChangeNumber DESC LIMIT 1 OFFSET 432", Connection))
+				{
+					Command.Parameters.AddWithValue("@param1", ProjectLikeString);
+					using (SQLiteDataReader Reader = Command.ExecuteReader())
+					{
+						while (Reader.Read())
+						{
+							LastBuildId = Reader.GetInt32(0);
+							break;
+						}
+					}
+				}
+			}
+			return new LatestData { LastBuildId = LastBuildId, LastCommentId = LastCommentId, LastEventId = LastEventId };
+		}
+		public static List<EventData> GetUserVotes(string Project, long LastEventId)
+		{
+			List<EventData> ReturnedEvents = new List<EventData>();
+			string ProjectLikeString = "%" + (Project == null ? String.Empty : GetProjectStream(Project)) + "%";
+			using (SQLiteConnection Connection = new SQLiteConnection(ConnectionString))
+			{
+				Connection.Open();
+				using (SQLiteCommand Command = new SQLiteCommand("SELECT Id, Changelist, UserName, Verdict, Project FROM [UserVotes] WHERE Id > @param1 AND Project LIKE @param2 ORDER BY Id", Connection))
+				{
+					Command.Parameters.AddWithValue("@param1", LastEventId);
+					Command.Parameters.AddWithValue("@param2", ProjectLikeString);
+					using (SQLiteDataReader Reader = Command.ExecuteReader())
+					{
+						while (Reader.Read())
+						{
+							EventData Review = new EventData();
+							Review.Id = Reader.GetInt64(0);
+							Review.Change = Reader.GetInt32(1);
+							Review.UserName = Reader.GetString(2);
+							Review.Project = Reader.IsDBNull(4) ? null : Reader.GetString(4);
+							if (Enum.TryParse(Reader.GetString(3), out Review.Type))
+							{
+								if (Review.Project == null || String.Compare(Review.Project, Project, true) == 0)
+								{
+									ReturnedEvents.Add(Review);
+								}
+							}
+						}
+					}
+				}
+			}
+			return ReturnedEvents;
+		}
+		public static List<CommentData> GetComments(string Project, long LastCommentId)
+		{
+			List<CommentData> ReturnedComments = new List<CommentData>();
+			string ProjectLikeString = "%" + (Project == null ? String.Empty : GetProjectStream(Project)) + "%";
+			using (SQLiteConnection Connection = new SQLiteConnection(ConnectionString))
+			{
+				Connection.Open();
+				using (SQLiteCommand Command = new SQLiteCommand("SELECT Id, ChangeNumber, UserName, Text, Project FROM [Comments] WHERE Id > @param1 AND Project LIKE @param2 ORDER BY Id", Connection))
+				{
+					Command.Parameters.AddWithValue("@param1", LastCommentId);
+					Command.Parameters.AddWithValue("@param2", ProjectLikeString);
+					using (SQLiteDataReader Reader = Command.ExecuteReader())
+					{
+						while (Reader.Read())
+						{
+							CommentData Comment = new CommentData();
+							Comment.Id = Reader.GetInt32(0);
+							Comment.ChangeNumber = Reader.GetInt32(1);
+							Comment.UserName = Reader.GetString(2);
+							Comment.Text = Reader.GetString(3);
+							Comment.Project = Reader.GetString(4);
+							if (Comment.Project == null || String.Compare(Comment.Project, Project, true) == 0)
+							{
+								ReturnedComments.Add(Comment);
+							}
+						}
+					}
+				}
+			}
+			return ReturnedComments;
+		}
+
+		public static List<BuildData> GetBuilds(string Project, long LastBuildId)
+		{
+			List<BuildData> ReturnedBuilds = new List<BuildData>();
+			string ProjectLikeString = "%" + (Project == null ? String.Empty : GetProjectStream(Project)) + "%";
+			using (SQLiteConnection Connection = new SQLiteConnection(ConnectionString))
+			{
+				Connection.Open();
+				using (SQLiteCommand Command = new SQLiteCommand("SELECT Id, ChangeNumber, BuildType, Result, Url, Project FROM [CIS] WHERE Id > @param1 AND Project LIKE @param2 ORDER BY Id", Connection))
+				{
+					Command.Parameters.AddWithValue("@param1", LastBuildId);
+					Command.Parameters.AddWithValue("@param2", ProjectLikeString);
+					using (SQLiteDataReader Reader = Command.ExecuteReader())
+					{
+						while (Reader.Read())
+						{
+							BuildData Build = new BuildData();
+							Build.Id = Reader.GetInt32(0);
+							Build.ChangeNumber = Reader.GetInt32(1);
+							Build.BuildType = Reader.GetString(2).TrimEnd();
+							if (Enum.TryParse(Reader.GetString(3).TrimEnd(), true, out Build.Result))
+							{
+								Build.Url = Reader.GetString(4);
+								Build.Project = Reader.IsDBNull(5) ? null : Reader.GetString(5);
+								if (Build.Project == null || String.Compare(Build.Project, Project, true) == 0 || MatchesWildcard(Build.Project, Project))
+								{
+									ReturnedBuilds.Add(Build);
+								}
+							}
+							LastBuildId = Math.Max(LastBuildId, Build.Id);
+						}
+					}
+				}
+			}
+			return ReturnedBuilds;
+		}
+
+		public static List<TelemetryErrorData> GetErrorData(int Records)
+		{
+			List<TelemetryErrorData> ReturnedErrors = new List<TelemetryErrorData>();
+			using (SQLiteConnection Connection = new SQLiteConnection(ConnectionString))
+			{
+				Connection.Open();
+				using (SQLiteCommand Command = new SQLiteCommand("SELECT Id, Type, Text, UserName, Project, Timestamp, Version, IpAddress FROM [Errors] ORDER BY Id DESC LIMIT @param1", Connection))
+				{
+					Command.Parameters.AddWithValue("@param1", Records);
+					using (SQLiteDataReader Reader = Command.ExecuteReader())
+					{
+						while (Reader.Read())
+						{
+							TelemetryErrorData Error = new TelemetryErrorData();
+							Error.Id = Reader.GetInt32(0);
+							Enum.TryParse(Reader.GetString(1), true, out Error.Type);
+							Error.Text = Reader.GetString(2);
+							Error.UserName = Reader.GetString(3);
+							Error.Project = Reader.IsDBNull(4) ? null : Reader.GetString(4);
+							Error.Timestamp = Reader.GetDateTime(5);
+							Error.Version = Reader.GetString(6);
+							Error.IpAddress = Reader.GetString(7);
+							ReturnedErrors.Add(Error);
+						}
+					}
+				}
+			}
+			return ReturnedErrors;
+		}
+
+
+		public static void PostBuild(BuildData Build)
+		{
+			using (SQLiteConnection Connection = new SQLiteConnection(ConnectionString))
+			{
+				Connection.Open();
+				using (SQLiteCommand Command = new SQLiteCommand("INSERT INTO [CIS] (ChangeNumber, BuildType, Result, URL, Project, ArchivePath) VALUES (@ChangeNumber, @BuildType, @Result, @URL, @Project, @ArchivePath)", Connection))
+				{
+					Command.Parameters.AddWithValue("@ChangeNumber", Build.ChangeNumber);
+					Command.Parameters.AddWithValue("@BuildType", Build.BuildType);
+					Command.Parameters.AddWithValue("@Result", Build.Result);
+					Command.Parameters.AddWithValue("@URL", Build.Url);
+					Command.Parameters.AddWithValue("@Project", Build.Project);
+					Command.Parameters.AddWithValue("@ArchivePath", Build.ArchivePath);
+					Command.ExecuteNonQuery();
+				}
+			}
+		}
+
+		public static void PostEvent(EventData Event)
+		{
+			using (SQLiteConnection Connection = new SQLiteConnection(ConnectionString))
+			{
+				Connection.Open();
+				using (SQLiteCommand Command = new SQLiteCommand("INSERT INTO [UserVotes] (Changelist, UserName, Verdict, Project) VALUES (@Changelist, @UserName, @Verdict, @Project)", Connection))
+				{
+					Command.Parameters.AddWithValue("@Changelist", Event.Change);
+					Command.Parameters.AddWithValue("@UserName", Event.UserName.ToString());
+					Command.Parameters.AddWithValue("@Verdict", Event.Type.ToString());
+					Command.Parameters.AddWithValue("@Project", Event.Project);
+					Command.ExecuteNonQuery();
+				}
+			}
+		}
+
+		public static void PostComment(CommentData Comment)
+		{
+			using (SQLiteConnection Connection = new SQLiteConnection(ConnectionString))
+			{
+				Connection.Open();
+				using (SQLiteCommand Command = new SQLiteCommand("INSERT INTO [Comments] (ChangeNumber, UserName, Text, Project) VALUES (@ChangeNumber, @UserName, @Text, @Project)", Connection))
+				{
+					Command.Parameters.AddWithValue("@ChangeNumber", Comment.ChangeNumber);
+					Command.Parameters.AddWithValue("@UserName", Comment.UserName);
+					Command.Parameters.AddWithValue("@Text", Comment.Text);
+					Command.Parameters.AddWithValue("@Project", Comment.Project);
+					Command.ExecuteNonQuery();
+				}
+			}
+		}
+
+		public static void PostTelemetryData(TelemetryTimingData Data, string Version, string IpAddress)
+		{
+			using (SQLiteConnection Connection = new SQLiteConnection(ConnectionString))
+			{
+				Connection.Open();
+				using (SQLiteCommand Command = new SQLiteCommand("INSERT INTO [Telemetry.v2] (Action, Result, UserName, Project, Timestamp, Duration, Version, IpAddress) VALUES (@Action, @Result, @UserName, @Project, @Timestamp, @Duration, @Version, @IpAddress)", Connection))
+				{
+					Command.Parameters.AddWithValue("@Action", Data.Action);
+					Command.Parameters.AddWithValue("@Result", Data.Result);
+					Command.Parameters.AddWithValue("@UserName", Data.UserName);
+					Command.Parameters.AddWithValue("@Project", Data.Project);
+					Command.Parameters.AddWithValue("@Timestamp", Data.Timestamp);
+					Command.Parameters.AddWithValue("@Duration", Data.Duration);
+					Command.Parameters.AddWithValue("@Version", Version);
+					Command.Parameters.AddWithValue("@IPAddress", IpAddress);
+					Command.ExecuteNonQuery();
+				}
+			}
+		}
+
+		public static void PostErrorData(TelemetryErrorData Data, string Version, string IpAddress)
+		{
+			using (SQLiteConnection Connection = new SQLiteConnection(ConnectionString))
+			{
+				Connection.Open();
+				using (SQLiteCommand Command = new SQLiteCommand("INSERT INTO [Errors] (Type, Text, UserName, Project, Timestamp, Version, IpAddress) VALUES (@Type, @Text, @UserName, @Project, @Timestamp, @Version, @IpAddress)", Connection))
+				{
+					Command.Parameters.AddWithValue("@Type", Data.Type.ToString());
+					Command.Parameters.AddWithValue("@Text", Data.Text);
+					Command.Parameters.AddWithValue("@UserName", Data.UserName);
+					if (Data.Project == null)
+					{
+						Command.Parameters.AddWithValue("@Project", DBNull.Value);
+					}
+					else
+					{
+						Command.Parameters.AddWithValue("@Project", Data.Project);
+					}
+					Command.Parameters.AddWithValue("@Timestamp", Data.Timestamp);
+					Command.Parameters.AddWithValue("@Version", Version);
+					Command.Parameters.AddWithValue("@IPAddress", IpAddress);
+					Command.ExecuteNonQuery();
+				}
+			}
+		}
+		private static string GetProjectStream(string Project)
+		{
+			// Get first two fragments of the p4 path.  If it doesn't work, just return back the project.
+			Regex StreamPattern = new Regex("(\\/\\/[a-zA-Z0-9\\.\\-_]{1,}\\/[a-zA-Z0-9\\.\\-_]{1,})");
+			Match StreamMatch = StreamPattern.Match(Project);
+			if(StreamMatch.Success)
+			{
+				return StreamMatch.Groups[1].Value;
+			}
+			return Project;
+		}
+		private static bool MatchesWildcard(string Wildcard, string Project)
+		{
+			return Wildcard.EndsWith("...") && Project.StartsWith(Wildcard.Substring(0, Wildcard.Length - 4), StringComparison.InvariantCultureIgnoreCase);
+		}
+
+		private static string NormalizeUserName(string UserName)
+		{
+			return UserName.ToUpperInvariant();
+		}
+
+		public static long FindOrAddUserId(string Name)
+		{
+			using (SQLiteConnection Connection = new SQLiteConnection(SqlConnector.ConnectionString))
+			{
+				Connection.Open();
+				return FindOrAddUserId(Name, Connection);
+			}
+		}
+
+		private static long FindOrAddUserId(string Name, SQLiteConnection Connection)
+		{
+			if(Name.Length == 0)
+			{
+				return -1;
+			}
+
+			string NormalizedName = NormalizeUserName(Name);
+
+			using (SQLiteCommand Command = new SQLiteCommand("SELECT [Id] FROM [Users] WHERE Name = @Name", Connection))
+			{
+				Command.Parameters.AddWithValue("@Name", NormalizedName);
+				object UserId = Command.ExecuteScalar();
+				if(UserId != null)
+				{
+					return (long)UserId;
+				}
+			}
+
+			using (SQLiteCommand Command = new SQLiteCommand("INSERT OR IGNORE INTO [Users] (Name) VALUES (@Name); SELECT [Id] FROM [Users] WHERE Name = @Name", Connection))
+			{
+				Command.Parameters.AddWithValue("@Name", NormalizedName);
+				object UserId = Command.ExecuteScalar();
+				return (long)UserId;
+			}
+		}
+
+		public static long AddIssue(IssueData Issue)
+		{
+			long IssueId;
+			using (SQLiteConnection Connection = new SQLiteConnection(SqlConnector.ConnectionString))
+			{
+				Connection.Open();
+
+				using (SQLiteCommand Command = new SQLiteCommand("INSERT INTO [Issues] (Project, Summary, CreatedAt, FixChange) VALUES (@Project, @Summary, DATETIME('now'), 0)", Connection))
+				{
+					Command.Parameters.AddWithValue("@Project", Issue.Project);
+					Command.Parameters.AddWithValue("@Summary", Issue.Summary);
+					Command.ExecuteNonQuery();
+
+					IssueId = Connection.LastInsertRowId;
+				}
+			}
+			return IssueId;
+		}
+
+		public static IssueData GetIssue(long IssueId)
+		{
+			List<IssueData> Issues = GetIssuesInternal(IssueId, null);
+			if(Issues.Count == 0)
+			{
+				return null;
+			}
+			else
+			{
+				return Issues[0];
+			}
+		}
+
+		public static List<IssueData> GetIssues()
+		{
+			return GetIssuesInternal(-1, null);
+		}
+
+		public static List<IssueData> GetIssues(string UserName)
+		{
+			return GetIssuesInternal(-1, UserName);
+		}
+
+		private static List<IssueData> GetIssuesInternal(long IssueId, string UserName)
+		{
+			List<IssueData> Issues = new List<IssueData>();
+			using (SQLiteConnection Connection = new SQLiteConnection(SqlConnector.ConnectionString))
+			{
+				Connection.Open();
+
+				long UserId = -1;
+				if(UserName != null)
+				{
+					UserId = FindOrAddUserId(UserName);
+				}
+
+				StringBuilder CommandBuilder = new StringBuilder();
+				CommandBuilder.Append("SELECT");
+				CommandBuilder.Append(" Issues.Id, Issues.CreatedAt, DATETIME('now'), Issues.Project, Issues.Summary, OwnerUsers.Name, NominatedByUsers.Name, Issues.AcknowledgedAt, Issues.FixChange, Issues.ResolvedAt");
+				if(UserName != null)
+				{
+					CommandBuilder.Append(", IssueWatchers.UserId");
+				}
+				CommandBuilder.Append(" FROM [Issues]");
+				CommandBuilder.Append(" LEFT JOIN [Users] AS [OwnerUsers] ON OwnerUsers.Id = Issues.OwnerId");
+				CommandBuilder.Append(" LEFT JOIN [Users] AS [NominatedByUsers] ON NominatedByUsers.Id = Issues.NominatedById");
+				if(UserName != null)
+				{
+					CommandBuilder.Append(" LEFT JOIN [IssueWatchers] ON IssueWatchers.IssueId = Issues.Id AND IssueWatchers.UserId = @UserId");
+				}
+				if(IssueId != -1)
+				{
+					CommandBuilder.Append(" WHERE Issues.Id = @IssueId");
+				}
+				else
+				{
+					CommandBuilder.Append(" WHERE Issues.ResolvedAt IS NULL");
+				}
+
+				using (SQLiteCommand Command = new SQLiteCommand(CommandBuilder.ToString(), Connection))
+				{
+					if(IssueId != -1)
+					{
+						Command.Parameters.AddWithValue("@IssueId", IssueId);
+					}
+					if(UserName != null)
+					{
+						Command.Parameters.AddWithValue("@UserId", UserId);
+					}
+
+					using(SQLiteDataReader Reader = Command.ExecuteReader())
+					{
+						while(Reader.Read())
+						{
+							IssueData Issue = new IssueData();
+							Issue.Id = Reader.GetInt64(0);
+							Issue.CreatedAt = Reader.GetDateTime(1);
+							Issue.RetrievedAt = Reader.GetDateTime(2);
+							Issue.Project = Reader.GetString(3);
+							Issue.Summary = Reader.GetString(4);
+							Issue.Owner = Reader.IsDBNull(5)? null : Reader.GetString(5);
+							Issue.NominatedBy = Reader.IsDBNull(6)? null : Reader.GetString(6);
+							Issue.AcknowledgedAt = Reader.IsDBNull(7)? (DateTime?)null : Reader.GetDateTime(7);
+							Issue.FixChange = Reader.GetInt32(8);
+							Issue.ResolvedAt = Reader.IsDBNull(9)? (DateTime?)null : Reader.GetDateTime(9);
+							if(UserName != null)
+							{
+								Issue.bNotify = !Reader.IsDBNull(10);
+							}
+							Issues.Add(Issue);
+						}
+					}
+				}
+			}
+			return Issues;
+		}
+
+		public static void UpdateIssue(long IssueId, IssueUpdateData Issue)
+		{
+			using (SQLiteConnection Connection = new SQLiteConnection(SqlConnector.ConnectionString))
+			{
+				Connection.Open();
+
+				using (SQLiteCommand Command = new SQLiteCommand(Connection))
+				{
+					List<string> Columns = new List<string>();
+					List<string> Values = new List<string>();
+					if(Issue.Summary != null)
+					{
+						Columns.Add("Summary");
+						Values.Add("@Summary");
+						Command.Parameters.AddWithValue("@Summary", Issue.Summary);
+					}
+					if(Issue.Owner != null)
+					{
+						Columns.Add("OwnerId");
+						Values.Add("@OwnerId");
+						Command.Parameters.AddWithValue("OwnerId", FindOrAddUserId(Issue.Owner, Connection));
+					}
+					if(Issue.NominatedBy != null)
+					{
+						Columns.Add("NominatedById");
+						Values.Add("@NominatedById");
+						Command.Parameters.AddWithValue("NominatedById", FindOrAddUserId(Issue.NominatedBy, Connection));
+					}
+					if(Issue.Acknowledged.HasValue)
+					{
+						Columns.Add("AcknowledgedAt");
+						Values.Add(Issue.Acknowledged.Value? "DATETIME('now')" : "NULL");
+					}
+					if(Issue.FixChange.HasValue)
+					{
+						Columns.Add("FixChange");
+						Values.Add("@FixChange");
+						Command.Parameters.AddWithValue("FixChange", Issue.FixChange.Value);
+					}
+					if(Issue.Resolved.HasValue)
+					{
+						Columns.Add("ResolvedAt");
+						Values.Add(Issue.Resolved.Value? "DATETIME('now')" : "NULL");
+					}
+
+					Command.CommandText = String.Format("UPDATE [Issues] SET ({0}) = ({1}) WHERE Id = @IssueId", String.Join(", ", Columns), String.Join(", ", Values));
+					Command.Parameters.AddWithValue("@IssueId", IssueId);
+					Command.ExecuteNonQuery();
+				}
+			}
+		}
+
+		public static void AddWatcher(long IssueId, string UserName)
+		{
+			using (SQLiteConnection Connection = new SQLiteConnection(SqlConnector.ConnectionString))
+			{
+				Connection.Open();
+
+				long UserId = FindOrAddUserId(UserName, Connection);
+
+				using(SQLiteCommand Command = new SQLiteCommand("INSERT OR IGNORE INTO [IssueWatchers] (IssueId, UserId) VALUES (@IssueId, @UserId)", Connection))
+				{
+					Command.Parameters.AddWithValue("@IssueId", IssueId);
+					Command.Parameters.AddWithValue("@UserId", UserId);
+					Command.ExecuteNonQuery();
+				}
+			}
+		}
+
+		public static List<string> GetWatchers(long IssueId)
+		{
+			List<string> Watchers = new List<string>();
+			using (SQLiteConnection Connection = new SQLiteConnection(SqlConnector.ConnectionString))
+			{
+				Connection.Open();
+
+				StringBuilder CommandBuilder = new StringBuilder();
+				CommandBuilder.Append("SELECT Users.Name FROM [IssueWatchers]");
+				CommandBuilder.Append(" LEFT JOIN [Users] ON IssueWatchers.UserId = Users.Id");
+				CommandBuilder.Append(" WHERE IssueWatchers.IssueId = @IssueId");
+
+				using(SQLiteCommand Command = new SQLiteCommand(CommandBuilder.ToString(), Connection))
+				{
+					Command.Parameters.AddWithValue("@IssueId", IssueId);
+					using(SQLiteDataReader Reader = Command.ExecuteReader())
+					{
+						while(Reader.Read())
+						{
+							Watchers.Add(Reader.GetString(0));
+						}
+					}
+				}
+			}
+			return Watchers;
+		}
+
+		public static void RemoveWatcher(long IssueId, string UserName)
+		{
+			using (SQLiteConnection Connection = new SQLiteConnection(SqlConnector.ConnectionString))
+			{
+				Connection.Open();
+
+				long UserId = FindOrAddUserId(UserName, Connection);
+
+				using(SQLiteCommand Command = new SQLiteCommand("DELETE FROM [IssueWatchers] WHERE IssueId = @IssueId AND UserId = @UserId", Connection))
+				{
+					Command.Parameters.AddWithValue("@IssueId", IssueId);
+					Command.Parameters.AddWithValue("@UserId", UserId);
+					Command.ExecuteNonQuery();
+				}
+			}
+		}
+
+		public static long AddBuild(long IssueId, string Stream, int Change, string Name, string Url, int Outcome)
+		{
+			using (SQLiteConnection Connection = new SQLiteConnection(SqlConnector.ConnectionString))
+			{
+				Connection.Open();
+
+				using (SQLiteCommand Command = new SQLiteCommand("INSERT INTO [IssueBuilds] (IssueId, Stream, Change, Name, Url, Outcome) VALUES (@IssueId, @Stream, @Change, @Name, @Url, @Outcome)", Connection))
+				{
+					Command.Parameters.AddWithValue("@IssueId", IssueId);
+					Command.Parameters.AddWithValue("@Stream", Stream);
+					Command.Parameters.AddWithValue("@Change", Change);
+					Command.Parameters.AddWithValue("@Name", Name);
+					Command.Parameters.AddWithValue("@Url", Url);
+					Command.Parameters.AddWithValue("@Outcome", Outcome);
+					Command.ExecuteNonQuery();
+
+					return Connection.LastInsertRowId;
+				}
+			}
+		}
+
+		public static List<IssueBuildData> GetBuilds(long IssueId)
+		{
+			List<IssueBuildData> Builds = new List<IssueBuildData>();
+			using (SQLiteConnection Connection = new SQLiteConnection(SqlConnector.ConnectionString))
+			{
+				Connection.Open();
+
+				using(SQLiteCommand Command = new SQLiteCommand("SELECT IssueBuilds.Id, IssueBuilds.Stream, IssueBuilds.Change, IssueBuilds.Name, IssueBuilds.Url, IssueBuilds.Outcome FROM [IssueBuilds] WHERE IssueBuilds.IssueId = @IssueId", Connection))
+				{
+					Command.Parameters.AddWithValue("@IssueId", IssueId);
+					using(SQLiteDataReader Reader = Command.ExecuteReader())
+					{
+						while(Reader.Read())
+						{
+							long Id = Reader.GetInt64(0);
+							string Stream = Reader.GetString(1);
+							int Change = Reader.GetInt32(2);
+							string Name = Reader.GetString(3);
+							string Url = Reader.GetString(4);
+							int Outcome = Reader.GetInt32(5);
+							Builds.Add(new IssueBuildData { Id = Id, Stream = Stream, Change = Change, Name = Name, Url = Url, Outcome = Outcome });
+						}
+					}
+				}
+			}
+			return Builds;
+		}
+
+		public static IssueBuildData GetBuild(long BuildId)
+		{
+			IssueBuildData Build = null;
+			using (SQLiteConnection Connection = new SQLiteConnection(SqlConnector.ConnectionString))
+			{
+				Connection.Open();
+
+				using(SQLiteCommand Command = new SQLiteCommand("SELECT IssueBuilds.Id, IssueBuilds.Stream, IssueBuilds.Change, IssueBuilds.Name, IssueBuilds.Url, IssueBuilds.Outcome FROM [IssueBuilds] WHERE IssueBuilds.Id = @BuildId", Connection))
+				{
+					Command.Parameters.AddWithValue("@BuildId", BuildId);
+					using(SQLiteDataReader Reader = Command.ExecuteReader())
+					{
+						while(Reader.Read())
+						{
+							long Id = Reader.GetInt64(0);
+							string Stream = Reader.GetString(1);
+							int Change = Reader.GetInt32(2);
+							string Name = Reader.GetString(3);
+							string Url = Reader.GetString(4);
+							int Outcome = Reader.GetInt32(5);
+
+							Build = new IssueBuildData { Id = Id, Stream = Stream, Change = Change, Name = Name, Url = Url, Outcome = Outcome };
+						}
+					}
+				}
+			}
+			return Build;
+		}
+
+		public static void UpdateBuild(long BuildId, int Outcome)
+		{
+			using (SQLiteConnection Connection = new SQLiteConnection(SqlConnector.ConnectionString))
+			{
+				Connection.Open();
+
+				using (SQLiteCommand Command = new SQLiteCommand("UPDATE [IssueBuilds] SET (Outcome) = (@Outcome) WHERE Id = @BuildId", Connection))
+				{
+					Command.Parameters.AddWithValue("@BuildId", BuildId);
+					Command.Parameters.AddWithValue("@Outcome", Outcome);
+					Command.ExecuteNonQuery();
+				}
+			}
+		}
+	}
+}
\ No newline at end of file
diff --git a/Engine/Source/Programs/UnrealGameSync/MetadataServer/Controllers/BuildController.cs b/Engine/Source/Programs/UnrealGameSync/MetadataServer/Controllers/BuildController.cs
new file mode 100644
--- /dev/null
+++ b/Engine/Source/Programs/UnrealGameSync/MetadataServer/Controllers/BuildController.cs
@@ -0,0 +1,21 @@
+﻿// Copyright 1998-2019 Epic Games, Inc. All Rights Reserved.
+
+using System.Collections.Generic;
+using System.Web.Http;
+using UnrealGameSyncMetadataServer.Connectors;
+using UnrealGameSyncMetadataServer.Models;
+
+namespace UnrealGameSyncMetadataServer.Controllers
+{
+    public class BuildController : ApiController
+    {
+		public List<BuildData> Get(string Project, long LastBuildId)
+		{
+			return SqlConnector.GetBuilds(Project, LastBuildId);
+		}
+		public void Post([FromBody]BuildData Build)
+		{
+			SqlConnector.PostBuild(Build);
+		}
+	}
+}
diff --git a/Engine/Source/Programs/UnrealGameSync/MetadataServer/Controllers/CISController.cs b/Engine/Source/Programs/UnrealGameSync/MetadataServer/Controllers/CISController.cs
new file mode 100644
--- /dev/null
+++ b/Engine/Source/Programs/UnrealGameSync/MetadataServer/Controllers/CISController.cs
@@ -0,0 +1,28 @@
+﻿// Copyright 1998-2019 Epic Games, Inc. All Rights Reserved.
+
+using System;
+using System.Collections.Generic;
+using System.Web.Http;
+using UnrealGameSyncMetadataServer.Connectors;
+using UnrealGameSyncMetadataServer.Models;
+
+namespace UnrealGameSyncMetadataServer.Controllers
+{
+	// Deprecated, use Latest or Build controllers instead
+    public class CISController : ApiController
+    {
+		public long[] Get(string Project = null)
+		{
+			LatestData LatestIDs = SqlConnector.GetLastIds(Project);
+			return new long[] { LatestIDs.LastEventId, LatestIDs.LastCommentId, LatestIDs.LastBuildId };
+		}
+		public List<BuildData> Get(string Project, long LastBuildId)
+		{
+			return SqlConnector.GetBuilds(Project, LastBuildId);
+		}
+		public void Post([FromBody]BuildData Build)
+		{
+			SqlConnector.PostBuild(Build);
+		}
+	}
+}
diff --git a/Engine/Source/Programs/UnrealGameSync/MetadataServer/Controllers/CommentController.cs b/Engine/Source/Programs/UnrealGameSync/MetadataServer/Controllers/CommentController.cs
new file mode 100644
--- /dev/null
+++ b/Engine/Source/Programs/UnrealGameSync/MetadataServer/Controllers/CommentController.cs
@@ -0,0 +1,21 @@
+﻿// Copyright 1998-2019 Epic Games, Inc. All Rights Reserved.
+
+using System.Collections.Generic;
+using System.Web.Http;
+using UnrealGameSyncMetadataServer.Connectors;
+using UnrealGameSyncMetadataServer.Models;
+
+namespace UnrealGameSyncMetadataServer.Controllers
+{
+	public class CommentController : ApiController
+	{
+		public List<CommentData> Get(string Project, long LastCommentId)
+		{
+			return SqlConnector.GetComments(Project, LastCommentId);
+		}
+		public void Post([FromBody] CommentData Comment)
+		{
+			SqlConnector.PostComment(Comment);
+		}
+	}
+}
\ No newline at end of file
diff --git a/Engine/Source/Programs/UnrealGameSync/MetadataServer/Controllers/ErrorController.cs b/Engine/Source/Programs/UnrealGameSync/MetadataServer/Controllers/ErrorController.cs
new file mode 100644
--- /dev/null
+++ b/Engine/Source/Programs/UnrealGameSync/MetadataServer/Controllers/ErrorController.cs
@@ -0,0 +1,21 @@
+﻿// Copyright 1998-2019 Epic Games, Inc. All Rights Reserved.
+
+using System.Collections.Generic;
+using System.Web.Http;
+using UnrealGameSyncMetadataServer.Connectors;
+using UnrealGameSyncMetadataServer.Models;
+
+namespace UnrealGameSyncMetadataServer.Controllers
+{
+    public class ErrorController : ApiController
+    {
+		public List<TelemetryErrorData> Get(int Records = 10)
+		{
+			return SqlConnector.GetErrorData(Records);
+		}
+		public void Post([FromBody] TelemetryErrorData Data, string Version, string IpAddress)
+		{
+			SqlConnector.PostErrorData(Data, Version, IpAddress);
+		}
+    }
+}
diff --git a/Engine/Source/Programs/UnrealGameSync/MetadataServer/Controllers/EventController.cs b/Engine/Source/Programs/UnrealGameSync/MetadataServer/Controllers/EventController.cs
new file mode 100644
--- /dev/null
+++ b/Engine/Source/Programs/UnrealGameSync/MetadataServer/Controllers/EventController.cs
@@ -0,0 +1,22 @@
+﻿// Copyright 1998-2019 Epic Games, Inc. All Rights Reserved.
+
+using System.Collections.Generic;
+using System.Web.Http;
+using UnrealGameSyncMetadataServer.Models;
+using UnrealGameSyncMetadataServer.Connectors;
+
+namespace UnrealGameSyncMetadataServer.Controllers
+{
+	public class EventController : ApiController
+	{
+		public List<EventData> Get(string Project, long LastEventId)
+		{
+			return SqlConnector.GetUserVotes(Project, LastEventId);
+		}
+
+		public void Post([FromBody] EventData Event)
+		{
+			SqlConnector.PostEvent(Event);
+		}
+	}
+}
\ No newline at end of file
diff --git a/Engine/Source/Programs/UnrealGameSync/MetadataServer/Controllers/IssuesController.cs b/Engine/Source/Programs/UnrealGameSync/MetadataServer/Controllers/IssuesController.cs
new file mode 100644
--- /dev/null
+++ b/Engine/Source/Programs/UnrealGameSync/MetadataServer/Controllers/IssuesController.cs
@@ -0,0 +1,100 @@
+// Copyright 1998-2018 Epic Games, Inc. All Rights Reserved.
+
+using Newtonsoft.Json.Linq;
+using System;
+using System.Collections.Generic;
+using System.Dynamic;
+using System.Linq;
+using System.Net.Http;
+using System.Web.Http;
+using System.Web.Http.Results;
+using UnrealGameSyncMetadataServer.Connectors;
+using UnrealGameSyncMetadataServer.Models;
+
+namespace UnrealGameSyncMetadataServer.Controllers
+{
+	public class IssuesController : ApiController
+	{
+		[HttpGet]
+		public List<IssueData> Get()
+		{
+			return SqlConnector.GetIssues();
+		}
+
+		[HttpGet]
+		public List<IssueData> Get(string User)
+		{
+			return SqlConnector.GetIssues(User);
+		}
+
+		[HttpGet]
+		public IssueData Get(long id)
+		{
+			return SqlConnector.GetIssue(id);
+		}
+
+		[HttpPut]
+		public void Put(long id, IssueUpdateData Issue)
+		{
+			SqlConnector.UpdateIssue(id, Issue);
+		}
+
+		[HttpPost]
+		public object Post(IssueData Issue)
+		{
+			long IssueId = SqlConnector.AddIssue(Issue);
+			return new { Id = IssueId };
+		}
+	}
+
+	public class IssueBuildsSubController : ApiController
+	{
+		[HttpGet]
+		public List<IssueBuildData> Get(long IssueId)
+		{
+			return SqlConnector.GetBuilds(IssueId);
+		}
+
+		[HttpPost]
+		public void Post(long IssueId, [FromBody] IssueBuildData Data)
+		{
+			SqlConnector.AddBuild(IssueId, Data.Stream, Data.Change, Data.Name, Data.Url, Data.Outcome);
+		}
+	}
+
+	public class IssueBuildsController : ApiController
+	{
+		[HttpGet]
+		public IssueBuildData Get(long BuildId)
+		{
+			return SqlConnector.GetBuild(BuildId);
+		}
+
+		[HttpPut]
+		public void Put(long BuildId, [FromBody] IssueBuildUpdateData Data)
+		{
+			SqlConnector.UpdateBuild(BuildId, Data.Outcome);
+		}
+	}
+
+	public class IssueWatchersController : ApiController
+	{
+		[HttpGet]
+		public List<string> Get(long IssueId)
+		{
+			return SqlConnector.GetWatchers(IssueId);
+		}
+
+		[HttpPost]
+		public void Post(long IssueId, [FromBody] IssueWatcherData Data)
+		{
+			SqlConnector.AddWatcher(IssueId, Data.UserName);
+		}
+
+		[HttpDelete]
+		public void Delete(long IssueId, [FromBody] IssueWatcherData Data)
+		{
+			SqlConnector.RemoveWatcher(IssueId, Data.UserName);
+		}
+	}
+}
diff --git a/Engine/Source/Programs/UnrealGameSync/MetadataServer/Controllers/LatestController.cs b/Engine/Source/Programs/UnrealGameSync/MetadataServer/Controllers/LatestController.cs
new file mode 100644
--- /dev/null
+++ b/Engine/Source/Programs/UnrealGameSync/MetadataServer/Controllers/LatestController.cs
@@ -0,0 +1,21 @@
+﻿// Copyright 1998-2019 Epic Games, Inc. All Rights Reserved.
+
+using System;
+using System.Collections.Generic;
+using System.Linq;
+using System.Net;
+using System.Net.Http;
+using System.Web.Http;
+using UnrealGameSyncMetadataServer.Connectors;
+using UnrealGameSyncMetadataServer.Models;
+
+namespace UnrealGameSyncMetadataServer.Controllers
+{
+    public class LatestController : ApiController
+    {
+		public LatestData Get(string Project = null)
+		{
+			return SqlConnector.GetLastIds(Project);
+		}
+	}
+}
diff --git a/Engine/Source/Programs/UnrealGameSync/MetadataServer/Controllers/TelemetryController.cs b/Engine/Source/Programs/UnrealGameSync/MetadataServer/Controllers/TelemetryController.cs
new file mode 100644
--- /dev/null
+++ b/Engine/Source/Programs/UnrealGameSync/MetadataServer/Controllers/TelemetryController.cs
@@ -0,0 +1,16 @@
+﻿// Copyright 1998-2019 Epic Games, Inc. All Rights Reserved.
+
+using System.Web.Http;
+using UnrealGameSyncMetadataServer.Connectors;
+using UnrealGameSyncMetadataServer.Models;
+
+namespace UnrealGameSyncMetadataServer.Controllers
+{
+    public class TelemetryController : ApiController
+    {
+		public void Post([FromBody] TelemetryTimingData Data, string Version, string IpAddress)
+		{
+			SqlConnector.PostTelemetryData(Data, Version, IpAddress);
+		}
+    }
+}
diff --git a/Engine/Source/Programs/UnrealGameSync/MetadataServer/Controllers/UserController.cs b/Engine/Source/Programs/UnrealGameSync/MetadataServer/Controllers/UserController.cs
new file mode 100644
--- /dev/null
+++ b/Engine/Source/Programs/UnrealGameSync/MetadataServer/Controllers/UserController.cs
@@ -0,0 +1,20 @@
+// Copyright 1998-2019 Epic Games, Inc. All Rights Reserved.
+
+using System;
+using System.Collections.Generic;
+using System.Data.SQLite;
+using System.Web.Http;
+using UnrealGameSyncMetadataServer.Connectors;
+using UnrealGameSyncMetadataServer.Models;
+
+namespace UnrealGameSyncMetadataServer.Controllers
+{
+	public class UserController : ApiController
+	{
+		public object Get(string Name)
+		{
+			long UserId = SqlConnector.FindOrAddUserId(Name);
+			return new { Id = UserId };
+		}
+	}
+}
diff --git a/Engine/Source/Programs/UnrealGameSync/MetadataServer/Global.asax.cs b/Engine/Source/Programs/UnrealGameSync/MetadataServer/Global.asax.cs
new file mode 100644
--- /dev/null
+++ b/Engine/Source/Programs/UnrealGameSync/MetadataServer/Global.asax.cs
@@ -0,0 +1,43 @@
+﻿// Copyright 1998-2019 Epic Games, Inc. All Rights Reserved.
+using System;
+using System.Data.SQLite;
+using System.IO;
+using System.Text.RegularExpressions;
+using System.Web;
+using System.Web.Http;
+
+namespace UnrealGameSyncMetadataServer
+{
+	public class WebApiApplication : System.Web.HttpApplication
+	{
+		protected void Application_Start()
+		{
+			try
+			{
+				string ConnectionString = System.Configuration.ConfigurationManager.ConnectionStrings["ConnectionString"].ConnectionString;
+				string FileName = Regex.Match(ConnectionString, "Data Source=(.+);.+").Groups[1].Value;
+				Directory.CreateDirectory(new FileInfo(FileName).Directory.FullName);
+				if (!File.Exists(FileName))
+				{
+					SQLiteConnection.CreateFile(FileName);
+				}
+				// initialize the db if it doesn't exist
+				using (SQLiteConnection Connection = new SQLiteConnection(ConnectionString, true))
+				{
+					Connection.Open();
+					using (SQLiteCommand Command = new SQLiteCommand(Properties.Resources.Setup, Connection))
+					{
+						Command.ExecuteNonQuery();
+					}
+				}
+
+			}
+			catch(Exception)
+			{
+				HttpRuntime.UnloadAppDomain();
+				return;
+			}
+			GlobalConfiguration.Configure(WebApiConfig.Register);
+		}
+	}
+}
diff --git a/Engine/Source/Programs/UnrealGameSync/MetadataServer/MetadataServer.csproj b/Engine/Source/Programs/UnrealGameSync/MetadataServer/MetadataServer.csproj
new file mode 100644
--- /dev/null
+++ b/Engine/Source/Programs/UnrealGameSync/MetadataServer/MetadataServer.csproj
@@ -0,0 +1,195 @@
+﻿<?xml version="1.0" encoding="utf-8"?>
+<Project ToolsVersion="14.0" DefaultTargets="Build" xmlns="http://schemas.microsoft.com/developer/msbuild/2003">
+  <Import Project="$(MSBuildExtensionsPath)\$(MSBuildToolsVersion)\Microsoft.Common.props" Condition="Exists('$(MSBuildExtensionsPath)\$(MSBuildToolsVersion)\Microsoft.Common.props')" />
+  <PropertyGroup>
+    <Configuration Condition=" '$(Configuration)' == '' ">Debug</Configuration>
+    <Platform Condition=" '$(Platform)' == '' ">AnyCPU</Platform>
+    <ProductVersion>
+    </ProductVersion>
+    <SchemaVersion>2.0</SchemaVersion>
+    <ProjectGuid>{C8D13E98-C8CD-42FB-9C79-F3EB11B0F3E8}</ProjectGuid>
+    <ProjectTypeGuids>{349c5851-65df-11da-9384-00065b846f21};{fae04ec0-301f-11d3-bf4b-00c04f79efbc}</ProjectTypeGuids>
+    <OutputType>Library</OutputType>
+    <AppDesignerFolder>Properties</AppDesignerFolder>
+    <RootNamespace>MetadataServer</RootNamespace>
+    <AssemblyName>MetadataServer</AssemblyName>
+    <TargetFrameworkVersion>v4.5</TargetFrameworkVersion>
+    <MvcBuildViews>false</MvcBuildViews>
+    <UseIISExpress>true</UseIISExpress>
+    <IISExpressSSLPort />
+    <IISExpressAnonymousAuthentication />
+    <IISExpressWindowsAuthentication />
+    <IISExpressUseClassicPipelineMode />
+    <UseGlobalApplicationHostFile />
+    <NuGetPackageImportStamp>
+    </NuGetPackageImportStamp>
+    <TargetFrameworkProfile />
+    <Use64BitIISExpress />
+  </PropertyGroup>
+  <PropertyGroup Condition=" '$(Configuration)|$(Platform)' == 'Debug|AnyCPU' ">
+    <TransformWebConfigEnabled>False</TransformWebConfigEnabled>
+    <DebugSymbols>true</DebugSymbols>
+    <DebugType>full</DebugType>
+    <Optimize>false</Optimize>
+    <OutputPath>bin\</OutputPath>
+    <DefineConstants>DEBUG;TRACE</DefineConstants>
+    <ErrorReport>prompt</ErrorReport>
+    <WarningLevel>4</WarningLevel>
+  </PropertyGroup>
+  <PropertyGroup Condition=" '$(Configuration)|$(Platform)' == 'Release|AnyCPU' ">
+    <TransformWebConfigEnabled>False</TransformWebConfigEnabled>
+    <DebugType>pdbonly</DebugType>
+    <Optimize>true</Optimize>
+    <OutputPath>bin\</OutputPath>
+    <DefineConstants>TRACE</DefineConstants>
+    <ErrorReport>prompt</ErrorReport>
+    <WarningLevel>4</WarningLevel>
+  </PropertyGroup>
+  <ItemGroup>
+    <Reference Include="Microsoft.CSharp" />
+    <Reference Include="Newtonsoft.Json, Version=11.0.0.0, Culture=neutral, PublicKeyToken=30ad4fe6b2a6aeed, processorArchitecture=MSIL">
+      <HintPath>..\packages\Newtonsoft.Json.11.0.2\lib\net45\Newtonsoft.Json.dll</HintPath>
+      <Private>True</Private>
+    </Reference>
+    <Reference Include="System" />
+    <Reference Include="System.Data" />
+    <Reference Include="System.Data.DataSetExtensions" />
+    <Reference Include="System.Data.SQLite, Version=1.0.108.0, Culture=neutral, PublicKeyToken=db937bc2d44ff139, processorArchitecture=MSIL">
+      <HintPath>..\packages\System.Data.SQLite.Core.1.0.108.0\lib\net45\System.Data.SQLite.dll</HintPath>
+      <Private>True</Private>
+    </Reference>
+    <Reference Include="System.Drawing" />
+    <Reference Include="System.Web.DynamicData" />
+    <Reference Include="System.Web.Entity" />
+    <Reference Include="System.Web.ApplicationServices" />
+    <Reference Include="System.ComponentModel.DataAnnotations" />
+    <Reference Include="System.Web.Extensions" />
+    <Reference Include="System.Web" />
+    <Reference Include="System.Web.Abstractions" />
+    <Reference Include="System.Web.Routing" />
+    <Reference Include="System.Xml" />
+    <Reference Include="System.Configuration" />
+    <Reference Include="System.Runtime.Serialization" />
+    <Reference Include="System.Net.Http">
+    </Reference>
+    <Reference Include="System.Net.Http.Formatting, Version=5.2.3.0, Culture=neutral, PublicKeyToken=31bf3856ad364e35, processorArchitecture=MSIL">
+      <HintPath>..\packages\Microsoft.AspNet.WebApi.Client.5.2.3\lib\net45\System.Net.Http.Formatting.dll</HintPath>
+    </Reference>
+    <Reference Include="System.Net.Http.WebRequest">
+    </Reference>
+    <Reference Include="System.Web.Http, Version=5.2.3.0, Culture=neutral, PublicKeyToken=31bf3856ad364e35, processorArchitecture=MSIL">
+      <HintPath>..\packages\Microsoft.AspNet.WebApi.Core.5.2.3\lib\net45\System.Web.Http.dll</HintPath>
+    </Reference>
+    <Reference Include="System.Web.Http.WebHost, Version=5.2.3.0, Culture=neutral, PublicKeyToken=31bf3856ad364e35, processorArchitecture=MSIL">
+      <HintPath>..\packages\Microsoft.AspNet.WebApi.WebHost.5.2.3\lib\net45\System.Web.Http.WebHost.dll</HintPath>
+    </Reference>
+    <Reference Include="System.Xml.Linq" />
+  </ItemGroup>
+  <ItemGroup>
+    <Folder Include="App_Data\" />
+  </ItemGroup>
+  <ItemGroup>
+    <Compile Include="App_Start\WebApiConfig.cs" />
+    <Compile Include="Controllers\IssuesController.cs" />
+    <Compile Include="Controllers\UserController.cs" />
+    <Compile Include="Models\IssueData.cs" />
+    <Compile Include="Properties\Resources.Designer.cs">
+      <AutoGen>True</AutoGen>
+      <DesignTime>True</DesignTime>
+      <DependentUpon>Resources.resx</DependentUpon>
+    </Compile>
+    <Content Include="Setup.sql" />
+    <Compile Include="Controllers\BuildController.cs" />
+    <Compile Include="Controllers\CISController.cs" />
+    <Compile Include="Controllers\CommentController.cs" />
+    <Compile Include="Controllers\ErrorController.cs" />
+    <Compile Include="Controllers\EventController.cs" />
+    <Compile Include="Controllers\LatestController.cs" />
+    <Compile Include="Controllers\TelemetryController.cs" />
+    <Compile Include="Global.asax.cs">
+      <DependentUpon>Global.asax</DependentUpon>
+    </Compile>
+    <Compile Include="Models\BuildData.cs" />
+    <Compile Include="Models\CommentData.cs" />
+    <Compile Include="Models\EventData.cs" />
+    <Compile Include="Models\EventSummary.cs" />
+    <Compile Include="Models\LatestData.cs" />
+    <Compile Include="Models\TelemetryErrorData.cs" />
+    <Compile Include="Models\TelemetryTimingData.cs" />
+    <Compile Include="Properties\AssemblyInfo.cs" />
+    <Compile Include="Connectors\SqlConnector.cs" />
+  </ItemGroup>
+  <ItemGroup>
+    <Content Include="Global.asax" />
+    <Content Include="Web.config">
+      <SubType>Designer</SubType>
+    </Content>
+    <None Include="Web.template.config" />
+    <None Include="Web.Debug.config">
+      <DependentUpon>Web.template.config</DependentUpon>
+    </None>
+    <None Include="Web.Release.config">
+      <DependentUpon>Web.template.config</DependentUpon>
+    </None>
+  </ItemGroup>
+  <ItemGroup>
+    <Content Include="NotForLicensees\Web.config">
+      <SubType>Designer</SubType>
+    </Content>
+  </ItemGroup>
+  <ItemGroup>
+    <Content Include="packages.config" />
+  </ItemGroup>
+  <ItemGroup>
+    <EmbeddedResource Include="Properties\Resources.resx">
+      <Generator>ResXFileCodeGenerator</Generator>
+      <LastGenOutput>Resources.Designer.cs</LastGenOutput>
+      <SubType>Designer</SubType>
+    </EmbeddedResource>
+  </ItemGroup>
+  <PropertyGroup>
+    <VisualStudioVersion Condition="'$(VisualStudioVersion)' == ''">10.0</VisualStudioVersion>
+    <VSToolsPath Condition="'$(VSToolsPath)' == ''">$(MSBuildExtensionsPath32)\Microsoft\VisualStudio\v$(VisualStudioVersion)</VSToolsPath>
+  </PropertyGroup>
+  <Import Project="$(MSBuildBinPath)\Microsoft.CSharp.targets" />
+  <Import Project="$(VSToolsPath)\WebApplications\Microsoft.WebApplication.targets" Condition="'$(VSToolsPath)' != ''" />
+  <Import Project="$(MSBuildExtensionsPath32)\Microsoft\VisualStudio\v10.0\WebApplications\Microsoft.WebApplication.targets" Condition="false" />
+  <Target Name="MvcBuildViews" AfterTargets="AfterBuild" Condition="'$(MvcBuildViews)'=='true'">
+    <AspNetCompiler VirtualPath="temp" PhysicalPath="$(WebProjectOutputDir)" />
+  </Target>
+  <ProjectExtensions>
+    <VisualStudio>
+      <FlavorProperties GUID="{349c5851-65df-11da-9384-00065b846f21}">
+        <WebProjectProperties>
+          <UseIIS>True</UseIIS>
+          <AutoAssignPort>True</AutoAssignPort>
+          <DevelopmentServerPort>62486</DevelopmentServerPort>
+          <DevelopmentServerVPath>/</DevelopmentServerVPath>
+          <IISUrl>http://localhost:62487/</IISUrl>
+          <NTLMAuthentication>False</NTLMAuthentication>
+          <UseCustomServer>False</UseCustomServer>
+          <CustomServerUrl>
+          </CustomServerUrl>
+          <SaveServerSettingsInUserFile>False</SaveServerSettingsInUserFile>
+        </WebProjectProperties>
+      </FlavorProperties>
+    </VisualStudio>
+  </ProjectExtensions>
+  <!-- To modify your build process, add your task inside one of the targets below and uncomment it.
+       Other similar extension points exist, see Microsoft.Common.targets.
+  <Target Name="BeforeBuild">
+  </Target> -->
+  <Target Name="BeforeBuild">
+    <Message Importance="High" Text="Creating web.config from web.template.config" />
+    <TransformXml Source="Web.template.config" Transform="Web.$(Configuration).config" Destination="Web.config" />
+    <Message Importance="High" Text="Found NotForLicensees\Web.config, applying transform" Condition="Exists('NotForLicensees\Web.config')" />
+    <TransformXml Source="Web.config" Transform="NotForLicensees\Web.config" Destination="Web.config" Condition="Exists('NotForLicensees\Web.config')" />
+  </Target>
+  <Import Project="..\packages\System.Data.SQLite.Core.1.0.108.0\build\net45\System.Data.SQLite.Core.targets" Condition="Exists('..\packages\System.Data.SQLite.Core.1.0.108.0\build\net45\System.Data.SQLite.Core.targets')" />
+  <Target Name="EnsureNuGetPackageBuildImports" BeforeTargets="PrepareForBuild">
+    <PropertyGroup>
+      <ErrorText>This project references NuGet package(s) that are missing on this computer. Use NuGet Package Restore to download them.  For more information, see http://go.microsoft.com/fwlink/?LinkID=322105. The missing file is {0}.</ErrorText>
+    </PropertyGroup>
+    <Error Condition="!Exists('..\packages\System.Data.SQLite.Core.1.0.108.0\build\net45\System.Data.SQLite.Core.targets')" Text="$([System.String]::Format('$(ErrorText)', '..\packages\System.Data.SQLite.Core.1.0.108.0\build\net45\System.Data.SQLite.Core.targets'))" />
+  </Target>
+</Project>
\ No newline at end of file
diff --git a/Engine/Source/Programs/UnrealGameSync/MetadataServer/Models/BuildData.cs b/Engine/Source/Programs/UnrealGameSync/MetadataServer/Models/BuildData.cs
new file mode 100644
--- /dev/null
+++ b/Engine/Source/Programs/UnrealGameSync/MetadataServer/Models/BuildData.cs
@@ -0,0 +1,39 @@
+﻿// Copyright 1998-2019 Epic Games, Inc. All Rights Reserved.
+
+using System;
+using System.Collections.Generic;
+using System.Linq;
+using System.Web;
+
+namespace UnrealGameSyncMetadataServer.Models
+{
+	public class BuildData
+	{
+		public enum BuildDataResult
+		{
+			Starting,
+			Failure,
+			Warning,
+			Success,
+			Skipped,
+		}
+
+		public long Id;
+		public int ChangeNumber;
+		public string BuildType;
+		public BuildDataResult Result;
+		public string Url;
+		public string Project;
+		public string ArchivePath;
+
+		public bool IsSuccess
+		{
+			get { return Result == BuildDataResult.Success || Result == BuildDataResult.Warning; }
+		}
+
+		public bool IsFailure
+		{
+			get { return Result == BuildDataResult.Failure; }
+		}
+	}
+}
\ No newline at end of file
diff --git a/Engine/Source/Programs/UnrealGameSync/MetadataServer/Models/CommentData.cs b/Engine/Source/Programs/UnrealGameSync/MetadataServer/Models/CommentData.cs
new file mode 100644
--- /dev/null
+++ b/Engine/Source/Programs/UnrealGameSync/MetadataServer/Models/CommentData.cs
@@ -0,0 +1,18 @@
+﻿// Copyright 1998-2019 Epic Games, Inc. All Rights Reserved.
+
+using System;
+using System.Collections.Generic;
+using System.Linq;
+using System.Web;
+
+namespace UnrealGameSyncMetadataServer.Models
+{
+	public class CommentData
+	{
+		public long Id;
+		public int ChangeNumber;
+		public string UserName;
+		public string Text;
+		public string Project;
+	}
+}
\ No newline at end of file
diff --git a/Engine/Source/Programs/UnrealGameSync/MetadataServer/Models/EventData.cs b/Engine/Source/Programs/UnrealGameSync/MetadataServer/Models/EventData.cs
new file mode 100644
--- /dev/null
+++ b/Engine/Source/Programs/UnrealGameSync/MetadataServer/Models/EventData.cs
@@ -0,0 +1,38 @@
+﻿// Copyright 1998-2019 Epic Games, Inc. All Rights Reserved.
+
+using System;
+using System.Collections.Generic;
+using System.Linq;
+using System.Web;
+
+namespace UnrealGameSyncMetadataServer.Models
+{
+	public class EventData
+	{
+		public enum EventType
+		{
+			Syncing,
+
+			// Reviews
+			Compiles,
+			DoesNotCompile,
+			Good,
+			Bad,
+			Unknown,
+
+			// Starred builds
+			Starred,
+			Unstarred,
+
+			// Investigating events
+			Investigating,
+			Resolved,
+		}
+		
+		public long Id;
+		public int Change;
+		public string UserName;
+		public EventType Type;
+		public string Project;
+	}
+}
\ No newline at end of file
diff --git a/Engine/Source/Programs/UnrealGameSync/MetadataServer/Models/EventSummary.cs b/Engine/Source/Programs/UnrealGameSync/MetadataServer/Models/EventSummary.cs
new file mode 100644
--- /dev/null
+++ b/Engine/Source/Programs/UnrealGameSync/MetadataServer/Models/EventSummary.cs
@@ -0,0 +1,29 @@
+﻿// Copyright 1998-2019 Epic Games, Inc. All Rights Reserved.
+
+using System;
+using System.Collections.Generic;
+using System.Linq;
+using System.Web;
+
+namespace UnrealGameSyncMetadataServer.Models
+{
+	public class EventSummary
+	{
+		public enum ReviewVerdict
+		{
+			Unknown,
+			Good,
+			Bad,
+			Mixed,
+		}
+
+		public int ChangeNumber;
+		public ReviewVerdict Verdict;
+		public List<EventData> SyncEvents = new List<EventData>();
+		public List<EventData> Reviews = new List<EventData>();
+		public List<string> CurrentUsers = new List<string>();
+		public EventData LastStarReview;
+		public List<BuildData> Builds = new List<BuildData>();
+		public List<CommentData> Comments = new List<CommentData>();
+	}
+}
\ No newline at end of file
diff --git a/Engine/Source/Programs/UnrealGameSync/MetadataServer/Models/IssueData.cs b/Engine/Source/Programs/UnrealGameSync/MetadataServer/Models/IssueData.cs
new file mode 100644
--- /dev/null
+++ b/Engine/Source/Programs/UnrealGameSync/MetadataServer/Models/IssueData.cs
@@ -0,0 +1,55 @@
+// Copyright 1998-2019 Epic Games, Inc. All Rights Reserved.
+
+using Newtonsoft.Json.Linq;
+using System;
+using System.Collections.Generic;
+using System.Linq;
+using System.Web;
+
+namespace UnrealGameSyncMetadataServer.Models
+{
+	public class IssueWatcherData
+	{
+		public string UserName;
+	}
+
+	public class IssueBuildData
+	{
+		public long Id;
+		public string Stream;
+		public int Change;
+		public string Name;
+		public string Url;
+		public int Outcome;
+	}
+
+	public class IssueBuildUpdateData
+	{
+		public int Outcome;
+	}
+
+	public class IssueData
+	{
+		public long Id;
+		public DateTime CreatedAt;
+		public DateTime RetrievedAt;
+		public string Project;
+		public string Summary;
+		public string Owner;
+		public string NominatedBy;
+		public DateTime? AcknowledgedAt;
+		public int FixChange;
+		public DateTime? ResolvedAt;
+		public bool bNotify;
+	}
+
+	public class IssueUpdateData
+	{
+		public string Summary;
+		public string Owner;
+		public string NominatedBy;
+		public bool? Acknowledged;
+		public int? FixChange;
+		public bool? Resolved;
+	}
+}
diff --git a/Engine/Source/Programs/UnrealGameSync/MetadataServer/Models/LatestData.cs b/Engine/Source/Programs/UnrealGameSync/MetadataServer/Models/LatestData.cs
new file mode 100644
--- /dev/null
+++ b/Engine/Source/Programs/UnrealGameSync/MetadataServer/Models/LatestData.cs
@@ -0,0 +1,16 @@
+﻿// Copyright 1998-2019 Epic Games, Inc. All Rights Reserved.
+
+using System;
+using System.Collections.Generic;
+using System.Linq;
+using System.Web;
+
+namespace UnrealGameSyncMetadataServer.Models
+{
+	public class LatestData
+	{
+		public long LastEventId { get; set; }
+		public long LastCommentId { get; set; }
+		public long LastBuildId { get; set; }
+	}
+}
\ No newline at end of file
diff --git a/Engine/Source/Programs/UnrealGameSync/MetadataServer/Models/TelemetryErrorData.cs b/Engine/Source/Programs/UnrealGameSync/MetadataServer/Models/TelemetryErrorData.cs
new file mode 100644
--- /dev/null
+++ b/Engine/Source/Programs/UnrealGameSync/MetadataServer/Models/TelemetryErrorData.cs
@@ -0,0 +1,25 @@
+﻿// Copyright 1998-2019 Epic Games, Inc. All Rights Reserved.
+
+using System;
+using System.Collections.Generic;
+using System.Linq;
+using System.Web;
+
+namespace UnrealGameSyncMetadataServer.Models
+{
+	public class TelemetryErrorData
+	{
+		public enum TelemetryErrorType
+		{
+			Crash,
+		}
+		public int Id;
+		public TelemetryErrorType Type;
+		public string Text;
+		public string UserName;
+		public string Project;
+		public DateTime Timestamp;
+		public string Version;
+		public string IpAddress;
+	}
+}
\ No newline at end of file
diff --git a/Engine/Source/Programs/UnrealGameSync/MetadataServer/Models/TelemetryTimingData.cs b/Engine/Source/Programs/UnrealGameSync/MetadataServer/Models/TelemetryTimingData.cs
new file mode 100644
--- /dev/null
+++ b/Engine/Source/Programs/UnrealGameSync/MetadataServer/Models/TelemetryTimingData.cs
@@ -0,0 +1,19 @@
+﻿// Copyright 1998-2019 Epic Games, Inc. All Rights Reserved.
+
+using System;
+using System.Collections.Generic;
+using System.Linq;
+using System.Web;
+
+namespace UnrealGameSyncMetadataServer.Models
+{
+	public class TelemetryTimingData
+	{
+		public string Action;
+		public string Result;
+		public string UserName;
+		public string Project;
+		public DateTime Timestamp;
+		public float Duration;
+	}
+}
\ No newline at end of file
diff --git a/Engine/Source/Programs/UnrealGameSync/MetadataServer/Properties/AssemblyInfo.cs b/Engine/Source/Programs/UnrealGameSync/MetadataServer/Properties/AssemblyInfo.cs
new file mode 100644
--- /dev/null
+++ b/Engine/Source/Programs/UnrealGameSync/MetadataServer/Properties/AssemblyInfo.cs
@@ -0,0 +1,37 @@
+﻿// Copyright 1998-2019 Epic Games, Inc. All Rights Reserved.
+
+using System.Reflection;
+using System.Runtime.CompilerServices;
+using System.Runtime.InteropServices;
+
+// General Information about an assembly is controlled through the following
+// set of attributes. Change these attribute values to modify the information
+// associated with an assembly.
+[assembly: AssemblyTitle("UnrealGameSyncMetadataServer")]
+[assembly: AssemblyDescription("")]
+[assembly: AssemblyConfiguration("")]
+[assembly: AssemblyCompany("")]
+[assembly: AssemblyProduct("UnrealGameSync")]
+[assembly: AssemblyCopyright("Copyright 1998-2019 Epic Games, Inc. All Rights Reserved.")]
+[assembly: AssemblyTrademark("")]
+[assembly: AssemblyCulture("")]
+
+// Setting ComVisible to false makes the types in this assembly not visible
+// to COM components.  If you need to access a type in this assembly from
+// COM, set the ComVisible attribute to true on that type.
+[assembly: ComVisible(false)]
+
+// The following GUID is for the ID of the typelib if this project is exposed to COM
+[assembly: Guid("f6044fb1-c977-4705-8cad-f17962adfb41")]
+
+// Version information for an assembly consists of the following four values:
+//
+//      Major Version
+//      Minor Version
+//      Build Number
+//      Revision
+//
+// You can specify all the values or you can default the Revision and Build Numbers
+// by using the '*' as shown below:
+[assembly: AssemblyVersion("1.0.0.0")]
+[assembly: AssemblyFileVersion("1.0.0.0")]
diff --git a/Engine/Source/Programs/UnrealGameSync/MetadataServer/Web.Debug.config b/Engine/Source/Programs/UnrealGameSync/MetadataServer/Web.Debug.config
new file mode 100644
--- /dev/null
+++ b/Engine/Source/Programs/UnrealGameSync/MetadataServer/Web.Debug.config
@@ -0,0 +1,30 @@
+﻿<?xml version="1.0"?>
+
+<!-- For more information on using Web.config transformation visit http://go.microsoft.com/fwlink/?LinkId=301874 -->
+
+<configuration xmlns:xdt="http://schemas.microsoft.com/XML-Document-Transform">
+  <!--
+    In the example below, the "SetAttributes" transform will change the value of
+    "connectionString" to use "ReleaseSQLServer" only when the "Match" locator
+    finds an attribute "name" that has a value of "MyDB".
+
+    <connectionStrings>
+      <add name="MyDB"
+        connectionString="Data Source=ReleaseSQLServer;Initial Catalog=MyReleaseDB;Integrated Security=True"
+        xdt:Transform="SetAttributes" xdt:Locator="Match(name)"/>
+    </connectionStrings>
+  -->
+  <system.web>
+    <!--
+      In the example below, the "Replace" transform will replace the entire
+      <customErrors> section of your Web.config file.
+      Note that because there is only one customErrors section under the
+      <system.web> node, there is no need to use the "xdt:Locator" attribute.
+
+      <customErrors defaultRedirect="GenericError.htm"
+        mode="RemoteOnly" xdt:Transform="Replace">
+        <error statusCode="500" redirect="InternalError.htm"/>
+      </customErrors>
+    -->
+  </system.web>
+</configuration>
diff --git a/Engine/Source/Programs/UnrealGameSync/MetadataServer/Web.Release.config b/Engine/Source/Programs/UnrealGameSync/MetadataServer/Web.Release.config
new file mode 100644
--- /dev/null
+++ b/Engine/Source/Programs/UnrealGameSync/MetadataServer/Web.Release.config
@@ -0,0 +1,31 @@
+﻿<?xml version="1.0"?>
+
+<!-- For more information on using Web.config transformation visit http://go.microsoft.com/fwlink/?LinkId=301874 -->
+
+<configuration xmlns:xdt="http://schemas.microsoft.com/XML-Document-Transform">
+  <!--
+    In the example below, the "SetAttributes" transform will change the value of
+    "connectionString" to use "ReleaseSQLServer" only when the "Match" locator
+    finds an attribute "name" that has a value of "MyDB".
+
+    <connectionStrings>
+      <add name="MyDB"
+        connectionString="Data Source=ReleaseSQLServer;Initial Catalog=MyReleaseDB;Integrated Security=True"
+        xdt:Transform="SetAttributes" xdt:Locator="Match(name)"/>
+    </connectionStrings>
+  -->
+  <system.web>
+    <compilation xdt:Transform="RemoveAttributes(debug)" />
+    <!--
+      In the example below, the "Replace" transform will replace the entire
+      <customErrors> section of your Web.config file.
+      Note that because there is only one customErrors section under the
+      <system.web> node, there is no need to use the "xdt:Locator" attribute.
+
+      <customErrors defaultRedirect="GenericError.htm"
+        mode="RemoteOnly" xdt:Transform="Replace">
+        <error statusCode="500" redirect="InternalError.htm"/>
+      </customErrors>
+    -->
+  </system.web>
+</configuration>
diff --git a/Engine/Source/Programs/UnrealGameSync/MetadataServer/Web.template.config b/Engine/Source/Programs/UnrealGameSync/MetadataServer/Web.template.config
new file mode 100644
--- /dev/null
+++ b/Engine/Source/Programs/UnrealGameSync/MetadataServer/Web.template.config
@@ -0,0 +1,42 @@
+﻿<?xml version="1.0" encoding="utf-8"?>
+<!--
+  For more information on how to configure your ASP.NET application, please visit
+  http://go.microsoft.com/fwlink/?LinkId=301879
+  -->
+<configuration>
+	<appSettings/>
+  <!--
+    For a description of web.config changes see http://go.microsoft.com/fwlink/?LinkId=235367.
+
+    The following attributes can be set on the <httpRuntime> tag.
+      <system.Web>
+        <httpRuntime targetFramework="4.5" />
+      </system.Web>
+  -->
+	<!-- Leave this connection string blank! It should go in either the .debug or .release transform config
+         There is a custom build step on BeforeBuild that generates a transformed Web.config on the fly. The generated Web.config should not be checked into source control, 
+         and this template.config should be checked in in it's place. -->
+	<connectionStrings>
+		<add name="ConnectionString" connectionString=""/>
+	</connectionStrings>
+	<system.web>
+    <compilation debug="true" targetFramework="4.5" />
+    <httpRuntime targetFramework="4.5" />
+  </system.web>
+  <system.webServer>
+    <handlers>
+      <remove name="ExtensionlessUrlHandler-Integrated-4.0" />
+      <remove name="OPTIONSVerbHandler" />
+      <remove name="TRACEVerbHandler" />
+      <add name="ExtensionlessUrlHandler-Integrated-4.0" path="*." verb="*" type="System.Web.Handlers.TransferRequestHandler" preCondition="integratedMode,runtimeVersionv4.0" />
+    </handlers>
+  </system.webServer>
+  <runtime>
+    <assemblyBinding xmlns="urn:schemas-microsoft-com:asm.v1">
+      <dependentAssembly>
+        <assemblyIdentity name="Newtonsoft.Json" publicKeyToken="30ad4fe6b2a6aeed" culture="neutral" />
+        <bindingRedirect oldVersion="0.0.0.0-11.0.0.0" newVersion="11.0.0.0" />
+      </dependentAssembly>
+    </assemblyBinding>
+  </runtime>
+</configuration>
\ No newline at end of file
diff --git a/Engine/Source/Programs/UnrealGameSync/MetadataServer/packages.config b/Engine/Source/Programs/UnrealGameSync/MetadataServer/packages.config
new file mode 100644
--- /dev/null
+++ b/Engine/Source/Programs/UnrealGameSync/MetadataServer/packages.config
@@ -0,0 +1,9 @@
+﻿<?xml version="1.0" encoding="utf-8"?>
+<packages>
+  <package id="Microsoft.AspNet.WebApi" version="5.2.3" targetFramework="net452" />
+  <package id="Microsoft.AspNet.WebApi.Client" version="5.2.3" targetFramework="net452" />
+  <package id="Microsoft.AspNet.WebApi.Core" version="5.2.3" targetFramework="net452" />
+  <package id="Microsoft.AspNet.WebApi.WebHost" version="5.2.3" targetFramework="net452" />
+  <package id="Newtonsoft.Json" version="11.0.2" targetFramework="net452" />
+  <package id="System.Data.SQLite.Core" version="1.0.108.0" targetFramework="net45" />
+</packages>
\ No newline at end of file
diff --git a/Engine/Source/Programs/UnrealGameSync/UnrealGameSyncMetadataServer/App_Start/WebApiConfig.cs b/Engine/Source/Programs/UnrealGameSync/UnrealGameSyncMetadataServer/App_Start/WebApiConfig.cs
deleted file mode 100644
diff --git a/Engine/Source/Programs/UnrealGameSync/UnrealGameSyncMetadataServer/Connectors/SqlConnector.cs b/Engine/Source/Programs/UnrealGameSync/UnrealGameSyncMetadataServer/Connectors/SqlConnector.cs
deleted file mode 100644
diff --git a/Engine/Source/Programs/UnrealGameSync/UnrealGameSyncMetadataServer/Controllers/BuildController.cs b/Engine/Source/Programs/UnrealGameSync/UnrealGameSyncMetadataServer/Controllers/BuildController.cs
deleted file mode 100644
diff --git a/Engine/Source/Programs/UnrealGameSync/UnrealGameSyncMetadataServer/Controllers/CISController.cs b/Engine/Source/Programs/UnrealGameSync/UnrealGameSyncMetadataServer/Controllers/CISController.cs
deleted file mode 100644
diff --git a/Engine/Source/Programs/UnrealGameSync/UnrealGameSyncMetadataServer/Controllers/CommentController.cs b/Engine/Source/Programs/UnrealGameSync/UnrealGameSyncMetadataServer/Controllers/CommentController.cs
deleted file mode 100644
diff --git a/Engine/Source/Programs/UnrealGameSync/UnrealGameSyncMetadataServer/Controllers/ErrorController.cs b/Engine/Source/Programs/UnrealGameSync/UnrealGameSyncMetadataServer/Controllers/ErrorController.cs
deleted file mode 100644
diff --git a/Engine/Source/Programs/UnrealGameSync/UnrealGameSyncMetadataServer/Controllers/EventController.cs b/Engine/Source/Programs/UnrealGameSync/UnrealGameSyncMetadataServer/Controllers/EventController.cs
deleted file mode 100644
diff --git a/Engine/Source/Programs/UnrealGameSync/UnrealGameSyncMetadataServer/Controllers/IssuesController.cs b/Engine/Source/Programs/UnrealGameSync/UnrealGameSyncMetadataServer/Controllers/IssuesController.cs
deleted file mode 100644
diff --git a/Engine/Source/Programs/UnrealGameSync/UnrealGameSyncMetadataServer/Controllers/LatestController.cs b/Engine/Source/Programs/UnrealGameSync/UnrealGameSyncMetadataServer/Controllers/LatestController.cs
deleted file mode 100644
diff --git a/Engine/Source/Programs/UnrealGameSync/UnrealGameSyncMetadataServer/Controllers/TelemetryController.cs b/Engine/Source/Programs/UnrealGameSync/UnrealGameSyncMetadataServer/Controllers/TelemetryController.cs
deleted file mode 100644
diff --git a/Engine/Source/Programs/UnrealGameSync/UnrealGameSyncMetadataServer/Global.asax.cs b/Engine/Source/Programs/UnrealGameSync/UnrealGameSyncMetadataServer/Global.asax.cs
deleted file mode 100644
diff --git a/Engine/Source/Programs/UnrealGameSync/UnrealGameSyncMetadataServer/Models/BuildData.cs b/Engine/Source/Programs/UnrealGameSync/UnrealGameSyncMetadataServer/Models/BuildData.cs
deleted file mode 100644
diff --git a/Engine/Source/Programs/UnrealGameSync/UnrealGameSyncMetadataServer/Models/CommentData.cs b/Engine/Source/Programs/UnrealGameSync/UnrealGameSyncMetadataServer/Models/CommentData.cs
deleted file mode 100644
diff --git a/Engine/Source/Programs/UnrealGameSync/UnrealGameSyncMetadataServer/Models/EventData.cs b/Engine/Source/Programs/UnrealGameSync/UnrealGameSyncMetadataServer/Models/EventData.cs
deleted file mode 100644
diff --git a/Engine/Source/Programs/UnrealGameSync/UnrealGameSyncMetadataServer/Models/EventSummary.cs b/Engine/Source/Programs/UnrealGameSync/UnrealGameSyncMetadataServer/Models/EventSummary.cs
deleted file mode 100644
diff --git a/Engine/Source/Programs/UnrealGameSync/UnrealGameSyncMetadataServer/Models/IssueData.cs b/Engine/Source/Programs/UnrealGameSync/UnrealGameSyncMetadataServer/Models/IssueData.cs
deleted file mode 100644
diff --git a/Engine/Source/Programs/UnrealGameSync/UnrealGameSyncMetadataServer/Models/LatestData.cs b/Engine/Source/Programs/UnrealGameSync/UnrealGameSyncMetadataServer/Models/LatestData.cs
deleted file mode 100644
diff --git a/Engine/Source/Programs/UnrealGameSync/UnrealGameSyncMetadataServer/Models/TelemetryErrorData.cs b/Engine/Source/Programs/UnrealGameSync/UnrealGameSyncMetadataServer/Models/TelemetryErrorData.cs
deleted file mode 100644
diff --git a/Engine/Source/Programs/UnrealGameSync/UnrealGameSyncMetadataServer/Models/TelemetryTimingData.cs b/Engine/Source/Programs/UnrealGameSync/UnrealGameSyncMetadataServer/Models/TelemetryTimingData.cs
deleted file mode 100644
diff --git a/Engine/Source/Programs/UnrealGameSync/UnrealGameSyncMetadataServer/Properties/AssemblyInfo.cs b/Engine/Source/Programs/UnrealGameSync/UnrealGameSyncMetadataServer/Properties/AssemblyInfo.cs
deleted file mode 100644
diff --git a/Engine/Source/Programs/UnrealGameSync/UnrealGameSyncMetadataServer/UnrealGameSyncMetadataServer.csproj b/Engine/Source/Programs/UnrealGameSync/UnrealGameSyncMetadataServer/UnrealGameSyncMetadataServer.csproj
deleted file mode 100644
diff --git a/Engine/Source/Programs/UnrealGameSync/UnrealGameSyncMetadataServer/Web.Debug.config b/Engine/Source/Programs/UnrealGameSync/UnrealGameSyncMetadataServer/Web.Debug.config
deleted file mode 100644
diff --git a/Engine/Source/Programs/UnrealGameSync/UnrealGameSyncMetadataServer/Web.Release.config b/Engine/Source/Programs/UnrealGameSync/UnrealGameSyncMetadataServer/Web.Release.config
deleted file mode 100644
diff --git a/Engine/Source/Programs/UnrealGameSync/UnrealGameSyncMetadataServer/Web.template.config b/Engine/Source/Programs/UnrealGameSync/UnrealGameSyncMetadataServer/Web.template.config
deleted file mode 100644
diff --git a/Engine/Source/Programs/UnrealGameSync/UnrealGameSyncMetadataServer/packages.config b/Engine/Source/Programs/UnrealGameSync/UnrealGameSyncMetadataServer/packages.config
deleted file mode 100644
