Update NuGet package versions for MetadataServer, and fix namespace names.

#rnx

diff --git a/Engine/Source/Programs/UnrealGameSync/MetadataServer/App_Start/WebApiConfig.cs b/Engine/Source/Programs/UnrealGameSync/MetadataServer/App_Start/WebApiConfig.cs
--- a/Engine/Source/Programs/UnrealGameSync/MetadataServer/App_Start/WebApiConfig.cs
+++ b/Engine/Source/Programs/UnrealGameSync/MetadataServer/App_Start/WebApiConfig.cs
@@ -2,7 +2,7 @@
 
 using System.Web.Http;
 
-namespace UnrealGameSyncMetadataServer
+namespace MetadataServer
 {
 	public static class WebApiConfig
 	{
diff --git a/Engine/Source/Programs/UnrealGameSync/MetadataServer/Connectors/SqlConnector.cs b/Engine/Source/Programs/UnrealGameSync/MetadataServer/Connectors/SqlConnector.cs
--- a/Engine/Source/Programs/UnrealGameSync/MetadataServer/Connectors/SqlConnector.cs
+++ b/Engine/Source/Programs/UnrealGameSync/MetadataServer/Connectors/SqlConnector.cs
@@ -10,9 +10,9 @@ using System.IO;
 using System.Linq;
 using System.Text;
 using System.Text.RegularExpressions;
-using UnrealGameSyncMetadataServer.Models;
+using MetadataServer.Models;
 
-namespace UnrealGameSyncMetadataServer.Connectors
+namespace MetadataServer.Connectors
 {
 	public static class SqlConnector
 	{
diff --git a/Engine/Source/Programs/UnrealGameSync/MetadataServer/Controllers/BuildController.cs b/Engine/Source/Programs/UnrealGameSync/MetadataServer/Controllers/BuildController.cs
--- a/Engine/Source/Programs/UnrealGameSync/MetadataServer/Controllers/BuildController.cs
+++ b/Engine/Source/Programs/UnrealGameSync/MetadataServer/Controllers/BuildController.cs
@@ -2,10 +2,10 @@
 
 using System.Collections.Generic;
 using System.Web.Http;
-using UnrealGameSyncMetadataServer.Connectors;
-using UnrealGameSyncMetadataServer.Models;
+using MetadataServer.Connectors;
+using MetadataServer.Models;
 
-namespace UnrealGameSyncMetadataServer.Controllers
+namespace MetadataServer.Controllers
 {
     public class BuildController : ApiController
     {
diff --git a/Engine/Source/Programs/UnrealGameSync/MetadataServer/Controllers/CISController.cs b/Engine/Source/Programs/UnrealGameSync/MetadataServer/Controllers/CISController.cs
--- a/Engine/Source/Programs/UnrealGameSync/MetadataServer/Controllers/CISController.cs
+++ b/Engine/Source/Programs/UnrealGameSync/MetadataServer/Controllers/CISController.cs
@@ -3,10 +3,10 @@
 using System;
 using System.Collections.Generic;
 using System.Web.Http;
-using UnrealGameSyncMetadataServer.Connectors;
-using UnrealGameSyncMetadataServer.Models;
+using MetadataServer.Connectors;
+using MetadataServer.Models;
 
-namespace UnrealGameSyncMetadataServer.Controllers
+namespace MetadataServer.Controllers
 {
 	// Deprecated, use Latest or Build controllers instead
     public class CISController : ApiController
diff --git a/Engine/Source/Programs/UnrealGameSync/MetadataServer/Controllers/CommentController.cs b/Engine/Source/Programs/UnrealGameSync/MetadataServer/Controllers/CommentController.cs
--- a/Engine/Source/Programs/UnrealGameSync/MetadataServer/Controllers/CommentController.cs
+++ b/Engine/Source/Programs/UnrealGameSync/MetadataServer/Controllers/CommentController.cs
@@ -2,10 +2,10 @@
 
 using System.Collections.Generic;
 using System.Web.Http;
-using UnrealGameSyncMetadataServer.Connectors;
-using UnrealGameSyncMetadataServer.Models;
+using MetadataServer.Connectors;
+using MetadataServer.Models;
 
-namespace UnrealGameSyncMetadataServer.Controllers
+namespace MetadataServer.Controllers
 {
 	public class CommentController : ApiController
 	{
diff --git a/Engine/Source/Programs/UnrealGameSync/MetadataServer/Controllers/ErrorController.cs b/Engine/Source/Programs/UnrealGameSync/MetadataServer/Controllers/ErrorController.cs
--- a/Engine/Source/Programs/UnrealGameSync/MetadataServer/Controllers/ErrorController.cs
+++ b/Engine/Source/Programs/UnrealGameSync/MetadataServer/Controllers/ErrorController.cs
@@ -2,10 +2,10 @@
 
 using System.Collections.Generic;
 using System.Web.Http;
-using UnrealGameSyncMetadataServer.Connectors;
-using UnrealGameSyncMetadataServer.Models;
+using MetadataServer.Connectors;
+using MetadataServer.Models;
 
-namespace UnrealGameSyncMetadataServer.Controllers
+namespace MetadataServer.Controllers
 {
     public class ErrorController : ApiController
     {
diff --git a/Engine/Source/Programs/UnrealGameSync/MetadataServer/Controllers/EventController.cs b/Engine/Source/Programs/UnrealGameSync/MetadataServer/Controllers/EventController.cs
--- a/Engine/Source/Programs/UnrealGameSync/MetadataServer/Controllers/EventController.cs
+++ b/Engine/Source/Programs/UnrealGameSync/MetadataServer/Controllers/EventController.cs
@@ -2,10 +2,10 @@
 
 using System.Collections.Generic;
 using System.Web.Http;
-using UnrealGameSyncMetadataServer.Models;
-using UnrealGameSyncMetadataServer.Connectors;
+using MetadataServer.Models;
+using MetadataServer.Connectors;
 
-namespace UnrealGameSyncMetadataServer.Controllers
+namespace MetadataServer.Controllers
 {
 	public class EventController : ApiController
 	{
diff --git a/Engine/Source/Programs/UnrealGameSync/MetadataServer/Controllers/IssuesController.cs b/Engine/Source/Programs/UnrealGameSync/MetadataServer/Controllers/IssuesController.cs
--- a/Engine/Source/Programs/UnrealGameSync/MetadataServer/Controllers/IssuesController.cs
+++ b/Engine/Source/Programs/UnrealGameSync/MetadataServer/Controllers/IssuesController.cs
@@ -8,10 +8,10 @@ using System.Linq;
 using System.Net.Http;
 using System.Web.Http;
 using System.Web.Http.Results;
-using UnrealGameSyncMetadataServer.Connectors;
-using UnrealGameSyncMetadataServer.Models;
+using MetadataServer.Connectors;
+using MetadataServer.Models;
 
-namespace UnrealGameSyncMetadataServer.Controllers
+namespace MetadataServer.Controllers
 {
 	public class IssuesController : ApiController
 	{
diff --git a/Engine/Source/Programs/UnrealGameSync/MetadataServer/Controllers/LatestController.cs b/Engine/Source/Programs/UnrealGameSync/MetadataServer/Controllers/LatestController.cs
--- a/Engine/Source/Programs/UnrealGameSync/MetadataServer/Controllers/LatestController.cs
+++ b/Engine/Source/Programs/UnrealGameSync/MetadataServer/Controllers/LatestController.cs
@@ -6,10 +6,10 @@ using System.Linq;
 using System.Net;
 using System.Net.Http;
 using System.Web.Http;
-using UnrealGameSyncMetadataServer.Connectors;
-using UnrealGameSyncMetadataServer.Models;
+using MetadataServer.Connectors;
+using MetadataServer.Models;
 
-namespace UnrealGameSyncMetadataServer.Controllers
+namespace MetadataServer.Controllers
 {
     public class LatestController : ApiController
     {
diff --git a/Engine/Source/Programs/UnrealGameSync/MetadataServer/Controllers/TelemetryController.cs b/Engine/Source/Programs/UnrealGameSync/MetadataServer/Controllers/TelemetryController.cs
--- a/Engine/Source/Programs/UnrealGameSync/MetadataServer/Controllers/TelemetryController.cs
+++ b/Engine/Source/Programs/UnrealGameSync/MetadataServer/Controllers/TelemetryController.cs
@@ -1,10 +1,10 @@
 ﻿// Copyright 1998-2019 Epic Games, Inc. All Rights Reserved.
 
 using System.Web.Http;
-using UnrealGameSyncMetadataServer.Connectors;
-using UnrealGameSyncMetadataServer.Models;
+using MetadataServer.Connectors;
+using MetadataServer.Models;
 
-namespace UnrealGameSyncMetadataServer.Controllers
+namespace MetadataServer.Controllers
 {
     public class TelemetryController : ApiController
     {
diff --git a/Engine/Source/Programs/UnrealGameSync/MetadataServer/Controllers/UserController.cs b/Engine/Source/Programs/UnrealGameSync/MetadataServer/Controllers/UserController.cs
--- a/Engine/Source/Programs/UnrealGameSync/MetadataServer/Controllers/UserController.cs
+++ b/Engine/Source/Programs/UnrealGameSync/MetadataServer/Controllers/UserController.cs
@@ -4,10 +4,10 @@ using System;
 using System.Collections.Generic;
 using System.Data.SQLite;
 using System.Web.Http;
-using UnrealGameSyncMetadataServer.Connectors;
-using UnrealGameSyncMetadataServer.Models;
+using MetadataServer.Connectors;
+using MetadataServer.Models;
 
-namespace UnrealGameSyncMetadataServer.Controllers
+namespace MetadataServer.Controllers
 {
 	public class UserController : ApiController
 	{
diff --git a/Engine/Source/Programs/UnrealGameSync/MetadataServer/Global.asax.cs b/Engine/Source/Programs/UnrealGameSync/MetadataServer/Global.asax.cs
--- a/Engine/Source/Programs/UnrealGameSync/MetadataServer/Global.asax.cs
+++ b/Engine/Source/Programs/UnrealGameSync/MetadataServer/Global.asax.cs
@@ -6,7 +6,7 @@ using System.Text.RegularExpressions;
 using System.Web;
 using System.Web.Http;
 
-namespace UnrealGameSyncMetadataServer
+namespace MetadataServer
 {
 	public class WebApiApplication : System.Web.HttpApplication
 	{
diff --git a/Engine/Source/Programs/UnrealGameSync/MetadataServer/MetadataServer.csproj b/Engine/Source/Programs/UnrealGameSync/MetadataServer/MetadataServer.csproj
--- a/Engine/Source/Programs/UnrealGameSync/MetadataServer/MetadataServer.csproj
+++ b/Engine/Source/Programs/UnrealGameSync/MetadataServer/MetadataServer.csproj
@@ -47,18 +47,19 @@
   </PropertyGroup>
   <ItemGroup>
     <Reference Include="Microsoft.CSharp" />
-    <Reference Include="Newtonsoft.Json, Version=11.0.0.0, Culture=neutral, PublicKeyToken=30ad4fe6b2a6aeed, processorArchitecture=MSIL">
-      <HintPath>..\packages\Newtonsoft.Json.11.0.2\lib\net45\Newtonsoft.Json.dll</HintPath>
-      <Private>True</Private>
+    <Reference Include="Newtonsoft.Json, Version=12.0.0.0, Culture=neutral, PublicKeyToken=30ad4fe6b2a6aeed, processorArchitecture=MSIL">
+      <HintPath>..\packages\Newtonsoft.Json.12.0.2\lib\net45\Newtonsoft.Json.dll</HintPath>
     </Reference>
     <Reference Include="System" />
     <Reference Include="System.Data" />
     <Reference Include="System.Data.DataSetExtensions" />
-    <Reference Include="System.Data.SQLite, Version=1.0.108.0, Culture=neutral, PublicKeyToken=db937bc2d44ff139, processorArchitecture=MSIL">
-      <HintPath>..\packages\System.Data.SQLite.Core.1.0.108.0\lib\net45\System.Data.SQLite.dll</HintPath>
-      <Private>True</Private>
+    <Reference Include="System.Data.SQLite, Version=1.0.110.0, Culture=neutral, PublicKeyToken=db937bc2d44ff139, processorArchitecture=MSIL">
+      <HintPath>..\packages\System.Data.SQLite.Core.1.0.110.0\lib\net45\System.Data.SQLite.dll</HintPath>
     </Reference>
     <Reference Include="System.Drawing" />
+    <Reference Include="System.Net.Http.Formatting, Version=5.2.7.0, Culture=neutral, PublicKeyToken=31bf3856ad364e35, processorArchitecture=MSIL">
+      <HintPath>..\packages\Microsoft.AspNet.WebApi.Client.5.2.7\lib\net45\System.Net.Http.Formatting.dll</HintPath>
+    </Reference>
     <Reference Include="System.Web.DynamicData" />
     <Reference Include="System.Web.Entity" />
     <Reference Include="System.Web.ApplicationServices" />
@@ -66,23 +67,20 @@
     <Reference Include="System.Web.Extensions" />
     <Reference Include="System.Web" />
     <Reference Include="System.Web.Abstractions" />
+    <Reference Include="System.Web.Http, Version=5.2.7.0, Culture=neutral, PublicKeyToken=31bf3856ad364e35, processorArchitecture=MSIL">
+      <HintPath>..\packages\Microsoft.AspNet.WebApi.Core.5.2.7\lib\net45\System.Web.Http.dll</HintPath>
+    </Reference>
+    <Reference Include="System.Web.Http.WebHost, Version=5.2.7.0, Culture=neutral, PublicKeyToken=31bf3856ad364e35, processorArchitecture=MSIL">
+      <HintPath>..\packages\Microsoft.AspNet.WebApi.WebHost.5.2.7\lib\net45\System.Web.Http.WebHost.dll</HintPath>
+    </Reference>
     <Reference Include="System.Web.Routing" />
     <Reference Include="System.Xml" />
     <Reference Include="System.Configuration" />
     <Reference Include="System.Runtime.Serialization" />
     <Reference Include="System.Net.Http">
     </Reference>
-    <Reference Include="System.Net.Http.Formatting, Version=5.2.3.0, Culture=neutral, PublicKeyToken=31bf3856ad364e35, processorArchitecture=MSIL">
-      <HintPath>..\packages\Microsoft.AspNet.WebApi.Client.5.2.3\lib\net45\System.Net.Http.Formatting.dll</HintPath>
-    </Reference>
     <Reference Include="System.Net.Http.WebRequest">
     </Reference>
-    <Reference Include="System.Web.Http, Version=5.2.3.0, Culture=neutral, PublicKeyToken=31bf3856ad364e35, processorArchitecture=MSIL">
-      <HintPath>..\packages\Microsoft.AspNet.WebApi.Core.5.2.3\lib\net45\System.Web.Http.dll</HintPath>
-    </Reference>
-    <Reference Include="System.Web.Http.WebHost, Version=5.2.3.0, Culture=neutral, PublicKeyToken=31bf3856ad364e35, processorArchitecture=MSIL">
-      <HintPath>..\packages\Microsoft.AspNet.WebApi.WebHost.5.2.3\lib\net45\System.Web.Http.WebHost.dll</HintPath>
-    </Reference>
     <Reference Include="System.Xml.Linq" />
   </ItemGroup>
   <ItemGroup>
@@ -124,7 +122,9 @@
     <Content Include="Web.config">
       <SubType>Designer</SubType>
     </Content>
-    <None Include="Web.template.config" />
+    <None Include="Web.template.config">
+      <SubType>Designer</SubType>
+    </None>
     <None Include="Web.Debug.config">
       <DependentUpon>Web.template.config</DependentUpon>
     </None>
@@ -185,11 +185,11 @@
     <Message Importance="High" Text="Found NotForLicensees\Web.config, applying transform" Condition="Exists('NotForLicensees\Web.config')" />
     <TransformXml Source="Web.config" Transform="NotForLicensees\Web.config" Destination="Web.config" Condition="Exists('NotForLicensees\Web.config')" />
   </Target>
-  <Import Project="..\packages\System.Data.SQLite.Core.1.0.108.0\build\net45\System.Data.SQLite.Core.targets" Condition="Exists('..\packages\System.Data.SQLite.Core.1.0.108.0\build\net45\System.Data.SQLite.Core.targets')" />
+  <Import Project="..\packages\System.Data.SQLite.Core.1.0.110.0\build\net45\System.Data.SQLite.Core.targets" Condition="Exists('..\packages\System.Data.SQLite.Core.1.0.110.0\build\net45\System.Data.SQLite.Core.targets')" />
   <Target Name="EnsureNuGetPackageBuildImports" BeforeTargets="PrepareForBuild">
     <PropertyGroup>
       <ErrorText>This project references NuGet package(s) that are missing on this computer. Use NuGet Package Restore to download them.  For more information, see http://go.microsoft.com/fwlink/?LinkID=322105. The missing file is {0}.</ErrorText>
     </PropertyGroup>
-    <Error Condition="!Exists('..\packages\System.Data.SQLite.Core.1.0.108.0\build\net45\System.Data.SQLite.Core.targets')" Text="$([System.String]::Format('$(ErrorText)', '..\packages\System.Data.SQLite.Core.1.0.108.0\build\net45\System.Data.SQLite.Core.targets'))" />
+    <Error Condition="!Exists('..\packages\System.Data.SQLite.Core.1.0.110.0\build\net45\System.Data.SQLite.Core.targets')" Text="$([System.String]::Format('$(ErrorText)', '..\packages\System.Data.SQLite.Core.1.0.110.0\build\net45\System.Data.SQLite.Core.targets'))" />
   </Target>
 </Project>
\ No newline at end of file
diff --git a/Engine/Source/Programs/UnrealGameSync/MetadataServer/Models/BuildData.cs b/Engine/Source/Programs/UnrealGameSync/MetadataServer/Models/BuildData.cs
--- a/Engine/Source/Programs/UnrealGameSync/MetadataServer/Models/BuildData.cs
+++ b/Engine/Source/Programs/UnrealGameSync/MetadataServer/Models/BuildData.cs
@@ -5,7 +5,7 @@ using System.Collections.Generic;
 using System.Linq;
 using System.Web;
 
-namespace UnrealGameSyncMetadataServer.Models
+namespace MetadataServer.Models
 {
 	public class BuildData
 	{
diff --git a/Engine/Source/Programs/UnrealGameSync/MetadataServer/Models/CommentData.cs b/Engine/Source/Programs/UnrealGameSync/MetadataServer/Models/CommentData.cs
--- a/Engine/Source/Programs/UnrealGameSync/MetadataServer/Models/CommentData.cs
+++ b/Engine/Source/Programs/UnrealGameSync/MetadataServer/Models/CommentData.cs
@@ -5,7 +5,7 @@ using System.Collections.Generic;
 using System.Linq;
 using System.Web;
 
-namespace UnrealGameSyncMetadataServer.Models
+namespace MetadataServer.Models
 {
 	public class CommentData
 	{
diff --git a/Engine/Source/Programs/UnrealGameSync/MetadataServer/Models/EventData.cs b/Engine/Source/Programs/UnrealGameSync/MetadataServer/Models/EventData.cs
--- a/Engine/Source/Programs/UnrealGameSync/MetadataServer/Models/EventData.cs
+++ b/Engine/Source/Programs/UnrealGameSync/MetadataServer/Models/EventData.cs
@@ -5,7 +5,7 @@ using System.Collections.Generic;
 using System.Linq;
 using System.Web;
 
-namespace UnrealGameSyncMetadataServer.Models
+namespace MetadataServer.Models
 {
 	public class EventData
 	{
diff --git a/Engine/Source/Programs/UnrealGameSync/MetadataServer/Models/EventSummary.cs b/Engine/Source/Programs/UnrealGameSync/MetadataServer/Models/EventSummary.cs
--- a/Engine/Source/Programs/UnrealGameSync/MetadataServer/Models/EventSummary.cs
+++ b/Engine/Source/Programs/UnrealGameSync/MetadataServer/Models/EventSummary.cs
@@ -5,7 +5,7 @@ using System.Collections.Generic;
 using System.Linq;
 using System.Web;
 
-namespace UnrealGameSyncMetadataServer.Models
+namespace MetadataServer.Models
 {
 	public class EventSummary
 	{
diff --git a/Engine/Source/Programs/UnrealGameSync/MetadataServer/Models/IssueData.cs b/Engine/Source/Programs/UnrealGameSync/MetadataServer/Models/IssueData.cs
--- a/Engine/Source/Programs/UnrealGameSync/MetadataServer/Models/IssueData.cs
+++ b/Engine/Source/Programs/UnrealGameSync/MetadataServer/Models/IssueData.cs
@@ -6,7 +6,7 @@ using System.Collections.Generic;
 using System.Linq;
 using System.Web;
 
-namespace UnrealGameSyncMetadataServer.Models
+namespace MetadataServer.Models
 {
 	public class IssueWatcherData
 	{
diff --git a/Engine/Source/Programs/UnrealGameSync/MetadataServer/Models/LatestData.cs b/Engine/Source/Programs/UnrealGameSync/MetadataServer/Models/LatestData.cs
--- a/Engine/Source/Programs/UnrealGameSync/MetadataServer/Models/LatestData.cs
+++ b/Engine/Source/Programs/UnrealGameSync/MetadataServer/Models/LatestData.cs
@@ -5,7 +5,7 @@ using System.Collections.Generic;
 using System.Linq;
 using System.Web;
 
-namespace UnrealGameSyncMetadataServer.Models
+namespace MetadataServer.Models
 {
 	public class LatestData
 	{
diff --git a/Engine/Source/Programs/UnrealGameSync/MetadataServer/Models/TelemetryErrorData.cs b/Engine/Source/Programs/UnrealGameSync/MetadataServer/Models/TelemetryErrorData.cs
--- a/Engine/Source/Programs/UnrealGameSync/MetadataServer/Models/TelemetryErrorData.cs
+++ b/Engine/Source/Programs/UnrealGameSync/MetadataServer/Models/TelemetryErrorData.cs
@@ -5,7 +5,7 @@ using System.Collections.Generic;
 using System.Linq;
 using System.Web;
 
-namespace UnrealGameSyncMetadataServer.Models
+namespace MetadataServer.Models
 {
 	public class TelemetryErrorData
 	{
diff --git a/Engine/Source/Programs/UnrealGameSync/MetadataServer/Models/TelemetryTimingData.cs b/Engine/Source/Programs/UnrealGameSync/MetadataServer/Models/TelemetryTimingData.cs
--- a/Engine/Source/Programs/UnrealGameSync/MetadataServer/Models/TelemetryTimingData.cs
+++ b/Engine/Source/Programs/UnrealGameSync/MetadataServer/Models/TelemetryTimingData.cs
@@ -5,7 +5,7 @@ using System.Collections.Generic;
 using System.Linq;
 using System.Web;
 
-namespace UnrealGameSyncMetadataServer.Models
+namespace MetadataServer.Models
 {
 	public class TelemetryTimingData
 	{
diff --git a/Engine/Source/Programs/UnrealGameSync/MetadataServer/Properties/AssemblyInfo.cs b/Engine/Source/Programs/UnrealGameSync/MetadataServer/Properties/AssemblyInfo.cs
--- a/Engine/Source/Programs/UnrealGameSync/MetadataServer/Properties/AssemblyInfo.cs
+++ b/Engine/Source/Programs/UnrealGameSync/MetadataServer/Properties/AssemblyInfo.cs
@@ -7,7 +7,7 @@ using System.Runtime.InteropServices;
 // General Information about an assembly is controlled through the following
 // set of attributes. Change these attribute values to modify the information
 // associated with an assembly.
-[assembly: AssemblyTitle("UnrealGameSyncMetadataServer")]
+[assembly: AssemblyTitle("MetadataServer")]
 [assembly: AssemblyDescription("")]
 [assembly: AssemblyConfiguration("")]
 [assembly: AssemblyCompany("")]
diff --git a/Engine/Source/Programs/UnrealGameSync/MetadataServer/Web.template.config b/Engine/Source/Programs/UnrealGameSync/MetadataServer/Web.template.config
--- a/Engine/Source/Programs/UnrealGameSync/MetadataServer/Web.template.config
+++ b/Engine/Source/Programs/UnrealGameSync/MetadataServer/Web.template.config
@@ -35,7 +35,7 @@
     <assemblyBinding xmlns="urn:schemas-microsoft-com:asm.v1">
       <dependentAssembly>
         <assemblyIdentity name="Newtonsoft.Json" publicKeyToken="30ad4fe6b2a6aeed" culture="neutral" />
-        <bindingRedirect oldVersion="0.0.0.0-11.0.0.0" newVersion="11.0.0.0" />
+        <bindingRedirect oldVersion="0.0.0.0-12.0.0.0" newVersion="12.0.0.0" />
       </dependentAssembly>
     </assemblyBinding>
   </runtime>
diff --git a/Engine/Source/Programs/UnrealGameSync/MetadataServer/packages.config b/Engine/Source/Programs/UnrealGameSync/MetadataServer/packages.config
--- a/Engine/Source/Programs/UnrealGameSync/MetadataServer/packages.config
+++ b/Engine/Source/Programs/UnrealGameSync/MetadataServer/packages.config
@@ -1,9 +1,9 @@
 ﻿<?xml version="1.0" encoding="utf-8"?>
 <packages>
-  <package id="Microsoft.AspNet.WebApi" version="5.2.3" targetFramework="net452" />
-  <package id="Microsoft.AspNet.WebApi.Client" version="5.2.3" targetFramework="net452" />
-  <package id="Microsoft.AspNet.WebApi.Core" version="5.2.3" targetFramework="net452" />
-  <package id="Microsoft.AspNet.WebApi.WebHost" version="5.2.3" targetFramework="net452" />
-  <package id="Newtonsoft.Json" version="11.0.2" targetFramework="net452" />
-  <package id="System.Data.SQLite.Core" version="1.0.108.0" targetFramework="net45" />
+  <package id="Microsoft.AspNet.WebApi" version="5.2.7" targetFramework="net45" />
+  <package id="Microsoft.AspNet.WebApi.Client" version="5.2.7" targetFramework="net45" />
+  <package id="Microsoft.AspNet.WebApi.Core" version="5.2.7" targetFramework="net45" />
+  <package id="Microsoft.AspNet.WebApi.WebHost" version="5.2.7" targetFramework="net45" />
+  <package id="Newtonsoft.Json" version="12.0.2" targetFramework="net45" />
+  <package id="System.Data.SQLite.Core" version="1.0.110.0" targetFramework="net45" />
 </packages>
\ No newline at end of file
