UGS: Show the build issue alert window in red if a user has been nominated to fix it.

diff --git a/Engine/Source/Programs/UnrealGameSync/UnrealGameSync/Controls/AlertButtonControl.cs b/Engine/Source/Programs/UnrealGameSync/UnrealGameSync/Controls/AlertButtonControl.cs
--- a/Engine/Source/Programs/UnrealGameSync/UnrealGameSync/Controls/AlertButtonControl.cs
+++ b/Engine/Source/Programs/UnrealGameSync/UnrealGameSync/Controls/AlertButtonControl.cs
@@ -21,6 +21,7 @@ namespace UnrealGameSync
 			Normal,
 			Green,
 			Red,
+			Strong,
 			Custom
 		}
 
@@ -47,7 +48,22 @@ namespace UnrealGameSync
 				get; set;
 			}
 
-			public Color BackgroundColorDown
+			public Color BackgroundColorHover1
+			{
+				get; set;
+			}
+
+			public Color BackgroundColorHover2
+			{
+				get; set;
+			}
+
+			public Color BackgroundColorDown1
+			{
+				get; set;
+			}
+
+			public Color BackgroundColorDown2
 			{
 				get; set;
 			}
@@ -71,11 +87,6 @@ namespace UnrealGameSync
 			set { CustomColorsValue = value; UpdateThemeColors(); }
 		}
 
-		public new Color ForeColor
-		{
-			get { return base.ForeColor; }
-		}
-
 		public AlertButtonControl()
 		{
 			Theme = AlertButtonTheme.Normal;
@@ -91,21 +102,40 @@ namespace UnrealGameSync
 					Colors.BorderColor = Color.FromArgb(230, 232, 235);
 					Colors.BackgroundColor1 = Color.FromArgb(255, 255, 255);
 					Colors.BackgroundColor2 = Color.FromArgb(244, 245, 247);
-					Colors.BackgroundColorDown = Color.FromArgb(234, 235, 237);
+					Colors.BackgroundColorHover1 = Color.FromArgb(244, 245, 247);
+					Colors.BackgroundColorHover2 = Color.FromArgb(244, 245, 247);
+					Colors.BackgroundColorDown1 = Color.FromArgb(234, 235, 237);
+					Colors.BackgroundColorDown2 = Color.FromArgb(234, 235, 237);
 					break;
 				case AlertButtonTheme.Green:
 					Colors.ForeColor = Color.FromArgb(255, 255, 255);
 					Colors.BorderColor = Color.FromArgb(143, 199, 156);
 					Colors.BackgroundColor1 = Color.FromArgb(116, 192, 134);
 					Colors.BackgroundColor2 = Color.FromArgb(99, 175, 117);
-					Colors.BackgroundColorDown = Color.FromArgb(90, 165, 107);
+					Colors.BackgroundColorHover1 = Color.FromArgb(99, 175, 117);
+					Colors.BackgroundColorHover2 = Color.FromArgb(99, 175, 117);
+					Colors.BackgroundColorDown1 = Color.FromArgb(90, 165, 107);
+					Colors.BackgroundColorDown2 = Color.FromArgb(90, 165, 107);
 					break;
 				case AlertButtonTheme.Red:
 					Colors.ForeColor = Color.FromArgb(255, 255, 255);
 					Colors.BorderColor = Color.FromArgb(230, 232, 235);
 					Colors.BackgroundColor1 = Color.FromArgb(222, 108, 86);
 					Colors.BackgroundColor2 = Color.FromArgb(214, 69, 64);
-					Colors.BackgroundColorDown = Color.FromArgb(204, 59, 54);
+					Colors.BackgroundColorHover1 = Color.FromArgb(214, 69, 64);
+					Colors.BackgroundColorHover2 = Color.FromArgb(214, 69, 64);
+					Colors.BackgroundColorDown1 = Color.FromArgb(204, 59, 54);
+					Colors.BackgroundColorDown2 = Color.FromArgb(204, 59, 54);
+					break;
+				case AlertButtonTheme.Strong:
+					Colors.ForeColor = Color.FromArgb(255, 255, 255);
+					Colors.BorderColor = Color.FromArgb(230, 232, 235);
+					Colors.BackgroundColor1 = Color.FromArgb(214, 69, 64);
+					Colors.BackgroundColor2 = Color.FromArgb(214, 69, 64);
+					Colors.BackgroundColorHover1 = Color.FromArgb(222, 108, 86);
+					Colors.BackgroundColorHover2 = Color.FromArgb(222, 108, 86);
+					Colors.BackgroundColorDown1 = Color.FromArgb(204, 59, 54);
+					Colors.BackgroundColorDown2 = Color.FromArgb(204, 59, 54);
 					break;
 				case AlertButtonTheme.Custom:
 					Colors = CustomColorsValue;
@@ -162,8 +192,8 @@ namespace UnrealGameSync
 				Path.AddArc(Width - 1 - Diameter, 0, Diameter, Diameter, 270, 90);
 				Path.CloseFigure();
 
-				Color BackgroundColorMin = (bMouseDown && bMouseOver)? Colors.BackgroundColorDown : bMouseOver? Colors.BackgroundColor2 : Colors.BackgroundColor1;
-				Color BackgroundColorMax = (bMouseDown && bMouseOver)? Colors.BackgroundColorDown : Colors.BackgroundColor2;
+				Color BackgroundColorMin = (bMouseDown && bMouseOver)? Colors.BackgroundColorDown1 : bMouseOver? Colors.BackgroundColorHover1 : Colors.BackgroundColor1;
+				Color BackgroundColorMax = (bMouseDown && bMouseOver)? Colors.BackgroundColorDown2 : bMouseOver? Colors.BackgroundColorHover2 : Colors.BackgroundColor2;
 
 				using(LinearGradientBrush Brush = new LinearGradientBrush(new Point(0, 0), new Point(0, Height), BackgroundColorMin, BackgroundColorMax))
 				{
diff --git a/Engine/Source/Programs/UnrealGameSync/UnrealGameSync/Forms/IssueAlertWindow.cs b/Engine/Source/Programs/UnrealGameSync/UnrealGameSync/Forms/IssueAlertWindow.cs
--- a/Engine/Source/Programs/UnrealGameSync/UnrealGameSync/Forms/IssueAlertWindow.cs
+++ b/Engine/Source/Programs/UnrealGameSync/UnrealGameSync/Forms/IssueAlertWindow.cs
@@ -25,7 +25,7 @@ namespace UnrealGameSync
 
 		public bool bIsWarning;
 
-		public bool bStrongAlert = false;
+		public bool? bStrongAlert;
 
 		public IssueAlertWindow(IssueMonitor IssueMonitor, IssueData Issue, IssueAlertReason Reason)
 		{
@@ -53,12 +53,12 @@ namespace UnrealGameSync
 		{
 			base.OnPaintBackground(e);
 
-			if(bStrongAlert)
+			if(bStrongAlert ?? false)
 			{
 				Color StripeColor = bIsWarning? Color.FromArgb(216, 167, 64) : Color.FromArgb(214, 69, 64);
 				using(Brush StripeBrush = new SolidBrush(StripeColor))
 				{
-					e.Graphics.FillRectangle(StripeBrush, 0, 0, Bounds.Width - 1, Bounds.Height - 1);
+					e.Graphics.FillRectangle(StripeBrush, 0, 0, Bounds.Width, Bounds.Height);
 				}
 			}
 			else
@@ -86,6 +86,8 @@ namespace UnrealGameSync
 
 		public void SetIssue(IssueData NewIssue, IssueAlertReason NewReason)
 		{
+			bool bNewStrongAlert = false;
+
 			StringBuilder OwnerTextBuilder = new StringBuilder();
 			if(NewIssue.Owner == null)
 			{
@@ -96,6 +98,7 @@ namespace UnrealGameSync
 				if(String.Compare(NewIssue.Owner, IssueMonitor.UserName, StringComparison.OrdinalIgnoreCase) == 0 && NewIssue.NominatedBy != null)
 				{
 					OwnerTextBuilder.AppendFormat("You have been nominated to fix this issue by {0}.", Utility.FormatUserName(NewIssue.NominatedBy));
+					bNewStrongAlert = true;
 				}
 				else
 				{
@@ -114,7 +117,7 @@ namespace UnrealGameSync
 
 			Issue = NewIssue;
 
-			if(NewIssue.Summary != SummaryLabel.Text || OwnerText != OwnerLabel.Text || Reason != NewReason || bIsWarning != bNewIsWarning)
+			if(NewIssue.Summary != SummaryLabel.Text || OwnerText != OwnerLabel.Text || Reason != NewReason || bIsWarning != bNewIsWarning || bStrongAlert != bNewStrongAlert)
 			{
 				Rectangle PrevBounds = Bounds;
 				SuspendLayout();
@@ -122,31 +125,42 @@ namespace UnrealGameSync
 				SummaryLabel.Text = NewIssue.Summary;
 				OwnerLabel.Text = OwnerText;
 
-				if(bStrongAlert)
+				bool bForceUpdateButtons = false;
+				if(bStrongAlert != bNewStrongAlert)
 				{
-					SummaryLabel.ForeColor = Color.FromArgb(255, 255, 255);
-					SummaryLabel.LinkColor = Color.FromArgb(255, 255, 255);
-					OwnerLabel.ForeColor = Color.FromArgb(255, 255, 255);
-//					AcceptBtn.BackgroundColor1 = Color.FromArgb(214, 69, 64);//Color.FromArgb(255, 255, 255);
-//					AcceptBtn.BackgroundColor2 = Color.FromArgb(214, 69, 64);//Color.FromArgb(255, 255, 255);
-//					AcceptBtn.BorderColor = Color.FromArgb(255, 255, 255);
-//					AcceptBtn.ForeColor = Color.FromArgb(255, 255, 255);//DetailsBtn.ForeColor;
-					AcceptBtn.Margin = new Padding(0,0, 20, 0);//.Right = 20;
-					LatestBuildLinkLabel.LinkColor = Color.FromArgb(255, 255, 255);
+					bStrongAlert = bNewStrongAlert;
+
+					if(bNewStrongAlert)
+					{
+						SummaryLabel.ForeColor = Color.FromArgb(255, 255, 255);
+						SummaryLabel.LinkColor = Color.FromArgb(255, 255, 255);
+						OwnerLabel.ForeColor = Color.FromArgb(255, 255, 255);
+						AcceptBtn.Theme = AlertButtonControl.AlertButtonTheme.Strong;
+						LatestBuildLinkLabel.LinkColor = Color.FromArgb(255, 255, 255);
+					}
+					else
+					{
+						SummaryLabel.ForeColor = Color.FromArgb(32, 32, 64);
+						SummaryLabel.LinkColor = Color.FromArgb(32, 32, 64);
+						OwnerLabel.ForeColor = Color.FromArgb(32, 32, 64);
+						AcceptBtn.Theme = AlertButtonControl.AlertButtonTheme.Green;
+						LatestBuildLinkLabel.LinkColor = Color.FromArgb(16, 102, 192);
+					}
+
+					bForceUpdateButtons = true;
 				}
 
-				if(bIsWarning != bNewIsWarning)
+				if (bIsWarning != bNewIsWarning)
 				{
 					bIsWarning = bNewIsWarning;
-					IconPictureBox.Image = bIsWarning? Properties.Resources.AlertWarningIcon : Properties.Resources.AlertErrorIcon;
 				}
 
-				if(Reason != NewReason)
+				if(Reason != NewReason || bForceUpdateButtons)
 				{
 					Reason = NewReason;
 
 					List<Button> Buttons = new List<Button>();
-					if(!bStrongAlert)
+					if(!bStrongAlert.Value)
 					{
 						Buttons.Add(DetailsBtn);
 					}
