UGS: Refactor MetadataTool to support different command types.

#rnx

diff --git a/Engine/Source/Programs/UnrealGameSync/MetadataTool/CommandHandler.cs b/Engine/Source/Programs/UnrealGameSync/MetadataTool/CommandHandler.cs
new file mode 100644
--- /dev/null
+++ b/Engine/Source/Programs/UnrealGameSync/MetadataTool/CommandHandler.cs
@@ -0,0 +1,15 @@
+using System;
+using System.Collections.Generic;
+using System.Linq;
+using System.Text;
+using System.Threading.Tasks;
+using Tools.DotNETCommon;
+
+namespace MetadataTool
+{
+	abstract class CommandHandler
+	{
+		public abstract void Exec(CommandLineArguments Arguments);
+		public abstract void Help();
+	}
+}
diff --git a/Engine/Source/Programs/UnrealGameSync/MetadataTool/CommandHandler_Http.cs b/Engine/Source/Programs/UnrealGameSync/MetadataTool/CommandHandler_Http.cs
new file mode 100644
--- /dev/null
+++ b/Engine/Source/Programs/UnrealGameSync/MetadataTool/CommandHandler_Http.cs
@@ -0,0 +1,243 @@
+using System;
+using System.Collections.Generic;
+using System.IO;
+using System.Linq;
+using System.Net;
+using System.Reflection;
+using System.Text;
+using System.Threading.Tasks;
+using System.Web.Script.Serialization;
+using Tools.DotNETCommon;
+
+namespace MetadataTool
+{
+	class CommandHandler_Http : CommandHandler
+	{
+		public string Verb;
+		public string Resource;
+		public Type BodyType;
+		public string[] OptionalParams;
+
+		public CommandHandler_Http(string Verb, string Resource, Type BodyType = null, string[] OptionalParams = null)
+		{
+			this.Verb = Verb;
+			this.Resource = Resource;
+			this.BodyType = BodyType;
+			this.OptionalParams = OptionalParams;
+		}
+
+		public override void Exec(CommandLineArguments Arguments)
+		{
+			// Parse the server URL
+			string ServerUrl = Arguments.GetStringOrDefault("-Server=", null) ?? Environment.GetEnvironmentVariable("UGS_METADATA_SERVER_URL");
+			if (String.IsNullOrEmpty(ServerUrl))
+			{
+				throw new CommandLineArgumentException("Missing -Server=... argument");
+			}
+
+			// Create the command url
+			StringBuilder CommandUrl = new StringBuilder();
+			CommandUrl.AppendFormat("{0}{1}", ServerUrl, Resource);
+
+			// Replace all the arguments that are embedded into the resource name
+			foreach (string ResourceArg in ParseArgumentNamesFromResource(Resource))
+			{
+				string Value = Arguments.GetString(String.Format("-{0}=", ResourceArg));
+				CommandUrl.Replace("{" + ResourceArg + "}", Value);
+			}
+
+			// Add all the required and optional parameters
+			List<string> QueryParams = new List<string>();
+			if (OptionalParams != null)
+			{
+				foreach (string OptionalParam in OptionalParams)
+				{
+					string Value = Arguments.GetStringOrDefault(String.Format("-{0}=", OptionalParam), null);
+					if (Value != null)
+					{
+						QueryParams.Add(String.Format("{0}={1}", OptionalParam, Value));
+					}
+				}
+			}
+
+			// Append the parameters to the URL
+			for (int Idx = 0; Idx < QueryParams.Count; Idx++)
+			{
+				if (Idx == 0)
+				{
+					CommandUrl.Append("?");
+				}
+				else
+				{
+					CommandUrl.Append("&");
+				}
+				CommandUrl.Append(QueryParams[Idx]);
+			}
+
+			// Parse additional options for the message body
+			string BodyText = null;
+			if (BodyType != null)
+			{
+				object BodyObject = ParseObject(Arguments, BodyType);
+				BodyText = new JavaScriptSerializer().Serialize(BodyObject);
+			}
+
+			// Make sure there are no unused arguments
+			Arguments.CheckAllArgumentsUsed();
+
+			// Create the request
+			HttpWebRequest Request = (HttpWebRequest)WebRequest.Create(CommandUrl.ToString());
+			Request.ContentType = "application/json";
+			Request.Method = Verb;
+			if (BodyText != null)
+			{
+				byte[] BodyData = Encoding.UTF8.GetBytes(BodyText);
+				using (Stream RequestStream = Request.GetRequestStream())
+				{
+					RequestStream.Write(BodyData, 0, BodyData.Length);
+				}
+			}
+
+			// Read the response
+			HttpWebResponse Response = (HttpWebResponse)Request.GetResponse();
+			Console.WriteLine("Response: {0} ({1})", (int)Response.StatusCode, Response.StatusDescription);
+			using (StreamReader ResponseReader = new StreamReader(Response.GetResponseStream(), Encoding.Default))
+			{
+				string ResponseContent = ResponseReader.ReadToEnd();
+				if (!String.IsNullOrEmpty(ResponseContent))
+				{
+					Console.WriteLine(Json.Format(ResponseContent));
+				}
+			}
+		}
+
+		public override void Help()
+		{
+			foreach (string ArgumentName in ParseArgumentNamesFromResource(Resource))
+			{
+				Console.WriteLine("  -{0}=...", ArgumentName);
+			}
+			if (OptionalParams != null)
+			{
+				foreach (string ArgumentName in OptionalParams)
+				{
+					Console.WriteLine("  -{0}=... (optional)", ArgumentName);
+				}
+			}
+			if (BodyType != null)
+			{
+				foreach (FieldInfo Field in BodyType.GetFields(BindingFlags.Instance | BindingFlags.Public))
+				{
+					if (IsOptionalObjectField(Field))
+					{
+						Console.WriteLine("  -{0}={1} (optional)", Field.Name, GetFieldValueDescription(Field.FieldType));
+					}
+					else
+					{
+						Console.WriteLine("  -{0}={1}", Field.Name, GetFieldValueDescription(Field.FieldType));
+					}
+				}
+			}
+		}
+
+		static IEnumerable<string> ParseArgumentNamesFromResource(string Resource)
+		{
+			for (int ArgIdx = Resource.IndexOf('{'); ArgIdx != -1; ArgIdx = Resource.IndexOf('{', ArgIdx))
+			{
+				int ArgEndIdx = Resource.IndexOf('}', ArgIdx + 1);
+				yield return Resource.Substring(ArgIdx + 1, ArgEndIdx - (ArgIdx + 1));
+				ArgIdx = ArgEndIdx + 1;
+			}
+		}
+
+		static object ParseObject(CommandLineArguments Arguments, Type ObjectType)
+		{
+			object Instance = Activator.CreateInstance(ObjectType);
+			foreach (FieldInfo Field in Instance.GetType().GetFields(BindingFlags.Instance | BindingFlags.Public))
+			{
+				string Value = Arguments.GetStringOrDefault(String.Format("-{0}=", Field.Name), null);
+				if (Value == null)
+				{
+					if (!IsOptionalObjectField(Field))
+					{
+						throw new CommandLineArgumentException(String.Format("Missing -{0}=... argument", Field.Name));
+					}
+				}
+				else
+				{
+					Type FieldType = Field.FieldType;
+					if (FieldType.IsGenericType && FieldType.GetGenericTypeDefinition() == typeof(Nullable<>))
+					{
+						FieldType = FieldType.GetGenericArguments()[0];
+					}
+
+					if (FieldType == typeof(string))
+					{
+						Field.SetValue(Instance, Value);
+					}
+					else if (FieldType == typeof(int))
+					{
+						int IntValue;
+						if (!int.TryParse(Value, out IntValue))
+						{
+							throw new CommandLineArgumentException(String.Format("Invalid value for '{0}' - {1} is not an integer.", Field.Name, Value));
+						}
+						Field.SetValue(Instance, IntValue);
+					}
+					else if (FieldType == typeof(bool))
+					{
+						bool BoolValue;
+						if (!bool.TryParse(Value, out BoolValue))
+						{
+							throw new CommandLineArgumentException(String.Format("Invalid value for '{0}' - {1} is not a boolean.", Field.Name, Value));
+						}
+						Field.SetValue(Instance, BoolValue);
+					}
+					else if (FieldType.IsEnum)
+					{
+						object EnumValue;
+						try
+						{
+							EnumValue = Enum.Parse(FieldType, Value, true);
+						}
+						catch (ArgumentException)
+						{
+							throw new CommandLineArgumentException(String.Format("Invalid value for '{0}' - should be {1}", Field.Name, String.Join("/", Enum.GetNames(FieldType))));
+						}
+						Field.SetValue(Instance, EnumValue);
+					}
+					else
+					{
+						throw new NotImplementedException(String.Format("Unsupported type '{0}'", FieldType));
+					}
+				}
+			}
+			return Instance;
+		}
+
+		static bool IsOptionalObjectField(FieldInfo Field)
+		{
+			if (Field.GetCustomAttribute<OptionalAttribute>() != null)
+			{
+				return true;
+			}
+			if (Field.FieldType.IsGenericType && Field.FieldType.GetGenericTypeDefinition() == typeof(Nullable<>))
+			{
+				return true;
+			}
+			return false;
+		}
+
+		static string GetFieldValueDescription(Type FieldType)
+		{
+			if (FieldType.IsEnum)
+			{
+				return String.Join("|", Enum.GetNames(FieldType));
+			}
+			else
+			{
+				return "...";
+			}
+		}
+	}
+}
diff --git a/Engine/Source/Programs/UnrealGameSync/MetadataTool/CommandTypes.cs b/Engine/Source/Programs/UnrealGameSync/MetadataTool/CommandTypes.cs
new file mode 100644
--- /dev/null
+++ b/Engine/Source/Programs/UnrealGameSync/MetadataTool/CommandTypes.cs
@@ -0,0 +1,63 @@
+using System;
+using System.Collections.Generic;
+using System.Linq;
+using System.Text;
+using System.Threading.Tasks;
+
+namespace MetadataTool
+{
+	[AttributeUsage(AttributeTargets.Field)]
+	class OptionalAttribute : Attribute
+	{
+	}
+
+	static class CommandTypes
+	{
+#pragma warning disable CS0649
+
+		public enum Outcome
+		{
+			Success = 1,
+			Error = 2,
+			Warning = 3,
+		}
+
+		public class AddIssue
+		{
+			public string Project;
+			public string Summary;
+		}
+
+		public class UpdateBuild
+		{
+			[Optional] public int? Outcome;
+		}
+
+		public class UpdateIssue
+		{
+			[Optional] public bool? Acknowledged;
+			[Optional] public int? FixChange;
+			[Optional] public string Owner;
+			[Optional] public string NominatedBy;
+			[Optional] public bool? Resolved;
+			[Optional] public string Summary;
+			[Optional] public string Url;
+		}
+
+		public class AddBuild
+		{
+			public string Name;
+			public string Url;
+			public string Stream;
+			public int Change;
+			public Outcome Outcome;
+		}
+
+		public class Watcher
+		{
+			public string UserName;
+		}
+
+#pragma warning restore CS0649
+	}
+}
diff --git a/Engine/Source/Programs/UnrealGameSync/MetadataTool/MetadataTool.csproj b/Engine/Source/Programs/UnrealGameSync/MetadataTool/MetadataTool.csproj
--- a/Engine/Source/Programs/UnrealGameSync/MetadataTool/MetadataTool.csproj
+++ b/Engine/Source/Programs/UnrealGameSync/MetadataTool/MetadataTool.csproj
@@ -45,6 +45,9 @@
     <Reference Include="System.Xml" />
   </ItemGroup>
   <ItemGroup>
+    <Compile Include="CommandHandler.cs" />
+    <Compile Include="CommandHandler_Http.cs" />
+    <Compile Include="CommandTypes.cs" />
     <Compile Include="Program.cs" />
     <Compile Include="Properties\AssemblyInfo.cs" />
   </ItemGroup>
diff --git a/Engine/Source/Programs/UnrealGameSync/MetadataTool/Program.cs b/Engine/Source/Programs/UnrealGameSync/MetadataTool/Program.cs
--- a/Engine/Source/Programs/UnrealGameSync/MetadataTool/Program.cs
+++ b/Engine/Source/Programs/UnrealGameSync/MetadataTool/Program.cs
@@ -13,258 +13,91 @@ using System.Web.Script.Serialization;
 
 namespace MetadataTool
 {
-	[AttributeUsage(AttributeTargets.Field)]
-	class OptionalAttribute : Attribute
-	{
-	}
-
 	class Program
 	{
-		class UserErrorException : Exception
-		{
-			public UserErrorException(string Message)
-				: base(Message)
-			{
-			}
-		}
-
-		class MissingArgumentException : UserErrorException
-		{
-			public MissingArgumentException(string Message)
-				: base(Message)
-			{
-			}
-		}
-
-		class Command
-		{
-			public string Name;
-			public string Verb;
-			public string Resource;
-			public Type BodyType;
-			public string[] OptionalParams;
-
-			public Command(string Name, string Verb, string Resource, Type BodyType = null, string[] OptionalParams = null)
-			{
-				this.Name = Name;
-				this.Verb = Verb;
-				this.Resource = Resource;
-				this.BodyType = BodyType;
-				this.OptionalParams = OptionalParams;
-			}
-		}
-
-		#pragma warning disable CS0649
-
-		enum Outcome
-		{
-			Success = 1,
-			Error = 2,
-			Warning = 3,
-		}
-
-		class AddIssueBody
-		{
-			public string Project;
-			public string Summary;
-		}
-
-		class UpdateBuildBody
-		{
-			[Optional] public int? Outcome;
-		}
-
-		class UpdateIssueBody
-		{
-			[Optional] public bool? Acknowledged;
-			[Optional] public int? FixChange;
-			[Optional] public string Owner;
-			[Optional] public string NominatedBy;
-			[Optional] public bool? Resolved;
-			[Optional] public string Summary;
-			[Optional] public string Url;
-		}
-
-		class AddBuildBody
-		{
-			public string Name;
-			public string Url;
-			public string Stream;
-			public int Change;
-			public Outcome Outcome;
-		}
-
-		class WatcherBody
-		{
-			public string UserName;
-		}
-
-		#pragma warning restore CS0649
-
-		static Command[] Commands =
+		static Dictionary<string, CommandHandler> CommandMap = new Dictionary<string, CommandHandler>(StringComparer.OrdinalIgnoreCase)
 		{
 			// Issues
-			new Command("add-issue", "POST", "/api/issues", typeof(AddIssueBody)),
-			new Command("get-issue", "GET", "/api/issues/{id}"),
-			new Command("get-issues", "GET", "/api/issues", OptionalParams: new string[]{ "user" }),
-			new Command("update-issue", "PUT", "/api/issues/{Id}", typeof(UpdateIssueBody)),
+			[ "AddIssue" ] = new CommandHandler_Http("POST", "/api/issues", typeof(CommandTypes.AddIssue)),
+			[ "GetIssue" ] = new CommandHandler_Http("GET", "/api/issues/{id}"),
+			[ "GetIssues" ] = new CommandHandler_Http("GET", "/api/issues", OptionalParams: new string[] { "user" }),
+			[ "UpdateIssue" ] = new CommandHandler_Http("PUT", "/api/issues/{Id}", typeof(CommandTypes.UpdateIssue)),
 
 			// Builds
-			new Command("add-build", "POST", "/api/issues/{Issue}/builds", typeof(AddBuildBody)),
-			new Command("get-builds", "GET", "/api/issues/{Issue}/builds"),
+			[ "AddBuild" ] = new CommandHandler_Http("POST", "/api/issues/{Issue}/builds", typeof(CommandTypes.AddBuild)),
+			[ "GetBuilds" ] = new CommandHandler_Http("GET", "/api/issues/{Issue}/builds"),
 
 			// Individual builds
-			new Command("get-build", "GET", "/api/issuebuilds/{Build}"),
-			new Command("update-build", "PUT", "/api/issuebuilds/{Build}", typeof(UpdateBuildBody)),
+			[ "GetBuild" ] = new CommandHandler_Http("GET", "/api/issuebuilds/{Build}"),
+			[ "UpdateBuild" ] = new CommandHandler_Http("PUT", "/api/issuebuilds/{Build}", typeof(CommandTypes.UpdateBuild)),
 
 			// Watchers
-			new Command("get-watchers", "GET", "/api/issues/{Issue}/watchers"),
-			new Command("add-watcher", "POST", "/api/issues/{Issue}/watchers", typeof(WatcherBody)),
-			new Command("remove-watcher", "DELETE", "/api/issues/{Issue}/watchers", typeof(WatcherBody))
+			[ "GetWatchers" ] = new CommandHandler_Http("GET", "/api/issues/{Issue}/watchers"),
+			[ "AddWatcher" ] = new CommandHandler_Http("POST", "/api/issues/{Issue}/watchers", typeof(CommandTypes.Watcher)),
+			[ "RemoveWatcher" ] = new CommandHandler_Http("DELETE", "/api/issues/{Issue}/watchers", typeof(CommandTypes.Watcher))
 		};
 
 		static int Main(string[] Args)
 		{
 			try
 			{
-				bool bResult = InnerMain(Args);
+				bool bResult = InnerMain(new CommandLineArguments(Args));
 				return bResult? 0 : 1;
 			}
-			catch(UserErrorException Ex)
+			catch(Exception Ex)
 			{
-				Console.WriteLine("{0}", Ex.Message);
+				Console.WriteLine("{0}", Ex.ToString());
 				return 1;
 			}
 		}
 
-		static bool InnerMain(string[] Args)
+		static bool InnerMain(CommandLineArguments Arguments)
 		{
-			List<string> RemainingArgs = new List<string>(Args);
-
-			// Get the server URL
-			string ServerUrl = ParseOptionalParam(RemainingArgs, "Server") ?? Environment.GetEnvironmentVariable("UGS_METADATA_SERVER_URL");
-			if(RemainingArgs.Count == 0)
+			// Find the command name
+			string CommandName = null;
+			for(int Idx = 0; Idx < Arguments.Count; Idx++)
 			{
-				PrintCommands();
-				return true;
+				if(!Arguments[Idx].StartsWith("-"))
+				{
+					CommandName = Arguments[Idx];
+					Arguments.MarkAsUsed(Idx);
+					break;
+				}
 			}
 
-			// Register all the commands
-			Command Command = Commands.FirstOrDefault(x => String.Compare(x.Name, RemainingArgs[0], StringComparison.OrdinalIgnoreCase) == 0);
-			if(Command == null)
+			// Make sure we've got a command name, and fall back to printing out help if not
+			if(Arguments.Count == 0 || CommandName == null)
 			{
-				Console.WriteLine("Unknown command '{0}'", RemainingArgs[0]);
+				PrintCommands();
 				return false;
 			}
-			RemainingArgs.RemoveAt(0);
 
-			// Handle help for this command
-			if(ParseOption(RemainingArgs, "-Help"))
+			// Register all the commands
+			CommandHandler Command;
+			if(!CommandMap.TryGetValue(CommandName, out Command))
 			{
-				PrintCommandHelp(Command);
+				Console.WriteLine("Unknown command '{0}'", CommandName);
 				return false;
 			}
 
-			// Make sure we have the server parameter
-			if(String.IsNullOrEmpty(ServerUrl))
+			// Handle help for this command
+			if (Arguments.HasOption("-Help"))
 			{
-				Console.WriteLine("Missing -Server=... argument.");
-				Console.WriteLine();
-				PrintCommandHelp(Command);
+				PrintCommandHelp(CommandName, Command);
 				return false;
 			}
 
 			// Try to execute the command
 			try
 			{
-				// Create the command url
-				StringBuilder CommandUrl = new StringBuilder();
-				CommandUrl.AppendFormat("{0}{1}", ServerUrl, Command.Resource);
-
-				// Replace all the arguments that are embedded into the resource name
-				foreach(string ResourceArg in ParseArgumentNamesFromResource(Command.Resource))
-				{
-					string Value = ParseRequiredParam(RemainingArgs, ResourceArg);
-					CommandUrl.Replace("{" + ResourceArg + "}", Value);
-				}
-
-				// Add all the required and optional parameters
-				List<string> QueryParams = new List<string>();
-				if(Command.OptionalParams != null)
-				{
-					foreach(string OptionalParam in Command.OptionalParams)
-					{
-						string Value = ParseOptionalParam(RemainingArgs, OptionalParam);
-						if(Value != null)
-						{
-							QueryParams.Add(String.Format("{0}={1}", OptionalParam, Value));
-						}
-					}
-				}
-
-				// Append the parameters to the URL
-				for(int Idx = 0; Idx < QueryParams.Count; Idx++)
-				{
-					if(Idx == 0)
-					{
-						CommandUrl.Append("?");
-					}
-					else
-					{
-						CommandUrl.Append("&");
-					}
-					CommandUrl.Append(QueryParams[Idx]);
-				}
-
-				// Parse additional options for the message body
-				string BodyText = null;
-				if(Command.BodyType != null)
-				{
-					object BodyObject = ParseObject(RemainingArgs, Command.BodyType);
-					BodyText = new JavaScriptSerializer().Serialize(BodyObject);
-				}
-
-				// Make sure there are no unused arguments
-				if(RemainingArgs.Count > 0)
-				{
-					foreach(string RemainingArg in RemainingArgs)
-					{
-						Console.WriteLine("Unused argument: {0}", RemainingArg);
-					}
-					return false;
-				}
-
-				// Create the request
-				HttpWebRequest Request = (HttpWebRequest)WebRequest.Create(CommandUrl.ToString());
-				Request.ContentType = "application/json";
-				Request.Method = Command.Verb;
-				if (BodyText != null)
-				{
-					byte[] BodyData = Encoding.UTF8.GetBytes(BodyText);
-					using (Stream RequestStream = Request.GetRequestStream())
-					{
-						RequestStream.Write(BodyData, 0, BodyData.Length);
-					}
-				}
-
-				// Read the response
-				HttpWebResponse Response = (HttpWebResponse)Request.GetResponse();
-				Console.WriteLine("Response: {0} ({1})", (int)Response.StatusCode, Response.StatusDescription);
-				using (StreamReader ResponseReader = new StreamReader(Response.GetResponseStream(), Encoding.Default))
-				{
-					string ResponseContent = ResponseReader.ReadToEnd();
-					if(!String.IsNullOrEmpty(ResponseContent))
-					{
-						Console.WriteLine(Json.Format(ResponseContent));
-					}
-				}
+				Command.Exec(Arguments);
 				return true;
 			}
 			catch(CommandLineArgumentException Ex)
 			{
 				Console.WriteLine("{0}", Ex.Message);
 				Console.WriteLine();
-				PrintCommandHelp(Command);
+				PrintCommandHelp(CommandName, Command);
 				return false;
 			}
 		}
@@ -277,180 +110,17 @@ namespace MetadataTool
 			Console.WriteLine("Global Options:");
 			Console.WriteLine("  -Server=<url>     Specifies the endpoint for the metadata server (as http://hostname:port)");
 			Console.WriteLine();
-			Console.WriteLine("Available commands:");
-			foreach(Command Command in Commands)
-			{
-				Console.WriteLine("  {0}", Command.Name);
-			}
-		}
-
-		static void PrintCommandHelp(Command Command)
-		{
-			Console.WriteLine("Arguments for command '{0}':", Command.Name);
-
-			foreach(string ArgumentName in ParseArgumentNamesFromResource(Command.Resource))
-			{
-				Console.WriteLine("  -{0}=...", ArgumentName);
-			}
-			if(Command.OptionalParams != null)
-			{
-				foreach(string ArgumentName in Command.OptionalParams)
-				{
-					Console.WriteLine("  -{0}=... (optional)", ArgumentName);
-				}
-			}
-			if(Command.BodyType != null)
-			{
-				foreach(FieldInfo Field in Command.BodyType.GetFields(BindingFlags.Instance | BindingFlags.Public))
-				{
-					if(IsOptionalObjectField(Field))
-					{
-						Console.WriteLine("  -{0}={1} (optional)", Field.Name, GetFieldValueDescription(Field.FieldType));
-					}
-					else
-					{
-						Console.WriteLine("  -{0}={1}", Field.Name, GetFieldValueDescription(Field.FieldType));
-					}
-				}
-			}
-		}
-
-		static string GetFieldValueDescription(Type FieldType)
-		{
-			if(FieldType.IsEnum)
-			{
-				return String.Join("|", Enum.GetNames(FieldType));
-			}
-			else
-			{
-				return "...";
-			}
-		}
-
-		static IEnumerable<string> ParseArgumentNamesFromResource(string Resource)
-		{
-			for(int ArgIdx = Resource.IndexOf('{'); ArgIdx != -1; ArgIdx = Resource.IndexOf('{', ArgIdx))
-			{
-				int ArgEndIdx = Resource.IndexOf('}', ArgIdx + 1);
-				yield return Resource.Substring(ArgIdx + 1, ArgEndIdx - (ArgIdx + 1));
-				ArgIdx = ArgEndIdx + 1;
-			}
-		}
-
-		static object ParseObject(List<string> CommandArgs, Type ObjectType)
-		{
-			object Instance = Activator.CreateInstance(ObjectType);
-			foreach(FieldInfo Field in Instance.GetType().GetFields(BindingFlags.Instance | BindingFlags.Public))
+			Console.WriteLine("Http commands:");
+			foreach(string CommandName in CommandMap.Keys)
 			{
-				string Value = ParseOptionalParam(CommandArgs, Field.Name);
-				if(Value == null)
-				{
-					if(!IsOptionalObjectField(Field))
-					{
-						throw new CommandLineArgumentException(String.Format("Missing -{0}=... argument", Field.Name));
-					}
-				}
-				else
-				{
-					Type FieldType = Field.FieldType;
-					if(FieldType.IsGenericType && FieldType.GetGenericTypeDefinition() == typeof(Nullable<>))
-					{
-						FieldType = FieldType.GetGenericArguments()[0];
-					}
-
-					if(FieldType == typeof(string))
-					{
-						Field.SetValue(Instance, Value);
-					}
-					else if(FieldType == typeof(int))
-					{
-						int IntValue;
-						if(!int.TryParse(Value, out IntValue))
-						{
-							throw new CommandLineArgumentException(String.Format("Invalid value for '{0}' - {1} is not an integer.", Field.Name, Value));
-						}
-						Field.SetValue(Instance, IntValue);
-					}
-					else if(FieldType == typeof(bool))
-					{
-						bool BoolValue;
-						if(!bool.TryParse(Value, out BoolValue))
-						{
-							throw new CommandLineArgumentException(String.Format("Invalid value for '{0}' - {1} is not a boolean.", Field.Name, Value));
-						}
-						Field.SetValue(Instance, BoolValue);
-					}
-					else if(FieldType.IsEnum)
-					{
-						object EnumValue;
-						try
-						{
-							EnumValue = Enum.Parse(FieldType, Value, true);
-						}
-						catch(ArgumentException)
-						{
-							throw new CommandLineArgumentException(String.Format("Invalid value for '{0}' - should be {1}", Field.Name, String.Join("/", Enum.GetNames(FieldType))));
-						}
-						Field.SetValue(Instance, EnumValue);
-					}
-					else
-					{
-						throw new NotImplementedException(String.Format("Unsupported type '{0}'", FieldType));
-					}
-				}
-			}
-			return Instance;
-		}
-
-		static bool IsOptionalObjectField(FieldInfo Field)
-		{
-			if(Field.GetCustomAttribute<OptionalAttribute>() != null)
-			{
-				return true;
-			}
-			if(Field.FieldType.IsGenericType && Field.FieldType.GetGenericTypeDefinition() == typeof(Nullable<>))
-			{
-				return true;
+				Console.WriteLine("  {0}", CommandName);
 			}
-			return false;
 		}
 
-		static bool ParseOption(List<string> RemainingArgs, string Option)
+		static void PrintCommandHelp(string Name, CommandHandler Command)
 		{
-			for(int Idx = 0; Idx < RemainingArgs.Count; Idx++)
-			{
-				if(RemainingArgs[Idx].Equals(Option, StringComparison.OrdinalIgnoreCase))
-				{
-					RemainingArgs.RemoveAt(Idx);
-					return true;
-				}
-			}
-			return false;
-		}
-
-		static string ParseRequiredParam(List<string> RemainingArgs, string Name)
-		{
-			string Result = ParseOptionalParam(RemainingArgs, Name);
-			if(Result == null)
-			{
-				throw new CommandLineArgumentException(String.Format("Missing -{0}=... argument", Name));
-			}
-			return Result;
-		}
-
-		static string ParseOptionalParam(List<string> RemainingArgs, string Name)
-		{
-			for(int Idx = 0; Idx < RemainingArgs.Count; Idx++)
-			{
-				string RemainingArg = RemainingArgs[Idx];
-				if(RemainingArg.Length >= Name.Length + 2 && RemainingArg[0] == '-' && RemainingArg[Name.Length + 1] == '=' && String.Compare(RemainingArg, 1, Name, 0, Name.Length, StringComparison.OrdinalIgnoreCase) == 0)
-				{
-					string Value = RemainingArgs[Idx].Substring(Name.Length + 2);
-					RemainingArgs.RemoveAt(Idx);
-					return Value;
-				}
-			}
-			return null;
+			Console.WriteLine("Arguments for '{0}':", Name);
+			Command.Help();
 		}
 	}
 }
