Fixing copyright notices.

#rnx

diff --git a/Engine/Source/Programs/UnrealBuildTool/Configuration/ModuleRulesContext.cs b/Engine/Source/Programs/UnrealBuildTool/Configuration/ModuleRulesContext.cs
--- a/Engine/Source/Programs/UnrealBuildTool/Configuration/ModuleRulesContext.cs
+++ b/Engine/Source/Programs/UnrealBuildTool/Configuration/ModuleRulesContext.cs
@@ -1,3 +1,5 @@
+// Copyright 1998-2019 Epic Games, Inc. All Rights Reserved.
+
 using System;
 using System.Collections.Generic;
 using System.Linq;
diff --git a/Engine/Source/Programs/UnrealBuildTool/System/RulesScope.cs b/Engine/Source/Programs/UnrealBuildTool/System/RulesScope.cs
--- a/Engine/Source/Programs/UnrealBuildTool/System/RulesScope.cs
+++ b/Engine/Source/Programs/UnrealBuildTool/System/RulesScope.cs
@@ -1,3 +1,5 @@
+// Copyright 1998-2019 Epic Games, Inc. All Rights Reserved.
+
 using System;
 using System.Collections.Generic;
 using System.Linq;
diff --git a/Engine/Source/Programs/UnrealGameSync/MetadataServer/Controllers/IssuesController.cs b/Engine/Source/Programs/UnrealGameSync/MetadataServer/Controllers/IssuesController.cs
--- a/Engine/Source/Programs/UnrealGameSync/MetadataServer/Controllers/IssuesController.cs
+++ b/Engine/Source/Programs/UnrealGameSync/MetadataServer/Controllers/IssuesController.cs
@@ -1,4 +1,4 @@
-// Copyright 1998-2018 Epic Games, Inc. All Rights Reserved.
+// Copyright 1998-2019 Epic Games, Inc. All Rights Reserved.
 
 using Newtonsoft.Json.Linq;
 using System;
diff --git a/Engine/Source/Programs/UnrealGameSync/MetadataTool/CommandHandler.cs b/Engine/Source/Programs/UnrealGameSync/MetadataTool/CommandHandler.cs
--- a/Engine/Source/Programs/UnrealGameSync/MetadataTool/CommandHandler.cs
+++ b/Engine/Source/Programs/UnrealGameSync/MetadataTool/CommandHandler.cs
@@ -1,3 +1,5 @@
+// Copyright 1998-2019 Epic Games, Inc. All Rights Reserved.
+
 using System;
 using System.Collections.Generic;
 using System.Linq;
diff --git a/Engine/Source/Programs/UnrealGameSync/MetadataTool/CommandHandler_Http.cs b/Engine/Source/Programs/UnrealGameSync/MetadataTool/CommandHandler_Http.cs
--- a/Engine/Source/Programs/UnrealGameSync/MetadataTool/CommandHandler_Http.cs
+++ b/Engine/Source/Programs/UnrealGameSync/MetadataTool/CommandHandler_Http.cs
@@ -1,3 +1,5 @@
+// Copyright 1998-2019 Epic Games, Inc. All Rights Reserved.
+
 using System;
 using System.Collections.Generic;
 using System.IO;
diff --git a/Engine/Source/Programs/UnrealGameSync/MetadataTool/CommandTypes.cs b/Engine/Source/Programs/UnrealGameSync/MetadataTool/CommandTypes.cs
--- a/Engine/Source/Programs/UnrealGameSync/MetadataTool/CommandTypes.cs
+++ b/Engine/Source/Programs/UnrealGameSync/MetadataTool/CommandTypes.cs
@@ -1,3 +1,5 @@
+// Copyright 1998-2019 Epic Games, Inc. All Rights Reserved.
+
 using System;
 using System.Collections.Generic;
 using System.Linq;
diff --git a/Engine/Source/Programs/UnrealGameSync/UnrealGameSync/IssueMonitor.cs b/Engine/Source/Programs/UnrealGameSync/UnrealGameSync/IssueMonitor.cs
--- a/Engine/Source/Programs/UnrealGameSync/UnrealGameSync/IssueMonitor.cs
+++ b/Engine/Source/Programs/UnrealGameSync/UnrealGameSync/IssueMonitor.cs
@@ -1,3 +1,5 @@
+// Copyright 1998-2019 Epic Games, Inc. All Rights Reserved.
+
 using System;
 using System.Collections.Concurrent;
 using System.Collections.Generic;
