UGS: Add a window to browse through all the open build health issues.

diff --git a/Engine/Source/Programs/UnrealGameSync/UnrealGameSync/Controls/CustomListViewControl.cs b/Engine/Source/Programs/UnrealGameSync/UnrealGameSync/Controls/CustomListViewControl.cs
--- a/Engine/Source/Programs/UnrealGameSync/UnrealGameSync/Controls/CustomListViewControl.cs
+++ b/Engine/Source/Programs/UnrealGameSync/UnrealGameSync/Controls/CustomListViewControl.cs
@@ -362,6 +362,39 @@ namespace UnrealGameSync
 
 			base.OnMouseMove(e);
 		}
+		public void DrawText(Graphics Graphics, Rectangle Bounds, HorizontalAlignment TextAlign, Color TextColor, string Text)
+		{
+			TextFormatFlags Flags = TextFormatFlags.EndEllipsis | TextFormatFlags.SingleLine | TextFormatFlags.VerticalCenter | TextFormatFlags.NoPrefix;
+			if (TextAlign == HorizontalAlignment.Left)
+			{
+				Flags |= TextFormatFlags.Left;
+			}
+			else if (TextAlign == HorizontalAlignment.Center)
+			{
+				Flags |= TextFormatFlags.HorizontalCenter;
+			}
+			else
+			{
+				Flags |= TextFormatFlags.Right;
+			}
+			TextRenderer.DrawText(Graphics, Text, Font, Bounds, TextColor, Flags);
+		}
+
+		public void DrawIcon(Graphics Graphics, Rectangle Bounds, Rectangle Icon)
+		{
+			float DpiScaleX = Graphics.DpiX / 96.0f;
+			float DpiScaleY = Graphics.DpiY / 96.0f;
+
+			float IconX = Bounds.Left + (Bounds.Width - 16 * DpiScaleX) / 2;
+			float IconY = Bounds.Top + (Bounds.Height - 16 * DpiScaleY) / 2;
+
+			Graphics.DrawImage(Properties.Resources.Icons, IconX, IconY, Icon, GraphicsUnit.Pixel);
+		}
+
+		public void DrawNormalSubItem(DrawListViewSubItemEventArgs e)
+		{
+			DrawText(e.Graphics, e.SubItem.Bounds, Columns[e.ColumnIndex].TextAlign, SystemColors.WindowText, e.SubItem.Text);
+		}
 
 		public void DrawCustomSubItem(Graphics Graphics, ListViewItem.ListViewSubItem SubItem)
 		{
diff --git a/Engine/Source/Programs/UnrealGameSync/UnrealGameSync/Controls/WorkspaceControl.cs b/Engine/Source/Programs/UnrealGameSync/UnrealGameSync/Controls/WorkspaceControl.cs
--- a/Engine/Source/Programs/UnrealGameSync/UnrealGameSync/Controls/WorkspaceControl.cs
+++ b/Engine/Source/Programs/UnrealGameSync/UnrealGameSync/Controls/WorkspaceControl.cs
@@ -2986,19 +2986,21 @@ namespace UnrealGameSync
 
 		private void ShowBuildHealthMenu(Rectangle Bounds)
 		{
-			while(BuildHealthContextMenu.Items[0] != BuildHealthContextMenu_Separator)
+			int MinSeparatorIdx = BuildHealthContextMenu.Items.IndexOf(BuildHealthContextMenu_MinSeparator);
+
+			while(BuildHealthContextMenu.Items[MinSeparatorIdx + 1] != BuildHealthContextMenu_MaxSeparator)
 			{
-				BuildHealthContextMenu.Items.RemoveAt(0);
+				BuildHealthContextMenu.Items.RemoveAt(MinSeparatorIdx + 1);
 			}
 
 			List<IssueData> Issues = IssueMonitor.GetIssues();
 			if(Issues.Count == 0)
 			{
-				BuildHealthContextMenu_Separator.Visible = false;
+				BuildHealthContextMenu_MaxSeparator.Visible = false;
 			}
 			else
 			{
-				BuildHealthContextMenu_Separator.Visible = true;
+				BuildHealthContextMenu_MaxSeparator.Visible = true;
 				for(int Idx = 0; Idx < Issues.Count; Idx++)
 				{
 					IssueData Issue = Issues[Idx];
@@ -3030,7 +3032,7 @@ namespace UnrealGameSync
 						Item.Font = new Font(Item.Font, FontStyle.Bold);
 					}
 					Item.Click += (s, e) => { BuildHealthContextMenu_Issue_Click(Issue); };
-					BuildHealthContextMenu.Items.Insert(Idx, Item);
+					BuildHealthContextMenu.Items.Insert(MinSeparatorIdx + Idx + 1, Item);
 				}
 			}
 
@@ -3085,6 +3087,11 @@ namespace UnrealGameSync
 			IssueDetailsWindow.ShowModal(this, IssueMonitor, Workspace.Perforce.ServerAndPort, Workspace.Perforce.UserName, Issue, Log, StreamName);
 		}
 
+		private void BuildHealthContextMenu_Browse_Click(object sender, EventArgs e)
+		{
+			IssueBrowserWindow.ShowModal(this, IssueMonitor, Workspace.Perforce.ServerAndPort, Workspace.Perforce.UserName, Log, StreamName);
+		}
+
 		private void BuildHealthContextMenu_Settings_Click(object sender, EventArgs e)
 		{
 			IssueSettingsWindow IssueSettings = new IssueSettingsWindow(Settings);
diff --git a/Engine/Source/Programs/UnrealGameSync/UnrealGameSync/Forms/IssueBrowserWindow.cs b/Engine/Source/Programs/UnrealGameSync/UnrealGameSync/Forms/IssueBrowserWindow.cs
new file mode 100644
--- /dev/null
+++ b/Engine/Source/Programs/UnrealGameSync/UnrealGameSync/Forms/IssueBrowserWindow.cs
@@ -0,0 +1,240 @@
+// Copyright 1998-2019 Epic Games, Inc. All Rights Reserved.
+
+using System;
+using System.Collections.Generic;
+using System.ComponentModel;
+using System.Data;
+using System.Drawing;
+using System.IO;
+using System.Linq;
+using System.Text;
+using System.Threading.Tasks;
+using System.Windows.Forms;
+
+namespace UnrealGameSync
+{
+	partial class IssueBrowserWindow : Form
+	{
+		IssueMonitor IssueMonitor;
+		string ServerAndPort;
+		string UserName;
+		TextWriter Log;
+		string CurrentStream;
+		int MaxResults = 0;
+
+		public IssueBrowserWindow(IssueMonitor IssueMonitor, string ServerAndPort, string UserName, TextWriter Log, string CurrentStream)
+		{
+			this.IssueMonitor = IssueMonitor;
+			this.ServerAndPort = ServerAndPort;
+			this.UserName = UserName;
+			this.Log = Log;
+			this.CurrentStream = CurrentStream;
+
+			InitializeComponent();
+
+			System.Reflection.PropertyInfo DoubleBufferedProperty = typeof(Control).GetProperty("DoubleBuffered", System.Reflection.BindingFlags.NonPublic | System.Reflection.BindingFlags.Instance);
+			DoubleBufferedProperty.SetValue(IssueListView, true, null);
+		}
+
+		private void IssueBrowserWindow_Load(object sender, EventArgs e)
+		{
+			FetchMoreResults();
+			UpdateDetailsBtn();
+		}
+
+		class QueryIssuesTask : IModalTask
+		{
+			string ApiUrl;
+			int MaxResults;
+			public List<IssueData> Issues;
+
+			public QueryIssuesTask(string ApiUrl, int MaxResults)
+			{
+				this.ApiUrl = ApiUrl;
+				this.MaxResults = MaxResults;
+			}
+
+			public bool Run(out string ErrorMessage)
+			{
+				Issues = RESTApi.GET<List<IssueData>>(ApiUrl, String.Format("issues?includeresolved=true&maxresults={0}", MaxResults));
+
+				ErrorMessage = null;
+				return true;
+			}
+		}
+
+		private void FetchMoreResults()
+		{
+			int NewMaxResults = MaxResults + 100;
+			QueryIssuesTask Task = new QueryIssuesTask(IssueMonitor.ApiUrl, NewMaxResults);
+
+			// Execute the task
+			string ErrorMessage;
+			ModalTaskResult Result = ModalTask.Execute(this, Task, "Querying issues", "Fetching data, please wait...", out ErrorMessage);
+			if(Result != ModalTaskResult.Succeeded)
+			{
+				if(Result != ModalTaskResult.Aborted)
+				{
+					StatusLabel.Text = "Unable to fetch issues.";
+					MessageBox.Show(ErrorMessage);
+				}
+				return;
+			}
+
+			// Get the time at midnight
+			DateTime Now = DateTime.Now;
+			DateTime Midnight = (Now - Now.TimeOfDay).ToUniversalTime();
+
+			// Fetch the new issues
+			IssueListView.BeginUpdate();
+			IssueListView.Items.Clear();
+			foreach(IssueData Issue in Task.Issues)
+			{
+				ListViewItem Item = new ListViewItem("");
+				Item.SubItems.Add(Issue.Id.ToString());
+				Item.SubItems.Add(FormatIssueDateTime(Issue.CreatedAt, Midnight));
+				Item.SubItems.Add(Issue.ResolvedAt.HasValue? FormatIssueDateTime(Issue.ResolvedAt.Value, Midnight) : "Unresolved");
+				Item.SubItems.Add(Issue.Summary);
+				Item.Tag = Issue;
+				IssueListView.Items.Add(Item);
+			}
+			IssueListView.EndUpdate();
+
+			// Update the maximum number of results
+			StatusLabel.Text = String.Format("Showing {0} results.", Task.Issues.Count);
+			MaxResults = NewMaxResults;
+		}
+
+		static string FormatIssueDateTime(DateTime DateTime, DateTime Midnight)
+		{
+			if(DateTime > Midnight)
+			{
+				return DateTime.ToShortTimeString();
+			}
+			else
+			{
+				return DateTime.ToShortDateString();
+			}
+		}
+
+		public static void ShowModal(IWin32Window Owner, IssueMonitor IssueMonitor, string ServerAndPort, string UserName, TextWriter Log, string CurrentStream)
+		{
+			IssueBrowserWindow Window = new IssueBrowserWindow(IssueMonitor, ServerAndPort, UserName, Log, CurrentStream);
+			Window.ShowDialog(Owner);
+		}
+
+		private void FetchMoreResultsLinkLabel_LinkClicked(object sender, LinkLabelLinkClickedEventArgs e)
+		{
+			FetchMoreResults();
+		}
+
+		private void OkBtn_Click(object sender, EventArgs e)
+		{
+			DialogResult = DialogResult.OK;
+			Close();
+		}
+
+		private void IssueListView_SelectedIndexChanged(object sender, EventArgs e)
+		{
+			UpdateDetailsBtn();
+		}
+
+		void UpdateDetailsBtn()
+		{
+			DetailsBtn.Enabled = (IssueListView.SelectedItems.Count != 0);
+		}
+
+		private void IssueListView_MouseDoubleClick(object sender, MouseEventArgs e)
+		{
+			ListViewHitTestInfo HitTest = IssueListView.HitTest(e.Location);
+			if(HitTest.Item != null)
+			{
+				IssueData Issue = (IssueData)HitTest.Item.Tag;
+				ShowIssue(Issue);
+			}
+		}
+
+		private void DetailsBtn_Click(object sender, EventArgs e)
+		{
+			foreach(ListViewItem Item in IssueListView.Items)
+			{
+				IssueData Issue = (IssueData)Item.Tag;
+				ShowIssue(Issue);
+			}
+		}
+
+		private void ShowIssue(IssueData Issue)
+		{
+			Issue.Builds = RESTApi.GET<List<IssueBuildData>>(IssueMonitor.ApiUrl, String.Format("issues/{0}/builds", Issue.Id));
+			IssueDetailsWindow.ShowModal(this, IssueMonitor, ServerAndPort, UserName, Issue, Log, CurrentStream);
+		}
+
+		private void IssueListView_DrawColumnHeader(object sender, DrawListViewColumnHeaderEventArgs e)
+		{
+			e.DrawDefault = true;
+		}
+
+		private void IssueListView_DrawItem(object sender, DrawListViewItemEventArgs e)
+		{
+			if (e.Item.Selected)
+			{
+				IssueListView.DrawSelectedBackground(e.Graphics, e.Bounds);
+			}
+			else if (e.ItemIndex == IssueListView.HoverItem)
+			{
+				IssueListView.DrawTrackedBackground(e.Graphics, e.Bounds);
+			}
+			else if (e.Item.Tag is IssueData)
+			{
+				IssueData Issue = (IssueData)e.Item.Tag;
+
+				Color BackgroundColor;
+				if (Issue.ResolvedAt.HasValue)
+				{
+					BackgroundColor = SystemColors.Window;//Color.FromArgb(248, 254, 246);
+				}
+				else if(Issue.FixChange > 0)
+				{
+					BackgroundColor = Color.FromArgb(245, 245, 245);
+				}
+				else
+				{
+					BackgroundColor = Color.FromArgb(254, 248, 246);
+				}
+
+				using (SolidBrush Brush = new SolidBrush(BackgroundColor))
+				{
+					e.Graphics.FillRectangle(Brush, e.Bounds);
+				}
+			}
+			else
+			{
+				IssueListView.DrawDefaultBackground(e.Graphics, e.Bounds);
+			}
+		}
+
+		private void IssueListView_DrawSubItem(object sender, DrawListViewSubItemEventArgs e)
+		{
+			IssueData Issue = (IssueData)e.Item.Tag;
+			if (e.ColumnIndex == IconHeader.Index)
+			{
+				if(!Issue.ResolvedAt.HasValue)
+				{
+					IssueListView.DrawIcon(e.Graphics, e.Bounds, WorkspaceControl.BadBuildIcon);
+				}
+				else if (Issue.FixChange > 0)
+				{
+					IssueListView.DrawIcon(e.Graphics, e.Bounds, WorkspaceControl.MixedBuildIcon);
+				}
+				else
+				{
+					IssueListView.DrawIcon(e.Graphics, e.Bounds, WorkspaceControl.GoodBuildIcon);
+				}
+			}
+			else
+			{
+				IssueListView.DrawNormalSubItem(e);
+			}
+		}
+	}
+}
diff --git a/Engine/Source/Programs/UnrealGameSync/UnrealGameSync/Forms/IssueDetailsWindow.cs b/Engine/Source/Programs/UnrealGameSync/UnrealGameSync/Forms/IssueDetailsWindow.cs
--- a/Engine/Source/Programs/UnrealGameSync/UnrealGameSync/Forms/IssueDetailsWindow.cs
+++ b/Engine/Source/Programs/UnrealGameSync/UnrealGameSync/Forms/IssueDetailsWindow.cs
@@ -255,6 +255,7 @@ namespace UnrealGameSync
 
 			InitializeComponent();
 
+			this.Text = String.Format("Issue {0}", Issue.Id);
 			this.StatusElementResources = new StatusElementResources(BuildListView.Font);
 			base.Disposed += IssueDetailsWindow_Disposed;
 
@@ -903,19 +904,19 @@ namespace UnrealGameSync
 				{
 					if(BuildData.Outcome == IssueBuildOutcome.Success)
 					{
-						DrawIcon(e.Graphics, e.Bounds, WorkspaceControl.GoodBuildIcon);
+						BuildListView.DrawIcon(e.Graphics, e.Bounds, WorkspaceControl.GoodBuildIcon);
 					}
 					else if(BuildData.Outcome == IssueBuildOutcome.Warning)
 					{
-						DrawIcon(e.Graphics, e.Bounds, WorkspaceControl.MixedBuildIcon);
+						BuildListView.DrawIcon(e.Graphics, e.Bounds, WorkspaceControl.MixedBuildIcon);
 					}
 					else if(BuildData.Outcome == IssueBuildOutcome.Error)
 					{
-						DrawIcon(e.Graphics, e.Bounds, WorkspaceControl.BadBuildIcon);
+						BuildListView.DrawIcon(e.Graphics, e.Bounds, WorkspaceControl.BadBuildIcon);
 					}
 					else
 					{
-						DrawIcon(e.Graphics, e.Bounds, WorkspaceControl.DefaultBuildIcon);
+						BuildListView.DrawIcon(e.Graphics, e.Bounds, WorkspaceControl.DefaultBuildIcon);
 					}
 				}
 				else if(e.ColumnIndex == ChangeHeader.Index)
@@ -968,17 +969,6 @@ namespace UnrealGameSync
 			}
 		}
 
-		private void DrawIcon(Graphics Graphics, Rectangle Bounds, Rectangle Icon)
-		{
-			float DpiScaleX = Graphics.DpiX / 96.0f;
-			float DpiScaleY = Graphics.DpiY / 96.0f;
-
-			float IconX = Bounds.Left + (Bounds.Width - 16 * DpiScaleX) / 2;
-			float IconY = Bounds.Top + (Bounds.Height - 16 * DpiScaleY) / 2;
-
-			Graphics.DrawImage(Properties.Resources.Icons, IconX, IconY, Icon, GraphicsUnit.Pixel);
-		}
-
 		private void BuildListView_FontChanged(object sender, EventArgs e)
 		{
 			if(BoldFont != null)
diff --git a/Engine/Source/Programs/UnrealGameSync/UnrealGameSync/UnrealGameSync.csproj b/Engine/Source/Programs/UnrealGameSync/UnrealGameSync/UnrealGameSync.csproj
--- a/Engine/Source/Programs/UnrealGameSync/UnrealGameSync/UnrealGameSync.csproj
+++ b/Engine/Source/Programs/UnrealGameSync/UnrealGameSync/UnrealGameSync.csproj
@@ -182,6 +182,12 @@
     <Compile Include="Forms\IssueFixedWindow.Designer.cs">
       <DependentUpon>IssueFixedWindow.cs</DependentUpon>
     </Compile>
+    <Compile Include="Forms\IssueBrowserWindow.cs">
+      <SubType>Form</SubType>
+    </Compile>
+    <Compile Include="Forms\IssueBrowserWindow.Designer.cs">
+      <DependentUpon>IssueBrowserWindow.cs</DependentUpon>
+    </Compile>
     <Compile Include="Forms\IssueSettingsWindow.cs">
       <SubType>Form</SubType>
     </Compile>
@@ -403,6 +409,9 @@
     <EmbeddedResource Include="Forms\IssueFixedWindow.resx">
       <DependentUpon>IssueFixedWindow.cs</DependentUpon>
     </EmbeddedResource>
+    <EmbeddedResource Include="Forms\IssueBrowserWindow.resx">
+      <DependentUpon>IssueBrowserWindow.cs</DependentUpon>
+    </EmbeddedResource>
     <EmbeddedResource Include="Forms\IssueSettingsWindow.resx">
       <DependentUpon>IssueSettingsWindow.cs</DependentUpon>
     </EmbeddedResource>
@@ -515,7 +524,7 @@
       <ItemGroup>
         <Compile Include="NotForLicensees\DeploymentSettingsInternal.cs" />
       </ItemGroup>
-	</When>
+    </When>
   </Choose>
   <ItemGroup>
     <EmbeddedResource Include="Resources\Icons.png" />
