UGS: Additional server functionality for filtering build issues.

#rnx

diff --git a/Engine/Source/Programs/UnrealGameSync/MetadataServer/Connectors/SqlConnector.cs b/Engine/Source/Programs/UnrealGameSync/MetadataServer/Connectors/SqlConnector.cs
--- a/Engine/Source/Programs/UnrealGameSync/MetadataServer/Connectors/SqlConnector.cs
+++ b/Engine/Source/Programs/UnrealGameSync/MetadataServer/Connectors/SqlConnector.cs
@@ -372,7 +372,7 @@ namespace MetadataServer.Connectors
 
 		public static IssueData GetIssue(long IssueId)
 		{
-			List<IssueData> Issues = GetIssuesInternal(IssueId, null);
+			List<IssueData> Issues = GetIssuesInternal(IssueId, null, true, -1);
 			if(Issues.Count == 0)
 			{
 				return null;
@@ -383,17 +383,17 @@ namespace MetadataServer.Connectors
 			}
 		}
 
-		public static List<IssueData> GetIssues()
+		public static List<IssueData> GetIssues(bool IncludeResolved, int NumResults)
 		{
-			return GetIssuesInternal(-1, null);
+			return GetIssuesInternal(-1, null, IncludeResolved, NumResults);
 		}
 
 		public static List<IssueData> GetIssues(string UserName)
 		{
-			return GetIssuesInternal(-1, UserName);
+			return GetIssuesInternal(-1, UserName, false, -1);
 		}
 
-		private static List<IssueData> GetIssuesInternal(long IssueId, string UserName)
+		private static List<IssueData> GetIssuesInternal(long IssueId, string UserName, bool IncludeResolved, int NumResults)
 		{
 			List<IssueData> Issues = new List<IssueData>();
 			using (SQLiteConnection Connection = new SQLiteConnection(SqlConnector.ConnectionString))
@@ -424,10 +424,14 @@ namespace MetadataServer.Connectors
 				{
 					CommandBuilder.Append(" WHERE Issues.Id = @IssueId");
 				}
-				else
+				else if(!IncludeResolved)
 				{
 					CommandBuilder.Append(" WHERE Issues.ResolvedAt IS NULL");
 				}
+				if(NumResults > 0)
+				{
+					CommandBuilder.AppendFormat(" ORDER BY Issues.Id DESC LIMIT {0}", NumResults);
+				}
 
 				using (SQLiteCommand Command = new SQLiteCommand(CommandBuilder.ToString(), Connection))
 				{
diff --git a/Engine/Source/Programs/UnrealGameSync/MetadataServer/Controllers/IssuesController.cs b/Engine/Source/Programs/UnrealGameSync/MetadataServer/Controllers/IssuesController.cs
--- a/Engine/Source/Programs/UnrealGameSync/MetadataServer/Controllers/IssuesController.cs
+++ b/Engine/Source/Programs/UnrealGameSync/MetadataServer/Controllers/IssuesController.cs
@@ -16,9 +16,9 @@ namespace MetadataServer.Controllers
 	public class IssuesController : ApiController
 	{
 		[HttpGet]
-		public List<IssueData> Get()
+		public List<IssueData> Get(bool IncludeResolved = false, int MaxResults = -1)
 		{
-			return SqlConnector.GetIssues();
+			return SqlConnector.GetIssues(IncludeResolved, MaxResults);
 		}
 
 		[HttpGet]
