UGS: Fix command name for "BuildHealth".

#rnx

diff --git a/Engine/Source/Programs/UnrealGameSync/MetadataTool/CommandHandler_BuildHealth.cs b/Engine/Source/Programs/UnrealGameSync/MetadataTool/CommandHandler_BuildHealth.cs
--- a/Engine/Source/Programs/UnrealGameSync/MetadataTool/CommandHandler_BuildHealth.cs
+++ b/Engine/Source/Programs/UnrealGameSync/MetadataTool/CommandHandler_BuildHealth.cs
@@ -54,7 +54,7 @@ namespace MetadataTool
 		/// Constructor
 		/// </summary>
 		public CommandHandler_BuildHealth()
-			: base("TickBuildHealth")
+			: base("BuildHealth")
 		{
 		}
 
