UGS: Add a rule in the project file to output MetadataTool to a different directory for internal Epic builds.

#rnx

diff --git a/Engine/Source/Programs/UnrealGameSync/MetadataTool/MetadataTool.csproj b/Engine/Source/Programs/UnrealGameSync/MetadataTool/MetadataTool.csproj
--- a/Engine/Source/Programs/UnrealGameSync/MetadataTool/MetadataTool.csproj
+++ b/Engine/Source/Programs/UnrealGameSync/MetadataTool/MetadataTool.csproj
@@ -35,6 +35,9 @@
     <WarningLevel>4</WarningLevel>
     <NoWarn>CS0649</NoWarn>
   </PropertyGroup>
+  <PropertyGroup Condition="Exists('..\..\..\..\Build\NotForLicensees\EpicInternal.txt')">
+    <OutputPath>..\..\..\..\Build\NotForLicensees\BuildFarm\MetadataTool\</OutputPath>
+  </PropertyGroup>
   <ItemGroup>
     <Reference Include="System" />
     <Reference Include="System.Core" />
