UGS: Create an extensible set of pattern matchers, making it easier to add handlers for different types of errors.

#rnx

diff --git a/Engine/Source/Programs/UnrealGameSync/MetadataTool/BuildHealth/Matchers/CompilePatternMatcher.cs b/Engine/Source/Programs/UnrealGameSync/MetadataTool/BuildHealth/Matchers/CompilePatternMatcher.cs
new file mode 100644
--- /dev/null
+++ b/Engine/Source/Programs/UnrealGameSync/MetadataTool/BuildHealth/Matchers/CompilePatternMatcher.cs
@@ -0,0 +1,69 @@
+using System;
+using System.Collections.Generic;
+using System.Linq;
+using System.Text;
+using System.Text.RegularExpressions;
+using System.Threading.Tasks;
+
+namespace MetadataTool
+{
+	class CompilePatternMatcher : PatternMatcher
+	{
+		static string[] SourceFileExtensions =
+		{
+			".cpp",
+			".h",
+			".cc",
+			".hh",
+			".m",
+			".mm",
+			".rc",
+			".inl",
+			".inc"
+		};
+
+		public override string Category => "Compile";
+
+		public override bool TryMatch(InputJob Job, InputJobStep JobStep, InputDiagnostic Diagnostic, List<TrackedIssueFingerprint> Fingerprints)
+		{
+			// Find any files in compiler output format
+			HashSet<string> SourceFileNames = new HashSet<string>(StringComparer.OrdinalIgnoreCase);
+			foreach (Match FileMatch in Regex.Matches(Diagnostic.Message, @"^\s*(?:In file included from\s*)?((?:[A-Za-z]:)?[^\s(:]+)[\(:]\d", RegexOptions.Multiline))
+			{
+				if (FileMatch.Success)
+				{
+					string FileName = GetNormalizedFileName(FileMatch.Groups[1].Value, JobStep.BaseDirectory);
+					if (SourceFileExtensions.Any(x => FileName.EndsWith(x, StringComparison.OrdinalIgnoreCase)))
+					{
+						SourceFileNames.Add(FileName);
+					}
+				}
+			}
+
+			// If we found any source files, create a diagnostic category for them
+			if (SourceFileNames.Count > 0)
+			{
+				TrackedIssueFingerprint Fingerprint = new TrackedIssueFingerprint(Category, Job.Change);
+				Fingerprint.FileNames.UnionWith(SourceFileNames);
+				Fingerprints.Add(Fingerprint);
+				return true;
+			}
+
+			// Otherwise pass
+			return false;
+		}
+
+		public override string GetSummary(TrackedIssue Issue)
+		{
+			SortedSet<string> ShortFileNames = Issue.Fingerprint.GetSourceFileNames();
+			if (ShortFileNames.Count == 0)
+			{
+				return "Compile errors";
+			}
+			else
+			{
+				return String.Format("Compile errors in {0}", String.Join(", ", ShortFileNames));
+			}
+		}
+	}
+}
diff --git a/Engine/Source/Programs/UnrealGameSync/MetadataTool/BuildHealth/PatternMatcher.cs b/Engine/Source/Programs/UnrealGameSync/MetadataTool/BuildHealth/PatternMatcher.cs
new file mode 100644
--- /dev/null
+++ b/Engine/Source/Programs/UnrealGameSync/MetadataTool/BuildHealth/PatternMatcher.cs
@@ -0,0 +1,170 @@
+using System;
+using System.Collections.Generic;
+using System.Linq;
+using System.Text;
+using System.Text.RegularExpressions;
+using System.Threading.Tasks;
+using Tools.DotNETCommon.Perforce;
+
+namespace MetadataTool
+{
+	/// <summary>
+	/// Class that implements a pattern matcher for a particular class of errors
+	/// </summary>
+	abstract class PatternMatcher
+	{
+		/// <summary>
+		/// The category name
+		/// </summary>
+		public abstract string Category
+		{
+			get;
+		}
+
+		/// <summary>
+		/// Creates fingerprints from any matching diagnostics
+		/// </summary>
+		/// <param name="Job">The job that was run</param>
+		/// <param name="JobStep">The job step that was run</param>
+		/// <param name="Diagnostics">List of diagnostics that were produced by the build. Items should be removed from this list if they match.</param>
+		/// <param name="Fingerprints">List which receives all the matched fingerprints.</param>
+		public virtual void Match(InputJob Job, InputJobStep JobStep, List<InputDiagnostic> Diagnostics, List<TrackedIssueFingerprint> Fingerprints)
+		{
+			for (int Idx = 0; Idx < Diagnostics.Count; Idx++)
+			{
+				InputDiagnostic Diagnostic = Diagnostics[Idx];
+				if(TryMatch(Job, JobStep, Diagnostic, Fingerprints))
+				{
+					Diagnostics.RemoveAt(Idx);
+					Idx--;
+				}
+			}
+		}
+
+		/// <summary>
+		/// Tries to create a fingerprint from an individual diagnostic.
+		/// </summary>
+		/// <param name="Job">The job that was run</param>
+		/// <param name="JobStep">The job step that was run</param>
+		/// <param name="Diagnostic">A diagnostic from the given job step</param>
+		/// <param name="Fingerprints">List which receives all the matched fingerprints.</param>
+		/// <returns>True if this diagnostic should be removed (usually because a fingerprint was created)</returns>
+		public abstract bool TryMatch(InputJob Job, InputJobStep JobStep, InputDiagnostic Diagnostic, List<TrackedIssueFingerprint> Fingerprints);
+
+		/// <summary>
+		/// Attempts to merge another fingerprint into this one
+		/// </summary>
+		/// <param name="Other">The issue to merge into this one</param>
+		public virtual bool TryMerge(TrackedIssueFingerprint Source, TrackedIssueFingerprint Target)
+		{
+			// Check that a filename or message matches
+			if (Source.InitialChange != Target.InitialChange)
+			{
+				if (!Source.FileNames.Any(x => Target.FileNames.Contains(x)) && !Source.Messages.Any(x => Target.Messages.Contains(x)))
+				{
+					return false;
+				}
+			}
+
+			// Allow merging
+			Target.FileNames.UnionWith(Source.FileNames);
+			Target.Messages.UnionWith(Source.Messages);
+			return true;
+		}
+
+		/// <summary>
+		/// Filters all the likely causers from the list of changes since an issue was created
+		/// </summary>
+		/// <param name="Perforce">The perforce connection</param>
+		/// <param name="Fingerprint">Fingerprint for the issue</param>
+		/// <param name="Changes">List of changes since the issue first occurred.</param>
+		/// <returns>List of changes which are causers for the issue</returns>
+		public virtual List<ChangeInfo> FindCausers(PerforceConnection Perforce, TrackedIssueFingerprint Fingerprint, List<ChangeInfo> Changes)
+		{
+			SortedSet<string> FileNamesWithoutPath = Fingerprint.GetFileNamesWithoutPath();
+
+			List<ChangeInfo> Causers = new List<ChangeInfo>();
+			foreach (ChangeInfo Change in Changes)
+			{
+				DescribeRecord Description = Perforce.Describe(Change.Record.Number).Data;
+				if (ContainsFileNames(Description, FileNamesWithoutPath))
+				{
+					Causers.Add(Change);
+				}
+			}
+
+			if(Causers.Count > 0)
+			{
+				return Causers;
+			}
+			else
+			{
+				return Changes;
+			}
+		}
+
+		/// <summary>
+		/// Determines if this change is a likely causer for an issue
+		/// </summary>
+		/// <param name="DescribeRecord">The change describe record</param>
+		/// <param name="Fingerprint">Fingerprint for the issue</param>
+		/// <returns>True if the change is a likely culprit</returns>
+		protected static bool ContainsFileNames(DescribeRecord DescribeRecord, SortedSet<string> FileNamesWithoutPath)
+		{
+			foreach (DescribeFileRecord File in DescribeRecord.Files)
+			{
+				int Idx = File.DepotFile.LastIndexOf('/');
+				if (Idx != -1)
+				{
+					string FileName = File.DepotFile.Substring(Idx + 1);
+					if (FileNamesWithoutPath.Contains(FileName))
+					{
+						return true;
+					}
+				}
+			}
+			return false;
+		}
+
+		/// <summary>
+		/// Gets the summary for an issue
+		/// </summary>
+		/// <param name="Issue">The issue to get a summary for</param>
+		/// <returns>Summary text for the given issue</returns>
+		public abstract string GetSummary(TrackedIssue Issue);
+
+		/// <summary>
+		/// Normalizes a filename to a path within the workspace
+		/// </summary>
+		/// <param name="FileName">Filename to normalize</param>
+		/// <param name="BaseDirectory">Base directory containing the workspace</param>
+		/// <returns>Normalized filename</returns>
+		protected string GetNormalizedFileName(string FileName, string BaseDirectory)
+		{
+			string NormalizedFileName = FileName.Replace('\\', '/');
+			if (!String.IsNullOrEmpty(BaseDirectory))
+			{
+				// Normalize the expected base directory for errors in this build, and attempt to strip it from the file name
+				string NormalizedBaseDirectory = BaseDirectory;
+				if (NormalizedBaseDirectory != null && NormalizedBaseDirectory.Length > 0)
+				{
+					NormalizedBaseDirectory = NormalizedBaseDirectory.Replace('\\', '/').TrimEnd('/') + "/";
+				}
+				if (NormalizedFileName.StartsWith(NormalizedBaseDirectory, StringComparison.OrdinalIgnoreCase))
+				{
+					NormalizedFileName = NormalizedFileName.Substring(NormalizedBaseDirectory.Length);
+				}
+			}
+			else
+			{
+				// Try to match anything under a 'Sync' folder.
+				Match FallbackRegex = Regex.Match(NormalizedFileName, "/Sync/(.*)");
+				if (FallbackRegex.Success)
+				{
+					NormalizedFileName = FallbackRegex.Groups[1].Value;
+				}
+			}
+			return NormalizedFileName;
+		}
+	}
+}
diff --git a/Engine/Source/Programs/UnrealGameSync/MetadataTool/BuildHealth/TrackedIssue.cs b/Engine/Source/Programs/UnrealGameSync/MetadataTool/BuildHealth/TrackedIssue.cs
--- a/Engine/Source/Programs/UnrealGameSync/MetadataTool/BuildHealth/TrackedIssue.cs
+++ b/Engine/Source/Programs/UnrealGameSync/MetadataTool/BuildHealth/TrackedIssue.cs
@@ -82,112 +82,11 @@ namespace MetadataTool
 			return Streams.Values.All(x => x.NextSuccessfulBuild != null);
 		}
 
-		/// <summary>
-		/// Creates a summary for this issue
-		/// </summary>
-		/// <returns>Summary for this issue</returns>
-		public string GetSummary()
-		{
-			StringBuilder Result = new StringBuilder();
-
-			SortedSet<string> StreamNames = new SortedSet<string>(StringComparer.OrdinalIgnoreCase);
-			foreach(string StreamName in Streams.Keys)
-			{
-				string TrimStreamName = StreamName.TrimEnd('/');
-
-				int Idx = TrimStreamName.LastIndexOf('/');
-				if(Idx != -1)
-				{
-					TrimStreamName = TrimStreamName.Substring(Idx + 1);
-				}
-
-				StreamNames.Add(TrimStreamName);
-			}
-			Result.AppendFormat("[{0}] ", String.Join("/", StreamNames));
-
-			if (Fingerprint.Category == TrackedIssueFingerprintCategory.Compile)
-			{
-				SortedSet<string> ShortFileNames = GetSourceFileNames();
-				if (ShortFileNames.Count == 0)
-				{
-					Result.Append("Compile errors");
-				}
-				else
-				{
-					Result.AppendFormat("Compile errors in {0}", String.Join(", ", ShortFileNames));
-				}
-			}
-			else
-			{
-				SortedSet<string> StepNames = GetStepNames();
-				Result.AppendFormat("Errors in {0}", String.Join(", ", StepNames));
-			}
-
-			const int MaxLength = 128;
-			if (Result.Length > MaxLength)
-			{
-				Result.Remove(MaxLength, Result.Length - MaxLength);
-				Result.Append("...");
-			}
-
-			return Result.ToString();
-
-			throw new NotImplementedException();
-		}
-
-		/// <summary>
-		/// Gets a set of unique source file names that relate to this issue
-		/// </summary>
-		/// <returns>Set of source file names</returns>
-		private SortedSet<string> GetSourceFileNames()
-		{
-			SortedSet<string> ShortFileNames = new SortedSet<string>(StringComparer.OrdinalIgnoreCase);
-			foreach (string FileName in Fingerprint.FileNames)
-			{
-				int Idx = FileName.LastIndexOfAny(new char[] { '/', '\\' });
-				if (Idx != -1)
-				{
-					string ShortFileName = FileName.Substring(Idx + 1);
-					if (!ShortFileName.StartsWith("Module.", StringComparison.OrdinalIgnoreCase))
-					{
-						ShortFileNames.Add(ShortFileName);
-					}
-				}
-			}
-			return ShortFileNames;
-		}
-
-		/// <summary>
-		/// Gets a set of unique asset filenames that relate to this issue
-		/// </summary>
-		/// <returns>Set of asset names</returns>
-		private SortedSet<string> GetAssetNames()
-		{
-			SortedSet<string> ShortFileNames = new SortedSet<string>(StringComparer.OrdinalIgnoreCase);
-			foreach (string FileName in Fingerprint.FileNames)
-			{
-				int Idx = FileName.LastIndexOfAny(new char[] { '/', '\\' });
-				if (Idx != -1)
-				{
-					string AssetName = FileName.Substring(Idx + 1);
-
-					int DotIdx = AssetName.LastIndexOf('.');
-					if (DotIdx != -1)
-					{
-						AssetName = AssetName.Substring(0, DotIdx);
-					}
-
-					ShortFileNames.Add(AssetName);
-				}
-			}
-			return ShortFileNames;
-		}
-
 		/// <summary>
 		/// Finds all the steps which are related to this issue
 		/// </summary>
 		/// <returns>Set of step names</returns>
-		private SortedSet<string> GetStepNames()
+		public SortedSet<string> GetStepNames()
 		{
 			return new SortedSet<string>(Streams.Values.SelectMany(x => x.FailedBuilds.SelectMany(y => y.StepNames)));
 		}
@@ -198,15 +97,17 @@ namespace MetadataTool
 		/// <returns>String representation of the issue</returns>
 		public override string ToString()
 		{
-			string Summary = GetSummary();
+			StringBuilder Result = new StringBuilder();
 			if(Id == -1)
 			{
-				return String.Format("[New] {0}", Summary);
+				Result.Append("[New] ");
 			}
 			else
 			{
-				return String.Format("[{0}] {1}", Id, Summary);
+				Result.AppendFormat("[{0}] ", Id);
 			}
+			Result.AppendFormat("Errors in {0}", String.Join(", ", GetStepNames()));
+			return Result.ToString();
 		}
 	}
 }
diff --git a/Engine/Source/Programs/UnrealGameSync/MetadataTool/CommandHandler_BuildHealth.cs b/Engine/Source/Programs/UnrealGameSync/MetadataTool/CommandHandler_BuildHealth.cs
--- a/Engine/Source/Programs/UnrealGameSync/MetadataTool/CommandHandler_BuildHealth.cs
+++ b/Engine/Source/Programs/UnrealGameSync/MetadataTool/CommandHandler_BuildHealth.cs
@@ -16,40 +16,12 @@ namespace MetadataTool
 {
 	class CommandHandler_BuildHealth : CommandHandler
 	{
-		/// <summary>
-		/// List of extensions for source files
-		/// </summary>
-		static string[] SourceFileExtensions =
-		{
-			".cpp",
-			".h",
-			".cc",
-			".hh",
-			".m",
-			".mm",
-			".rc",
-			".inl",
-			".inc"
+		// Register all the pattern matchers
+		static readonly List<PatternMatcher> Matchers = new List<PatternMatcher>()
+		{ 
+			new CompilePatternMatcher()
 		};
 
-		/// <summary>
-		/// List of extensions for content
-		/// </summary>
-		static string[] ContentFileExtensions =
-		{
-			".umap",
-			".uasset"
-		};
-
-		/// <summary>
-		/// Stores information about a change and its merge history
-		/// </summary>
-		class ChangeInfo
-		{
-			public ChangesRecord Record;
-			public List<int> SourceChanges = new List<int>();
-		}
-
 		/// <summary>
 		/// Constructor
 		/// </summary>
@@ -73,6 +45,13 @@ namespace MetadataTool
 			string ServerUrl = Arguments.GetStringOrDefault("-Server=", null);
 			Arguments.CheckAllArgumentsUsed();
 
+			// Build a mapping from category to matching
+			Dictionary<string, PatternMatcher> CategoryNameToMatcher = new Dictionary<string, PatternMatcher>();
+			foreach (PatternMatcher Matcher in Matchers)
+			{
+				CategoryNameToMatcher[Matcher.Category] = Matcher;
+			}
+
 			// Parse the input file
 			Log.TraceInformation("Reading build results from {0}", InputFile);
 			InputData InputData = DeserializeJson<InputData>(InputFile);
@@ -157,7 +136,13 @@ namespace MetadataTool
 				// Post any issue updates
 				foreach(TrackedIssue Issue in State.Issues)
 				{
-					string Summary = Issue.GetSummary();
+					PatternMatcher Matcher;
+					if(!CategoryNameToMatcher.TryGetValue(Issue.Fingerprint.Category, out Matcher))
+					{
+						continue;
+					}
+
+					string Summary = Matcher.GetSummary(Issue);
 					if(Issue.Id == -1)
 					{
 						Log.TraceInformation("Adding issue: {0}", Issue.Fingerprint);
@@ -307,6 +292,13 @@ namespace MetadataTool
 			// TODO: VERIFY ISSUES ARE CLOSED
 		}
 
+		/// <summary>
+		/// Sends an arbitrary HTTP request
+		/// </summary>
+		/// <param name="Url">Endpoint to send to</param>
+		/// <param name="Method">The method to use for sending the request</param>
+		/// <param name="BodyObject">Object to be serialized as json in the body</param>
+		/// <returns>HTTP response</returns>
 		HttpWebResponse SendHttpRequest(string Url, string Method, object BodyObject)
 		{
 			// Create the request
@@ -327,6 +319,12 @@ namespace MetadataTool
 			return (HttpWebResponse)Request.GetResponse();
 		}
 
+		/// <summary>
+		/// Parses an HTTP response object as JSON
+		/// </summary>
+		/// <typeparam name="T">The type of object to parse</typeparam>
+		/// <param name="Response">The web response instance</param>
+		/// <returns>Response object</returns>
 		T ParseHttpResponse<T>(HttpWebResponse Response)
 		{
 			using (StreamReader ResponseReader = new StreamReader(Response.GetResponseStream(), Encoding.Default))
@@ -348,6 +346,27 @@ namespace MetadataTool
 			Log.TraceInformation("  -Clean             Removes the existing state file and creates a new one");
 		}
 
+		/// <summary>
+		/// Finds the matcher for a particular category
+		/// </summary>
+		/// <param name="Category">The category to find</param>
+		/// <returns>Pattern matcher for this category</returns>
+		bool TryMergeFingerprint(TrackedIssueFingerprint Source, TrackedIssueFingerprint Target)
+		{
+			if(Source.Category != Target.Category)
+			{
+				return false;
+			}
+
+			PatternMatcher Matcher = Matchers.First(x => x.Category == Source.Category);
+			if(!Matcher.TryMerge(Source, Target))
+			{
+				return false;
+			}
+
+			return true;
+		}
+
 		/// <summary>
 		/// Duplicates a build for use as a sentinel value
 		/// </summary>
@@ -380,7 +399,11 @@ namespace MetadataTool
 			List<TrackedIssueFingerprint> Fingerprints = new List<TrackedIssueFingerprint>();
 			if(InputJobStep.Diagnostics != null)
 			{
-				CreateFingerprints(InputJob, InputJobStep, Fingerprints);
+				List<InputDiagnostic> Diagnostics = new List<InputDiagnostic>(InputJobStep.Diagnostics); 
+				foreach(PatternMatcher PatternMatcher in Matchers)
+				{
+					PatternMatcher.Match(InputJob, InputJobStep, Diagnostics, Fingerprints);
+				}
 			}
 
 			// Early out if there are no remaining fingerprints to add
@@ -405,13 +428,12 @@ namespace MetadataTool
 					{
 						continue;
 					}
-					if(!Issue.Fingerprint.CanMerge(Fingerprint))
+					if(!TryMergeFingerprint(Fingerprint, Issue.Fingerprint))
 					{
 						continue;
 					}
 
 					// Merge the issue
-					Issue.Fingerprint.Merge(Fingerprint);
 					History.AddFailedBuild(CreateBuildForJobStep(InputJob, InputJobStep));
 					Fingerprints.RemoveAt(Idx--);
 					break;
@@ -460,13 +482,12 @@ namespace MetadataTool
 								{
 									continue;
 								}
-								if(!Issue.Fingerprint.CanMerge(Fingerprint))
+								if (!TryMergeFingerprint(Fingerprint, Issue.Fingerprint))
 								{
 									continue;
 								}
 
 								// Merge the issue
-								Issue.Fingerprint.Merge(Fingerprint);
 								AddFailureToIssue(Issue, InputJob, InputJobStep, State);
 								Fingerprints.RemoveAt(Idx--);
 								break;
@@ -486,13 +507,14 @@ namespace MetadataTool
 			List<TrackedIssue> NewIssues = new List<TrackedIssue>();
 			foreach(TrackedIssueFingerprint Fingerprint in Fingerprints)
 			{
-				TrackedIssue Issue = NewIssues.FirstOrDefault(NewIssue => NewIssue.Fingerprint.TryMerge(Fingerprint));
+				TrackedIssue Issue = NewIssues.FirstOrDefault(NewIssue => TryMergeFingerprint(Fingerprint, NewIssue.Fingerprint));
 				if(Issue == null)
 				{
 					Issue = new TrackedIssue(Fingerprint, InputJob.Change);
 					if(Changes != null)
 					{
-						List<ChangeInfo> Causers = FilterCausers(Perforce, Fingerprint, Changes);
+						PatternMatcher Matcher = Matchers.First(x => x.Category == Fingerprint.Category);
+						List<ChangeInfo> Causers = Matcher.FindCausers(Perforce, Fingerprint, Changes);
 						foreach(ChangeInfo Causer in Causers)
 						{
 							Issue.SourceChanges.UnionWith(Causer.SourceChanges);
@@ -506,73 +528,6 @@ namespace MetadataTool
 			State.Issues.AddRange(NewIssues);
 		}
 
-		/// <summary>
-		/// Filters the list of causers for a possible issue
-		/// </summary>
-		/// <param name="Perforce">The Perforce connection</param>
-		/// <param name="Fingerprint">The issue fingerprint</param>
-		/// <param name="Changes">List of changes</param>
-		/// <returns>Receives a possible list of causers</returns>
-		List<ChangeInfo> FilterCausers(PerforceConnection Perforce, TrackedIssueFingerprint Fingerprint, List<ChangeInfo> Changes)
-		{
-			HashSet<string> ShortFileNames = new HashSet<string>(StringComparer.OrdinalIgnoreCase);
-			foreach(string FileName in Fingerprint.FileNames)
-			{
-				int Idx = FileName.LastIndexOf('/');
-				if(Idx != -1)
-				{
-					ShortFileNames.Add(FileName.Substring(Idx + 1));
-				}
-			}
-
-			List<ChangeInfo> LikelyCausers = new List<ChangeInfo>();
-			List<ChangeInfo> PossibleCausers = new List<ChangeInfo>();
-			foreach(ChangeInfo Change in Changes)
-			{
-				DescribeRecord Description = Perforce.Describe(Change.Record.Number).Data;
-				if(IsLikelyCauser(Description, ShortFileNames))
-				{
-					LikelyCausers.Add(Change);
-				}
-				else
-				{
-					PossibleCausers.Add(Change);
-				}
-			}
-
-			if(LikelyCausers.Count > 0)
-			{
-				return LikelyCausers;
-			}
-			else
-			{
-				return PossibleCausers;
-			}
-		}
-
-		/// <summary>
-		/// Determines if this change is a likely causer for an issue
-		/// </summary>
-		/// <param name="DescribeRecord">The change describe record</param>
-		/// <param name="Fingerprint">Fingerprint for the issue</param>
-		/// <returns>True if the change is a likely culprit</returns>
-		bool IsLikelyCauser(DescribeRecord DescribeRecord, HashSet<string> ShortFileNames)
-		{
-			foreach(DescribeFileRecord File in DescribeRecord.Files)
-			{
-				int Idx = File.DepotFile.LastIndexOf('/');
-				if(Idx != -1)
-				{
-					string FileName = File.DepotFile.Substring(Idx + 1);
-					if(ShortFileNames.Contains(FileName))
-					{
-						return true;
-					}
-				}
-			}
-			return false;
-		}
-
 		/// <summary>
 		/// Creates a TrackedBuild instance for the given jobstep
 		/// </summary>
@@ -608,74 +563,6 @@ namespace MetadataTool
 			}
 		}
 
-		/// <summary>
-		/// Parse a list of fingerprints from the build output
-		/// </summary>
-		/// <param name="Job">The job to generate fingerprints for</param>
-		/// <param name="JobStep">The jobstep to create fingerprints for</param>
-		/// <param name="Fingerprints">List which receives the created fingerprints</param>
-		void CreateFingerprints(InputJob Job, InputJobStep JobStep, List<TrackedIssueFingerprint> Fingerprints)
-		{
-			foreach(InputDiagnostic Diagnostic in JobStep.Diagnostics)
-			{
-				// Find any files in compiler output format
-				HashSet<string> SourceFileNames = new HashSet<string>(StringComparer.OrdinalIgnoreCase);
-				foreach (Match FileMatch in Regex.Matches(Diagnostic.Message, @"^\s*(?:In file included from\s*)?((?:[A-Za-z]:)?[^\s(:]+)[\(:]\d", RegexOptions.Multiline))
-				{
-					if (FileMatch.Success)
-					{
-						string FileName = GetNormalizedFileName(FileMatch.Groups[1].Value, JobStep.BaseDirectory);
-						if(SourceFileExtensions.Any(x => FileName.EndsWith(x, StringComparison.OrdinalIgnoreCase)))
-						{
-							SourceFileNames.Add(FileName);
-						}
-					}
-				}
-
-				// If we found any source files, create a diagnostic category for them
-				if (SourceFileNames.Count > 0)
-				{
-					TrackedIssueFingerprint Fingerprint = new TrackedIssueFingerprint(TrackedIssueFingerprintCategory.Compile, Job.Change);
-					Fingerprint.FileNames.UnionWith(SourceFileNames);
-					Fingerprints.Add(Fingerprint);
-				}
-			}
-		}
-
-		/// <summary>
-		/// Normalizes a filename to a path within the workspace
-		/// </summary>
-		/// <param name="FileName">Filename to normalize</param>
-		/// <param name="BaseDirectory">Base directory containing the workspace</param>
-		/// <returns>Normalized filename</returns>
-		protected string GetNormalizedFileName(string FileName, string BaseDirectory)
-		{
-			string NormalizedFileName = FileName.Replace('\\', '/');
-			if (!String.IsNullOrEmpty(BaseDirectory))
-			{
-				// Normalize the expected base directory for errors in this build, and attempt to strip it from the file name
-				string NormalizedBaseDirectory = BaseDirectory;
-				if (NormalizedBaseDirectory != null && NormalizedBaseDirectory.Length > 0)
-				{
-					NormalizedBaseDirectory = NormalizedBaseDirectory.Replace('\\', '/').TrimEnd('/') + "/";
-				}
-				if (NormalizedFileName.StartsWith(NormalizedBaseDirectory, StringComparison.OrdinalIgnoreCase))
-				{
-					NormalizedFileName = NormalizedFileName.Substring(NormalizedBaseDirectory.Length);
-				}
-			}
-			else
-			{
-				// Try to match anything under a 'Sync' folder.
-				Match FallbackRegex = Regex.Match(NormalizedFileName, "/Sync/(.*)");
-				if (FallbackRegex.Success)
-				{
-					NormalizedFileName = FallbackRegex.Groups[1].Value;
-				}
-			}
-			return NormalizedFileName;
-		}
-
 		/// <summary>
 		/// Find all changes PLUS all robomerge source changes
 		/// </summary>
diff --git a/Engine/Source/Programs/UnrealGameSync/MetadataTool/MetadataTool.csproj b/Engine/Source/Programs/UnrealGameSync/MetadataTool/MetadataTool.csproj
--- a/Engine/Source/Programs/UnrealGameSync/MetadataTool/MetadataTool.csproj
+++ b/Engine/Source/Programs/UnrealGameSync/MetadataTool/MetadataTool.csproj
@@ -51,6 +51,9 @@
     <Reference Include="System.Xml" />
   </ItemGroup>
   <ItemGroup>
+    <Compile Include="BuildHealth\ChangeInfo.cs" />
+    <Compile Include="BuildHealth\Matchers\CompilePatternMatcher.cs" />
+    <Compile Include="BuildHealth\PatternMatcher.cs" />
     <Compile Include="BuildHealth\TrackedIssueFingerprint.cs" />
     <Compile Include="CommandHandler.cs" />
     <Compile Include="CommandHandler_BuildHealth.cs" />
