Add missing copyright notices.

#rnx

diff --git a/Engine/Source/Programs/UnrealBuildTool/System/StringToInstanceRegistry.cs b/Engine/Source/Programs/UnrealBuildTool/System/StringToInstanceRegistry.cs
--- a/Engine/Source/Programs/UnrealBuildTool/System/StringToInstanceRegistry.cs
+++ b/Engine/Source/Programs/UnrealBuildTool/System/StringToInstanceRegistry.cs
@@ -1,3 +1,5 @@
+// Copyright 1998-2019 Epic Games, Inc. All Rights Reserved.
+
 using System;
 using System.Collections.Generic;
 using System.Linq;
diff --git a/Engine/Source/Programs/UnrealGameSync/MetadataTool/BuildHealth/ChangeInfo.cs b/Engine/Source/Programs/UnrealGameSync/MetadataTool/BuildHealth/ChangeInfo.cs
new file mode 100644
--- /dev/null
+++ b/Engine/Source/Programs/UnrealGameSync/MetadataTool/BuildHealth/ChangeInfo.cs
@@ -0,0 +1,20 @@
+// Copyright 1998-2019 Epic Games, Inc. All Rights Reserved.
+
+using System;
+using System.Collections.Generic;
+using System.Linq;
+using System.Text;
+using System.Threading.Tasks;
+using Tools.DotNETCommon.Perforce;
+
+namespace MetadataTool
+{
+	/// <summary>
+	/// Stores information about a change and its merge history
+	/// </summary>
+	class ChangeInfo
+	{
+		public ChangesRecord Record;
+		public List<int> SourceChanges = new List<int>();
+	}
+}
diff --git a/Engine/Source/Programs/UnrealGameSync/MetadataTool/BuildHealth/InputData.cs b/Engine/Source/Programs/UnrealGameSync/MetadataTool/BuildHealth/InputData.cs
--- a/Engine/Source/Programs/UnrealGameSync/MetadataTool/BuildHealth/InputData.cs
+++ b/Engine/Source/Programs/UnrealGameSync/MetadataTool/BuildHealth/InputData.cs
@@ -1,3 +1,5 @@
+// Copyright 1998-2019 Epic Games, Inc. All Rights Reserved.
+
 using System;
 using System.Collections.Generic;
 using System.Diagnostics;
diff --git a/Engine/Source/Programs/UnrealGameSync/MetadataTool/BuildHealth/Matchers/CompilePatternMatcher.cs b/Engine/Source/Programs/UnrealGameSync/MetadataTool/BuildHealth/Matchers/CompilePatternMatcher.cs
--- a/Engine/Source/Programs/UnrealGameSync/MetadataTool/BuildHealth/Matchers/CompilePatternMatcher.cs
+++ b/Engine/Source/Programs/UnrealGameSync/MetadataTool/BuildHealth/Matchers/CompilePatternMatcher.cs
@@ -1,3 +1,5 @@
+// Copyright 1998-2019 Epic Games, Inc. All Rights Reserved.
+
 using System;
 using System.Collections.Generic;
 using System.Linq;
diff --git a/Engine/Source/Programs/UnrealGameSync/MetadataTool/BuildHealth/PatternMatcher.cs b/Engine/Source/Programs/UnrealGameSync/MetadataTool/BuildHealth/PatternMatcher.cs
--- a/Engine/Source/Programs/UnrealGameSync/MetadataTool/BuildHealth/PatternMatcher.cs
+++ b/Engine/Source/Programs/UnrealGameSync/MetadataTool/BuildHealth/PatternMatcher.cs
@@ -1,3 +1,5 @@
+// Copyright 1998-2019 Epic Games, Inc. All Rights Reserved.
+
 using System;
 using System.Collections.Generic;
 using System.Linq;
diff --git a/Engine/Source/Programs/UnrealGameSync/MetadataTool/BuildHealth/TrackedBuild.cs b/Engine/Source/Programs/UnrealGameSync/MetadataTool/BuildHealth/TrackedBuild.cs
--- a/Engine/Source/Programs/UnrealGameSync/MetadataTool/BuildHealth/TrackedBuild.cs
+++ b/Engine/Source/Programs/UnrealGameSync/MetadataTool/BuildHealth/TrackedBuild.cs
@@ -1,3 +1,5 @@
+// Copyright 1998-2019 Epic Games, Inc. All Rights Reserved.
+
 using System;
 using System.Collections.Generic;
 using System.Diagnostics;
diff --git a/Engine/Source/Programs/UnrealGameSync/MetadataTool/BuildHealth/TrackedIssue.cs b/Engine/Source/Programs/UnrealGameSync/MetadataTool/BuildHealth/TrackedIssue.cs
--- a/Engine/Source/Programs/UnrealGameSync/MetadataTool/BuildHealth/TrackedIssue.cs
+++ b/Engine/Source/Programs/UnrealGameSync/MetadataTool/BuildHealth/TrackedIssue.cs
@@ -1,3 +1,5 @@
+// Copyright 1998-2019 Epic Games, Inc. All Rights Reserved.
+
 using System;
 using System.Collections.Generic;
 using System.Diagnostics;
diff --git a/Engine/Source/Programs/UnrealGameSync/MetadataTool/BuildHealth/TrackedIssueFingerprint.cs b/Engine/Source/Programs/UnrealGameSync/MetadataTool/BuildHealth/TrackedIssueFingerprint.cs
new file mode 100644
--- /dev/null
+++ b/Engine/Source/Programs/UnrealGameSync/MetadataTool/BuildHealth/TrackedIssueFingerprint.cs
@@ -0,0 +1,168 @@
+// Copyright 1998-2019 Epic Games, Inc. All Rights Reserved.
+
+using System;
+using System.Collections.Generic;
+using System.Diagnostics;
+using System.Linq;
+using System.Runtime.Serialization;
+using System.Text;
+using System.Threading.Tasks;
+
+namespace MetadataTool
+{
+	/// <summary>
+	/// Used to group diagnostics together
+	/// </summary>
+	[DataContract]
+	class TrackedIssueFingerprint
+	{
+		/// <summary>
+		/// Type common to all diagnostics within this issue.
+		/// </summary>
+		[DataMember(IsRequired = true)]
+		public string Category;
+
+		/// <summary>
+		/// The initial change that this changelist was seen on. Allows other fingerprints at the same change to be merged.
+		/// </summary>
+		[DataMember(IsRequired = true)]
+		public int InitialChange;
+
+		/// <summary>
+		/// List of files associated with this issue
+		/// </summary>
+		[DataMember]
+		public HashSet<string> FileNames = new HashSet<string>();
+
+		/// <summary>
+		/// List of messages associated with this issue
+		/// </summary>
+		[DataMember]
+		public SortedSet<string> Messages = new SortedSet<string>();
+
+		/// <summary>
+		/// Constructor
+		/// </summary>
+		/// <param name="Category">Category for this issue</param>
+		/// <param name="InitialChange">The initial change that this issue was seen on. Fingerprints in the same category at the same change can be merged as long as the type matches.</param>
+		public TrackedIssueFingerprint(string Category, int InitialChange)
+		{
+			this.Category = Category;
+			this.InitialChange = InitialChange;
+		}
+
+		/// <summary>
+		/// Finds all the unique filenames without their path components
+		/// </summary>
+		/// <returns>Set of sorted filenames</returns>
+		public SortedSet<string> GetFileNamesWithoutPath()
+		{
+			SortedSet<string> FileNamesWithoutPath = new SortedSet<string>(StringComparer.OrdinalIgnoreCase);
+			foreach (string FileName in FileNames)
+			{
+				int Idx = FileName.LastIndexOf('/');
+				if (Idx != -1)
+				{
+					FileNamesWithoutPath.Add(FileName.Substring(Idx + 1));
+				}
+			}
+			return FileNamesWithoutPath;
+		}
+
+		/// <summary>
+		/// Gets a set of unique source file names that relate to this issue
+		/// </summary>
+		/// <returns>Set of source file names</returns>
+		public SortedSet<string> GetSourceFileNames()
+		{
+			SortedSet<string> ShortFileNames = new SortedSet<string>(StringComparer.OrdinalIgnoreCase);
+			foreach (string FileName in FileNames)
+			{
+				int Idx = FileName.LastIndexOfAny(new char[] { '/', '\\' });
+				if (Idx != -1)
+				{
+					string ShortFileName = FileName.Substring(Idx + 1);
+					if (!ShortFileName.StartsWith("Module.", StringComparison.OrdinalIgnoreCase))
+					{
+						ShortFileNames.Add(ShortFileName);
+					}
+				}
+			}
+			return ShortFileNames;
+		}
+
+		/// <summary>
+		/// Gets a set of unique asset filenames that relate to this issue
+		/// </summary>
+		/// <returns>Set of asset names</returns>
+		public SortedSet<string> GetAssetNames()
+		{
+			SortedSet<string> ShortFileNames = new SortedSet<string>(StringComparer.OrdinalIgnoreCase);
+			foreach (string FileName in FileNames)
+			{
+				int Idx = FileName.LastIndexOfAny(new char[] { '/', '\\' });
+				if (Idx != -1)
+				{
+					string AssetName = FileName.Substring(Idx + 1);
+
+					int DotIdx = AssetName.LastIndexOf('.');
+					if (DotIdx != -1)
+					{
+						AssetName = AssetName.Substring(0, DotIdx);
+					}
+
+					ShortFileNames.Add(AssetName);
+				}
+			}
+			return ShortFileNames;
+		}
+
+		/// <summary>
+		/// Formats this object for debugging
+		/// </summary>
+		/// <returns>String representation of the issue</returns>
+		public override string ToString()
+		{
+			StringBuilder Result = new StringBuilder();
+			Result.AppendFormat("{0}", Category);
+
+			if(FileNames.Count > 0 || Messages.Count > 0)
+			{
+				string Separator = ": ";
+				foreach(string Item in Enumerable.Concat(FileNames, Messages))
+				{
+					Result.Append(Separator);
+
+					int StartIdx = 0;
+					while(StartIdx < Item.Length && Char.IsWhiteSpace(Item[StartIdx]))
+					{
+						StartIdx++;
+					}
+
+					int EndIdx = Item.IndexOfAny(new char[] { '\r', '\n' });
+					if(EndIdx == -1)
+					{
+						EndIdx = Item.Length;
+					}
+					while(EndIdx > StartIdx && Char.IsWhiteSpace(Item[EndIdx - 1]))
+					{
+						EndIdx--;
+					}
+
+					Result.Append(Item, StartIdx, EndIdx - StartIdx);
+					Separator = ", ";
+				}
+			}
+
+			const int MaxLength = 128;
+			if(Result.Length > MaxLength)
+			{
+				Result.Remove(MaxLength, Result.Length - MaxLength);
+				Result.Append("...");
+			}
+
+			return Result.ToString();
+		}
+	}
+
+}
diff --git a/Engine/Source/Programs/UnrealGameSync/MetadataTool/BuildHealth/TrackedIssueHistory.cs b/Engine/Source/Programs/UnrealGameSync/MetadataTool/BuildHealth/TrackedIssueHistory.cs
--- a/Engine/Source/Programs/UnrealGameSync/MetadataTool/BuildHealth/TrackedIssueHistory.cs
+++ b/Engine/Source/Programs/UnrealGameSync/MetadataTool/BuildHealth/TrackedIssueHistory.cs
@@ -1,3 +1,5 @@
+// Copyright 1998-2019 Epic Games, Inc. All Rights Reserved.
+
 using System;
 using System.Collections.Generic;
 using System.Linq;
diff --git a/Engine/Source/Programs/UnrealGameSync/MetadataTool/BuildHealth/TrackedState.cs b/Engine/Source/Programs/UnrealGameSync/MetadataTool/BuildHealth/TrackedState.cs
--- a/Engine/Source/Programs/UnrealGameSync/MetadataTool/BuildHealth/TrackedState.cs
+++ b/Engine/Source/Programs/UnrealGameSync/MetadataTool/BuildHealth/TrackedState.cs
@@ -1,3 +1,5 @@
+// Copyright 1998-2019 Epic Games, Inc. All Rights Reserved.
+
 using System;
 using System.Collections.Generic;
 using System.Linq;
diff --git a/Engine/Source/Programs/UnrealGameSync/MetadataTool/CommandHandler_BuildHealth.cs b/Engine/Source/Programs/UnrealGameSync/MetadataTool/CommandHandler_BuildHealth.cs
--- a/Engine/Source/Programs/UnrealGameSync/MetadataTool/CommandHandler_BuildHealth.cs
+++ b/Engine/Source/Programs/UnrealGameSync/MetadataTool/CommandHandler_BuildHealth.cs
@@ -1,3 +1,5 @@
+// Copyright 1998-2019 Epic Games, Inc. All Rights Reserved.
+
 using System;
 using System.Collections.Generic;
 using System.IO;
diff --git a/Engine/Source/Programs/Windows/CompileTimeAnalyzer/CompileTimeAnalyzer/App.xaml.cs b/Engine/Source/Programs/Windows/CompileTimeAnalyzer/CompileTimeAnalyzer/App.xaml.cs
--- a/Engine/Source/Programs/Windows/CompileTimeAnalyzer/CompileTimeAnalyzer/App.xaml.cs
+++ b/Engine/Source/Programs/Windows/CompileTimeAnalyzer/CompileTimeAnalyzer/App.xaml.cs
@@ -1,4 +1,6 @@
-﻿using System;
+﻿// Copyright 1998-2019 Epic Games, Inc. All Rights Reserved.
+
+using System;
 using System.Collections.Generic;
 using System.Configuration;
 using System.Data;
diff --git a/Engine/Source/Programs/Windows/CompileTimeAnalyzer/CompileTimeAnalyzer/Properties/AssemblyInfo.cs b/Engine/Source/Programs/Windows/CompileTimeAnalyzer/CompileTimeAnalyzer/Properties/AssemblyInfo.cs
--- a/Engine/Source/Programs/Windows/CompileTimeAnalyzer/CompileTimeAnalyzer/Properties/AssemblyInfo.cs
+++ b/Engine/Source/Programs/Windows/CompileTimeAnalyzer/CompileTimeAnalyzer/Properties/AssemblyInfo.cs
@@ -1,4 +1,6 @@
-﻿using System.Reflection;
+﻿// Copyright 1998-2019 Epic Games, Inc. All Rights Reserved.
+
+using System.Reflection;
 using System.Resources;
 using System.Runtime.CompilerServices;
 using System.Runtime.InteropServices;
