UGS: Add a pattern matcher for linker undefined symbol errors.

diff --git a/Engine/Source/Programs/UnrealGameSync/MetadataTool/BuildHealth/Matchers/UndefinedSymbolPatternMatcher.cs b/Engine/Source/Programs/UnrealGameSync/MetadataTool/BuildHealth/Matchers/UndefinedSymbolPatternMatcher.cs
new file mode 100644
--- /dev/null
+++ b/Engine/Source/Programs/UnrealGameSync/MetadataTool/BuildHealth/Matchers/UndefinedSymbolPatternMatcher.cs
@@ -0,0 +1,109 @@
+using System;
+using System.Collections.Generic;
+using System.Linq;
+using System.Text;
+using System.Text.RegularExpressions;
+using System.Threading.Tasks;
+
+namespace MetadataTool
+{
+	class UndefinedSymbolPatternMatcher : PatternMatcher
+	{
+		public override string Category => "UndefinedSymbol";
+
+		public override bool TryMatch(InputJob Job, InputJobStep JobStep, InputDiagnostic Diagnostic, List<TrackedIssueFingerprint> Fingerprints)
+		{
+			SortedSet<string> SymbolNames = new SortedSet<string>(StringComparer.Ordinal);
+
+			// Mac link error:
+			//   Undefined symbols for architecture arm64:
+            //     "Foo::Bar() const", referenced from:
+			if (Regex.IsMatch(Diagnostic.Message, "^Undefined symbols"))
+			{
+				foreach(string Line in Diagnostic.Message.Split('\n'))
+				{
+					Match SymbolMatch = Regex.Match(Line, "^  \"([^\"\\(]+)");
+					if (SymbolMatch.Success)
+					{
+						SymbolNames.Add(SymbolMatch.Groups[1].Value);
+					}
+				}
+			}
+
+			// Android link error:
+			//   Foo.o:(.data.rel.ro + 0x5d88): undefined reference to `Foo::Bar()'
+			Match UndefinedReference = Regex.Match(Diagnostic.Message, ": undefined reference to [`']([^`'(]+)");
+			if(UndefinedReference.Success)
+			{
+				SymbolNames.Add(UndefinedReference.Groups[1].Value);
+			}
+
+			// LLD link error:
+			//   ld.lld.exe: error: undefined symbol: Foo::Bar() const
+			Match LldMatch = Regex.Match(Diagnostic.Message, "error: undefined symbol:\\s*([^ (]+)");
+			if (LldMatch.Success)
+			{
+				SymbolNames.Add(LldMatch.Groups[1].Value);
+			}
+
+			// Link error:
+			//   Link: error: L0039: reference to undefined symbol `Foo::Bar() const' in file
+			Match LinkMatch = Regex.Match(Diagnostic.Message, ": reference to undefined symbol [`']([^`'(]+)");
+			if (LinkMatch.Success)
+			{
+				SymbolNames.Add(LinkMatch.Groups[1].Value);
+			}
+
+			// Microsoft linker error:
+			//   Foo.cpp.obj : error LNK2001: unresolved external symbol \"private: virtual void __cdecl UAssetManager::InitializeAssetBundlesFromMetadata_Recursive(class UStruct const *,void const *,struct FAssetBundleData &,class FName,class TSet<void const *,struct DefaultKeyFuncs<void const *,0>,class FDefaultSetAllocator> &)const \" (?InitializeAssetBundlesFromMetadata_Recursive@UAssetManager@@EEBAXPEBVUStruct@@PEBXAEAUFAssetBundleData@@VFName@@AEAV?$TSet@PEBXU?$DefaultKeyFuncs@PEBX$0A@@@VFDefaultSetAllocator@@@@@Z)",
+			Match MicrosoftMatch = Regex.Match(Diagnostic.Message, ": unresolved external symbol \"([^\"]*)\"");
+			if(MicrosoftMatch.Success)
+			{
+				string SymbolName = MicrosoftMatch.Groups[1].Value;
+
+				// Remove any argument lists for functions (anything after the first paren)
+				SymbolName = Regex.Replace(SymbolName, "\\(.*$", "");
+
+				// Remove any decorators and type information (greedy match up to the last space)
+				SymbolName = Regex.Replace(SymbolName, "^.* ", "");
+
+				SymbolNames.Add(SymbolName);
+			}
+
+			// If we found any symbol names, create a fingerprint for them
+			if (SymbolNames.Count > 0)
+			{
+				TrackedIssueFingerprint Fingerprint = new TrackedIssueFingerprint(Category, GetSummary(SymbolNames), Diagnostic.Url, Job.Change);
+				Fingerprint.Details.Add(Diagnostic.Message);
+				Fingerprint.Messages.UnionWith(SymbolNames);
+				Fingerprints.Add(Fingerprint);
+				return true;
+			}
+
+			// Otherwise pass
+			return false;
+		}
+
+		public override bool TryMerge(TrackedIssueFingerprint Source, TrackedIssueFingerprint Target)
+		{
+			if (base.TryMerge(Source, Target))
+			{
+				Target.Summary = GetSummary(Target.Messages);
+				return true;
+			}
+			return false;
+		}
+
+		string GetSummary(SortedSet<string> SymbolNames)
+		{
+			if(SymbolNames.Count == 1)
+			{
+				return String.Format("Undefined symbol: {0}", SymbolNames.First());
+			}
+			else
+			{
+				return String.Format("Undefined symbols: {0}", String.Join(", ", SymbolNames));
+			}
+		}
+	}
+}
diff --git a/Engine/Source/Programs/UnrealGameSync/MetadataTool/CommandHandler_BuildHealth.cs b/Engine/Source/Programs/UnrealGameSync/MetadataTool/CommandHandler_BuildHealth.cs
--- a/Engine/Source/Programs/UnrealGameSync/MetadataTool/CommandHandler_BuildHealth.cs
+++ b/Engine/Source/Programs/UnrealGameSync/MetadataTool/CommandHandler_BuildHealth.cs
@@ -21,7 +21,8 @@ namespace MetadataTool
 		// Register all the pattern matchers
 		static readonly List<PatternMatcher> Matchers = new List<PatternMatcher>()
 		{ 
-			new CompilePatternMatcher()
+			new CompilePatternMatcher(),
+			new UndefinedSymbolPatternMatcher()
 		};
 
 		/// <summary>
diff --git a/Engine/Source/Programs/UnrealGameSync/MetadataTool/MetadataTool.csproj b/Engine/Source/Programs/UnrealGameSync/MetadataTool/MetadataTool.csproj
--- a/Engine/Source/Programs/UnrealGameSync/MetadataTool/MetadataTool.csproj
+++ b/Engine/Source/Programs/UnrealGameSync/MetadataTool/MetadataTool.csproj
@@ -53,6 +53,7 @@
   <ItemGroup>
     <Compile Include="BuildHealth\ChangeInfo.cs" />
     <Compile Include="BuildHealth\Matchers\CompilePatternMatcher.cs" />
+    <Compile Include="BuildHealth\Matchers\UndefinedSymbolPatternMatcher.cs" />
     <Compile Include="BuildHealth\PatternMatcher.cs" />
     <Compile Include="BuildHealth\TrackedIssueFingerprint.cs" />
     <Compile Include="CommandHandler.cs" />
