UGS: Clean up some output for the undefined symbol pattern matcher.

#rnx

diff --git a/Engine/Source/Programs/UnrealGameSync/MetadataTool/BuildHealth/Matchers/UndefinedSymbolPatternMatcher.cs b/Engine/Source/Programs/UnrealGameSync/MetadataTool/BuildHealth/Matchers/UndefinedSymbolPatternMatcher.cs
--- a/Engine/Source/Programs/UnrealGameSync/MetadataTool/BuildHealth/Matchers/UndefinedSymbolPatternMatcher.cs
+++ b/Engine/Source/Programs/UnrealGameSync/MetadataTool/BuildHealth/Matchers/UndefinedSymbolPatternMatcher.cs
@@ -98,11 +98,11 @@ namespace MetadataTool
 		{
 			if(SymbolNames.Count == 1)
 			{
-				return String.Format("Undefined symbol: {0}", SymbolNames.First());
+				return String.Format("Undefined symbol '{0}'", SymbolNames.First());
 			}
 			else
 			{
-				return String.Format("Undefined symbols: {0}", String.Join(", ", SymbolNames));
+				return String.Format("Undefined symbols - '{0}'", String.Join("', '", SymbolNames));
 			}
 		}
 	}
