UGS: Fix issue where clients with multiple proxies will fail to get server info and refuse to connect.

diff --git a/Engine/Source/Programs/UnrealGameSync/UnrealGameSync/Perforce.cs b/Engine/Source/Programs/UnrealGameSync/UnrealGameSync/Perforce.cs
--- a/Engine/Source/Programs/UnrealGameSync/UnrealGameSync/Perforce.cs
+++ b/Engine/Source/Programs/UnrealGameSync/UnrealGameSync/Perforce.cs
@@ -157,7 +157,6 @@ namespace UnrealGameSync
 		}
 	}
 
-	[DebuggerDisplay("{ServerAddress}")]
 	class PerforceInfoRecord
 	{
 		public string UserName;
@@ -165,14 +164,14 @@ namespace UnrealGameSync
 		public string ClientAddress;
 		public TimeSpan ServerTimeZone;
 
-		public PerforceInfoRecord(Dictionary<string, string> Tags)
+		public PerforceInfoRecord(List<Dictionary<string, string>> TagRecords)
 		{
-			Tags.TryGetValue("userName", out UserName);
-			Tags.TryGetValue("clientHost", out HostName);
-			Tags.TryGetValue("clientAddress", out ClientAddress);
+			TryGetValue(TagRecords, "userName", out UserName);
+			TryGetValue(TagRecords, "clientHost", out HostName);
+			TryGetValue(TagRecords, "clientAddress", out ClientAddress);
 
 			string ServerDateTime;
-			if(Tags.TryGetValue("serverDate", out ServerDateTime))
+			if(TryGetValue(TagRecords, "serverDate", out ServerDateTime))
 			{
 				string[] Fields = ServerDateTime.Split(new char[]{ ' ' }, StringSplitOptions.RemoveEmptyEntries);
 				if(Fields.Length >= 3)
@@ -185,6 +184,22 @@ namespace UnrealGameSync
 				}
 			}
 		}
+
+		static bool TryGetValue(List<Dictionary<string, string>> TagRecords, string Key, out string Value)
+		{
+			foreach(Dictionary<string, string> TagRecord in TagRecords)
+			{
+				string TagValue;
+				if(TagRecord.TryGetValue(Key, out TagValue))
+				{
+					Value = TagValue;
+					return true;
+				}
+			}
+
+			Value = null;
+			return false;
+		}
 	}
 
 	[Flags]
@@ -568,14 +583,14 @@ namespace UnrealGameSync
 		public bool Info(out PerforceInfoRecord Info, TextWriter Log)
 		{
 			List<Dictionary<string, string>> TagRecords;
-			if(!RunCommand("info -s", out TagRecords, CommandOptions.NoClient, Log) || TagRecords.Count != 1)
+			if(!RunCommand("info -s", out TagRecords, CommandOptions.NoClient, Log) || TagRecords.Count == 0)
 			{
 				Info = null;
 				return false;
 			}
 			else
 			{
-				Info = new PerforceInfoRecord(TagRecords[0]);
+				Info = new PerforceInfoRecord(TagRecords);
 				return true;
 			}
 		}
