UGS: Add missing files.

#rnx

diff --git a/Engine/Source/Programs/UnrealGameSync/MetadataTool/BuildHealth/Matchers/UndefinedSymbolPatternMatcher.cs b/Engine/Source/Programs/UnrealGameSync/MetadataTool/BuildHealth/Matchers/UndefinedSymbolPatternMatcher.cs
--- a/Engine/Source/Programs/UnrealGameSync/MetadataTool/BuildHealth/Matchers/UndefinedSymbolPatternMatcher.cs
+++ b/Engine/Source/Programs/UnrealGameSync/MetadataTool/BuildHealth/Matchers/UndefinedSymbolPatternMatcher.cs
@@ -1,3 +1,5 @@
+// Copyright 1998-2019 Epic Games, Inc. All Rights Reserved.
+
 using System;
 using System.Collections.Generic;
 using System.Linq;
diff --git a/Engine/Source/Programs/UnrealGameSync/UnrealGameSync/Forms/IssueFixedWindow.cs b/Engine/Source/Programs/UnrealGameSync/UnrealGameSync/Forms/IssueFixedWindow.cs
new file mode 100644
--- /dev/null
+++ b/Engine/Source/Programs/UnrealGameSync/UnrealGameSync/Forms/IssueFixedWindow.cs
@@ -0,0 +1,352 @@
+using System;
+using System.Collections.Generic;
+using System.ComponentModel;
+using System.Data;
+using System.Drawing;
+using System.Linq;
+using System.Text;
+using System.Threading;
+using System.Threading.Tasks;
+using System.Windows.Forms;
+
+namespace UnrealGameSync
+{
+	partial class IssueFixedWindow : Form
+	{
+		class FindChangesWorker : Component
+		{
+			public delegate void OnCompleteDelegate(string UserName, List<PerforceDescribeRecord> Changes);
+
+			BufferedTextWriter Log = new BufferedTextWriter();
+			PerforceConnection Perforce;
+			SynchronizationContext MainThreadSyncContext;
+			AutoResetEvent WakeEvent = new AutoResetEvent(false);
+			string RequestedUserName;
+			bool bStopRequested;
+			Thread Thread;
+			OnCompleteDelegate OnComplete;
+
+			public FindChangesWorker(PerforceConnection Perforce, OnCompleteDelegate OnComplete)
+			{
+				this.Perforce = Perforce;
+				this.MainThreadSyncContext = SynchronizationContext.Current;
+				this.OnComplete = OnComplete;
+			}
+
+			public void Start()
+			{
+				if(Thread == null)
+				{
+					bStopRequested = false;
+					Thread = new Thread(DoWork);
+					Thread.Start();
+				}
+			}
+
+			public void Stop()
+			{
+				if(Thread != null)
+				{
+					bStopRequested = true;
+					Thread.Join();
+					Thread = null;
+				}
+			}
+
+			public void FetchChanges(string UserName)
+			{
+				RequestedUserName = UserName;
+				WakeEvent.Set();
+			}
+
+			protected override void Dispose(bool disposing)
+			{
+				base.Dispose(disposing);
+
+				Stop();
+			}
+
+			public void DoWork()
+			{
+				for(;;)
+				{
+					WakeEvent.WaitOne();
+
+					if(bStopRequested)
+					{
+						break;
+					}
+
+					string UserName = RequestedUserName;
+
+					List<PerforceDescribeRecord> Descriptions = null;
+
+					List<PerforceChangeSummary> Changes;
+					if(Perforce.FindChanges(new string[]{ "//..." }, UserName, 250, out Changes, Log))
+					{
+						Perforce.DescribeMultiple(Changes.Select(x => x.Number), out Descriptions, Log);
+					}
+					
+					MainThreadSyncContext.Post((o) => { OnComplete(UserName, Descriptions); }, null);
+				}
+			}
+		}
+
+		PerforceConnection Perforce;
+		int ChangeNumber;
+		FindChangesWorker Worker;
+	
+		public IssueFixedWindow(PerforceConnection Perforce, int InitialChangeNumber)
+		{
+			InitializeComponent();
+
+			this.Perforce = Perforce;
+			this.Worker = new FindChangesWorker(Perforce, PopulateChanges);
+
+			UserNameTextBox.Text = Perforce.UserName;
+			UserNameTextBox.SelectionStart = UserNameTextBox.Text.Length;
+
+			if(InitialChangeNumber != 0)
+			{
+				if(InitialChangeNumber < 0)
+				{
+					SpecifyChangeRadioButton.Checked = false;
+					SystemicFixRadioButton.Checked = true;
+				}
+				else
+				{
+					SpecifyChangeRadioButton.Checked = true;
+					ChangeNumberTextBox.Text = InitialChangeNumber.ToString();
+				}
+			}
+
+			UpdateOkButton();
+		}
+
+		protected override void OnLoad(EventArgs e)
+		{
+			base.OnLoad(e);
+
+			if(SystemicFixRadioButton.Checked)
+			{
+				SystemicFixRadioButton.Select();
+			}
+			else if(SpecifyChangeRadioButton.Checked)
+			{
+				SpecifyChangeRadioButton.Select();
+			}
+			else if(RecentChangeRadioButton.Checked)
+			{
+				ChangesListView.Select();
+			}
+
+			Worker.Start();
+			Worker.FetchChanges(UserNameTextBox.Text);
+		}
+
+		public static bool ShowModal(IWin32Window Owner, PerforceConnection Perforce, ref int FixChangeNumber)
+		{
+			using(IssueFixedWindow FixedWindow = new IssueFixedWindow(Perforce, FixChangeNumber))
+			{
+				if(FixedWindow.ShowDialog(Owner) == DialogResult.OK)
+				{
+					FixChangeNumber = FixedWindow.ChangeNumber;
+					return true;
+				}
+				else
+				{
+					FixChangeNumber = 0;
+					return false;
+				}
+			}
+		}
+
+		private void OkBtn_Click(object sender, EventArgs e)
+		{
+			if(TryGetSelectedChange(out ChangeNumber))
+			{
+				DialogResult = DialogResult.OK;
+				Close();
+			}
+		}
+
+		private void PopulateChanges(string UserName, List<PerforceDescribeRecord> Changes)
+		{
+			if(!IsDisposed)
+			{
+				ChangesListView.BeginUpdate();
+				ChangesListView.Items.Clear();
+				if(Changes != null)
+				{
+					foreach(PerforceDescribeRecord Change in Changes)
+					{
+						if(Change.Description != null && Change.Description.IndexOf("#ROBOMERGE-SOURCE", 0) == -1)
+						{
+							string Stream = "";
+							if(Change.Files.Count > 0)
+							{
+								string DepotFile = Change.Files[0].DepotFile;
+
+								int Idx = 0;
+								for(int Count = 0; Idx < DepotFile.Length; Idx++)
+								{
+									if(DepotFile[Idx] == '/' && ++Count >= 4)
+									{
+										break;
+									}
+								}
+
+								Stream = DepotFile.Substring(0, Idx);
+							}
+
+							ListViewItem Item = new ListViewItem("");
+							Item.Tag = Change;
+							Item.SubItems.Add(Change.ChangeNumber.ToString());
+							Item.SubItems.Add(Stream);
+							Item.SubItems.Add(Change.Description.Replace('\n', ' '));
+							ChangesListView.Items.Add(Item);
+						}
+					}
+				}
+				ChangesListView.EndUpdate();
+			}
+		}
+
+		private void ChangesListView_MouseClick(object Sender, MouseEventArgs Args)
+		{
+			if(Args.Button == MouseButtons.Right)
+			{
+				ListViewHitTestInfo HitTest = ChangesListView.HitTest(Args.Location);
+				if(HitTest.Item != null && HitTest.Item.Tag != null)
+				{
+					PerforceDescribeRecord Record = HitTest.Item.Tag as PerforceDescribeRecord; 
+					if(Record != null)
+					{
+						ChangesListContextMenu.Tag = Record;
+						ChangesListContextMenu.Show(ChangesListView, Args.Location);
+					}
+				}
+			}
+		}
+
+		private void ChangesContextMenu_MoreInfo_Click(object sender, EventArgs e)
+		{
+			PerforceDescribeRecord Record = (PerforceDescribeRecord)ChangesListContextMenu.Tag;
+			if(!Utility.SpawnHiddenProcess("p4vc.exe", String.Format("-p\"{0}\" change {1}", Perforce.ServerAndPort, Record.ChangeNumber)))
+			{
+				MessageBox.Show("Unable to spawn p4vc. Check you have P4V installed.");
+			}
+		}
+
+		private void ChangeNumberTextBox_Enter(object sender, EventArgs e)
+		{
+			SpecifyChangeRadioButton.Checked = true;
+			RecentChangeRadioButton.Checked = false;
+		}
+
+		private void ChangesListView_Enter(object sender, EventArgs e)
+		{
+			SpecifyChangeRadioButton.Checked = false;
+			RecentChangeRadioButton.Checked = true;
+		}
+
+		private void UserNameTextBox_Enter(object sender, EventArgs e)
+		{
+			SpecifyChangeRadioButton.Checked = false;
+			RecentChangeRadioButton.Checked = true;
+		}
+
+		private void UserBrowseBtn_Enter(object sender, EventArgs e)
+		{
+			SpecifyChangeRadioButton.Checked = false;
+			RecentChangeRadioButton.Checked = true;
+		}
+
+		private void ChangeNumberTextBox_TextChanged(object sender, EventArgs e)
+		{
+			UpdateOkButton();
+		}
+
+		private void ChangesListView_ItemSelectionChanged(object sender, ListViewItemSelectionChangedEventArgs e)
+		{
+			UpdateOkButton();
+		}
+
+		private bool TryGetSelectedChange(out int ChangeNumber)
+		{
+			if(SpecifyChangeRadioButton.Checked)
+			{
+				return int.TryParse(ChangeNumberTextBox.Text, out ChangeNumber);
+			}
+			else if(SystemicFixRadioButton.Checked)
+			{
+				ChangeNumber = -1;
+				return true;
+			}
+			else
+			{
+				PerforceDescribeRecord Change = (ChangesListView.SelectedItems.Count > 0)? ChangesListView.SelectedItems[0].Tag as PerforceDescribeRecord : null;
+				if(Change == null)
+				{
+					ChangeNumber = 0;
+					return false;
+				}
+				else
+				{
+					ChangeNumber = Change.ChangeNumber;
+					return true;
+				}
+			}
+		}
+
+		private void UpdateOkButton()
+		{
+			int ChangeNumber;
+			OkBtn.Enabled = TryGetSelectedChange(out ChangeNumber);
+		}
+
+		private void SpecifyChangeRadioButton_CheckedChanged(object sender, EventArgs e)
+		{
+			if(SpecifyChangeRadioButton.Checked)
+			{
+				RecentChangeRadioButton.Checked = false;
+				SystemicFixRadioButton.Checked = false;
+			}
+			UpdateOkButton();
+		}
+
+		private void RecentChangeRadioButton_CheckedChanged(object sender, EventArgs e)
+		{
+			if(RecentChangeRadioButton.Checked)
+			{
+				SpecifyChangeRadioButton.Checked = false;
+				SystemicFixRadioButton.Checked = false;
+			}
+			UpdateOkButton();
+		}
+
+		private void SystemicFixRadioButton_CheckedChanged(object sender, EventArgs e)
+		{
+			if(SystemicFixRadioButton.Checked)
+			{
+				RecentChangeRadioButton.Checked = false;
+				SpecifyChangeRadioButton.Checked = false;
+			}
+			UpdateOkButton();
+		}
+
+		private void UserBrowseBtn_Click(object sender, EventArgs e)
+		{
+			string SelectedUserName;
+			if(SelectUserWindow.ShowModal(this, Perforce.ServerAndPort, Perforce.UserName, new BufferedTextWriter(), out SelectedUserName))
+			{
+				UserNameTextBox.Text = SelectedUserName;
+			}
+		}
+
+		private void UserNameTextBox_TextChanged(object sender, EventArgs e)
+		{
+			Worker.FetchChanges(UserNameTextBox.Text);
+		}
+	}
+}
diff --git a/Engine/Source/Programs/UnrealGameSync/UnrealGameSync/Forms/IssueSettingsWindow.cs b/Engine/Source/Programs/UnrealGameSync/UnrealGameSync/Forms/IssueSettingsWindow.cs
new file mode 100644
--- /dev/null
+++ b/Engine/Source/Programs/UnrealGameSync/UnrealGameSync/Forms/IssueSettingsWindow.cs
@@ -0,0 +1,111 @@
+// Copyright 1998-2019 Epic Games, Inc. All Rights Reserved.
+
+using System;
+using System.Collections.Generic;
+using System.ComponentModel;
+using System.Data;
+using System.Drawing;
+using System.Linq;
+using System.Text;
+using System.Threading.Tasks;
+using System.Windows.Forms;
+
+namespace UnrealGameSync
+{
+	partial class IssueSettingsWindow : Form
+	{
+		UserSettings Settings;
+
+		public IssueSettingsWindow(UserSettings Settings)
+		{
+			this.Settings = Settings;
+
+			InitializeComponent();
+
+			if(Settings.NotifyUnassignedMinutes < 0)
+			{
+				NotifyUnassignedCheckBox.Checked = false;
+				NotifyUnassignedTextBox.Text = "5";
+			}
+			else
+			{
+				NotifyUnassignedCheckBox.Checked = true;
+				NotifyUnassignedTextBox.Text = Settings.NotifyUnassignedMinutes.ToString();
+			}
+
+			if(Settings.NotifyUnresolvedMinutes < 0)
+			{
+				NotifyUnresolvedCheckBox.Checked = false;
+				NotifyUnresolvedTextBox.Text = "20";
+			}
+			else
+			{
+				NotifyUnresolvedCheckBox.Checked = true;
+				NotifyUnresolvedTextBox.Text = Settings.NotifyUnresolvedMinutes.ToString();
+			}
+
+			UpdateEnabledTextBoxes();
+		}
+
+		private void UpdateEnabledTextBoxes()
+		{
+			NotifyUnassignedTextBox.Enabled = NotifyUnassignedCheckBox.Checked;
+			NotifyUnresolvedTextBox.Enabled = NotifyUnresolvedCheckBox.Checked;
+		}
+
+		private void NotifyUnassignedCheckBox_CheckedChanged(object sender, EventArgs e)
+		{
+			UpdateEnabledTextBoxes();
+		}
+
+		private void NotifyUnresolvedCheckBox_CheckedChanged(object sender, EventArgs e)
+		{
+			UpdateEnabledTextBoxes();
+		}
+
+		private void OkBtn_Click(object sender, EventArgs e)
+		{
+			int NewNotifyUnresolvedMinutes = -1;
+			if(NotifyUnresolvedCheckBox.Checked)
+			{
+				ushort NewNotifyUnresolvedMinutesValue;
+				if(!ushort.TryParse(NotifyUnresolvedTextBox.Text, out NewNotifyUnresolvedMinutesValue))
+				{
+					MessageBox.Show("Invalid time");
+					return;
+				}
+				NewNotifyUnresolvedMinutes = NewNotifyUnresolvedMinutesValue;
+			}
+
+			int NewNotifyUnassignedMinutes = -1;
+			if(NotifyUnassignedCheckBox.Checked)
+			{
+				ushort NewNotifyUnassignedMinutesValue;
+				if(!ushort.TryParse(NotifyUnassignedTextBox.Text, out NewNotifyUnassignedMinutesValue))
+				{
+					MessageBox.Show("Invalid time");
+					return;
+				}
+				NewNotifyUnassignedMinutes = NewNotifyUnassignedMinutesValue;
+			}
+
+			Settings.NotifyUnresolvedMinutes = NewNotifyUnresolvedMinutes;
+			Settings.NotifyUnassignedMinutes = NewNotifyUnassignedMinutes;
+			Settings.Save();
+
+			DialogResult = DialogResult.Cancel;
+			Close();
+		}
+
+		private void CancelBtn_Click(object sender, EventArgs e)
+		{
+			DialogResult = DialogResult.Cancel;
+			Close();
+		}
+
+		private void groupBox2_Enter(object sender, EventArgs e)
+		{
+
+		}
+	}
+}
