UGS: Include information about each failed step in a job, so the client can link to it as necessary. Also purge old jobs from the state file.

#rnx

diff --git a/Engine/Source/Programs/UnrealGameSync/MetadataServer/Connectors/SqlConnector.cs b/Engine/Source/Programs/UnrealGameSync/MetadataServer/Connectors/SqlConnector.cs
--- a/Engine/Source/Programs/UnrealGameSync/MetadataServer/Connectors/SqlConnector.cs
+++ b/Engine/Source/Programs/UnrealGameSync/MetadataServer/Connectors/SqlConnector.cs
@@ -612,20 +612,23 @@ namespace MetadataServer.Connectors
 			}
 		}
 
-		public static long AddBuild(long IssueId, string Stream, int Change, string Name, string Url, int Outcome)
+		public static long AddBuild(long IssueId, IssueBuildData Build)
 		{
 			using (SQLiteConnection Connection = new SQLiteConnection(SqlConnector.ConnectionString))
 			{
 				Connection.Open();
 
-				using (SQLiteCommand Command = new SQLiteCommand("INSERT INTO [IssueBuilds] (IssueId, Stream, Change, Name, Url, Outcome) VALUES (@IssueId, @Stream, @Change, @Name, @Url, @Outcome)", Connection))
+				using (SQLiteCommand Command = new SQLiteCommand("INSERT INTO [IssueBuilds] (IssueId, Stream, Change, JobName, JobUrl, JobStepName, JobStepUrl, ErrorUrl, Outcome) VALUES (@IssueId, @Stream, @Change, @JobName, @JobUrl, @JobStepName, @JobStepUrl, @ErrorUrl, @Outcome)", Connection))
 				{
 					Command.Parameters.AddWithValue("@IssueId", IssueId);
-					Command.Parameters.AddWithValue("@Stream", Stream);
-					Command.Parameters.AddWithValue("@Change", Change);
-					Command.Parameters.AddWithValue("@Name", Name);
-					Command.Parameters.AddWithValue("@Url", Url);
-					Command.Parameters.AddWithValue("@Outcome", Outcome);
+					Command.Parameters.AddWithValue("@Stream", Build.Stream);
+					Command.Parameters.AddWithValue("@Change", Build.Change);
+					Command.Parameters.AddWithValue("@JobName", Build.JobName);
+					Command.Parameters.AddWithValue("@JobUrl", Build.JobUrl);
+					Command.Parameters.AddWithValue("@JobStepName", Build.JobStepName);
+					Command.Parameters.AddWithValue("@JobStepUrl", Build.JobStepUrl);
+					Command.Parameters.AddWithValue("@ErrorUrl", Build.ErrorUrl);
+					Command.Parameters.AddWithValue("@Outcome", Build.Outcome);
 					Command.ExecuteNonQuery();
 
 					return Connection.LastInsertRowId;
@@ -640,7 +643,7 @@ namespace MetadataServer.Connectors
 			{
 				Connection.Open();
 
-				using(SQLiteCommand Command = new SQLiteCommand("SELECT IssueBuilds.Id, IssueBuilds.Stream, IssueBuilds.Change, IssueBuilds.Name, IssueBuilds.Url, IssueBuilds.Outcome FROM [IssueBuilds] WHERE IssueBuilds.IssueId = @IssueId", Connection))
+				using(SQLiteCommand Command = new SQLiteCommand("SELECT IssueBuilds.Id, IssueBuilds.Stream, IssueBuilds.Change, IssueBuilds.JobName, IssueBuilds.JobUrl, IssueBuilds.JobStepName, IssueBuilds.JobStepUrl, IssueBuilds.ErrorUrl, IssueBuilds.Outcome FROM [IssueBuilds] WHERE IssueBuilds.IssueId = @IssueId", Connection))
 				{
 					Command.Parameters.AddWithValue("@IssueId", IssueId);
 					using(SQLiteDataReader Reader = Command.ExecuteReader())
@@ -650,10 +653,13 @@ namespace MetadataServer.Connectors
 							long Id = Reader.GetInt64(0);
 							string Stream = Reader.GetString(1);
 							int Change = Reader.GetInt32(2);
-							string Name = Reader.GetString(3);
-							string Url = Reader.GetString(4);
-							int Outcome = Reader.GetInt32(5);
-							Builds.Add(new IssueBuildData { Id = Id, Stream = Stream, Change = Change, Name = Name, Url = Url, Outcome = Outcome });
+							string JobName = Reader.GetString(3);
+							string JobUrl = Reader.GetString(4);
+							string JobStepName = Reader.GetString(5);
+							string JobStepUrl = Reader.GetString(6);
+							string ErrorUrl = Reader.IsDBNull(7)? null : Reader.GetString(7);
+							int Outcome = Reader.GetInt32(8);
+							Builds.Add(new IssueBuildData { Id = Id, Stream = Stream, Change = Change, JobName = JobName, JobUrl = JobUrl, JobStepName = JobStepName, JobStepUrl = JobStepUrl, ErrorUrl = ErrorUrl, Outcome = Outcome });
 						}
 					}
 				}
@@ -668,7 +674,7 @@ namespace MetadataServer.Connectors
 			{
 				Connection.Open();
 
-				using(SQLiteCommand Command = new SQLiteCommand("SELECT IssueBuilds.Id, IssueBuilds.Stream, IssueBuilds.Change, IssueBuilds.Name, IssueBuilds.Url, IssueBuilds.Outcome FROM [IssueBuilds] WHERE IssueBuilds.Id = @BuildId", Connection))
+				using(SQLiteCommand Command = new SQLiteCommand("SELECT IssueBuilds.Id, IssueBuilds.Stream, IssueBuilds.Change, IssueBuilds.JobName, IssueBuilds.JobUrl, IssueBuilds.JobStepName, IssueBuilds.JobStepUrl, IssueBuilds.ErrorUrl, IssueBuilds.Outcome FROM [IssueBuilds] WHERE IssueBuilds.Id = @BuildId", Connection))
 				{
 					Command.Parameters.AddWithValue("@BuildId", BuildId);
 					using(SQLiteDataReader Reader = Command.ExecuteReader())
@@ -678,11 +684,14 @@ namespace MetadataServer.Connectors
 							long Id = Reader.GetInt64(0);
 							string Stream = Reader.GetString(1);
 							int Change = Reader.GetInt32(2);
-							string Name = Reader.GetString(3);
-							string Url = Reader.GetString(4);
-							int Outcome = Reader.GetInt32(5);
-
-							Build = new IssueBuildData { Id = Id, Stream = Stream, Change = Change, Name = Name, Url = Url, Outcome = Outcome };
+							string JobName = Reader.GetString(3);
+							string JobUrl = Reader.GetString(4);
+							string JobStepName = Reader.GetString(5);
+							string JobStepUrl = Reader.GetString(6);
+							string ErrorUrl = Reader.GetString(7);
+							int Outcome = Reader.GetInt32(8);
+
+							Build = new IssueBuildData { Id = Id, Stream = Stream, Change = Change, JobName = JobName, JobUrl = JobUrl, JobStepName = JobStepName, JobStepUrl = JobStepUrl, ErrorUrl = ErrorUrl, Outcome = Outcome };
 						}
 					}
 				}
diff --git a/Engine/Source/Programs/UnrealGameSync/MetadataServer/Controllers/IssuesController.cs b/Engine/Source/Programs/UnrealGameSync/MetadataServer/Controllers/IssuesController.cs
--- a/Engine/Source/Programs/UnrealGameSync/MetadataServer/Controllers/IssuesController.cs
+++ b/Engine/Source/Programs/UnrealGameSync/MetadataServer/Controllers/IssuesController.cs
@@ -58,7 +58,7 @@ namespace MetadataServer.Controllers
 		[HttpPost]
 		public void Post(long IssueId, [FromBody] IssueBuildData Data)
 		{
-			SqlConnector.AddBuild(IssueId, Data.Stream, Data.Change, Data.Name, Data.Url, Data.Outcome);
+			SqlConnector.AddBuild(IssueId, Data);
 		}
 	}
 
diff --git a/Engine/Source/Programs/UnrealGameSync/MetadataServer/Models/IssueData.cs b/Engine/Source/Programs/UnrealGameSync/MetadataServer/Models/IssueData.cs
--- a/Engine/Source/Programs/UnrealGameSync/MetadataServer/Models/IssueData.cs
+++ b/Engine/Source/Programs/UnrealGameSync/MetadataServer/Models/IssueData.cs
@@ -18,8 +18,11 @@ namespace MetadataServer.Models
 		public long Id;
 		public string Stream;
 		public int Change;
-		public string Name;
-		public string Url;
+		public string JobName;
+		public string JobUrl;
+		public string JobStepName;
+		public string JobStepUrl;
+		public string ErrorUrl;
 		public int Outcome;
 	}
 
diff --git a/Engine/Source/Programs/UnrealGameSync/MetadataTool/BuildHealth/InputData.cs b/Engine/Source/Programs/UnrealGameSync/MetadataTool/BuildHealth/InputData.cs
--- a/Engine/Source/Programs/UnrealGameSync/MetadataTool/BuildHealth/InputData.cs
+++ b/Engine/Source/Programs/UnrealGameSync/MetadataTool/BuildHealth/InputData.cs
@@ -37,8 +37,14 @@ namespace MetadataTool
 		/// <summary>
 		/// Name of the build step. Multiple step results can be submitted for a single build, and can be used to augment a single issue with additional diagnostic information.
 		/// </summary>
-		[DataMember]
-		public string Name = "Default";
+		[DataMember(IsRequired = true)]
+		public string Name;
+
+		/// <summary>
+		/// Url of this job step
+		/// </summary>
+		[DataMember(IsRequired = true)]
+		public string Url;
 
 		/// <summary>
 		/// Base directory that the build was executed in. Optional. This is used to determine branch-relative file paths for diagnostics.
@@ -67,22 +73,22 @@ namespace MetadataTool
 		public string Name;
 
 		/// <summary>
-		/// Stream that was built
+		/// Url for this job
 		/// </summary>
 		[DataMember(IsRequired = true)]
-		public string Stream;
+		public string Url;
 
 		/// <summary>
-		/// Changelist number that was built
+		/// Stream that was built
 		/// </summary>
 		[DataMember(IsRequired = true)]
-		public int Change;
+		public string Stream;
 
 		/// <summary>
-		/// Url for this job
+		/// Changelist number that was built
 		/// </summary>
 		[DataMember(IsRequired = true)]
-		public string Url;
+		public int Change;
 
 		/// <summary>
 		/// Steps that are part of this job
diff --git a/Engine/Source/Programs/UnrealGameSync/MetadataTool/BuildHealth/Matchers/CompilePatternMatcher.cs b/Engine/Source/Programs/UnrealGameSync/MetadataTool/BuildHealth/Matchers/CompilePatternMatcher.cs
--- a/Engine/Source/Programs/UnrealGameSync/MetadataTool/BuildHealth/Matchers/CompilePatternMatcher.cs
+++ b/Engine/Source/Programs/UnrealGameSync/MetadataTool/BuildHealth/Matchers/CompilePatternMatcher.cs
@@ -56,14 +56,11 @@ namespace MetadataTool
 			return false;
 		}
 
-		public override bool TryMerge(TrackedIssueFingerprint Source, TrackedIssueFingerprint Target)
+		public override void Merge(TrackedIssueFingerprint Source, TrackedIssueFingerprint Target)
 		{
-			if(base.TryMerge(Source, Target))
-			{
-				Target.Summary = GetSummary(Target.FileNames);
-				return true;
-			}
-			return false;
+			base.Merge(Source, Target);
+
+			Target.Summary = GetSummary(Target.FileNames);
 		}
 
 		string GetSummary(IEnumerable<string> SourceFileNames)
diff --git a/Engine/Source/Programs/UnrealGameSync/MetadataTool/BuildHealth/Matchers/UndefinedSymbolPatternMatcher.cs b/Engine/Source/Programs/UnrealGameSync/MetadataTool/BuildHealth/Matchers/UndefinedSymbolPatternMatcher.cs
--- a/Engine/Source/Programs/UnrealGameSync/MetadataTool/BuildHealth/Matchers/UndefinedSymbolPatternMatcher.cs
+++ b/Engine/Source/Programs/UnrealGameSync/MetadataTool/BuildHealth/Matchers/UndefinedSymbolPatternMatcher.cs
@@ -86,14 +86,11 @@ namespace MetadataTool
 			return false;
 		}
 
-		public override bool TryMerge(TrackedIssueFingerprint Source, TrackedIssueFingerprint Target)
+		public override void Merge(TrackedIssueFingerprint Source, TrackedIssueFingerprint Target)
 		{
-			if (base.TryMerge(Source, Target))
-			{
-				Target.Summary = GetSummary(Target.Messages);
-				return true;
-			}
-			return false;
+			base.Merge(Source, Target);
+
+			Target.Summary = GetSummary(Target.Messages);
 		}
 
 		string GetSummary(SortedSet<string> SymbolNames)
diff --git a/Engine/Source/Programs/UnrealGameSync/MetadataTool/BuildHealth/PatternMatcher.cs b/Engine/Source/Programs/UnrealGameSync/MetadataTool/BuildHealth/PatternMatcher.cs
--- a/Engine/Source/Programs/UnrealGameSync/MetadataTool/BuildHealth/PatternMatcher.cs
+++ b/Engine/Source/Programs/UnrealGameSync/MetadataTool/BuildHealth/PatternMatcher.cs
@@ -54,11 +54,18 @@ namespace MetadataTool
 		public abstract bool TryMatch(InputJob Job, InputJobStep JobStep, InputDiagnostic Diagnostic, List<TrackedIssueFingerprint> Fingerprints);
 
 		/// <summary>
-		/// Attempts to merge another fingerprint into this one
+		/// Determines if one fingerprint can be merged with another one
 		/// </summary>
-		/// <param name="Other">The issue to merge into this one</param>
-		public virtual bool TryMerge(TrackedIssueFingerprint Source, TrackedIssueFingerprint Target)
+		/// <param name="Source">The source fingerprint</param>
+		/// <param name="Target">The fingerprint to merge into</param>
+		public virtual bool CanMerge(TrackedIssueFingerprint Source, TrackedIssueFingerprint Target)
 		{
+			// Make sure the categories match
+			if (Source.Category != Target.Category)
+			{
+				return false;
+			}
+
 			// Check that a filename or message matches
 			if (Source.InitialChange != Target.InitialChange)
 			{
@@ -67,11 +74,18 @@ namespace MetadataTool
 					return false;
 				}
 			}
+			return true;
+		}
 
-			// Allow merging
+		/// <summary>
+		/// Merge one fingerprint with another
+		/// </summary>
+		/// <param name="Source">The source fingerprint</param>
+		/// <param name="Target">The fingerprint to merge into</param>
+		public virtual void Merge(TrackedIssueFingerprint Source, TrackedIssueFingerprint Target)
+		{
 			Target.FileNames.UnionWith(Source.FileNames);
 			Target.Messages.UnionWith(Source.Messages);
-			return true;
 		}
 
 		/// <summary>
@@ -81,7 +95,7 @@ namespace MetadataTool
 		/// <param name="Fingerprint">Fingerprint for the issue</param>
 		/// <param name="Changes">List of changes since the issue first occurred.</param>
 		/// <returns>List of changes which are causers for the issue</returns>
-		public virtual List<ChangeInfo> FindCausers(PerforceConnection Perforce, TrackedIssueFingerprint Fingerprint, List<ChangeInfo> Changes)
+		public virtual List<ChangeInfo> FindCausers(PerforceConnection Perforce, TrackedIssueFingerprint Fingerprint, IReadOnlyList<ChangeInfo> Changes)
 		{
 			SortedSet<string> FileNamesWithoutPath = TrackedIssueFingerprint.GetFileNamesWithoutPath(Fingerprint.FileNames);
 
@@ -101,7 +115,7 @@ namespace MetadataTool
 			}
 			else
 			{
-				return Changes;
+				return new List<ChangeInfo>(Changes);
 			}
 		}
 
diff --git a/Engine/Source/Programs/UnrealGameSync/MetadataTool/BuildHealth/TrackedBuild.cs b/Engine/Source/Programs/UnrealGameSync/MetadataTool/BuildHealth/TrackedBuild.cs
--- a/Engine/Source/Programs/UnrealGameSync/MetadataTool/BuildHealth/TrackedBuild.cs
+++ b/Engine/Source/Programs/UnrealGameSync/MetadataTool/BuildHealth/TrackedBuild.cs
@@ -14,38 +14,44 @@ namespace MetadataTool
 	/// Represents information about a build associated with a particular issue
 	/// </summary>
 	[DataContract]
-	[DebuggerDisplay("{Change}: {Name}")]
+	[DebuggerDisplay("{Change}: {JobName}")]
 	class TrackedBuild : IComparable<TrackedBuild>
 	{
 		/// <summary>
-		/// Name of this build
+		/// The changelist that this build was run at.
 		/// </summary>
 		[DataMember(IsRequired = true)]
-		public string Name;
+		public int Change;
 
 		/// <summary>
-		/// The changelist that this build was run at.
+		/// Name of this job
 		/// </summary>
 		[DataMember(IsRequired = true)]
-		public int Change;
+		public string JobName;
 
 		/// <summary>
-		/// Unique id to identify this job (typically the job url)
+		/// Url for this job
 		/// </summary>
 		[DataMember(IsRequired = true)]
-		public string UniqueId;
+		public string JobUrl;
 
 		/// <summary>
-		/// Url of the first error within this job
+		/// Name of this job step (typically null for sentinel builds)
 		/// </summary>
 		[DataMember(IsRequired = true)]
-		public string Url;
+		public string JobStepName;
 
 		/// <summary>
-		/// Set of step names that contributed to this build.
+		/// Url for this job step (typically null for sentinel builds)
 		/// </summary>
 		[DataMember(IsRequired = true)]
-		public HashSet<string> StepNames = new HashSet<string>();
+		public string JobStepUrl;
+
+		/// <summary>
+		/// Url of the first error within this job
+		/// </summary>
+		[DataMember]
+		public string ErrorUrl;
 
 		/// <summary>
 		/// Whether this build has been posted to the server or not
@@ -56,12 +62,14 @@ namespace MetadataTool
 		/// <summary>
 		/// Constructor
 		/// </summary>
-		public TrackedBuild(string Name, int Change, string UniqueId, string Url)
+		public TrackedBuild(int Change, string JobName, string JobUrl, string JobStepName, string JobStepUrl, string ErrorUrl)
 		{
-			this.Name = Name;
 			this.Change = Change;
-			this.UniqueId = UniqueId;
-			this.Url = Url;
+			this.JobName = JobName;
+			this.JobUrl = JobUrl;
+			this.JobStepName = JobStepName;
+			this.JobStepUrl = JobStepUrl;
+			this.ErrorUrl = ErrorUrl;
 		}
 
 		/// <summary>
@@ -74,10 +82,10 @@ namespace MetadataTool
 			int Delta = Change - Other.Change;
 			if (Delta == 0)
 			{
-				Delta = Name.CompareTo(Other.Name);
+				Delta = JobUrl.CompareTo(Other.JobUrl);
 				if(Delta == 0)
 				{
-					Delta = UniqueId.CompareTo(Other.UniqueId);
+					Delta = JobStepUrl.CompareTo(Other.JobStepUrl);
 				}
 			}
 			return Delta;
diff --git a/Engine/Source/Programs/UnrealGameSync/MetadataTool/BuildHealth/TrackedIssue.cs b/Engine/Source/Programs/UnrealGameSync/MetadataTool/BuildHealth/TrackedIssue.cs
--- a/Engine/Source/Programs/UnrealGameSync/MetadataTool/BuildHealth/TrackedIssue.cs
+++ b/Engine/Source/Programs/UnrealGameSync/MetadataTool/BuildHealth/TrackedIssue.cs
@@ -47,10 +47,10 @@ namespace MetadataTool
 		public int InitialChange;
 
 		/// <summary>
-		/// All the streams that are exhibiting this issue
+		/// Map of stream name -> step name -> history for builds exhibiting this issue
 		/// </summary>
 		[DataMember(IsRequired = true)]
-		public Dictionary<string, TrackedIssueHistory> Streams = new Dictionary<string, TrackedIssueHistory>();
+		public Dictionary<string, Dictionary<string, TrackedIssueHistory>> Streams = new Dictionary<string, Dictionary<string, TrackedIssueHistory>>();
 
 		/// <summary>
 		/// Set of changes which may have caused this issue. Used to de-duplicate issues between streams.
@@ -87,7 +87,7 @@ namespace MetadataTool
 		/// <returns>True if the issue can be closed</returns>
 		public bool CanClose()
 		{
-			return Streams.Values.All(x => x.NextSuccessfulBuild != null);
+			return Streams.Values.All(x => x.Values.All(y => y.NextSuccessfulBuild != null));
 		}
 
 		/// <summary>
@@ -96,7 +96,7 @@ namespace MetadataTool
 		/// <returns>Set of step names</returns>
 		public SortedSet<string> GetStepNames()
 		{
-			return new SortedSet<string>(Streams.Values.SelectMany(x => x.FailedBuilds.SelectMany(y => y.StepNames)));
+			return new SortedSet<string>(Streams.Values.SelectMany(x => x.Keys));
 		}
 
 		/// <summary>
diff --git a/Engine/Source/Programs/UnrealGameSync/MetadataTool/BuildHealth/TrackedIssueFingerprint.cs b/Engine/Source/Programs/UnrealGameSync/MetadataTool/BuildHealth/TrackedIssueFingerprint.cs
--- a/Engine/Source/Programs/UnrealGameSync/MetadataTool/BuildHealth/TrackedIssueFingerprint.cs
+++ b/Engine/Source/Programs/UnrealGameSync/MetadataTool/BuildHealth/TrackedIssueFingerprint.cs
@@ -44,7 +44,7 @@ namespace MetadataTool
 		/// Url for this issue
 		/// </summary>
 		[DataMember(IsRequired = true)]
-		public string Url;
+		public string ErrorUrl;
 
 		/// <summary>
 		/// List of files associated with this issue
@@ -63,13 +63,13 @@ namespace MetadataTool
 		/// </summary>
 		/// <param name="Category">Category for this issue</param>
 		/// <param name="Summary">The summary text for this issue</param>
-		/// <param name="Url">The Url for this issue</param>
+		/// <param name="ErrorUrl">The Url for this issue</param>
 		/// <param name="InitialChange">The initial change that this issue was seen on. Fingerprints in the same category at the same change can be merged as long as the type matches.</param>
-		public TrackedIssueFingerprint(string Category, string Summary, string Url, int InitialChange)
+		public TrackedIssueFingerprint(string Category, string Summary, string ErrorUrl, int InitialChange)
 		{
 			this.Category = Category;
 			this.Summary = Summary;
-			this.Url = Url;
+			this.ErrorUrl = ErrorUrl;
 			this.InitialChange = InitialChange;
 		}
 
diff --git a/Engine/Source/Programs/UnrealGameSync/MetadataTool/BuildHealth/TrackedIssueHistory.cs b/Engine/Source/Programs/UnrealGameSync/MetadataTool/BuildHealth/TrackedIssueHistory.cs
--- a/Engine/Source/Programs/UnrealGameSync/MetadataTool/BuildHealth/TrackedIssueHistory.cs
+++ b/Engine/Source/Programs/UnrealGameSync/MetadataTool/BuildHealth/TrackedIssueHistory.cs
@@ -36,10 +36,9 @@ namespace MetadataTool
 		/// <summary>
 		/// Constructs a new history for a particular stream
 		/// </summary>
-		public TrackedIssueHistory(TrackedBuild PrevSuccessfulBuild, TrackedBuild FailedBuild)
+		public TrackedIssueHistory(TrackedBuild PrevSuccessfulBuild)
 		{
-			this.PrevSuccessfulBuild = PrevSuccessfulBuild;
-			FailedBuilds.Add(FailedBuild);
+			this.PrevSuccessfulBuild  = PrevSuccessfulBuild;
 		}
 
 		/// <summary>
@@ -49,11 +48,7 @@ namespace MetadataTool
 		public void AddFailedBuild(TrackedBuild Build)
 		{
 			int Index = FailedBuilds.BinarySearch(Build);
-			if (Index >= 0)
-			{
-				FailedBuilds[Index].StepNames.UnionWith(Build.StepNames);
-			}
-			else
+			if (Index < 0)
 			{
 				FailedBuilds.Insert(~Index, Build);
 			}
@@ -96,7 +91,7 @@ namespace MetadataTool
 				{
 					yield return FailedBuild;
 				}
-				if (NextSuccessfulBuild != null)
+				if(NextSuccessfulBuild != null)
 				{
 					yield return NextSuccessfulBuild;
 				}
diff --git a/Engine/Source/Programs/UnrealGameSync/MetadataTool/BuildHealth/TrackedState.cs b/Engine/Source/Programs/UnrealGameSync/MetadataTool/BuildHealth/TrackedState.cs
--- a/Engine/Source/Programs/UnrealGameSync/MetadataTool/BuildHealth/TrackedState.cs
+++ b/Engine/Source/Programs/UnrealGameSync/MetadataTool/BuildHealth/TrackedState.cs
@@ -27,6 +27,30 @@ namespace MetadataTool
 		[DataMember]
 		public Dictionary<string, List<TrackedBuild>> Streams = new Dictionary<string, List<TrackedBuild>>();
 
+		/// <summary>
+		/// Adds a new build for tracking
+		/// </summary>
+		/// <param name="Stream">The stream containing the build</param>
+		/// <param name="Build">The build to add</param>
+		public void AddBuild(string Stream, TrackedBuild Build)
+		{
+			// Get the list of tracked builds for this stream
+			List<TrackedBuild> Builds;
+			if (!Streams.TryGetValue(Stream, out Builds))
+			{
+				Builds = new List<TrackedBuild>();
+				Streams.Add(Stream, Builds);
+			}
+
+			// Add this build to the tracked state data
+			int BuildIdx = Builds.BinarySearch(Build);
+			if (BuildIdx < 0)
+			{
+				BuildIdx = ~BuildIdx;
+				Builds.Insert(BuildIdx, Build);
+			}
+		}
+
 		/// <summary>
 		/// Finds the last build before the one given which executes the same steps. Used to determine the last succesful build before a failure.
 		/// </summary>
@@ -34,7 +58,7 @@ namespace MetadataTool
 		/// <param name="Change">The change to search for</param>
 		/// <param name="StepNames">The step names which the job step needs to have</param>
 		/// <returns>The last build known before the given changelist</returns>
-		public TrackedBuild FindBuildBefore(string Stream, int Change, IEnumerable<string> StepNames)
+		public TrackedBuild FindBuildBefore(string Stream, int Change, string StepName)
 		{
 			TrackedBuild Result = null;
 
@@ -43,7 +67,7 @@ namespace MetadataTool
 			{
 				for(int Idx = 0; Idx < Builds.Count && Builds[Idx].Change < Change; Idx++)
 				{
-					if(StepNames.All(x => Builds[Idx].StepNames.Contains(x)))
+					if(Builds[Idx].JobStepName == StepName)
 					{
 						Result = Builds[Idx];
 					}
@@ -58,9 +82,9 @@ namespace MetadataTool
 		/// </summary>
 		/// <param name="Stream">The stream to search for</param>
 		/// <param name="Change">The change to search for</param>
-		/// <param name="StepNames">The step names which the job step needs to have</param>
+		/// <param name="StepName">Name of the job step to find</param>
 		/// <returns>The last build known before the given changelist</returns>
-		public TrackedBuild FindBuildAfter(string Stream, int Change, IEnumerable<string> StepNames)
+		public TrackedBuild FindBuildAfter(string Stream, int Change, string StepName)
 		{
 			TrackedBuild Result = null;
 
@@ -69,7 +93,7 @@ namespace MetadataTool
 			{
 				for (int Idx = Builds.Count - 1; Idx >= 0 && Builds[Idx].Change > Change; Idx--)
 				{
-					if (StepNames.All(x => Builds[Idx].StepNames.Contains(x)))
+					if (Builds[Idx].JobStepName == StepName)
 					{
 						Result = Builds[Idx];
 					}
diff --git a/Engine/Source/Programs/UnrealGameSync/MetadataTool/CommandHandler_BuildHealth.cs b/Engine/Source/Programs/UnrealGameSync/MetadataTool/CommandHandler_BuildHealth.cs
--- a/Engine/Source/Programs/UnrealGameSync/MetadataTool/CommandHandler_BuildHealth.cs
+++ b/Engine/Source/Programs/UnrealGameSync/MetadataTool/CommandHandler_BuildHealth.cs
@@ -25,6 +25,18 @@ namespace MetadataTool
 			new UndefinedSymbolPatternMatcher()
 		};
 
+		static readonly Dictionary<string, PatternMatcher> CategoryToMatcher = Matchers.ToDictionary(x => x.Category, x => x);
+
+		class CachedChangeInfo
+		{
+			public string Stream;
+			public int PrevChange;
+			public int NextChange;
+			public IReadOnlyList<ChangeInfo> Changes;
+		}
+
+		List<CachedChangeInfo> CachedChanges = new List<CachedChangeInfo>();
+
 		/// <summary>
 		/// Constructor
 		/// </summary>
@@ -46,6 +58,7 @@ namespace MetadataTool
 			FileReference InputFile = Arguments.GetFileReference("-InputFile=");
 			FileReference StateFile = Arguments.GetFileReference("-StateFile=");
 			string ServerUrl = Arguments.GetStringOrDefault("-Server=", null);
+			bool bKeepHistory = Arguments.HasOption("-KeepHistory");
 			Arguments.CheckAllArgumentsUsed();
 
 			// Build a mapping from category to matching
@@ -79,28 +92,13 @@ namespace MetadataTool
 			List<InputJob> InputJobs = InputData.Jobs.OrderBy(x => x.Change).ThenBy(x => x.Stream).ToList();
 			foreach(InputJob InputJob in InputJobs)
 			{
-				// Get the list of tracked builds for this stream
-				List<TrackedBuild> StreamBuilds;
-				if (!State.Streams.TryGetValue(InputJob.Stream, out StreamBuilds))
-				{
-					StreamBuilds = new List<TrackedBuild>();
-					State.Streams.Add(InputJob.Stream, StreamBuilds);
-				}
-
-				// Create a new tracked build for this input build. Even if we discard this, we can use it for searching the existing build list for the place to insert.
-				TrackedBuild NewBuild = new TrackedBuild(InputJob.Name, InputJob.Change, InputJob.Url, InputJob.Url);
-
-				// Add this build to the tracked state data
-				int StreamBuildIndex = StreamBuilds.BinarySearch(NewBuild);
-				if (StreamBuildIndex < 0)
+				// Add a new build for each job step
+				foreach(InputJobStep InputJobStep in InputJob.Steps)
 				{
-					StreamBuildIndex = ~StreamBuildIndex;
-					StreamBuilds.Insert(StreamBuildIndex, NewBuild);
+					TrackedBuild NewBuild = new TrackedBuild(InputJob.Change, InputJob.Name, InputJob.Url, InputJobStep.Name, InputJobStep.Url, null);
+					State.AddBuild(InputJob.Stream, NewBuild);
 				}
 
-				// Get the build for this changelist
-				StreamBuilds[StreamBuildIndex].StepNames.UnionWith(InputJob.Steps.Select(x => x.Name));
-
 				// Add all the job steps
 				List<InputJobStep> InputJobSteps = InputJob.Steps.OrderBy(x => x.Name).ToList();
 				foreach (InputJobStep InputJobStep in InputJobSteps)
@@ -115,16 +113,50 @@ namespace MetadataTool
 				TrackedIssue Issue = State.Issues[Idx];
 				foreach(string Stream in Issue.Streams.Keys)
 				{
-					TrackedIssueHistory IssueHistory = Issue.Streams[Stream];
-					if(IssueHistory.FailedBuilds.Count > 0 && IssueHistory.NextSuccessfulBuild == null)
+					Dictionary<string, TrackedIssueHistory> StepNameToHistory = Issue.Streams[Stream];
+					foreach(string StepName in StepNameToHistory.Keys)
 					{
-						// Figure out all the step names that have to be 
-						HashSet<string> RequiredStepNames = new HashSet<string>(IssueHistory.FailedBuilds.SelectMany(x => x.StepNames));
+						TrackedIssueHistory IssueHistory = StepNameToHistory[StepName];
+						if(IssueHistory.FailedBuilds.Count > 0 && IssueHistory.NextSuccessfulBuild == null)
+						{
+							// Find the successful build after this change
+							TrackedBuild LastFailedBuild = IssueHistory.FailedBuilds[IssueHistory.FailedBuilds.Count - 1];
+							IssueHistory.NextSuccessfulBuild = State.FindBuildAfter(Stream, LastFailedBuild.Change, StepName);
+						}
+					}
+				}
+			}
 
-						// Find the successful build after this change
-						TrackedBuild LastFailedBuild = IssueHistory.FailedBuilds[IssueHistory.FailedBuilds.Count - 1];
-						IssueHistory.NextSuccessfulBuild = DuplicateSentinelBuild(State.FindBuildAfter(Stream, LastFailedBuild.Change, LastFailedBuild.StepNames));
+			// Find the change two days before the latest change being added
+			if(InputData.Jobs.Count > 0 && !bKeepHistory)
+			{
+				// Find all the unique change numbers for each stream
+				SortedSet<int> ChangeNumbers = new SortedSet<int>();
+				foreach (List<TrackedBuild> Builds in State.Streams.Values)
+				{
+					ChangeNumbers.UnionWith(Builds.Select(x => x.Change));
+				}
+
+				// Get the latest change record
+				int LatestChangeNumber = InputData.Jobs.Min(x => x.Change);
+				ChangeRecord LatestChangeRecord = Perforce.GetChange(GetChangeOptions.None, LatestChangeNumber).Data;
+
+				// Step forward through all the changelists until we get to one we don't want to delete
+				int DeleteChangeNumber = -1;
+				foreach(int ChangeNumber in ChangeNumbers)
+				{
+					ChangeRecord ChangeRecord = Perforce.GetChange(GetChangeOptions.None, ChangeNumber).Data;
+					if (ChangeRecord.Date > LatestChangeRecord.Date - TimeSpan.FromDays(2))
+					{
+						break;
 					}
+					DeleteChangeNumber = ChangeNumber;
+				}
+
+				// Remove any builds we no longer want to track
+				foreach (List<TrackedBuild> Builds in State.Streams.Values)
+				{
+					Builds.RemoveAll(x => x.Change <= DeleteChangeNumber);
 				}
 			}
 
@@ -211,33 +243,38 @@ namespace MetadataTool
 				// Add any new builds associated with issues
 				foreach (TrackedIssue Issue in State.Issues)
 				{
-					foreach(KeyValuePair<string, TrackedIssueHistory> StreamPair in Issue.Streams)
+					foreach(KeyValuePair<string, Dictionary<string, TrackedIssueHistory>> StreamPair in Issue.Streams)
 					{
-						TrackedIssueHistory StreamHistory = StreamPair.Value;
-						foreach(TrackedBuild Build in StreamHistory.Builds)
+						foreach(TrackedIssueHistory StreamHistory in StreamPair.Value.Values)
 						{
-							if(!Build.bPostedToServer)
+							foreach(TrackedBuild Build in StreamHistory.Builds)
 							{
-								Log.TraceInformation("Adding {0} to issue {1}", Build.UniqueId, Issue.Id);
-
-								CommandTypes.AddBuild AddBuild = new CommandTypes.AddBuild();
-								AddBuild.Change = Build.Change;
-								AddBuild.Stream = StreamPair.Key;
-								AddBuild.Name = Build.Name;
-								AddBuild.Outcome = (Build == StreamHistory.PrevSuccessfulBuild || Build == StreamHistory.NextSuccessfulBuild)? CommandTypes.Outcome.Success : CommandTypes.Outcome.Error;
-								AddBuild.Url = Build.Url;
-
-								using (HttpWebResponse Response = SendHttpRequest(String.Format("{0}/api/issues/{1}/builds", ServerUrl, Issue.Id), "POST", AddBuild))
+								if(!Build.bPostedToServer)
 								{
-									int ResponseCode = (int)Response.StatusCode;
-									if (!(ResponseCode >= 200 && ResponseCode <= 299))
+									Log.TraceInformation("Adding {0} to issue {1}", Build.JobStepUrl, Issue.Id);
+
+									CommandTypes.AddBuild AddBuild = new CommandTypes.AddBuild();
+									AddBuild.Stream = StreamPair.Key;
+									AddBuild.Change = Build.Change;
+									AddBuild.JobName = Build.JobName;
+									AddBuild.JobUrl = Build.JobUrl;
+									AddBuild.JobStepName = Build.JobStepName;
+									AddBuild.JobStepUrl = Build.JobStepUrl;
+									AddBuild.ErrorUrl = Build.ErrorUrl;
+									AddBuild.Outcome = (Build == StreamHistory.PrevSuccessfulBuild || Build == StreamHistory.NextSuccessfulBuild)? CommandTypes.Outcome.Success : CommandTypes.Outcome.Error;
+
+									using (HttpWebResponse Response = SendHttpRequest(String.Format("{0}/api/issues/{1}/builds", ServerUrl, Issue.Id), "POST", AddBuild))
 									{
-										throw new Exception("Unable to add build");
+										int ResponseCode = (int)Response.StatusCode;
+										if (!(ResponseCode >= 200 && ResponseCode <= 299))
+										{
+											throw new Exception("Unable to add build");
+										}
 									}
-								}
 
-								Build.bPostedToServer = true;
-								SerializeJson(StateFile, State);
+									Build.bPostedToServer = true;
+									SerializeJson(StateFile, State);
+								}
 							}
 						}
 					}
@@ -350,44 +387,6 @@ namespace MetadataTool
 			Log.TraceInformation("  -Clean             Removes the existing state file and creates a new one");
 		}
 
-		/// <summary>
-		/// Finds the matcher for a particular category
-		/// </summary>
-		/// <param name="Category">The category to find</param>
-		/// <returns>Pattern matcher for this category</returns>
-		bool TryMergeFingerprint(TrackedIssueFingerprint Source, TrackedIssueFingerprint Target)
-		{
-			if(Source.Category != Target.Category)
-			{
-				return false;
-			}
-
-			PatternMatcher Matcher = Matchers.First(x => x.Category == Source.Category);
-			if(!Matcher.TryMerge(Source, Target))
-			{
-				return false;
-			}
-
-			return true;
-		}
-
-		/// <summary>
-		/// Duplicates a build for use as a sentinel value
-		/// </summary>
-		/// <param name="Build">The build to duplicate</param>
-		/// <returns>The new build instance</returns>
-		static TrackedBuild DuplicateSentinelBuild(TrackedBuild Build)
-		{
-			if(Build == null)
-			{
-				return null;
-			}
-			else
-			{
-				return new TrackedBuild(Build.Name, Build.Change, Build.UniqueId, Build.Url);
-			}
-		}
-
 		/// <summary>
 		/// Adds diagnostics from a job step into the issue database
 		/// </summary>
@@ -410,48 +409,68 @@ namespace MetadataTool
 				}
 			}
 
-			// Early out if there are no remaining fingerprints to add
-			if (Fingerprints.Count == 0)
+			// Add all the fingerprints to issues
+			foreach(TrackedIssueFingerprint Fingerprint in Fingerprints)
 			{
-				return;
+				TrackedIssue Issue = FindOrAddIssueForFingerprint(Perforce, State, InputJob, InputJobStep, Fingerprint);
+				AddFailureToIssue(Issue, Fingerprint, InputJob, InputJobStep, State);
 			}
+		}
 
-			// Merge any fingerprints with open issues
-			for (int Idx = 0; Idx < Fingerprints.Count; Idx++)
+		/// <summary>
+		/// Finds or adds an issue for a particular fingerprint
+		/// </summary>
+		/// <param name="Perforce">Perforce connection used to find possible causers</param>
+		/// <param name="State">The current set of tracked issues</param>
+		/// <param name="Build">The new build</param>
+		/// <param name="PreviousChange">The last changelist that was built before this one</param>
+		/// <param name="InputJob">Job containing the step to add</param>
+		/// <param name="InputJobStep">The job step to add</param>
+		TrackedIssue FindOrAddIssueForFingerprint(PerforceConnection Perforce, TrackedState State, InputJob InputJob, InputJobStep InputJobStep, TrackedIssueFingerprint Fingerprint)
+		{
+			// Find the pattern matcher for this fingerprint
+			PatternMatcher Matcher = CategoryToMatcher[Fingerprint.Category];
+
+			// Check if it can be added to an existing open issue
+			foreach (TrackedIssue Issue in State.Issues)
 			{
-				TrackedIssueFingerprint Fingerprint = Fingerprints[Idx];
-				foreach (TrackedIssue Issue in State.Issues)
+				// Check this issue already exists in the current stream
+				Dictionary<string, TrackedIssueHistory> StepNameToHistory;
+				if(!Issue.Streams.TryGetValue(InputJob.Stream, out StepNameToHistory))
+				{
+					continue;
+				}
+
+				// Check that this issue has not already been closed
+				TrackedIssueHistory History;
+				if (StepNameToHistory.TryGetValue(InputJobStep.Name, out History))
 				{
-					// Check that this issue is present in the current stream, and that this build is not either side of a successful build
-					TrackedIssueHistory History;
-					if (!Issue.Streams.TryGetValue(InputJob.Stream, out History))
-					{
-						continue;
-					}
 					if(!History.CanAddFailedBuild(InputJob.Change))
 					{
 						continue;
 					}
-					if(!TryMergeFingerprint(Fingerprint, Issue.Fingerprint))
+				}
+				else
+				{
+					if(!StepNameToHistory.Values.Any(x => x.CanAddFailedBuild(InputJob.Change)))
 					{
 						continue;
 					}
+				}
 
-					// Merge the issue
-					History.AddFailedBuild(CreateBuildForJobStep(InputJob, InputJobStep, Fingerprint));
-					Fingerprints.RemoveAt(Idx--);
-					break;
+				// Try to merge the fingerprint
+				if(!Matcher.CanMerge(Fingerprint, Issue.Fingerprint))
+				{
+					continue;
 				}
-			}
 
-			// Early out if there are no remaining issues
-			if (Fingerprints.Count == 0)
-			{
-				return;
+				// Add the new build
+				Matcher.Merge(Fingerprint, Issue.Fingerprint);
+				return Issue;
 			}
 
 			// List of changes since the last successful build in this stream
-			List<ChangeInfo> Changes = null;
+			IReadOnlyList<ChangeInfo> Changes = null;
 
 			// Find the previous changelist that was built in this stream
 			List<TrackedBuild> StreamBuilds;
@@ -472,64 +491,43 @@ namespace MetadataTool
 					if (Changes.Count > 0)
 					{
 						SortedSet<int> SourceChanges = new SortedSet<int>(Changes.SelectMany(x => x.SourceChanges));
-						for (int Idx = 0; Idx < Fingerprints.Count; Idx++)
+						foreach (TrackedIssue Issue in State.Issues)
 						{
-							TrackedIssueFingerprint Fingerprint = Fingerprints[Idx];
-							foreach (TrackedIssue Issue in State.Issues)
+							// Check if this issue does not already contain this stream, but contains one of the causing changes
+							if (Issue.Streams.ContainsKey(InputJob.Stream))
 							{
-								// Check if this issue does not already contain this stream, but contains one of the causing changes
-								if (Issue.Streams.ContainsKey(InputJob.Stream))
-								{
-									continue;
-								}
-								if(!SourceChanges.Any(x => Issue.SourceChanges.Contains(x)))
-								{
-									continue;
-								}
-								if (!TryMergeFingerprint(Fingerprint, Issue.Fingerprint))
-								{
-									continue;
-								}
-
-								// Merge the issue
-								AddFailureToIssue(Issue, Fingerprint, InputJob, InputJobStep, State);
-								Fingerprints.RemoveAt(Idx--);
-								break;
+								continue;
+							}
+							if(!SourceChanges.Any(x => Issue.SourceChanges.Contains(x)))
+							{
+								continue;
 							}
+							if (!Matcher.CanMerge(Fingerprint, Issue.Fingerprint))
+							{
+								continue;
+							}
+
+							// Merge the issue
+							Matcher.Merge(Fingerprint, Issue.Fingerprint);
+							return Issue;
 						}
 					}
 				}
 			}
 
-			// Early out if there are no remaining issues
-			if (Fingerprints.Count == 0)
-			{
-				return;
-			}
-
 			// Create new issues for everything else in this stream
-			List<TrackedIssue> NewIssues = new List<TrackedIssue>();
-			foreach(TrackedIssueFingerprint Fingerprint in Fingerprints)
+			TrackedIssue NewIssue = new TrackedIssue(Fingerprint, InputJob.Change);
+			if (Changes != null)
 			{
-				TrackedIssue Issue = NewIssues.FirstOrDefault(NewIssue => TryMergeFingerprint(Fingerprint, NewIssue.Fingerprint));
-				if(Issue == null)
+				List<ChangeInfo> Causers = Matcher.FindCausers(Perforce, Fingerprint, Changes);
+				foreach(ChangeInfo Causer in Causers)
 				{
-					Issue = new TrackedIssue(Fingerprint, InputJob.Change);
-					if(Changes != null)
-					{
-						PatternMatcher Matcher = Matchers.First(x => x.Category == Fingerprint.Category);
-						List<ChangeInfo> Causers = Matcher.FindCausers(Perforce, Fingerprint, Changes);
-						foreach(ChangeInfo Causer in Causers)
-						{
-							Issue.SourceChanges.UnionWith(Causer.SourceChanges);
-							Issue.PendingWatchers.Add(Causer.Record.User);
-						}
-					}
-					NewIssues.Add(Issue);
+					NewIssue.SourceChanges.UnionWith(Causer.SourceChanges);
+					NewIssue.PendingWatchers.Add(Causer.Record.User);
 				}
-				AddFailureToIssue(Issue, Fingerprint, InputJob, InputJobStep, State);
 			}
-			State.Issues.AddRange(NewIssues);
+			State.Issues.Add(NewIssue);
+			return NewIssue;
 		}
 
 		/// <summary>
@@ -540,9 +538,7 @@ namespace MetadataTool
 		/// <returns>New build instance</returns>
 		TrackedBuild CreateBuildForJobStep(InputJob InputJob, InputJobStep InputJobStep, TrackedIssueFingerprint Fingerprint)
 		{
-			TrackedBuild FailedBuild = new TrackedBuild(InputJob.Name, InputJob.Change, InputJob.Url, Fingerprint.Url);
-			FailedBuild.StepNames.Add(InputJobStep.Name);
-			return FailedBuild;
+			return new TrackedBuild(InputJob.Change, InputJob.Name, InputJob.Url, InputJobStep.Name, InputJobStep.Url, Fingerprint.ErrorUrl);
 		}
 
 		/// <summary>
@@ -555,17 +551,24 @@ namespace MetadataTool
 		/// <param name="State">Current persistent state. Used to find previous build history.</param>
 		void AddFailureToIssue(TrackedIssue Issue, TrackedIssueFingerprint Fingerprint, InputJob InputJob, InputJobStep InputJobStep, TrackedState State)
 		{
-			TrackedIssueHistory History;
-			if(Issue.Streams.TryGetValue(InputJob.Stream, out History))
+			// Find or add a step name to history mapping
+			Dictionary<string, TrackedIssueHistory> StepNameToHistory;
+			if(!Issue.Streams.TryGetValue(InputJob.Stream, out StepNameToHistory))
 			{
-				History.AddFailedBuild(CreateBuildForJobStep(InputJob, InputJobStep, Fingerprint));
+				StepNameToHistory = new Dictionary<string, TrackedIssueHistory>();
+				Issue.Streams.Add(InputJob.Stream, StepNameToHistory);
 			}
-			else
+
+			// Find or add a history for this step
+			TrackedIssueHistory History;
+			if(!StepNameToHistory.TryGetValue(InputJobStep.Name, out History))
 			{
-				TrackedBuild LastSuccessfulBuild = DuplicateSentinelBuild(State.FindBuildBefore(InputJob.Stream, InputJob.Change, new HashSet<string>{ InputJobStep.Name }));
-				History = new TrackedIssueHistory(LastSuccessfulBuild, CreateBuildForJobStep(InputJob, InputJobStep, Fingerprint));
-				Issue.Streams.Add(InputJob.Stream, History);
+				History = new TrackedIssueHistory(State.FindBuildBefore(InputJob.Stream, InputJob.Change, InputJobStep.Name));
+				StepNameToHistory.Add(InputJobStep.Name, History);
 			}
+
+			// Add the new build
+			History.AddFailedBuild(CreateBuildForJobStep(InputJob, InputJobStep, Fingerprint));
 		}
 
 		/// <summary>
@@ -576,35 +579,43 @@ namespace MetadataTool
 		/// <param name="PrevChange">The first change in the range to query</param>
 		/// <param name="NextChange">The last change in the range to query</param>
 		/// <returns>Set of changelist numbers</returns>
-		List<ChangeInfo> FindChanges(PerforceConnection Perforce, string Stream, int PrevChange, int NextChange)
+		IReadOnlyList<ChangeInfo> FindChanges(PerforceConnection Perforce, string Stream, int PrevChange, int NextChange)
 		{
-			// Query for all the changes since then
-			List<ChangesRecord> ChangeRecords = Perforce.Changes(ChangesOptions.LongOutput, null, -1, ChangeStatus.Submitted, null, String.Format("{0}/...@{1},{2}", Stream, PrevChange, NextChange)).Data.ToList();
-
-			// Figure out all the original changelists that these were merged from, and see if any of those matches with an existing issue
-			List<ChangeInfo> Changes = new List<ChangeInfo>();
-			foreach(ChangesRecord ChangeRecord in ChangeRecords)
+			CachedChangeInfo CachedInfo = CachedChanges.FirstOrDefault(x => x.Stream == Stream && x.PrevChange == PrevChange && x.NextChange == NextChange);
+			if(CachedInfo == null)
 			{
-				ChangeInfo Change = new ChangeInfo();
-				Change.Record = ChangeRecord;
-				Change.SourceChanges.Add(ChangeRecord.Number);
-				Changes.Add(Change);
+				// Query for all the changes since then
+				List<ChangesRecord> ChangeRecords = Perforce.Changes(ChangesOptions.LongOutput, null, -1, ChangeStatus.Submitted, null, String.Format("{0}/...@{1},{2}", Stream, PrevChange, NextChange)).Data.ToList();
 
-				Match SourceMatch = Regex.Match(ChangeRecord.Description, "^#ROBOMERGE-SOURCE: (.*)$", RegexOptions.Multiline);
-				if(SourceMatch.Success)
+				// Figure out all the original changelists that these were merged from, and see if any of those matches with an existing issue
+				List<ChangeInfo> Changes = new List<ChangeInfo>();
+				foreach (ChangesRecord ChangeRecord in ChangeRecords)
 				{
-					string SourceText = SourceMatch.Groups[1].Value;
-					foreach(Match ChangeMatch in Regex.Matches(SourceText, "CL\\s*(\\d+)"))
+					ChangeInfo Change = new ChangeInfo();
+					Change.Record = ChangeRecord;
+					Change.SourceChanges.Add(ChangeRecord.Number);
+					Changes.Add(Change);
+
+					Match SourceMatch = Regex.Match(ChangeRecord.Description, "^#ROBOMERGE-SOURCE: (.*)$", RegexOptions.Multiline);
+					if (SourceMatch.Success)
 					{
-						int SourceChange;
-						if(int.TryParse(ChangeMatch.Groups[1].Value, out SourceChange))
+						string SourceText = SourceMatch.Groups[1].Value;
+						foreach (Match ChangeMatch in Regex.Matches(SourceText, "CL\\s*(\\d+)"))
 						{
-							Change.SourceChanges.Add(SourceChange);
+							int SourceChange;
+							if (int.TryParse(ChangeMatch.Groups[1].Value, out SourceChange))
+							{
+								Change.SourceChanges.Add(SourceChange);
+							}
 						}
 					}
 				}
+
+				// Create the new cached info
+				CachedInfo = new CachedChangeInfo() { Stream = Stream, PrevChange = PrevChange, NextChange = NextChange, Changes = Changes };
+				CachedChanges.Add(CachedInfo);
 			}
-			return Changes;
+			return CachedInfo.Changes;
 		}
 
 		/// <summary>
diff --git a/Engine/Source/Programs/UnrealGameSync/MetadataTool/CommandTypes.cs b/Engine/Source/Programs/UnrealGameSync/MetadataTool/CommandTypes.cs
--- a/Engine/Source/Programs/UnrealGameSync/MetadataTool/CommandTypes.cs
+++ b/Engine/Source/Programs/UnrealGameSync/MetadataTool/CommandTypes.cs
@@ -54,10 +54,13 @@ namespace MetadataTool
 
 		public class AddBuild
 		{
-			public string Name;
-			public string Url;
 			public string Stream;
 			public int Change;
+			public string JobName;
+			public string JobUrl;
+			public string JobStepName;
+			public string JobStepUrl;
+			[Optional] public string ErrorUrl;
 			public Outcome Outcome;
 		}
 
diff --git a/Engine/Source/Programs/UnrealGameSync/UnrealGameSync/Controls/CustomListViewControl.cs b/Engine/Source/Programs/UnrealGameSync/UnrealGameSync/Controls/CustomListViewControl.cs
--- a/Engine/Source/Programs/UnrealGameSync/UnrealGameSync/Controls/CustomListViewControl.cs
+++ b/Engine/Source/Programs/UnrealGameSync/UnrealGameSync/Controls/CustomListViewControl.cs
@@ -304,7 +304,14 @@ namespace UnrealGameSync
 				ConditionalRedrawWidget(MouseDownWidget);
 			}
 
-			MouseDownWidget = FindWidget(e.Location);
+			if((e.Button & MouseButtons.Left) != 0)
+			{
+				MouseDownWidget = FindWidget(e.Location);
+			}
+			else
+			{
+				MouseDownWidget = null;
+			}
 
 			if(MouseDownWidget != null)
 			{
diff --git a/Engine/Source/Programs/UnrealGameSync/UnrealGameSync/Forms/IssueAlertWindow.cs b/Engine/Source/Programs/UnrealGameSync/UnrealGameSync/Forms/IssueAlertWindow.cs
--- a/Engine/Source/Programs/UnrealGameSync/UnrealGameSync/Forms/IssueAlertWindow.cs
+++ b/Engine/Source/Programs/UnrealGameSync/UnrealGameSync/Forms/IssueAlertWindow.cs
@@ -219,7 +219,7 @@ namespace UnrealGameSync
 			IssueBuildData LastBuild = Issue.Builds.OrderByDescending(x => x.Change).FirstOrDefault();
 			if(LastBuild != null)
 			{
-				System.Diagnostics.Process.Start(LastBuild.Url);
+				System.Diagnostics.Process.Start(LastBuild.ErrorUrl);
 			}
 		}
 	}
diff --git a/Engine/Source/Programs/UnrealGameSync/UnrealGameSync/Forms/IssueBrowserWindow.cs b/Engine/Source/Programs/UnrealGameSync/UnrealGameSync/Forms/IssueBrowserWindow.cs
--- a/Engine/Source/Programs/UnrealGameSync/UnrealGameSync/Forms/IssueBrowserWindow.cs
+++ b/Engine/Source/Programs/UnrealGameSync/UnrealGameSync/Forms/IssueBrowserWindow.cs
@@ -156,10 +156,11 @@ namespace UnrealGameSync
 
 		private void DetailsBtn_Click(object sender, EventArgs e)
 		{
-			foreach(ListViewItem Item in IssueListView.Items)
+			foreach(ListViewItem Item in IssueListView.SelectedItems)
 			{
 				IssueData Issue = (IssueData)Item.Tag;
 				ShowIssue(Issue);
+				break;
 			}
 		}
 
diff --git a/Engine/Source/Programs/UnrealGameSync/UnrealGameSync/Forms/IssueDetailsWindow.cs b/Engine/Source/Programs/UnrealGameSync/UnrealGameSync/Forms/IssueDetailsWindow.cs
--- a/Engine/Source/Programs/UnrealGameSync/UnrealGameSync/Forms/IssueDetailsWindow.cs
+++ b/Engine/Source/Programs/UnrealGameSync/UnrealGameSync/Forms/IssueDetailsWindow.cs
@@ -16,6 +16,15 @@ namespace UnrealGameSync
 {
 	partial class IssueDetailsWindow : Form
 	{
+		class BuildGroup
+		{
+			public string JobName;
+			public string JobUrl;
+			public int Change;
+			public IssueBuildOutcome Outcome;
+			public List<IssueBuildData> Builds;
+		}
+
 		class PerforceChangeRange
 		{
 			public bool bExpanded;
@@ -23,7 +32,7 @@ namespace UnrealGameSync
 			public int MaxChange;
 			public List<PerforceChangeSummary> Changes;
 			public string ErrorMessage;
-			public IssueBuildData Build;
+			public BuildGroup BuildGroup;
 		}
 
 		class PerforceChangeDetailsWithDescribeRecord : PerforceChangeDetails
@@ -78,6 +87,7 @@ namespace UnrealGameSync
 				if(WorkerThread != null)
 				{
 					bTerminate = true;
+					RefreshEvent.Set();
 					if(!WorkerThread.Join(100))
 					{
 						WorkerThread.Abort();
@@ -332,14 +342,15 @@ namespace UnrealGameSync
 		{
 			if(TextBox.Text != NewText)
 			{
-				TextBox.SelectionLength = 0;
 				TextBox.Text = NewText;
+				TextBox.SelectionLength = 0;
+				TextBox.SelectionStart = NewText.Length;
 			}
 		}
 
 		void UpdateCurrentIssue()
 		{
-			SummaryTextBox.Text = Issue.Summary.ToString();
+			UpdateSummaryTextIfChanged(SummaryTextBox, Issue.Summary.ToString());
 
 			StringBuilder Status = new StringBuilder();
 			if(IssueMonitor.HasPendingUpdate())
@@ -412,6 +423,9 @@ namespace UnrealGameSync
 				BuildListView.Invalidate();
 			}
 
+			UpdateSummaryTextIfChanged(StepNamesTextBox, String.Join(", ", Issue.Builds.Select(x => x.JobStepName).Distinct().OrderBy(x => x)));
+			UpdateSummaryTextIfChanged(StreamNamesTextBox, String.Join(", ", Issue.Builds.Select(x => x.Stream).Distinct().OrderBy(x => x)));
+
 			UpdateSummaryTextIfChanged(DetailsTextBox, Issue.Details.Replace("\n", "\r\n"));
 
 			ReopenBtn.Enabled = (Issue.FixChange != 0);
@@ -612,13 +626,15 @@ namespace UnrealGameSync
 				}
 				*/
 				ListViewItem BuildItem = new ListViewItem("");
-				BuildItem.Tag = Range.Build;
-				BuildItem.SubItems.Add(Range.Build.Change.ToString());
+				BuildItem.Tag = Range.BuildGroup;
+				BuildItem.SubItems.Add(Range.BuildGroup.Change.ToString());
 				BuildItem.SubItems.Add("");
 
 				StatusLineListViewWidget BuildWidget = new StatusLineListViewWidget(BuildItem, StatusElementResources);
 				BuildWidget.HorizontalAlignment = HorizontalAlignment.Left;
-				BuildWidget.Line.AddLink(Range.Build.Name, FontStyle.Underline, () => System.Diagnostics.Process.Start(Range.Build.Url));
+
+				string DefaultUrl = Range.BuildGroup.Builds.Select(x => x.ErrorUrl).FirstOrDefault(x => x != null) ?? Range.BuildGroup.JobUrl;
+				BuildWidget.Line.AddLink(Range.BuildGroup.JobName, FontStyle.Underline, () => System.Diagnostics.Process.Start(DefaultUrl));
 				BuildItem.SubItems.Add(new ListViewItem.ListViewSubItem(BuildItem, ""){ Tag = BuildWidget });
 				BuildItem.SubItems.Add(new ListViewItem.ListViewSubItem(BuildItem, ""){ Tag = BuildWidget });
 
@@ -644,7 +660,35 @@ namespace UnrealGameSync
 			BuildListView.EndUpdate();
 			BuildListView.Invalidate();
 		}
-		
+
+		BuildGroup SelectedBuildGroup;
+
+		void ShowJobContextMenu(Point Point, BuildGroup BuildGroup)
+		{
+			SelectedBuildGroup = BuildGroup;
+
+			bool bHasErrorUrl = BuildGroup.Builds.Any(x => x.ErrorUrl != null);
+			JobContextMenu_StepSeparatorMax.Visible = bHasErrorUrl;
+			JobContextMenu_ShowFirstError.Visible = bHasErrorUrl;
+
+			int MinIndex = JobContextMenu.Items.IndexOf(JobContextMenu_StepSeparatorMin) + 1;
+			while(JobContextMenu.Items[MinIndex] != JobContextMenu_StepSeparatorMax)
+			{
+				JobContextMenu.Items.RemoveAt(MinIndex);
+			}
+
+			JobContextMenu_ViewJob.Text = String.Format("View Job: {0}", BuildGroup.JobName);
+
+			foreach (IssueBuildData Build in BuildGroup.Builds.OrderBy(x => x.JobStepName))
+			{
+				ToolStripMenuItem MenuItem = new ToolStripMenuItem(String.Format("View Step: {0}", Build.JobStepName));
+				MenuItem.Click += (S, E) => System.Diagnostics.Process.Start(Build.JobStepUrl);
+				JobContextMenu.Items.Insert(MinIndex++, MenuItem);
+			}
+
+			JobContextMenu.Show(BuildListView, Point, ToolStripDropDownDirection.BelowRight);
+		}
+
 		void UpdateChangeTypeWidget(StatusLineListViewWidget TypeWidget, PerforceChangeDetails Details)
 		{
 			TypeWidget.Line.Clear();
@@ -682,15 +726,22 @@ namespace UnrealGameSync
 					SelectedStreamRanges = new List<PerforceChangeRange>();
 
 					int MaxChange = -1;
-					foreach(IssueBuildData Build in Issue.Builds.Where(x => x.Stream == SelectedStream).OrderByDescending(x => x.Change).ThenByDescending(x => x.Url))
+					foreach(IGrouping<string, IssueBuildData> Group in Issue.Builds.Where(x => x.Stream == SelectedStream).OrderByDescending(x => x.Change).ThenByDescending(x => x.JobUrl).GroupBy(x => x.JobUrl))
 					{
+						BuildGroup BuildGroup = new BuildGroup();
+						BuildGroup.JobName = Group.First().JobName;
+						BuildGroup.JobUrl = Group.Key;
+						BuildGroup.Change = Group.First().Change;
+						BuildGroup.Builds = Group.ToList();
+						BuildGroup.Outcome = Group.Any(x => x.Outcome == IssueBuildOutcome.Error)? IssueBuildOutcome.Error : Group.Any(x => x.Outcome == IssueBuildOutcome.Warning)? IssueBuildOutcome.Warning : IssueBuildOutcome.Success;
+
 						PerforceChangeRange Range = new PerforceChangeRange();
-						Range.MinChange = Build.Change + 1;
+						Range.MinChange = BuildGroup.Change + 1;
 						Range.MaxChange = MaxChange;
-						Range.Build = Build;
+						Range.BuildGroup = BuildGroup;
 						SelectedStreamRanges.Add(Range);
 
-						MaxChange = Build.Change;
+						MaxChange = BuildGroup.Change;
 					}
 
 					PerforceWorker = new PerforceWorkerThread(Perforce, String.Format("{0}/...", SelectedStream), OnRequestCompleteAsync, UpdateChangeMetadataAsync, Log);
@@ -699,7 +750,7 @@ namespace UnrealGameSync
 
 					for(int Idx = 0; Idx + 2 < SelectedStreamRanges.Count; Idx++)
 					{
-						if(SelectedStreamRanges[Idx].Build.Outcome != SelectedStreamRanges[Idx + 1].Build.Outcome)
+						if(SelectedStreamRanges[Idx].BuildGroup.Outcome != SelectedStreamRanges[Idx + 1].BuildGroup.Outcome)
 						{
 							FetchBuildChanges(SelectedStreamRanges[Idx + 1]);
 						}
@@ -830,20 +881,20 @@ namespace UnrealGameSync
 			{
 				BuildListView.DrawTrackedBackground(e.Graphics, e.Bounds);
 			}
-			else if(e.Item.Tag is IssueBuildData)
+			else if(e.Item.Tag is BuildGroup)
 			{
-				IssueBuildData BuildData = (IssueBuildData)e.Item.Tag;
+				BuildGroup BuildGroup = (BuildGroup)e.Item.Tag;
 
 				Color BackgroundColor;
-				if(BuildData.Outcome == IssueBuildOutcome.Error)
+				if(BuildGroup.Outcome == IssueBuildOutcome.Error)
 				{
 					BackgroundColor = Color.FromArgb(254, 248, 246);
 				}
-				else if(BuildData.Outcome == IssueBuildOutcome.Warning)
+				else if(BuildGroup.Outcome == IssueBuildOutcome.Warning)
 				{
 					BackgroundColor = Color.FromArgb(254, 254, 246);
 				}
-				else if(BuildData.Outcome == IssueBuildOutcome.Success)
+				else if(BuildGroup.Outcome == IssueBuildOutcome.Success)
 				{
 					BackgroundColor = Color.FromArgb(248, 254, 246);
 				}
@@ -949,7 +1000,7 @@ namespace UnrealGameSync
 				else if(e.ColumnIndex == AuthorHeader.Index)
 				{
 					Rectangle Bounds = new Rectangle(e.Bounds.X, e.Bounds.Y, e.Bounds.Width + e.Item.SubItems[e.ColumnIndex].Bounds.Width, e.Bounds.Height);
-					TextRenderer.DrawText(e.Graphics, BuildData.Name, Font, Bounds, TextColor, TextFormatFlags.EndEllipsis | TextFormatFlags.SingleLine | TextFormatFlags.Left | TextFormatFlags.VerticalCenter | TextFormatFlags.NoPrefix);
+					TextRenderer.DrawText(e.Graphics, BuildData.JobName, Font, Bounds, TextColor, TextFormatFlags.EndEllipsis | TextFormatFlags.SingleLine | TextFormatFlags.Left | TextFormatFlags.VerticalCenter | TextFormatFlags.NoPrefix);
 				}
 			}
 			else
@@ -1096,10 +1147,10 @@ namespace UnrealGameSync
 
 		private void DescriptionLabel_LinkClicked(object sender, LinkLabelLinkClickedEventArgs e)
 		{
-			IssueBuildData LastBuild = Issue.Builds.Where(x => x.Stream == SelectedStream).OrderByDescending(x => x.Change).ThenByDescending(x => x.Url).FirstOrDefault();
+			IssueBuildData LastBuild = Issue.Builds.Where(x => x.Stream == SelectedStream).OrderByDescending(x => x.Change).ThenByDescending(x => x.ErrorUrl).FirstOrDefault();
 			if(LastBuild != null)
 			{
-				System.Diagnostics.Process.Start(LastBuild.Url);
+				System.Diagnostics.Process.Start(LastBuild.ErrorUrl);
 			}
 		}
 
@@ -1113,5 +1164,38 @@ namespace UnrealGameSync
 				IssueMonitor.PostUpdate(Update);
 			}
 		}
+
+		private void JobContextMenu_ShowError_Click(object sender, EventArgs e)
+		{
+			foreach(IssueBuildData Build in SelectedBuildGroup.Builds)
+			{
+				if(Build.ErrorUrl != null)
+				{
+					System.Diagnostics.Process.Start(Build.ErrorUrl);
+					break;
+				}
+			}
+		}
+
+		private void JobContextMenu_ViewJob_Click(object sender, EventArgs e)
+		{
+			System.Diagnostics.Process.Start(SelectedBuildGroup.JobUrl);
+		}
+
+		private void BuildListView_MouseUp(object sender, MouseEventArgs e)
+		{
+			if((e.Button & MouseButtons.Right) != 0)
+			{
+				ListViewHitTestInfo HitTest = BuildListView.HitTest(e.Location);
+				if(HitTest.Item != null)
+				{
+					BuildGroup Group = HitTest.Item.Tag as BuildGroup;
+					if(Group != null)
+					{
+						ShowJobContextMenu(e.Location, Group);
+					}
+				}
+			}
+		}
 	}
 }
diff --git a/Engine/Source/Programs/UnrealGameSync/UnrealGameSync/IssueMonitor.cs b/Engine/Source/Programs/UnrealGameSync/UnrealGameSync/IssueMonitor.cs
--- a/Engine/Source/Programs/UnrealGameSync/UnrealGameSync/IssueMonitor.cs
+++ b/Engine/Source/Programs/UnrealGameSync/UnrealGameSync/IssueMonitor.cs
@@ -23,8 +23,11 @@ namespace UnrealGameSync
 	{
 		public string Stream;
 		public int Change;
-		public string Name;
-		public string Url;
+		public string JobName;
+		public string JobUrl;
+		public string JobStepName;
+		public string JobStepUrl;
+		public string ErrorUrl;
 		public IssueBuildOutcome Outcome;
 	}
 
