UGS: Include the owner of an issue in the browser window.

#rnx

diff --git a/Engine/Source/Programs/UnrealGameSync/UnrealGameSync/Forms/IssueBrowserWindow.cs b/Engine/Source/Programs/UnrealGameSync/UnrealGameSync/Forms/IssueBrowserWindow.cs
--- a/Engine/Source/Programs/UnrealGameSync/UnrealGameSync/Forms/IssueBrowserWindow.cs
+++ b/Engine/Source/Programs/UnrealGameSync/UnrealGameSync/Forms/IssueBrowserWindow.cs
@@ -94,6 +94,7 @@ namespace UnrealGameSync
 				Item.SubItems.Add(Issue.Id.ToString());
 				Item.SubItems.Add(FormatIssueDateTime(Issue.CreatedAt, Midnight));
 				Item.SubItems.Add(Issue.ResolvedAt.HasValue? FormatIssueDateTime(Issue.ResolvedAt.Value, Midnight) : "Unresolved");
+				Item.SubItems.Add((Issue.Owner == null)? "-" : Utility.FormatUserName(Issue.Owner));
 				Item.SubItems.Add(Issue.Summary);
 				Item.Tag = Issue;
 				IssueListView.Items.Add(Item);
