Few tweaks to the undefined symbol pattern matcher, so extract symbol names in more cases.

#rnx

diff --git a/Engine/Source/Programs/UnrealGameSync/MetadataTool/BuildHealth/Matchers/UndefinedSymbolPatternMatcher.cs b/Engine/Source/Programs/UnrealGameSync/MetadataTool/BuildHealth/Matchers/UndefinedSymbolPatternMatcher.cs
--- a/Engine/Source/Programs/UnrealGameSync/MetadataTool/BuildHealth/Matchers/UndefinedSymbolPatternMatcher.cs
+++ b/Engine/Source/Programs/UnrealGameSync/MetadataTool/BuildHealth/Matchers/UndefinedSymbolPatternMatcher.cs
@@ -15,7 +15,7 @@ namespace MetadataTool
 
 		public override bool TryMatch(InputJob Job, InputJobStep JobStep, InputDiagnostic Diagnostic, List<TrackedIssueFingerprint> Fingerprints)
 		{
-			SortedSet<string> SymbolNames = new SortedSet<string>(StringComparer.Ordinal);
+			List<string> SymbolMatches = new List<string>();
 
 			// Mac link error:
 			//   Undefined symbols for architecture arm64:
@@ -24,36 +24,36 @@ namespace MetadataTool
 			{
 				foreach(string Line in Diagnostic.Message.Split('\n'))
 				{
-					Match SymbolMatch = Regex.Match(Line, "^  \"([^\"\\(]+)");
+					Match SymbolMatch = Regex.Match(Line, "^  \"(.+)\"");
 					if (SymbolMatch.Success)
 					{
-						SymbolNames.Add(SymbolMatch.Groups[1].Value);
+						SymbolMatches.Add(SymbolMatch.Groups[1].Value);
 					}
 				}
 			}
 
 			// Android link error:
 			//   Foo.o:(.data.rel.ro + 0x5d88): undefined reference to `Foo::Bar()'
-			Match UndefinedReference = Regex.Match(Diagnostic.Message, ": undefined reference to [`']([^`'(]+)");
+			Match UndefinedReference = Regex.Match(Diagnostic.Message, ": undefined reference to [`']([^`']+)");
 			if(UndefinedReference.Success)
 			{
-				SymbolNames.Add(UndefinedReference.Groups[1].Value);
+				SymbolMatches.Add(UndefinedReference.Groups[1].Value);
 			}
 
 			// LLD link error:
 			//   ld.lld.exe: error: undefined symbol: Foo::Bar() const
-			Match LldMatch = Regex.Match(Diagnostic.Message, "error: undefined symbol:\\s*([^ (]+)");
+			Match LldMatch = Regex.Match(Diagnostic.Message, "error: undefined symbol:\\s*(.+)");
 			if (LldMatch.Success)
 			{
-				SymbolNames.Add(LldMatch.Groups[1].Value);
+				SymbolMatches.Add(LldMatch.Groups[1].Value);
 			}
 
 			// Link error:
 			//   Link: error: L0039: reference to undefined symbol `Foo::Bar() const' in file
-			Match LinkMatch = Regex.Match(Diagnostic.Message, ": reference to undefined symbol [`']([^`'(]+)");
+			Match LinkMatch = Regex.Match(Diagnostic.Message, ": reference to undefined symbol [`']([^`']+)");
 			if (LinkMatch.Success)
 			{
-				SymbolNames.Add(LinkMatch.Groups[1].Value);
+				SymbolMatches.Add(LinkMatch.Groups[1].Value);
 			}
 
 			// Microsoft linker error:
@@ -61,7 +61,14 @@ namespace MetadataTool
 			Match MicrosoftMatch = Regex.Match(Diagnostic.Message, ": unresolved external symbol \"([^\"]*)\"");
 			if(MicrosoftMatch.Success)
 			{
-				string SymbolName = MicrosoftMatch.Groups[1].Value;
+				SymbolMatches.Add(MicrosoftMatch.Groups[1].Value);
+			}
+
+			// Clean up all the symbol names
+			SortedSet<string> SymbolNames = new SortedSet<string>(StringComparer.Ordinal);
+			foreach(string SymbolMatch in SymbolMatches)
+			{
+				string SymbolName = SymbolMatch;
 
 				// Remove any argument lists for functions (anything after the first paren)
 				SymbolName = Regex.Replace(SymbolName, "\\(.*$", "");
@@ -69,6 +76,7 @@ namespace MetadataTool
 				// Remove any decorators and type information (greedy match up to the last space)
 				SymbolName = Regex.Replace(SymbolName, "^.* ", "");
 
+				// Add it to the list
 				SymbolNames.Add(SymbolName);
 			}
 
