UGS: Make updates to the state transaction file atomic.

#rnx

diff --git a/Engine/Source/Programs/UnrealGameSync/MetadataTool/CommandHandler_BuildHealth.cs b/Engine/Source/Programs/UnrealGameSync/MetadataTool/CommandHandler_BuildHealth.cs
--- a/Engine/Source/Programs/UnrealGameSync/MetadataTool/CommandHandler_BuildHealth.cs
+++ b/Engine/Source/Programs/UnrealGameSync/MetadataTool/CommandHandler_BuildHealth.cs
@@ -72,6 +72,9 @@ namespace MetadataTool
 			Log.TraceInformation("Reading build results from {0}", InputFile);
 			InputData InputData = DeserializeJson<InputData>(InputFile);
 
+			// Complete any interrupted operation to update the state file
+			CompleteStateTransaction(StateFile);
+
 			// Read the persistent data file
 			TrackedState State;
 			if (!bClean && FileReference.Exists(StateFile))
@@ -163,7 +166,7 @@ namespace MetadataTool
 			// Save the persistent data
 			Log.TraceInformation("Writing persistent data to {0}", StateFile);
 			DirectoryReference.CreateDirectory(StateFile.Directory);
-			SerializeJson(StateFile, State);
+			WriteState(StateFile, State);
 
 			// Synchronize with the server
 			if (ServerUrl != null)
@@ -207,7 +210,7 @@ namespace MetadataTool
 						}
 
 						Issue.PostedSummary = Issue.Fingerprint.Summary;
-						SerializeJson(StateFile, State);
+						WriteState(StateFile, State);
 					}
 					else if(Issue.PostedSummary == null || !String.Equals(Issue.PostedSummary, Issue.Fingerprint.Summary, StringComparison.Ordinal))
 					{
@@ -227,7 +230,7 @@ namespace MetadataTool
 						}
 
 						Issue.PostedSummary = Issue.Fingerprint.Summary;
-						SerializeJson(StateFile, State);
+						WriteState(StateFile, State);
 					}
 				}
 
@@ -273,7 +276,7 @@ namespace MetadataTool
 									}
 
 									Build.bPostedToServer = true;
-									SerializeJson(StateFile, State);
+									WriteState(StateFile, State);
 								}
 							}
 						}
@@ -301,7 +304,7 @@ namespace MetadataTool
 						}
 
 						State.Issues.RemoveAt(Idx--);
-						SerializeJson(StateFile, State);
+						WriteState(StateFile, State);
 					}
 				}
 
@@ -325,7 +328,7 @@ namespace MetadataTool
 						Issue.PendingWatchers.Remove(Watcher.UserName);
 						Issue.Watchers.Add(Watcher.UserName);
 
-						SerializeJson(StateFile, State);
+						WriteState(StateFile, State);
 					}
 				}
 			}
@@ -618,6 +621,49 @@ namespace MetadataTool
 			return CachedInfo.Changes;
 		}
 
+		/// <summary>
+		/// Gets the path to a temporary file used for ensuring that serialization from the state file is transactional
+		/// </summary>
+		/// <param name="StateFile">Path to the state file</param>
+		/// <returns>Path to the temporary state transaction file</returns>
+		static FileReference GetStateTransactionFile(FileReference StateFile)
+		{
+			return new FileReference(StateFile.FullName + ".transaction");
+		}
+
+		/// <summary>
+		/// Completes an interrupted transaction to write the state file
+		/// </summary>
+		/// <param name="StateFile">Path to the state file</param>
+		static void CompleteStateTransaction(FileReference StateFile)
+		{
+			if(!FileReference.Exists(StateFile))
+			{
+				FileReference StateTransactionFile = GetStateTransactionFile(StateFile);
+				if(FileReference.Exists(StateTransactionFile))
+				{
+					Log.TraceInformation("Completing partial transaction through {0}", StateTransactionFile);
+					FileReference.Move(StateTransactionFile, StateFile);
+				}
+			}
+		}
+
+		/// <summary>
+		/// Writes the state to disk in a way that can be recovered if the operation is interrupted
+		/// </summary>
+		/// <param name="StateFile">The file to write to</param>
+		/// <param name="State">The state object</param>
+		static void WriteState(FileReference StateFile, TrackedState State)
+		{
+			// Write out the state to the transaction file
+			FileReference StateTransactionFile = GetStateTransactionFile(StateFile);
+			SerializeJson(StateTransactionFile, State);
+
+			// Remove the original file, then move the transaction file into place
+			FileReference.Delete(StateFile);
+			FileReference.Move(StateTransactionFile, StateFile);
+		}
+
 		/// <summary>
 		/// Serializes an object to a file as JSON
 		/// </summary>
