UGS: Retain information about issues for 24 hours after they are closed. This prevents issues being recreated because jobs are still running other steps, and fixes problems caused by robomerge taking time to move changes between branches.

#rnx

diff --git a/Engine/Source/Programs/UnrealGameSync/MetadataTool/BuildHealth/TrackedIssue.cs b/Engine/Source/Programs/UnrealGameSync/MetadataTool/BuildHealth/TrackedIssue.cs
--- a/Engine/Source/Programs/UnrealGameSync/MetadataTool/BuildHealth/TrackedIssue.cs
+++ b/Engine/Source/Programs/UnrealGameSync/MetadataTool/BuildHealth/TrackedIssue.cs
@@ -46,6 +46,18 @@ namespace MetadataTool
 		[DataMember]
 		public int InitialChange;
 
+		/// <summary>
+		/// Whether we've posted an update to the resolved flag to the server
+		/// </summary>
+		[DataMember]
+		public bool bPostedResolved;
+
+		/// <summary>
+		/// The time at which the issue was closed. Issues will be retained for 24 hours after they are closed, in case the same issue appears in another stream and to prevent the issue being added again. 
+		/// </summary>
+		[DataMember]
+		public DateTime? ResolvedAt;
+
 		/// <summary>
 		/// Map of stream name -> step name -> history for builds exhibiting this issue
 		/// </summary>
@@ -85,7 +97,7 @@ namespace MetadataTool
 		/// Determines whether the issue can be closed
 		/// </summary>
 		/// <returns>True if the issue can be closed</returns>
-		public bool CanClose()
+		public bool IsResolved()
 		{
 			return Streams.Values.All(x => x.Values.All(y => y.NextSuccessfulBuild != null));
 		}
diff --git a/Engine/Source/Programs/UnrealGameSync/MetadataTool/CommandHandler_BuildHealth.cs b/Engine/Source/Programs/UnrealGameSync/MetadataTool/CommandHandler_BuildHealth.cs
--- a/Engine/Source/Programs/UnrealGameSync/MetadataTool/CommandHandler_BuildHealth.cs
+++ b/Engine/Source/Programs/UnrealGameSync/MetadataTool/CommandHandler_BuildHealth.cs
@@ -163,6 +163,25 @@ namespace MetadataTool
 				}
 			}
 
+			// Mark any issues as resolved
+			foreach(TrackedIssue Issue in State.Issues)
+			{
+				if(Issue.IsResolved())
+				{
+					if(!Issue.ResolvedAt.HasValue)
+					{
+						Issue.ResolvedAt = DateTime.UtcNow;
+					}
+				}
+				else
+				{
+					if(Issue.ResolvedAt.HasValue)
+					{
+						Issue.ResolvedAt = null;
+					}
+				}
+			}
+
 			// Save the persistent data
 			Log.TraceInformation("Writing persistent data to {0}", StateFile);
 			DirectoryReference.CreateDirectory(StateFile.Directory);
@@ -287,12 +306,12 @@ namespace MetadataTool
 				for (int Idx = 0; Idx < State.Issues.Count; Idx++)
 				{
 					TrackedIssue Issue = State.Issues[Idx];
-					if (Issue.CanClose())
+					if (Issue.ResolvedAt.HasValue != Issue.bPostedResolved)
 					{
-						Log.TraceInformation("Marking issue {0} as resolved", Issue.Id);
+						Log.TraceInformation("Setting issue {0} resolved flag to {1}", Issue.Id, Issue.ResolvedAt.HasValue);
 
 						CommandTypes.UpdateIssue UpdateBody = new CommandTypes.UpdateIssue();
-						UpdateBody.Resolved = true;
+						UpdateBody.Resolved = Issue.ResolvedAt.HasValue;
 
 						using(HttpWebResponse Response = SendHttpRequest(String.Format("{0}/api/issues/{1}", ServerUrl, Issue.Id), "PUT", UpdateBody))
 						{
@@ -303,7 +322,7 @@ namespace MetadataTool
 							}
 						}
 
-						State.Issues.RemoveAt(Idx--);
+						Issue.bPostedResolved = Issue.ResolvedAt.HasValue;
 						WriteState(StateFile, State);
 					}
 				}
@@ -333,6 +352,19 @@ namespace MetadataTool
 				}
 			}
 
+			// Remove any issues which have been resolved for 24 hours. We have to keep information about issues that have been fixed for some time; we may be updating the same job 
+			// multiple times while other steps are running, and we don't want to keep opening new issues for it. Also, it can take time for changes to propagate between streams.
+			DateTime RemoveIssueTime = DateTime.UtcNow - TimeSpan.FromHours(24.0);
+			for(int Idx = 0; Idx < State.Issues.Count; Idx++)
+			{
+				TrackedIssue Issue = State.Issues[Idx];
+				if(Issue.ResolvedAt.HasValue && Issue.ResolvedAt.Value < RemoveIssueTime)
+				{
+					State.Issues.RemoveAt(Idx--);
+					continue;
+				}
+			}
+
 			// TODO: VERIFY ISSUES ARE CLOSED
 		}
 
