UGS: Fix form links in project files.

#rnx

diff --git a/Engine/Source/Programs/UnrealGameSync/UnrealGameSyncLauncher/UnrealGameSyncLauncher.csproj b/Engine/Source/Programs/UnrealGameSync/UnrealGameSyncLauncher/UnrealGameSyncLauncher.csproj
--- a/Engine/Source/Programs/UnrealGameSync/UnrealGameSyncLauncher/UnrealGameSyncLauncher.csproj
+++ b/Engine/Source/Programs/UnrealGameSync/UnrealGameSyncLauncher/UnrealGameSyncLauncher.csproj
@@ -81,6 +81,7 @@
     </Compile>
     <Compile Include="..\UnrealGameSync\Forms\ModalTaskWindow.Designer.cs">
       <Link>UnrealGameSync\ModalTaskWindow.Designer.cs</Link>
+      <DependentUpon>ModalTaskWindow.cs</DependentUpon>
     </Compile>
     <Compile Include="..\UnrealGameSync\Forms\PasswordWindow.cs">
       <Link>UnrealGameSync\PasswordWindow.cs</Link>
@@ -88,6 +89,7 @@
     </Compile>
     <Compile Include="..\UnrealGameSync\Forms\PasswordWindow.Designer.cs">
       <Link>UnrealGameSync\PasswordWindow.Designer.cs</Link>
+      <DependentUpon>PasswordWindow.cs</DependentUpon>
     </Compile>
     <Compile Include="..\UnrealGameSync\DeploymentSettings.cs">
       <Link>Shared\DeploymentSettings.cs</Link>
