UGS: Prompt for a password immediately if the current login has expired.

#rnx

diff --git a/Engine/Source/Programs/UnrealGameSync/UnrealGameSync/PerforceModalTask.cs b/Engine/Source/Programs/UnrealGameSync/UnrealGameSync/PerforceModalTask.cs
--- a/Engine/Source/Programs/UnrealGameSync/UnrealGameSync/PerforceModalTask.cs
+++ b/Engine/Source/Programs/UnrealGameSync/UnrealGameSync/PerforceModalTask.cs
@@ -155,6 +155,10 @@ namespace UnrealGameSync
 				else if(Task.LoginResult == LoginResult.MissingPassword)
 				{
 					PasswordWindow PasswordWindow = new PasswordWindow(String.Format("Enter the password for user '{0}' on server '{1}'.", Task.UserName, Task.ServerAndPort), Task.Password);
+					if(Owner == null)
+					{
+						PasswordWindow.StartPosition = FormStartPosition.CenterScreen;
+					}
 					if(PasswordWindow.ShowDialog() != DialogResult.OK)
 					{
 						ErrorMessage = null;
@@ -165,7 +169,11 @@ namespace UnrealGameSync
 				else if(Task.LoginResult == LoginResult.IncorrectPassword)
 				{
 					PasswordWindow PasswordWindow = new PasswordWindow(String.Format("Authentication failed. Enter the password for user '{0}' on server '{1}'.", Task.UserName, Task.ServerAndPort), Task.Password);
-					if(PasswordWindow.ShowDialog() != DialogResult.OK)
+					if (Owner == null)
+					{
+						PasswordWindow.StartPosition = FormStartPosition.CenterScreen;
+					}
+					if (PasswordWindow.ShowDialog() != DialogResult.OK)
 					{
 						ErrorMessage = null;
 						return ModalTaskResult.Aborted;
diff --git a/Engine/Source/Programs/UnrealGameSync/UnrealGameSyncLauncher/Program.cs b/Engine/Source/Programs/UnrealGameSync/UnrealGameSyncLauncher/Program.cs
--- a/Engine/Source/Programs/UnrealGameSync/UnrealGameSyncLauncher/Program.cs
+++ b/Engine/Source/Programs/UnrealGameSync/UnrealGameSyncLauncher/Program.cs
@@ -80,15 +80,16 @@ namespace UnrealGameSyncLauncher
 				else
 				{
 					// Try to do a sync with the current settings first
-					SyncAndRunModalTask SyncApplication = new SyncAndRunModalTask(ServerAndPort, UserName, (Perforce, LogWriter) => SyncAndRun(Perforce, DepotPath, bUnstable, Args, InstanceMutex, LogWriter));
+					StringWriter Log = new StringWriter();
+					SyncAndRunPerforceTask SyncApplication = new SyncAndRunPerforceTask((Perforce, LogWriter) => SyncAndRun(Perforce, DepotPath, bUnstable, Args, InstanceMutex, LogWriter));
 
 					string ErrorMessage;
-					if(ModalTask.Execute(null, SyncApplication, "Updating", "Checking for updates, please wait...", out ErrorMessage) == ModalTaskResult.Succeeded)
+					if(PerforceModalTask.Execute(null, null, ServerAndPort, UserName, SyncApplication, "Updating", "Checking for updates, please wait...", Log, out ErrorMessage) == ModalTaskResult.Succeeded)
 					{
 						return 0;
 					}
 
-					SettingsWindow UpdateError = new SettingsWindow("Unable to update UnrealGameSync from Perforce. Verify that your connection settings are correct.", SyncApplication.LogWriter.ToString(), SyncApplication.ServerAndPort, SyncApplication.UserName, DepotPath, bUnstable, (Perforce, DepotParam, bUnstableParam, LogWriter) => SyncAndRun(Perforce, DepotParam, bUnstableParam, Args, InstanceMutex, LogWriter));
+					SettingsWindow UpdateError = new SettingsWindow("Unable to update UnrealGameSync from Perforce. Verify that your connection settings are correct.", Log.ToString(), ServerAndPort, UserName, DepotPath, bUnstable, (Perforce, DepotParam, bUnstableParam, LogWriter) => SyncAndRun(Perforce, DepotParam, bUnstableParam, Args, InstanceMutex, LogWriter));
 					if(UpdateError.ShowDialog() == DialogResult.OK)
 					{
 						return 0;
