UGS: Remove the separate "Close Build Issue" button. Show the "Build Health" button as a red badge if a user has open build issues.

#rnx

diff --git a/Engine/Source/Programs/UnrealGameSync/UnrealGameSync/Controls/WorkspaceControl.cs b/Engine/Source/Programs/UnrealGameSync/UnrealGameSync/Controls/WorkspaceControl.cs
--- a/Engine/Source/Programs/UnrealGameSync/UnrealGameSync/Controls/WorkspaceControl.cs
+++ b/Engine/Source/Programs/UnrealGameSync/UnrealGameSync/Controls/WorkspaceControl.cs
@@ -2874,12 +2874,13 @@ namespace UnrealGameSync
 				if(Workspace.ProjectConfigFile.GetValue("Options.ShowBuildHealth", false))
 				{
 					ProgramsLine.AddText("  |  ");
-					ProgramsLine.AddLink("Build Health", FontStyle.Regular, (P, R) => { ShowBuildHealthMenu(R); });
-
 					if(bUserHasOpenIssues)
 					{
-						ProgramsLine.AddText("  |  ");
-						ProgramsLine.AddBadge("Close Build Issue", GetBuildBadgeColor(BadgeResult.Failure), (P, R) => ShowCloseBuildIssueMenu(R));
+						ProgramsLine.AddBadge("Build Health", GetBuildBadgeColor(BadgeResult.Failure), (P, R) => ShowBuildHealthMenu(R));
+					}
+					else
+					{
+						ProgramsLine.AddLink("Build Health", FontStyle.Regular, (P, R) => { ShowBuildHealthMenu(R); });
 					}
 				}
 
@@ -3039,49 +3040,11 @@ namespace UnrealGameSync
 			BuildHealthContextMenu.Show(StatusPanel, new Point(Bounds.Right, Bounds.Bottom), ToolStripDropDownDirection.BelowLeft);
 		}
 
-		private void ShowCloseBuildIssueMenu(Rectangle Bounds)
-		{
-			List<IssueData> Issues = GetOpenIssuesForUser();
-			if(Issues.Count == 1)
-			{
-				CloseBuildIssueContextMenu_Issue_Click(Issues[0]);
-			}
-			else if(Issues.Count > 1)
-			{
-				CloseBuildIssueMenu.Items.Clear();
-
-				foreach(IssueData Issue in Issues)
-				{
-					StringBuilder Description = new StringBuilder();
-					Description.AppendFormat("{0}: {1}", Issue.Id, Issue.Summary);
-					Description.AppendFormat(" ({0})", Utility.FormatDurationMinutes((int)((Issue.RetrievedAt - Issue.CreatedAt).TotalMinutes + 1)));
-
-					ToolStripMenuItem Item = new ToolStripMenuItem(Description.ToString());
-					Item.Click += (s, e) => { CloseBuildIssueContextMenu_Issue_Click(Issue); };
-					CloseBuildIssueMenu.Items.Add(Item);
-				}
-
-				CloseBuildIssueMenu.Show(StatusPanel, new Point(Bounds.Right, Bounds.Bottom), ToolStripDropDownDirection.BelowLeft);
-			}
-		}
-
 		private void BuildHealthContextMenu_Issue_Click(IssueData Issue)
 		{
 			ShowIssueDetails(Issue);
 		}
 
-		private void CloseBuildIssueContextMenu_Issue_Click(IssueData Issue)
-		{
-			int FixChangeNumber = 0;
-			if(IssueFixedWindow.ShowModal(this, Workspace.Perforce, ref FixChangeNumber))
-			{
-				IssueUpdateData Update = new IssueUpdateData();
-				Update.Id = Issue.Id;
-				Update.FixChange = FixChangeNumber;
-				IssueMonitor.PostUpdate(Update);
-			}
-		}
-
 		public void ShowIssueDetails(IssueData Issue)
 		{
 			IssueDetailsWindow.ShowModal(this, IssueMonitor, Workspace.Perforce.ServerAndPort, Workspace.Perforce.UserName, Issue, Log, StreamName);
