UGS: Allow double-clicking in the issue fixed window to close an issue.

#rnx

diff --git a/Engine/Source/Programs/UnrealGameSync/UnrealGameSync/Forms/IssueFixedWindow.cs b/Engine/Source/Programs/UnrealGameSync/UnrealGameSync/Forms/IssueFixedWindow.cs
--- a/Engine/Source/Programs/UnrealGameSync/UnrealGameSync/Forms/IssueFixedWindow.cs
+++ b/Engine/Source/Programs/UnrealGameSync/UnrealGameSync/Forms/IssueFixedWindow.cs
@@ -163,15 +163,20 @@ namespace UnrealGameSync
 			}
 		}
 
-		private void OkBtn_Click(object sender, EventArgs e)
+		private void UpdateSelectedChangeAndClose()
 		{
-			if(TryGetSelectedChange(out ChangeNumber))
+			if (TryGetSelectedChange(out ChangeNumber))
 			{
 				DialogResult = DialogResult.OK;
 				Close();
 			}
 		}
 
+		private void OkBtn_Click(object sender, EventArgs e)
+		{
+			UpdateSelectedChangeAndClose();
+		}
+
 		private void PopulateChanges(string UserName, List<PerforceDescribeRecord> Changes)
 		{
 			if(!IsDisposed)
@@ -231,6 +236,11 @@ namespace UnrealGameSync
 			}
 		}
 
+		private void ChangesListView_MouseDoubleClick(object sender, MouseEventArgs e)
+		{
+			UpdateSelectedChangeAndClose();
+		}
+
 		private void ChangesContextMenu_MoreInfo_Click(object sender, EventArgs e)
 		{
 			PerforceDescribeRecord Record = (PerforceDescribeRecord)ChangesListContextMenu.Tag;
