UGS: Change the "Mark Fixed" button in the issue details window to "Reopen" depending on context, rather than having a separate disabled button for it.

#rnx

diff --git a/Engine/Source/Programs/UnrealGameSync/UnrealGameSync/Forms/IssueDetailsWindow.cs b/Engine/Source/Programs/UnrealGameSync/UnrealGameSync/Forms/IssueDetailsWindow.cs
--- a/Engine/Source/Programs/UnrealGameSync/UnrealGameSync/Forms/IssueDetailsWindow.cs
+++ b/Engine/Source/Programs/UnrealGameSync/UnrealGameSync/Forms/IssueDetailsWindow.cs
@@ -428,7 +428,14 @@ namespace UnrealGameSync
 
 			UpdateSummaryTextIfChanged(DetailsTextBox, Issue.Details.Replace("\n", "\r\n"));
 
-			ReopenBtn.Enabled = (Issue.FixChange != 0);
+			if(Issue.FixChange == 0)
+			{
+				MarkFixedBtn.Text = "Mark Fixed...";
+			}
+			else
+			{
+				MarkFixedBtn.Text = "Reopen";
+			}
 		}
 
 		private void IssueDetailsWindow_Disposed(object sender, EventArgs e)
@@ -1136,11 +1143,21 @@ namespace UnrealGameSync
 		private void MarkFixedBtn_Click(object sender, EventArgs e)
 		{
 			int FixChangeNumber = Issue.FixChange;
-			if(IssueFixedWindow.ShowModal(this, Perforce, ref FixChangeNumber))
+			if(FixChangeNumber == 0)
+			{
+				if(IssueFixedWindow.ShowModal(this, Perforce, ref FixChangeNumber))
+				{
+					IssueUpdateData Update = new IssueUpdateData();
+					Update.Id = Issue.Id;
+					Update.FixChange = FixChangeNumber;
+					IssueMonitor.PostUpdate(Update);
+				}
+			}
+			else
 			{
 				IssueUpdateData Update = new IssueUpdateData();
 				Update.Id = Issue.Id;
-				Update.FixChange = FixChangeNumber;
+				Update.FixChange = 0;
 				IssueMonitor.PostUpdate(Update);
 			}
 		}
@@ -1154,17 +1171,6 @@ namespace UnrealGameSync
 			}
 		}
 
-		private void ReopenBtn_Click(object sender, EventArgs e)
-		{
-			if(Issue.FixChange != 0)
-			{
-				IssueUpdateData Update = new IssueUpdateData();
-				Update.Id = Issue.Id;
-				Update.FixChange = 0;
-				IssueMonitor.PostUpdate(Update);
-			}
-		}
-
 		private void JobContextMenu_ShowError_Click(object sender, EventArgs e)
 		{
 			foreach(IssueBuildData Build in SelectedBuildGroup.Builds)
