UGS: Remove the "Assign" button from the details panel. This is already exposed through the context menu, and can also be done through the "assign to user" menu.

#rnx

diff --git a/Engine/Source/Programs/UnrealGameSync/UnrealGameSync/Forms/IssueDetailsWindow.cs b/Engine/Source/Programs/UnrealGameSync/UnrealGameSync/Forms/IssueDetailsWindow.cs
--- a/Engine/Source/Programs/UnrealGameSync/UnrealGameSync/Forms/IssueDetailsWindow.cs
+++ b/Engine/Source/Programs/UnrealGameSync/UnrealGameSync/Forms/IssueDetailsWindow.cs
@@ -305,7 +305,6 @@ namespace UnrealGameSync
 			CreateWorker();
 
 			UpdateCurrentIssue();
-			UpdateEnabledButtons();
 		}
 
 		private void StartUpdateTimer()
@@ -770,11 +769,6 @@ namespace UnrealGameSync
 			}
 		}
 
-		void UpdateEnabledButtons()
-		{
-			AssignBtn.Enabled = BuildListView.SelectedItems.Count > 0 && (BuildListView.SelectedItems[0].Tag as PerforceChangeSummary) != null;
-		}
-
 		void OnUpdateIssues()
 		{
 			List<IssueData> NewIssues = IssueMonitor.GetIssues();
@@ -1069,11 +1063,6 @@ namespace UnrealGameSync
 			}
 		}
 
-		private void BuildListView_SelectedIndexChanged(object sender, EventArgs e)
-		{
-			UpdateEnabledButtons();
-		}
-
 		private void BuildListContextMenu_Blame_Click(object sender, EventArgs e)
 		{
 			if(ContextMenuChange != null)
@@ -1103,18 +1092,6 @@ namespace UnrealGameSync
 			UpdateCurrentIssue();
 		}
 
-		private void AssignBtn_Click(object sender, EventArgs e)
-		{
-			if(BuildListView.SelectedItems.Count > 0)
-			{
-				PerforceChangeSummary Change = BuildListView.SelectedItems[0].Tag as PerforceChangeSummary;
-				if(Change != null)
-				{
-					AssignToUser(Change.User);
-				}
-			}
-		}
-
 		private void AssignToMeBtn_Click(object sender, EventArgs e)
 		{
 			AssignToUser(IssueMonitor.UserName);
