UGS: Strip out paths to source files from compile errors, to make output easier to read.

#rnx

diff --git a/Engine/Source/Programs/UnrealGameSync/MetadataTool/BuildHealth/Matchers/CompilePatternMatcher.cs b/Engine/Source/Programs/UnrealGameSync/MetadataTool/BuildHealth/Matchers/CompilePatternMatcher.cs
--- a/Engine/Source/Programs/UnrealGameSync/MetadataTool/BuildHealth/Matchers/CompilePatternMatcher.cs
+++ b/Engine/Source/Programs/UnrealGameSync/MetadataTool/BuildHealth/Matchers/CompilePatternMatcher.cs
@@ -46,7 +46,7 @@ namespace MetadataTool
 			if (SourceFileNames.Count > 0)
 			{
 				TrackedIssueFingerprint Fingerprint = new TrackedIssueFingerprint(Category, GetSummary(SourceFileNames), Diagnostic.Url, Job.Change);
-				Fingerprint.Details.Add(Diagnostic.Message);
+				Fingerprint.Details.Add(ShortenPaths(Diagnostic.Message));
 				Fingerprint.FileNames.UnionWith(SourceFileNames);
 				Fingerprints.Add(Fingerprint);
 				return true;
@@ -56,6 +56,12 @@ namespace MetadataTool
 			return false;
 		}
 
+		static string ShortenPaths(string Text)
+		{
+			Text = Regex.Replace(Text, @"^(\s*(?:In file included from\s*)?)(?:[A-Za-z]:)?[^\s(:]+[/\\]([^\s/\\]+[\(:]\d)", "$1$2", RegexOptions.Multiline);
+			return Text;
+		}
+
 		public override void Merge(TrackedIssueFingerprint Source, TrackedIssueFingerprint Target)
 		{
 			base.Merge(Source, Target);
