UGS: Make issue browser and details forms modeless.

#rnx

diff --git a/Engine/Source/Programs/UnrealGameSync/UnrealGameSync/Controls/WorkspaceControl.cs b/Engine/Source/Programs/UnrealGameSync/UnrealGameSync/Controls/WorkspaceControl.cs
--- a/Engine/Source/Programs/UnrealGameSync/UnrealGameSync/Controls/WorkspaceControl.cs
+++ b/Engine/Source/Programs/UnrealGameSync/UnrealGameSync/Controls/WorkspaceControl.cs
@@ -3047,12 +3047,12 @@ namespace UnrealGameSync
 
 		public void ShowIssueDetails(IssueData Issue)
 		{
-			IssueDetailsWindow.ShowModal(this, IssueMonitor, Workspace.Perforce.ServerAndPort, Workspace.Perforce.UserName, Issue, Log, StreamName);
+			IssueDetailsWindow.Show(ParentForm, IssueMonitor, Workspace.Perforce.ServerAndPort, Workspace.Perforce.UserName, Issue, Log, StreamName);
 		}
 
 		private void BuildHealthContextMenu_Browse_Click(object sender, EventArgs e)
 		{
-			IssueBrowserWindow.ShowModal(this, IssueMonitor, Workspace.Perforce.ServerAndPort, Workspace.Perforce.UserName, Log, StreamName);
+			IssueBrowserWindow.Show(ParentForm, IssueMonitor, Workspace.Perforce.ServerAndPort, Workspace.Perforce.UserName, Log, StreamName);
 		}
 
 		private void BuildHealthContextMenu_Settings_Click(object sender, EventArgs e)
diff --git a/Engine/Source/Programs/UnrealGameSync/UnrealGameSync/Forms/IssueBrowserWindow.cs b/Engine/Source/Programs/UnrealGameSync/UnrealGameSync/Forms/IssueBrowserWindow.cs
--- a/Engine/Source/Programs/UnrealGameSync/UnrealGameSync/Forms/IssueBrowserWindow.cs
+++ b/Engine/Source/Programs/UnrealGameSync/UnrealGameSync/Forms/IssueBrowserWindow.cs
@@ -30,12 +30,26 @@ namespace UnrealGameSync
 			this.Log = Log;
 			this.CurrentStream = CurrentStream;
 
+			IssueMonitor.AddRef();
+
 			InitializeComponent();
 
 			System.Reflection.PropertyInfo DoubleBufferedProperty = typeof(Control).GetProperty("DoubleBuffered", System.Reflection.BindingFlags.NonPublic | System.Reflection.BindingFlags.Instance);
 			DoubleBufferedProperty.SetValue(IssueListView, true, null);
 		}
 
+		protected override void Dispose(bool disposing)
+		{
+			if (disposing && (components != null))
+			{
+				components.Dispose();
+			}
+
+			IssueMonitor.Release();
+
+			base.Dispose(disposing);
+		}
+
 		private void IssueBrowserWindow_Load(object sender, EventArgs e)
 		{
 			FetchMoreResults();
@@ -118,10 +132,26 @@ namespace UnrealGameSync
 			}
 		}
 
-		public static void ShowModal(IWin32Window Owner, IssueMonitor IssueMonitor, string ServerAndPort, string UserName, TextWriter Log, string CurrentStream)
+		static List<IssueBrowserWindow> ExistingWindows = new List<IssueBrowserWindow>();
+
+		public static void Show(Form Owner, IssueMonitor IssueMonitor, string ServerAndPort, string UserName, TextWriter Log, string CurrentStream)
 		{
-			IssueBrowserWindow Window = new IssueBrowserWindow(IssueMonitor, ServerAndPort, UserName, Log, CurrentStream);
-			Window.ShowDialog(Owner);
+			IssueBrowserWindow Window = ExistingWindows.FirstOrDefault(x => x.IssueMonitor == IssueMonitor);
+			if(Window == null)
+			{
+				Window = new IssueBrowserWindow(IssueMonitor, ServerAndPort, UserName, Log, CurrentStream);
+				Window.Owner = Owner;
+				Window.StartPosition = FormStartPosition.Manual;
+				Window.Location = new Point(Owner.Location.X + (Owner.Width - Window.Width) / 2, Owner.Location.Y + (Owner.Height - Window.Height) / 2);
+				Window.Show(Owner);
+
+				ExistingWindows.Add(Window);
+				Window.FormClosed += (E, S) => ExistingWindows.Remove(Window);
+			}
+			else
+			{
+				Window.Activate();
+			}
 		}
 
 		private void FetchMoreResultsLinkLabel_LinkClicked(object sender, LinkLabelLinkClickedEventArgs e)
@@ -168,7 +198,7 @@ namespace UnrealGameSync
 		private void ShowIssue(IssueData Issue)
 		{
 			Issue.Builds = RESTApi.GET<List<IssueBuildData>>(IssueMonitor.ApiUrl, String.Format("issues/{0}/builds", Issue.Id));
-			IssueDetailsWindow.ShowModal(this, IssueMonitor, ServerAndPort, UserName, Issue, Log, CurrentStream);
+			IssueDetailsWindow.Show(Owner, IssueMonitor, ServerAndPort, UserName, Issue, Log, CurrentStream);
 		}
 
 		private void IssueListView_DrawColumnHeader(object sender, DrawListViewColumnHeaderEventArgs e)
diff --git a/Engine/Source/Programs/UnrealGameSync/UnrealGameSync/Forms/IssueDetailsWindow.cs b/Engine/Source/Programs/UnrealGameSync/UnrealGameSync/Forms/IssueDetailsWindow.cs
--- a/Engine/Source/Programs/UnrealGameSync/UnrealGameSync/Forms/IssueDetailsWindow.cs
+++ b/Engine/Source/Programs/UnrealGameSync/UnrealGameSync/Forms/IssueDetailsWindow.cs
@@ -261,6 +261,8 @@ namespace UnrealGameSync
 			this.Perforce = new PerforceConnection(UserName, null, ServerAndPort);
 			this.Log = Log;
 
+			IssueMonitor.AddRef();
+
 			MainThreadSynchronizationContext = SynchronizationContext.Current;
 
 			InitializeComponent();
@@ -307,6 +309,18 @@ namespace UnrealGameSync
 			UpdateCurrentIssue();
 		}
 
+		protected override void Dispose(bool disposing)
+		{
+			if (disposing && (components != null))
+			{
+				components.Dispose();
+			}
+
+			IssueMonitor.Release();
+
+			base.Dispose(disposing);
+		}
+
 		private void StartUpdateTimer()
 		{
 			if(UpdateTimer == null)
@@ -847,11 +861,25 @@ namespace UnrealGameSync
 			MainThreadSynchronizationContext.Post((o) => { if(!bIsDisposing) { UpdateChangeMetadata(Change); } }, null);
 		}
 
-		public static void ShowModal(IWin32Window Owner, IssueMonitor IssueMonitor, string ServerAndPort, string UserName, IssueData Issue, TextWriter Log, string CurrentStream)
+		static List<IssueDetailsWindow> ExistingWindows = new List<IssueDetailsWindow>();
+
+		public static void Show(Form Owner, IssueMonitor IssueMonitor, string ServerAndPort, string UserName, IssueData Issue, TextWriter Log, string CurrentStream)
 		{
-			using(IssueDetailsWindow Window = new IssueDetailsWindow(IssueMonitor, Issue, ServerAndPort, UserName, Log, CurrentStream))
+			IssueDetailsWindow Window = ExistingWindows.FirstOrDefault(x => x.IssueMonitor == IssueMonitor && x.Issue.Id == Issue.Id);
+			if(Window == null)
+			{
+				Window = new IssueDetailsWindow(IssueMonitor, Issue, ServerAndPort, UserName, Log, CurrentStream);
+				Window.Owner = Owner;
+				Window.StartPosition = FormStartPosition.Manual;
+				Window.Location = new Point(Owner.Location.X + (Owner.Width - Window.Width) / 2, Owner.Location.Y + (Owner.Height - Window.Height) / 2);
+				Window.Show();
+
+				ExistingWindows.Add(Window);
+				Window.FormClosed += (S, E) => ExistingWindows.Remove(Window);
+			}
+			else
 			{
-				Window.ShowDialog(Owner);
+				Window.Activate();
 			}
 		}
 
diff --git a/Engine/Source/Programs/UnrealGameSync/UnrealGameSync/Forms/MainWindow.cs b/Engine/Source/Programs/UnrealGameSync/UnrealGameSync/Forms/MainWindow.cs
--- a/Engine/Source/Programs/UnrealGameSync/UnrealGameSync/Forms/MainWindow.cs
+++ b/Engine/Source/Programs/UnrealGameSync/UnrealGameSync/Forms/MainWindow.cs
@@ -419,7 +419,7 @@ namespace UnrealGameSync
 							CloseAlertWindow(AlertWindow);
 						}
 					}
-					IssueMonitor.Dispose();
+					IssueMonitor.Release();
 				}
 			}
 			IssueMonitors = NewIssueMonitors.ToList();
@@ -453,7 +453,7 @@ namespace UnrealGameSync
 
 			foreach(IssueMonitor IssueMonitor in IssueMonitors)
 			{
-				IssueMonitor.Dispose();
+				IssueMonitor.Release();
 			}
 			IssueMonitors.Clear();
 
diff --git a/Engine/Source/Programs/UnrealGameSync/UnrealGameSync/IssueMonitor.cs b/Engine/Source/Programs/UnrealGameSync/UnrealGameSync/IssueMonitor.cs
--- a/Engine/Source/Programs/UnrealGameSync/UnrealGameSync/IssueMonitor.cs
+++ b/Engine/Source/Programs/UnrealGameSync/UnrealGameSync/IssueMonitor.cs
@@ -71,12 +71,13 @@ namespace UnrealGameSync
 	{
 		public readonly string ApiUrl;
 		public readonly string UserName;
+		int RefCount = 1;
 		Thread WorkerThread;
 		BoundedLogWriter LogWriter;
 		bool bDisposing;
 		AutoResetEvent RefreshEvent;
 		int UpdateIntervalMs;
-		long TrackingIssueId = -1;
+		List<long> TrackingIssueIds = new List<long>();
 		List<IssueData> Issues = new List<IssueData>();
 		object LockObject = new object();
 		List<IssueUpdateData> PendingUpdates = new List<IssueUpdateData>();
@@ -131,7 +132,7 @@ namespace UnrealGameSync
 		{
 			lock(LockObject)
 			{
-				TrackingIssueId = IssueId;
+				TrackingIssueIds.Add(IssueId);
 			}
 			RefreshEvent.Set();
 		}
@@ -140,7 +141,7 @@ namespace UnrealGameSync
 		{
 			lock(LockObject)
 			{
-				TrackingIssueId = -1;
+				TrackingIssueIds.RemoveAt(TrackingIssueIds.IndexOf(IssueId));
 			}
 		}
 
@@ -213,7 +214,34 @@ namespace UnrealGameSync
 			}
 		}
 
-		public void Dispose()
+		public void AddRef()
+		{
+			if(RefCount == 0)
+			{
+				throw new Exception("Invalid reference count for IssueMonitor (zero)");
+			}
+			RefCount++;
+		}
+
+		public void Release()
+		{
+			RefCount--;
+			if(RefCount < 0)
+			{
+				throw new Exception("Invalid reference count for IssueMonitor (ltz)");
+			}
+			if(RefCount == 0)
+			{
+				DisposeInternal();
+			}
+		}
+
+		void IDisposable.Dispose()
+		{
+			DisposeInternal();
+		}
+
+		void DisposeInternal()
 		{
 			bDisposing = true;
 			if(WorkerThread != null)
@@ -303,24 +331,27 @@ namespace UnrealGameSync
 				List<IssueData> NewIssues = RESTApi.GET<List<IssueData>>(ApiUrl, String.Format("issues?user={0}", UserName));
 
 				// Check if we're tracking a particular issue. If so, we want updates even when it's resolved.
-				long LocalTrackingIssueId;
+				long[] LocalTrackingIssueIds;
 				lock(LockObject)
 				{
-					LocalTrackingIssueId = TrackingIssueId;
+					LocalTrackingIssueIds = TrackingIssueIds.Distinct().ToArray();
 				}
-				if(LocalTrackingIssueId != -1 && !NewIssues.Any(x => x.Id == LocalTrackingIssueId))
+				foreach(long LocalTrackingIssueId in LocalTrackingIssueIds)
 				{
-					try
+					if(!NewIssues.Any(x => x.Id == LocalTrackingIssueId))
 					{
-						IssueData Issue = RESTApi.GET<IssueData>(ApiUrl, String.Format("issues/{0}", LocalTrackingIssueId));
-						if(Issue != null)
+						try
 						{
-							NewIssues.Add(Issue);
+							IssueData Issue = RESTApi.GET<IssueData>(ApiUrl, String.Format("issues/{0}", LocalTrackingIssueId));
+							if(Issue != null)
+							{
+								NewIssues.Add(Issue);
+							}
+						}
+						catch(Exception Ex)
+						{
+							LogWriter.WriteException(Ex, "Exception while fetching tracked issue");
 						}
-					}
-					catch(Exception Ex)
-					{
-						LogWriter.WriteException(Ex, "Exception while fetching tracked issue");
 					}
 				}
 
