UGS: Remove the first error link from the job context menu.

#rnx

diff --git a/Engine/Source/Programs/UnrealGameSync/UnrealGameSync/Forms/IssueDetailsWindow.cs b/Engine/Source/Programs/UnrealGameSync/UnrealGameSync/Forms/IssueDetailsWindow.cs
--- a/Engine/Source/Programs/UnrealGameSync/UnrealGameSync/Forms/IssueDetailsWindow.cs
+++ b/Engine/Source/Programs/UnrealGameSync/UnrealGameSync/Forms/IssueDetailsWindow.cs
@@ -687,12 +687,8 @@ namespace UnrealGameSync
 		{
 			SelectedBuildGroup = BuildGroup;
 
-			bool bHasErrorUrl = BuildGroup.Builds.Any(x => x.ErrorUrl != null);
-			JobContextMenu_StepSeparatorMax.Visible = bHasErrorUrl;
-			JobContextMenu_ShowFirstError.Visible = bHasErrorUrl;
-
 			int MinIndex = JobContextMenu.Items.IndexOf(JobContextMenu_StepSeparatorMin) + 1;
-			while(JobContextMenu.Items[MinIndex] != JobContextMenu_StepSeparatorMax)
+			while(JobContextMenu.Items.Count > MinIndex)
 			{
 				JobContextMenu.Items.RemoveAt(MinIndex);
 			}
