UGS: Move the changes list onto a second tab of the issue details panel, to avoid overloading people with information.

#rnx

diff --git a/Engine/Source/Programs/UnrealGameSync/UnrealGameSync/Forms/IssueDetailsWindow.cs b/Engine/Source/Programs/UnrealGameSync/UnrealGameSync/Forms/IssueDetailsWindow.cs
--- a/Engine/Source/Programs/UnrealGameSync/UnrealGameSync/Forms/IssueDetailsWindow.cs
+++ b/Engine/Source/Programs/UnrealGameSync/UnrealGameSync/Forms/IssueDetailsWindow.cs
@@ -351,6 +351,14 @@ namespace UnrealGameSync
 			StopUpdateTimer();
 		}
 
+		void UpdateSummaryTextIfChanged(Label Label, string NewText)
+		{
+			if (Label.Text != NewText)
+			{
+				Label.Text = NewText;
+			}
+		}
+
 		void UpdateSummaryTextIfChanged(TextBox TextBox, string NewText)
 		{
 			if(TextBox.Text != NewText)
@@ -365,6 +373,8 @@ namespace UnrealGameSync
 		{
 			UpdateSummaryTextIfChanged(SummaryTextBox, Issue.Summary.ToString());
 
+			BuildLinkLabel.Text = (Issue.Builds.Count > 0) ? Issue.Builds[0].JobName : "Unknown";
+
 			StringBuilder Status = new StringBuilder();
 			if(IssueMonitor.HasPendingUpdate())
 			{
@@ -1204,5 +1214,14 @@ namespace UnrealGameSync
 				}
 			}
 		}
+
+		private void BuildLinkLabel_LinkClicked(object sender, LinkLabelLinkClickedEventArgs e)
+		{
+			IssueBuildData Build = Issue.Builds.FirstOrDefault(x => x.ErrorUrl != null);
+			if (Build != null)
+			{
+				System.Diagnostics.Process.Start(Build.ErrorUrl);
+			}
+		}
 	}
 }
