UGS: Sync down the top level /Platforms/ folder in UGS, if it exists.

diff --git a/Engine/Source/Programs/UnrealGameSync/UnrealGameSync/Perforce.cs b/Engine/Source/Programs/UnrealGameSync/UnrealGameSync/Perforce.cs
--- a/Engine/Source/Programs/UnrealGameSync/UnrealGameSync/Perforce.cs
+++ b/Engine/Source/Programs/UnrealGameSync/UnrealGameSync/Perforce.cs
@@ -1299,7 +1299,7 @@ namespace UnrealGameSync
 			using(PerforceTagRecordParser Parser = new PerforceTagRecordParser(x => SyncOutput(new PerforceFileRecord(x))))
 			{
 				string CommandLine = String.Format("-ztag sync -n {0}@{1}{2}", Filter, bOnlyFilesInThisChange? "=" : "", ChangeNumber);
-				return RunCommand(CommandLine.ToString(), null, Line => FilterTaggedOutput(Line, Parser, Log), CommandOptions.NoFailOnErrors | CommandOptions.IgnoreFilesUpToDateError | CommandOptions.IgnoreExitCode, Log);
+				return RunCommand(CommandLine.ToString(), null, Line => FilterTaggedOutput(Line, Parser, Log), CommandOptions.NoFailOnErrors | CommandOptions.IgnoreFilesUpToDateError | CommandOptions.IgnoreExitCode | CommandOptions.IgnoreNoSuchFilesError, Log);
 			}
 		}
 
diff --git a/Engine/Source/Programs/UnrealGameSync/UnrealGameSync/Workspace.cs b/Engine/Source/Programs/UnrealGameSync/UnrealGameSync/Workspace.cs
--- a/Engine/Source/Programs/UnrealGameSync/UnrealGameSync/Workspace.cs
+++ b/Engine/Source/Programs/UnrealGameSync/UnrealGameSync/Workspace.cs
@@ -1091,6 +1091,7 @@ namespace UnrealGameSync
 				{
 					SyncPaths.Add(ClientRootPath + "/Enterprise/...");
 				}
+				SyncPaths.Add(ClientRootPath + "/Platforms/...");
 				SyncPaths.Add(PerforceUtils.GetClientOrDepotDirectoryName(SelectedClientFileName) + "/...");
 			}
 
