UGS: Modify existing items the main list view control rather than deleting everything and reinserting whenever there are changes. This appears to fix issues with the scroll position being reset whenever it's updated.

diff --git a/Engine/Source/Programs/UnrealGameSync/UnrealGameSync/Controls/WorkspaceControl.cs b/Engine/Source/Programs/UnrealGameSync/UnrealGameSync/Controls/WorkspaceControl.cs
--- a/Engine/Source/Programs/UnrealGameSync/UnrealGameSync/Controls/WorkspaceControl.cs
+++ b/Engine/Source/Programs/UnrealGameSync/UnrealGameSync/Controls/WorkspaceControl.cs
@@ -1057,31 +1057,11 @@ namespace UnrealGameSync
 		{
 			if(SelectedFileName != null)
 			{
-				List<int> SelectedChanges = new List<int>();
-				foreach(ListViewItem SelectedItem in BuildList.SelectedItems)
-				{
-					PerforceChangeSummary Change = (PerforceChangeSummary)SelectedItem.Tag;
-					if(Change != null)
-					{
-						SelectedChanges.Add(Change.Number);
-					}
-				}
-
 				ChangeNumberToArchivePath.Clear();
 				ChangeNumberToLayoutInfo.Clear();
 
 				ExpandItem = null;
 
-				BuildList.BeginUpdate();
-
-				foreach(ListViewGroup Group in BuildList.Groups)
-				{
-					Group.Name = "xxx " + Group.Name;
-				}
-
-				int RemoveItems = BuildList.Items.Count;
-				int RemoveGroups = BuildList.Groups.Count;
-
 				List<PerforceChangeSummary> Changes = PerforceMonitor.GetChanges();
 				EventMonitor.FilterChanges(Changes.Select(x => x.Number));
 
@@ -1104,6 +1084,7 @@ namespace UnrealGameSync
 				ListIndexToChangeIndex = new List<int>();
 				SortedChangeNumbers = new List<int>();
 				
+				List<PerforceChangeSummary> FilteredChanges = new List<PerforceChangeSummary>();
 				for(int ChangeIdx = 0; ChangeIdx < Changes.Count; ChangeIdx++)
 				{
 					PerforceChangeSummary Change = Changes[ChangeIdx];
@@ -1115,55 +1096,8 @@ namespace UnrealGameSync
 						{
 							bFirstChange = false;
 
-							ListViewGroup Group;
-
-							DateTime DisplayTime = Change.Date;
-							if(Settings.bShowLocalTimes)
-							{
-								DisplayTime = (DisplayTime - PerforceMonitor.ServerTimeZone).ToLocalTime();
-							}
-
-							string GroupName = DisplayTime.ToString("D");//"dddd\\,\\ h\\.mmtt");
-							for(int Idx = 0;;Idx++)
-							{
-								if(Idx == BuildList.Groups.Count)
-								{
-									Group = new ListViewGroup(GroupName);
-									Group.Name = GroupName;
-									BuildList.Groups.Add(Group);
-									break;
-								}
-								else if(BuildList.Groups[Idx].Name == GroupName)
-								{
-									Group = BuildList.Groups[Idx];
-									break;
-								}
-							}
-
-							string[] SubItemLabels = new string[BuildList.Columns.Count];
-							SubItemLabels[ChangeColumn.Index] = String.Format("{0}", Change.Number);
-							SubItemLabels[TimeColumn.Index] = DisplayTime.ToString("h\\.mmtt");
-							SubItemLabels[AuthorColumn.Index] = FormatUserName(Change.User);
-							SubItemLabels[DescriptionColumn.Index] = Change.Description.Replace('\n', ' ');
-
-							ListViewItem Item = new ListViewItem(Group);
-							Item.Tag = Change;
-							Item.Selected = SelectedChanges.Contains(Change.Number);
-							for(int ColumnIdx = 1; ColumnIdx < BuildList.Columns.Count; ColumnIdx++)
-							{
-								Item.SubItems.Add(new ListViewItem.ListViewSubItem(Item, SubItemLabels[ColumnIdx] ?? ""));
-							}
-
-							// Insert it at the right position within the group
-							int GroupInsertIdx = 0;
-							while(GroupInsertIdx < Group.Items.Count && Change.Number < ((PerforceChangeSummary)Group.Items[GroupInsertIdx].Tag).Number)
-							{
-								GroupInsertIdx++;
-							}
-							Group.Items.Insert(GroupInsertIdx, Item);
-
-							// Insert it at the right place in the list
-							BuildList.Items.Add(Item);
+							// Add the new change
+							FilteredChanges.Add(Change);
 
 							// Store off the list index for this change
 							ListIndexToChangeIndex.Add(ChangeIdx);
@@ -1171,47 +1105,183 @@ namespace UnrealGameSync
 					}
 				}
 
+				UpdateBuildListInternal(FilteredChanges, Changes.Count > 0);
+
 				SortedChangeNumbers.Sort();
 
-				for(int Idx = 0; Idx < RemoveItems; Idx++)
+				if (PendingSelectedChangeNumber != -1)
 				{
-					BuildList.Items.RemoveAt(0);
+					SelectChange(PendingSelectedChangeNumber);
 				}
-				for(int Idx = 0; Idx < RemoveGroups; Idx++)
+			}
+
+			if (BuildList.HoverItem > BuildList.Items.Count)
+			{
+				BuildList.HoverItem = -1;
+			}
+
+			UpdateBuildFailureNotification();
+
+			UpdateBuildSteps();
+			UpdateSyncActionCheckboxes();
+		}
+
+		void UpdateBuildListInternal(List<PerforceChangeSummary> Changes, bool bShowExpandItem)
+		{
+			BuildList.BeginUpdate();
+
+			// Find the item before the expand row
+			ListViewItem BeforeExpandItem = null;
+			if(ExpandItem != null && ExpandItem.Index > 0)
+			{
+				BeforeExpandItem = BuildList.Items[ExpandItem.Index - 1];
+			}
+
+			// Merge the new changes into the build list
+			int BuildListIndex = 0;
+			foreach (PerforceChangeSummary Change in Changes)
+			{
+				for (; ; )
 				{
-					BuildList.Groups.RemoveAt(0);
+					if (BuildListIndex == BuildList.Items.Count)
+					{
+						// Insert
+						BuildList_InsertItem(BuildListIndex, Change);
+						break;
+					}
+
+					ListViewItem NextItem = BuildList.Items[BuildListIndex];
+					if (NextItem.Tag == null)
+					{
+						// Insert
+						BuildList_InsertItem(BuildListIndex, Change);
+						break;
+					}
+
+					PerforceChangeSummary NextChange = (PerforceChangeSummary)NextItem.Tag;
+					if (Change.Number > NextChange.Number)
+					{
+						// Insert
+						BuildList_InsertItem(BuildListIndex, Change);
+						break;
+					}
+					else if (NextChange.Number == Change.Number)
+					{
+						// Update
+						NextItem.SubItems[DescriptionColumn.Index].Text = Change.Description.Replace('\n', ' ');
+						break;
+					}
+					else
+					{
+						// Delete
+						BuildList.Items.RemoveAt(BuildListIndex);
+						continue;
+					}
 				}
 
-				if(Changes.Count > 0)
+				// Move to the next item
+				BuildListIndex++;
+			}
+
+			// Figure out if we need to invalidate the expand row
+			if(bShowExpandItem)
+			{
+				ListViewGroup Group = (BuildListIndex > 0)? BuildList.Items[BuildListIndex - 1].Group : null;
+				if (ExpandItem == null || ExpandItem.Index != BuildListIndex || ExpandItem.Group != Group)
 				{
-					ExpandItem = new ListViewItem((BuildList.Groups.Count > 0)? BuildList.Groups[BuildList.Groups.Count - 1] : null);
+					ExpandItem = new ListViewItem(Group);
 					ExpandItem.Tag = null;
 					ExpandItem.Selected = false;
 					ExpandItem.Text = "";
-					for(int ColumnIdx = 1; ColumnIdx < BuildList.Columns.Count; ColumnIdx++)
+					for (int ColumnIdx = 1; ColumnIdx < BuildList.Columns.Count; ColumnIdx++)
 					{
 						ExpandItem.SubItems.Add(new ListViewItem.ListViewSubItem(ExpandItem, ""));
 					}
-					BuildList.Items.Add(ExpandItem);
+					BuildList.Items.Insert(BuildListIndex, ExpandItem);
 				}
+				BuildListIndex++;
+			}
 
-				BuildList.EndUpdate();
+			// Remove everything else
+			while (BuildList.Items.Count > BuildListIndex)
+			{
+				BuildList.Items.RemoveAt(BuildList.Items.Count - 1);
+			}
+
+			BuildList.EndUpdate();
+		}
 
-				if(PendingSelectedChangeNumber != -1)
+		ListViewGroup BuildList_FindOrAddGroup(PerforceChangeSummary Change, string GroupName)
+		{
+			// Find or add the new group
+			int GroupIndex = 0;
+			for (; GroupIndex < BuildList.Groups.Count; GroupIndex++)
+			{
+				ListViewGroup NextGroup = BuildList.Groups[GroupIndex];
+				if (NextGroup.Header == GroupName)
 				{
-					SelectChange(PendingSelectedChangeNumber);
+					return NextGroup;
+				}
+
+				PerforceChangeSummary LastChange = null;
+				for(int Idx = NextGroup.Items.Count - 1; Idx >= 0 && LastChange == null; Idx--)
+				{
+					LastChange = NextGroup.Items[Idx].Tag as PerforceChangeSummary;
+				}
+				if (LastChange == null || Change.Number > LastChange.Number)
+				{
+					break;
 				}
 			}
 
-			if(BuildList.HoverItem > BuildList.Items.Count)
+			// Create the new group
+			ListViewGroup Group = new ListViewGroup(GroupName);
+			BuildList.Groups.Insert(GroupIndex, Group);
+			return Group;
+		}
+
+		void BuildList_InsertItem(int Index, PerforceChangeSummary Change)
+		{
+			// Get the display time for this item
+			DateTime DisplayTime = Change.Date;
+			if (Settings.bShowLocalTimes)
 			{
-				BuildList.HoverItem = -1;
+				DisplayTime = (DisplayTime - PerforceMonitor.ServerTimeZone).ToLocalTime();
 			}
 
-			UpdateBuildFailureNotification();
+			// Find or add the new group
+			ListViewGroup Group = BuildList_FindOrAddGroup(Change, DisplayTime.ToString("D"));
 
-			UpdateBuildSteps();
-			UpdateSyncActionCheckboxes();
+			// Create the new item
+			ListViewItem Item = new ListViewItem(Group);
+			Item.Tag = Change;
+
+			// Get the new text for each column
+			string[] Columns = new string[BuildList.Columns.Count];
+			Columns[ChangeColumn.Index] = String.Format("{0}", Change.Number);
+			Columns[TimeColumn.Index] = DisplayTime.ToString("h\\.mmtt");
+			Columns[AuthorColumn.Index] = FormatUserName(Change.User);
+			Columns[DescriptionColumn.Index] = Change.Description.Replace('\n', ' ');
+
+			for (int ColumnIdx = 1; ColumnIdx < BuildList.Columns.Count; ColumnIdx++)
+			{
+				Item.SubItems.Add(new ListViewItem.ListViewSubItem(Item, Columns[ColumnIdx] ?? ""));
+			}
+
+			// Insert it at the right position within the group
+			int GroupInsertIdx = 0;
+			for(; GroupInsertIdx < Group.Items.Count; GroupInsertIdx++)
+			{
+				PerforceChangeSummary OtherChange = Group.Items[GroupInsertIdx].Tag as PerforceChangeSummary;
+				if(OtherChange == null || Change.Number >= OtherChange.Number)
+				{
+					break;
+				}
+			}
+			Group.Items.Insert(GroupInsertIdx, Item);
+
+			// Insert it into the build list
+			BuildList.Items.Insert(Index, Item);
 		}
 
 		bool ShouldShowChange(PerforceChangeSummary Change, string[] ExcludeChanges)
@@ -4233,6 +4303,7 @@ namespace UnrealGameSync
 		{
 			Settings.bShowLocalTimes = false;
 			Settings.Save();
+			BuildList.Items.Clear(); // Need to clear list to rebuild groups, populate time column again
 			UpdateBuildList();
 		}
 
@@ -4240,6 +4311,7 @@ namespace UnrealGameSync
 		{
 			Settings.bShowLocalTimes = true;
 			Settings.Save();
+			BuildList.Items.Clear(); // Need to clear list to rebuild groups, populate time column again
 			UpdateBuildList();
 		}
 
diff --git a/Engine/Source/Programs/UnrealGameSync/UnrealGameSync/PerforceMonitor.cs b/Engine/Source/Programs/UnrealGameSync/UnrealGameSync/PerforceMonitor.cs
--- a/Engine/Source/Programs/UnrealGameSync/UnrealGameSync/PerforceMonitor.cs
+++ b/Engine/Source/Programs/UnrealGameSync/UnrealGameSync/PerforceMonitor.cs
@@ -382,7 +382,7 @@ namespace UnrealGameSync
 			foreach(int QueryChangeNumber in QueryChangeNumbers)
 			{
 				// Skip this stuff if the user wants us to query for more changes
-				if(PendingMaxChanges > CurrentMaxChanges)
+				if(PendingMaxChanges != CurrentMaxChanges)
 				{
 					break;
 				}
