UGS: Use a rich text control to display build health errors. Allows for additional formatting information (eg. headings for each error).

diff --git a/Engine/Source/Programs/UnrealGameSync/UnrealGameSync/Forms/IssueDetailsWindow.cs b/Engine/Source/Programs/UnrealGameSync/UnrealGameSync/Forms/IssueDetailsWindow.cs
--- a/Engine/Source/Programs/UnrealGameSync/UnrealGameSync/Forms/IssueDetailsWindow.cs
+++ b/Engine/Source/Programs/UnrealGameSync/UnrealGameSync/Forms/IssueDetailsWindow.cs
@@ -251,6 +251,7 @@ namespace UnrealGameSync
 		Font BoldFont;
 		PerforceChangeSummary ContextMenuChange;
 		string LastOwner;
+		string LastDetailsText;
 		System.Windows.Forms.Timer UpdateTimer;
 		StatusElementResources StatusElementResources;
 
@@ -369,6 +370,92 @@ namespace UnrealGameSync
 			}
 		}
 
+		void AppendEscapedRtf(StringBuilder Result, string Text)
+		{
+			for (int Idx = 0; Idx < Text.Length; Idx++)
+			{
+				char Character = Text[Idx];
+				if(Character == '\n')
+				{
+					Result.Append(@"\line");
+				}
+				else if(Character >= 0x20 && Character <= 0x7f)
+				{
+					if(Character == '\\' || Character == '{' || Character == '}')
+					{
+						Result.Append('\\');
+					}
+					Result.Append(Character);
+				}
+				else
+				{
+					Result.AppendFormat("\\u{0}?", (int)Character);
+				}
+			}
+		}
+
+		string CreateRichTextErrors(string ErrorText)
+		{
+			List<List<string>> Errors = new List<List<string>>();
+			List<string> CurrentError = null;
+			foreach (string Line in ErrorText.Split('\n'))
+			{
+				if (Line == "***")
+				{
+					CurrentError = null;
+				}
+				else
+				{
+					if (CurrentError == null)
+					{
+						CurrentError = new List<string>();
+						Errors.Add(CurrentError);
+					}
+					CurrentError.Add(Line);
+				}
+			}
+
+			StringBuilder RichText = new StringBuilder();
+
+			RichText.AppendLine(@"{\rtf1\ansi");
+			RichText.AppendLine(@"{\fonttbl{\f0\fnil\fcharset0 Arial;}{\f1\fnil\fcharset0 Courier New;}{\f2\fnil\fcharset0 Calibri;}}");
+			RichText.AppendLine(@"{\colortbl;\red192\green80\blue77;}");
+
+			for(int ErrorIdx = 0; ErrorIdx < Errors.Count; ErrorIdx++)
+			{
+				// Error X/Y:
+				RichText.Append(@"\pard");   // Paragraph default
+				RichText.Append(@"\cf1");    // Foreground color 1
+				RichText.Append(@"\b1");     // Bold on
+				RichText.Append(@"\f0");     // Font 0
+				RichText.Append(@"\fs16");   // Font size 16 (8pt * 2)
+				if(ErrorIdx > 0)
+				{
+					RichText.Append(@"\sb300"); // Space before 200
+				}
+				RichText.Append(@"\sa100");  // Space after 100
+				RichText.AppendFormat(@"Error {0}/{1}:", ErrorIdx + 1, Errors.Count);
+				RichText.AppendLine(@"\par");
+
+				// Error text
+				foreach(string ErrorLine in Errors[ErrorIdx])
+				{
+					RichText.Append(@"\pard");   // Paragraph default
+					RichText.Append(@"\cf0");    // Foreground color 0
+					RichText.Append(@"\b0");     // Bold off
+					RichText.Append(@"\f1");     // Font 1
+					RichText.Append(@"\fs16");   // Font size 16 (8pt * 2)
+					RichText.Append(@"\fi150");  // First line indent 150
+					RichText.Append(@"\li150");  // Other line indent 150
+					AppendEscapedRtf(RichText, ErrorLine);
+					RichText.Append(@"\par");
+				}
+			}
+
+			RichText.AppendLine("}");
+			return RichText.ToString();
+		}
+
 		void UpdateCurrentIssue()
 		{
 			UpdateSummaryTextIfChanged(SummaryTextBox, Issue.Summary.ToString());
@@ -449,7 +536,16 @@ namespace UnrealGameSync
 			UpdateSummaryTextIfChanged(StepNamesTextBox, String.Join(", ", Issue.Builds.Select(x => x.JobStepName).Distinct().OrderBy(x => x)));
 			UpdateSummaryTextIfChanged(StreamNamesTextBox, String.Join(", ", Issue.Builds.Select(x => x.Stream).Distinct().OrderBy(x => x)));
 
-			UpdateSummaryTextIfChanged(DetailsTextBox, Issue.Details.Replace("\n", "\r\n"));
+			if(LastDetailsText != Issue.Details)
+			{
+				string RtfText = CreateRichTextErrors(Issue.Details);
+				using (MemoryStream Stream = new MemoryStream(Encoding.UTF8.GetBytes(RtfText), false))
+				{
+					DetailsTextBox.LoadFile(Stream, RichTextBoxStreamType.RichText);
+					DetailsTextBox.Select(0, 0);
+				}
+				LastDetailsText = Issue.Details;
+			}
 
 			if(Issue.FixChange == 0)
 			{
