UGS: Rename the fingerprint "Messages" field to "Identifiers" to more accurately reflect what it's used for.

#rnx

diff --git a/Engine/Source/Programs/UnrealGameSync/MetadataTool/BuildHealth/Matchers/UndefinedSymbolPatternMatcher.cs b/Engine/Source/Programs/UnrealGameSync/MetadataTool/BuildHealth/Matchers/UndefinedSymbolPatternMatcher.cs
--- a/Engine/Source/Programs/UnrealGameSync/MetadataTool/BuildHealth/Matchers/UndefinedSymbolPatternMatcher.cs
+++ b/Engine/Source/Programs/UnrealGameSync/MetadataTool/BuildHealth/Matchers/UndefinedSymbolPatternMatcher.cs
@@ -85,7 +85,7 @@ namespace MetadataTool
 			{
 				TrackedIssueFingerprint Fingerprint = new TrackedIssueFingerprint(Category, GetSummary(SymbolNames), Diagnostic.Url, Job.Change);
 				Fingerprint.Details.Add(Diagnostic.Message);
-				Fingerprint.Messages.UnionWith(SymbolNames);
+				Fingerprint.Identifiers.UnionWith(SymbolNames);
 				Fingerprints.Add(Fingerprint);
 				return true;
 			}
@@ -98,7 +98,7 @@ namespace MetadataTool
 		{
 			base.Merge(Source, Target);
 
-			Target.Summary = GetSummary(Target.Messages);
+			Target.Summary = GetSummary(Target.Identifiers);
 		}
 
 		string GetSummary(SortedSet<string> SymbolNames)
diff --git a/Engine/Source/Programs/UnrealGameSync/MetadataTool/BuildHealth/PatternMatcher.cs b/Engine/Source/Programs/UnrealGameSync/MetadataTool/BuildHealth/PatternMatcher.cs
--- a/Engine/Source/Programs/UnrealGameSync/MetadataTool/BuildHealth/PatternMatcher.cs
+++ b/Engine/Source/Programs/UnrealGameSync/MetadataTool/BuildHealth/PatternMatcher.cs
@@ -69,7 +69,7 @@ namespace MetadataTool
 			// Check that a filename or message matches
 			if (Source.InitialChange != Target.InitialChange)
 			{
-				if (!Source.FileNames.Any(x => Target.FileNames.Contains(x)) && !Source.Messages.Any(x => Target.Messages.Contains(x)))
+				if (!Source.FileNames.Any(x => Target.FileNames.Contains(x)) && !Source.Identifiers.Any(x => Target.Identifiers.Contains(x)))
 				{
 					return false;
 				}
@@ -85,7 +85,7 @@ namespace MetadataTool
 		public virtual void Merge(TrackedIssueFingerprint Source, TrackedIssueFingerprint Target)
 		{
 			Target.FileNames.UnionWith(Source.FileNames);
-			Target.Messages.UnionWith(Source.Messages);
+			Target.Identifiers.UnionWith(Source.Identifiers);
 		}
 
 		/// <summary>
diff --git a/Engine/Source/Programs/UnrealGameSync/MetadataTool/BuildHealth/TrackedIssueFingerprint.cs b/Engine/Source/Programs/UnrealGameSync/MetadataTool/BuildHealth/TrackedIssueFingerprint.cs
--- a/Engine/Source/Programs/UnrealGameSync/MetadataTool/BuildHealth/TrackedIssueFingerprint.cs
+++ b/Engine/Source/Programs/UnrealGameSync/MetadataTool/BuildHealth/TrackedIssueFingerprint.cs
@@ -56,7 +56,7 @@ namespace MetadataTool
 		/// List of messages associated with this issue
 		/// </summary>
 		[DataMember]
-		public SortedSet<string> Messages = new SortedSet<string>();
+		public SortedSet<string> Identifiers = new SortedSet<string>();
 
 		/// <summary>
 		/// Constructor
