UGS: Rename Tracked* to BuildHealth*.

#rnx

diff --git a/Engine/Source/Programs/UnrealGameSync/MetadataTool/BuildHealth/Matchers/CompilePatternMatcher.cs b/Engine/Source/Programs/UnrealGameSync/MetadataTool/BuildHealth/Matchers/CompilePatternMatcher.cs
--- a/Engine/Source/Programs/UnrealGameSync/MetadataTool/BuildHealth/Matchers/CompilePatternMatcher.cs
+++ b/Engine/Source/Programs/UnrealGameSync/MetadataTool/BuildHealth/Matchers/CompilePatternMatcher.cs
@@ -26,7 +26,7 @@ namespace MetadataTool
 
 		public override string Category => "Compile";
 
-		public override bool TryMatch(InputJob Job, InputJobStep JobStep, InputDiagnostic Diagnostic, List<TrackedIssue> Issues)
+		public override bool TryMatch(InputJob Job, InputJobStep JobStep, InputDiagnostic Diagnostic, List<BuildHealthIssue> Issues)
 		{
 			// Find any files in compiler output format
 			HashSet<string> SourceFileNames = new HashSet<string>(StringComparer.OrdinalIgnoreCase);
@@ -45,7 +45,7 @@ namespace MetadataTool
 			// If we found any source files, create a diagnostic category for them
 			if (SourceFileNames.Count > 0)
 			{
-				TrackedIssue Issue = new TrackedIssue(Category, Job.Url, Diagnostic.Url);
+				BuildHealthIssue Issue = new BuildHealthIssue(Category, Job.Url, Diagnostic.Url);
 				Issue.Details.Add(ShortenPaths(Diagnostic.Message));
 				Issue.FileNames.UnionWith(SourceFileNames);
 				Issues.Add(Issue);
@@ -62,7 +62,7 @@ namespace MetadataTool
 			return Text;
 		}
 
-		public override string GetSummary(TrackedIssue Issue)
+		public override string GetSummary(BuildHealthIssue Issue)
 		{
 			SortedSet<string> ShortFileNames = GetSourceFileNames(Issue.FileNames);
 			if (ShortFileNames.Count == 0)
diff --git a/Engine/Source/Programs/UnrealGameSync/MetadataTool/BuildHealth/Matchers/UndefinedSymbolPatternMatcher.cs b/Engine/Source/Programs/UnrealGameSync/MetadataTool/BuildHealth/Matchers/UndefinedSymbolPatternMatcher.cs
--- a/Engine/Source/Programs/UnrealGameSync/MetadataTool/BuildHealth/Matchers/UndefinedSymbolPatternMatcher.cs
+++ b/Engine/Source/Programs/UnrealGameSync/MetadataTool/BuildHealth/Matchers/UndefinedSymbolPatternMatcher.cs
@@ -13,7 +13,7 @@ namespace MetadataTool
 	{
 		public override string Category => "UndefinedSymbol";
 
-		public override bool TryMatch(InputJob Job, InputJobStep JobStep, InputDiagnostic Diagnostic, List<TrackedIssue> Issues)
+		public override bool TryMatch(InputJob Job, InputJobStep JobStep, InputDiagnostic Diagnostic, List<BuildHealthIssue> Issues)
 		{
 			List<string> SymbolMatches = new List<string>();
 
@@ -83,7 +83,7 @@ namespace MetadataTool
 			// If we found any symbol names, create a fingerprint for them
 			if (SymbolNames.Count > 0)
 			{
-				TrackedIssue Issue = new TrackedIssue(Category, Job.Url, Diagnostic.Url);
+				BuildHealthIssue Issue = new BuildHealthIssue(Category, Job.Url, Diagnostic.Url);
 				Issue.Details.Add(Diagnostic.Message);
 				Issue.Identifiers.UnionWith(SymbolNames);
 				Issues.Add(Issue);
@@ -94,7 +94,7 @@ namespace MetadataTool
 			return false;
 		}
 
-		public override string GetSummary(TrackedIssue Issue)
+		public override string GetSummary(BuildHealthIssue Issue)
 		{
 			if(Issue.Identifiers.Count == 1)
 			{
diff --git a/Engine/Source/Programs/UnrealGameSync/MetadataTool/BuildHealth/PatternMatcher.cs b/Engine/Source/Programs/UnrealGameSync/MetadataTool/BuildHealth/PatternMatcher.cs
--- a/Engine/Source/Programs/UnrealGameSync/MetadataTool/BuildHealth/PatternMatcher.cs
+++ b/Engine/Source/Programs/UnrealGameSync/MetadataTool/BuildHealth/PatternMatcher.cs
@@ -30,7 +30,7 @@ namespace MetadataTool
 		/// <param name="JobStep">The job step that was run</param>
 		/// <param name="Diagnostics">List of diagnostics that were produced by the build. Items should be removed from this list if they match.</param>
 		/// <param name="Issues">List which receives all the matched issues.</param>
-		public virtual void Match(InputJob Job, InputJobStep JobStep, List<InputDiagnostic> Diagnostics, List<TrackedIssue> Issues)
+		public virtual void Match(InputJob Job, InputJobStep JobStep, List<InputDiagnostic> Diagnostics, List<BuildHealthIssue> Issues)
 		{
 			for (int Idx = 0; Idx < Diagnostics.Count; Idx++)
 			{
@@ -51,14 +51,14 @@ namespace MetadataTool
 		/// <param name="Diagnostic">A diagnostic from the given job step</param>
 		/// <param name="Issues">List which receives all the matched issues.</param>
 		/// <returns>True if this diagnostic should be removed (usually because a fingerprint was created)</returns>
-		public abstract bool TryMatch(InputJob Job, InputJobStep JobStep, InputDiagnostic Diagnostic, List<TrackedIssue> Issues);
+		public abstract bool TryMatch(InputJob Job, InputJobStep JobStep, InputDiagnostic Diagnostic, List<BuildHealthIssue> Issues);
 
 		/// <summary>
 		/// Determines if one fingerprint can be merged with another one
 		/// </summary>
 		/// <param name="Source">The source fingerprint</param>
 		/// <param name="Target">The fingerprint to merge into</param>
-		public virtual bool CanMerge(TrackedIssue Source, TrackedIssue Target)
+		public virtual bool CanMerge(BuildHealthIssue Source, BuildHealthIssue Target)
 		{
 			// Make sure the categories match
 			if (Source.Category != Target.Category)
@@ -82,7 +82,7 @@ namespace MetadataTool
 		/// </summary>
 		/// <param name="Source">The source fingerprint</param>
 		/// <param name="Target">The fingerprint to merge into</param>
-		public virtual void Merge(TrackedIssue Source, TrackedIssue Target)
+		public virtual void Merge(BuildHealthIssue Source, BuildHealthIssue Target)
 		{
 			Target.Details.AddRange(Source.Details);
 			Target.FileNames.UnionWith(Source.FileNames);
@@ -96,7 +96,7 @@ namespace MetadataTool
 		/// <param name="Fingerprint">Fingerprint for the issue</param>
 		/// <param name="Changes">List of changes since the issue first occurred.</param>
 		/// <returns>List of changes which are causers for the issue</returns>
-		public virtual List<ChangeInfo> FindCausers(PerforceConnection Perforce, TrackedIssue Issue, IReadOnlyList<ChangeInfo> Changes)
+		public virtual List<ChangeInfo> FindCausers(PerforceConnection Perforce, BuildHealthIssue Issue, IReadOnlyList<ChangeInfo> Changes)
 		{
 			SortedSet<string> FileNamesWithoutPath = GetFileNamesWithoutPath(Issue.FileNames);
 
@@ -248,6 +248,6 @@ namespace MetadataTool
 		/// </summary>
 		/// <param name="Issue">The issue to summarize</param>
 		/// <returns>The summary text for this issue</returns>
-		public abstract string GetSummary(TrackedIssue Issue);
+		public abstract string GetSummary(BuildHealthIssue Issue);
 	}
 }
diff --git a/Engine/Source/Programs/UnrealGameSync/MetadataTool/BuildHealth/TrackedBuild.cs b/Engine/Source/Programs/UnrealGameSync/MetadataTool/BuildHealth/TrackedBuild.cs
--- a/Engine/Source/Programs/UnrealGameSync/MetadataTool/BuildHealth/TrackedBuild.cs
+++ b/Engine/Source/Programs/UnrealGameSync/MetadataTool/BuildHealth/TrackedBuild.cs
@@ -15,7 +15,7 @@ namespace MetadataTool
 	/// </summary>
 	[DataContract]
 	[DebuggerDisplay("{Change}: {JobName}")]
-	class TrackedBuild : IComparable<TrackedBuild>
+	class BuildHealthJobStep : IComparable<BuildHealthJobStep>
 	{
 		/// <summary>
 		/// The changelist that this build was run at.
@@ -62,7 +62,7 @@ namespace MetadataTool
 		/// <summary>
 		/// Constructor
 		/// </summary>
-		public TrackedBuild(int Change, string JobName, string JobUrl, string JobStepName, string JobStepUrl, string ErrorUrl)
+		public BuildHealthJobStep(int Change, string JobName, string JobUrl, string JobStepName, string JobStepUrl, string ErrorUrl)
 		{
 			this.Change = Change;
 			this.JobName = JobName;
@@ -77,7 +77,7 @@ namespace MetadataTool
 		/// </summary>
 		/// <param name="Other">Build to compare to</param>
 		/// <returns>Value indicating how the two builds should be ordered</returns>
-		public int CompareTo(TrackedBuild Other)
+		public int CompareTo(BuildHealthJobStep Other)
 		{
 			int Delta = Change - Other.Change;
 			if (Delta == 0)
diff --git a/Engine/Source/Programs/UnrealGameSync/MetadataTool/BuildHealth/TrackedIssue.cs b/Engine/Source/Programs/UnrealGameSync/MetadataTool/BuildHealth/TrackedIssue.cs
--- a/Engine/Source/Programs/UnrealGameSync/MetadataTool/BuildHealth/TrackedIssue.cs
+++ b/Engine/Source/Programs/UnrealGameSync/MetadataTool/BuildHealth/TrackedIssue.cs
@@ -14,7 +14,7 @@ namespace MetadataTool
 	/// Information about a particular issue
 	/// </summary>
 	[DataContract]
-	class TrackedIssue
+	class BuildHealthIssue
 	{
 		/// <summary>
 		/// The issue id in the database. -1 for issues that have not been posted yet.
@@ -86,7 +86,7 @@ namespace MetadataTool
 		/// Map of stream name -> step name -> history for builds exhibiting this issue
 		/// </summary>
 		[DataMember(Order = 30, IsRequired = true)]
-		public Dictionary<string, Dictionary<string, TrackedIssueHistory>> Streams = new Dictionary<string, Dictionary<string, TrackedIssueHistory>>();
+		public Dictionary<string, Dictionary<string, BuildHealthJobHistory>> Streams = new Dictionary<string, Dictionary<string, BuildHealthJobHistory>>();
 
 		/// <summary>
 		/// Set of changes which may have caused this issue. Used to de-duplicate issues between streams.
@@ -112,7 +112,7 @@ namespace MetadataTool
 		/// <param name="Category">Category of this issue</param>
 		/// <param name="InitialJobUrl">Url of the initial job that this error was seen with</param>
 		/// <param name="ErrorUrl"></param>
-		public TrackedIssue(string Category, string InitialJobUrl, string ErrorUrl)
+		public BuildHealthIssue(string Category, string InitialJobUrl, string ErrorUrl)
 		{
 			this.Category = Category;
 			this.InitialJobUrl = InitialJobUrl;
diff --git a/Engine/Source/Programs/UnrealGameSync/MetadataTool/BuildHealth/TrackedIssueHistory.cs b/Engine/Source/Programs/UnrealGameSync/MetadataTool/BuildHealth/TrackedIssueHistory.cs
--- a/Engine/Source/Programs/UnrealGameSync/MetadataTool/BuildHealth/TrackedIssueHistory.cs
+++ b/Engine/Source/Programs/UnrealGameSync/MetadataTool/BuildHealth/TrackedIssueHistory.cs
@@ -13,30 +13,30 @@ namespace MetadataTool
 	/// History of builds within a particular stream that contribute to an issue
 	/// </summary>
 	[DataContract]
-	class TrackedIssueHistory
+	class BuildHealthJobHistory
 	{
 		/// <summary>
 		/// The previous build before it started failing. This should be updated as new builds come in.
 		/// </summary>
 		[DataMember]
-		public TrackedBuild PrevSuccessfulBuild;
+		public BuildHealthJobStep PrevSuccessfulBuild;
 
 		/// <summary>
 		/// List of failing builds contributing to this issue
 		/// </summary>
 		[DataMember]
-		public List<TrackedBuild> FailedBuilds = new List<TrackedBuild>();
+		public List<BuildHealthJobStep> FailedBuilds = new List<BuildHealthJobStep>();
 
 		/// <summary>
 		/// The first successful build after the failures.
 		/// </summary>
 		[DataMember]
-		public TrackedBuild NextSuccessfulBuild;
+		public BuildHealthJobStep NextSuccessfulBuild;
 
 		/// <summary>
 		/// Constructs a new history for a particular stream
 		/// </summary>
-		public TrackedIssueHistory(TrackedBuild PrevSuccessfulBuild)
+		public BuildHealthJobHistory(BuildHealthJobStep PrevSuccessfulBuild)
 		{
 			this.PrevSuccessfulBuild  = PrevSuccessfulBuild;
 		}
@@ -45,7 +45,7 @@ namespace MetadataTool
 		/// Adds a failed build to this object
 		/// </summary>
 		/// <param name="Build">The failed build</param>
-		public void AddFailedBuild(TrackedBuild Build)
+		public void AddFailedBuild(BuildHealthJobStep Build)
 		{
 			int Index = FailedBuilds.BinarySearch(Build);
 			if (Index < 0)
@@ -79,7 +79,7 @@ namespace MetadataTool
 		/// <summary>
 		/// Enumerates all the builds in this stream
 		/// </summary>
-		public IEnumerable<TrackedBuild> Builds
+		public IEnumerable<BuildHealthJobStep> Builds
 		{
 			get
 			{
@@ -87,7 +87,7 @@ namespace MetadataTool
 				{
 					yield return PrevSuccessfulBuild;
 				}
-				foreach(TrackedBuild FailedBuild in FailedBuilds)
+				foreach(BuildHealthJobStep FailedBuild in FailedBuilds)
 				{
 					yield return FailedBuild;
 				}
diff --git a/Engine/Source/Programs/UnrealGameSync/MetadataTool/BuildHealth/TrackedState.cs b/Engine/Source/Programs/UnrealGameSync/MetadataTool/BuildHealth/TrackedState.cs
--- a/Engine/Source/Programs/UnrealGameSync/MetadataTool/BuildHealth/TrackedState.cs
+++ b/Engine/Source/Programs/UnrealGameSync/MetadataTool/BuildHealth/TrackedState.cs
@@ -13,32 +13,32 @@ namespace MetadataTool
 	/// The persistent state tracked by the program between runs.
 	/// </summary>
 	[DataContract]
-	class TrackedState
+	class BuildHealthState
 	{
 		/// <summary>
 		/// All open issues.
 		/// </summary>
 		[DataMember]
-		public List<TrackedIssue> Issues = new List<TrackedIssue>();
+		public List<BuildHealthIssue> Issues = new List<BuildHealthIssue>();
 
 		/// <summary>
 		/// Map of stream name to list of builds within it. Only retains a few builds in order to determine the last succesful CL for a build.
 		/// </summary>
 		[DataMember]
-		public Dictionary<string, List<TrackedBuild>> Streams = new Dictionary<string, List<TrackedBuild>>();
+		public Dictionary<string, List<BuildHealthJobStep>> Streams = new Dictionary<string, List<BuildHealthJobStep>>();
 
 		/// <summary>
 		/// Adds a new build for tracking
 		/// </summary>
 		/// <param name="Stream">The stream containing the build</param>
 		/// <param name="Build">The build to add</param>
-		public void AddBuild(string Stream, TrackedBuild Build)
+		public void AddBuild(string Stream, BuildHealthJobStep Build)
 		{
 			// Get the list of tracked builds for this stream
-			List<TrackedBuild> Builds;
+			List<BuildHealthJobStep> Builds;
 			if (!Streams.TryGetValue(Stream, out Builds))
 			{
-				Builds = new List<TrackedBuild>();
+				Builds = new List<BuildHealthJobStep>();
 				Streams.Add(Stream, Builds);
 			}
 
@@ -58,11 +58,11 @@ namespace MetadataTool
 		/// <param name="Change">The change to search for</param>
 		/// <param name="StepNames">The step names which the job step needs to have</param>
 		/// <returns>The last build known before the given changelist</returns>
-		public TrackedBuild FindBuildBefore(string Stream, int Change, string StepName)
+		public BuildHealthJobStep FindBuildBefore(string Stream, int Change, string StepName)
 		{
-			TrackedBuild Result = null;
+			BuildHealthJobStep Result = null;
 
-			List<TrackedBuild> Builds;
+			List<BuildHealthJobStep> Builds;
 			if(Streams.TryGetValue(Stream, out Builds))
 			{
 				for(int Idx = 0; Idx < Builds.Count && Builds[Idx].Change < Change; Idx++)
@@ -84,11 +84,11 @@ namespace MetadataTool
 		/// <param name="Change">The change to search for</param>
 		/// <param name="StepName">Name of the job step to find</param>
 		/// <returns>The last build known before the given changelist</returns>
-		public TrackedBuild FindBuildAfter(string Stream, int Change, string StepName)
+		public BuildHealthJobStep FindBuildAfter(string Stream, int Change, string StepName)
 		{
-			TrackedBuild Result = null;
+			BuildHealthJobStep Result = null;
 
-			List<TrackedBuild> Builds;
+			List<BuildHealthJobStep> Builds;
 			if (Streams.TryGetValue(Stream, out Builds))
 			{
 				for (int Idx = Builds.Count - 1; Idx >= 0 && Builds[Idx].Change > Change; Idx--)
diff --git a/Engine/Source/Programs/UnrealGameSync/MetadataTool/CommandHandler_BuildHealth.cs b/Engine/Source/Programs/UnrealGameSync/MetadataTool/CommandHandler_BuildHealth.cs
--- a/Engine/Source/Programs/UnrealGameSync/MetadataTool/CommandHandler_BuildHealth.cs
+++ b/Engine/Source/Programs/UnrealGameSync/MetadataTool/CommandHandler_BuildHealth.cs
@@ -76,16 +76,16 @@ namespace MetadataTool
 			CompleteStateTransaction(StateFile);
 
 			// Read the persistent data file
-			TrackedState State;
+			BuildHealthState State;
 			if (!bClean && FileReference.Exists(StateFile))
 			{
 				Log.TraceInformation("Reading persistent data from {0}", StateFile);
-				State = DeserializeJson<TrackedState>(StateFile);
+				State = DeserializeJson<BuildHealthState>(StateFile);
 			}
 			else
 			{
 				Log.TraceInformation("Creating new persistent data");
-				State = new TrackedState();
+				State = new BuildHealthState();
 			}
 
 			// Create the Perforce connection
@@ -98,7 +98,7 @@ namespace MetadataTool
 				// Add a new build for each job step
 				foreach(InputJobStep InputJobStep in InputJob.Steps)
 				{
-					TrackedBuild NewBuild = new TrackedBuild(InputJob.Change, InputJob.Name, InputJob.Url, InputJobStep.Name, InputJobStep.Url, null);
+					BuildHealthJobStep NewBuild = new BuildHealthJobStep(InputJob.Change, InputJob.Name, InputJob.Url, InputJobStep.Name, InputJobStep.Url, null);
 					State.AddBuild(InputJob.Stream, NewBuild);
 				}
 
@@ -113,17 +113,17 @@ namespace MetadataTool
 			// Try to find the next successful build for each stream, so we can close it as part of updating the server
 			for (int Idx = 0; Idx < State.Issues.Count; Idx++)
 			{
-				TrackedIssue Issue = State.Issues[Idx];
+				BuildHealthIssue Issue = State.Issues[Idx];
 				foreach(string Stream in Issue.Streams.Keys)
 				{
-					Dictionary<string, TrackedIssueHistory> StepNameToHistory = Issue.Streams[Stream];
+					Dictionary<string, BuildHealthJobHistory> StepNameToHistory = Issue.Streams[Stream];
 					foreach(string StepName in StepNameToHistory.Keys)
 					{
-						TrackedIssueHistory IssueHistory = StepNameToHistory[StepName];
+						BuildHealthJobHistory IssueHistory = StepNameToHistory[StepName];
 						if(IssueHistory.FailedBuilds.Count > 0 && IssueHistory.NextSuccessfulBuild == null)
 						{
 							// Find the successful build after this change
-							TrackedBuild LastFailedBuild = IssueHistory.FailedBuilds[IssueHistory.FailedBuilds.Count - 1];
+							BuildHealthJobStep LastFailedBuild = IssueHistory.FailedBuilds[IssueHistory.FailedBuilds.Count - 1];
 							IssueHistory.NextSuccessfulBuild = State.FindBuildAfter(Stream, LastFailedBuild.Change, StepName);
 						}
 					}
@@ -135,7 +135,7 @@ namespace MetadataTool
 			{
 				// Find all the unique change numbers for each stream
 				SortedSet<int> ChangeNumbers = new SortedSet<int>();
-				foreach (List<TrackedBuild> Builds in State.Streams.Values)
+				foreach (List<BuildHealthJobStep> Builds in State.Streams.Values)
 				{
 					ChangeNumbers.UnionWith(Builds.Select(x => x.Change));
 				}
@@ -157,14 +157,14 @@ namespace MetadataTool
 				}
 
 				// Remove any builds we no longer want to track
-				foreach (List<TrackedBuild> Builds in State.Streams.Values)
+				foreach (List<BuildHealthJobStep> Builds in State.Streams.Values)
 				{
 					Builds.RemoveAll(x => x.Change <= DeleteChangeNumber);
 				}
 			}
 
 			// Mark any issues as resolved
-			foreach(TrackedIssue Issue in State.Issues)
+			foreach(BuildHealthIssue Issue in State.Issues)
 			{
 				if(Issue.IsResolved())
 				{
@@ -191,7 +191,7 @@ namespace MetadataTool
 			if (ServerUrl != null)
 			{
 				// Post any issue updates
-				foreach(TrackedIssue Issue in State.Issues)
+				foreach(BuildHealthIssue Issue in State.Issues)
 				{
 					PatternMatcher Matcher;
 					if(!CategoryNameToMatcher.TryGetValue(Issue.Category, out Matcher))
@@ -255,7 +255,7 @@ namespace MetadataTool
 				}
 
 				// Update the summary for any issues that are still open
-				foreach (TrackedIssue Issue in State.Issues)
+				foreach (BuildHealthIssue Issue in State.Issues)
 				{
 					if (Issue.Id == -1)
 					{
@@ -264,13 +264,13 @@ namespace MetadataTool
 				}
 
 				// Add any new builds associated with issues
-				foreach (TrackedIssue Issue in State.Issues)
+				foreach (BuildHealthIssue Issue in State.Issues)
 				{
-					foreach(KeyValuePair<string, Dictionary<string, TrackedIssueHistory>> StreamPair in Issue.Streams)
+					foreach(KeyValuePair<string, Dictionary<string, BuildHealthJobHistory>> StreamPair in Issue.Streams)
 					{
-						foreach(TrackedIssueHistory StreamHistory in StreamPair.Value.Values)
+						foreach(BuildHealthJobHistory StreamHistory in StreamPair.Value.Values)
 						{
-							foreach(TrackedBuild Build in StreamHistory.Builds)
+							foreach(BuildHealthJobStep Build in StreamHistory.Builds)
 							{
 								if(!Build.bPostedToServer)
 								{
@@ -306,7 +306,7 @@ namespace MetadataTool
 				// Close any issues which are complete
 				for (int Idx = 0; Idx < State.Issues.Count; Idx++)
 				{
-					TrackedIssue Issue = State.Issues[Idx];
+					BuildHealthIssue Issue = State.Issues[Idx];
 					if (Issue.ResolvedAt.HasValue != Issue.bPostedResolved)
 					{
 						Log.TraceInformation("Setting issue {0} resolved flag to {1}", Issue.Id, Issue.ResolvedAt.HasValue);
@@ -329,7 +329,7 @@ namespace MetadataTool
 				}
 
 				// Update watchers on any open builds
-				foreach(TrackedIssue Issue in State.Issues)
+				foreach(BuildHealthIssue Issue in State.Issues)
 				{
 					while (Issue.PendingWatchers.Count > 0)
 					{
@@ -358,7 +358,7 @@ namespace MetadataTool
 			DateTime RemoveIssueTime = DateTime.UtcNow - TimeSpan.FromHours(24.0);
 			for(int Idx = 0; Idx < State.Issues.Count; Idx++)
 			{
-				TrackedIssue Issue = State.Issues[Idx];
+				BuildHealthIssue Issue = State.Issues[Idx];
 				if(Issue.ResolvedAt.HasValue && Issue.ResolvedAt.Value < RemoveIssueTime)
 				{
 					State.Issues.RemoveAt(Idx--);
@@ -432,10 +432,10 @@ namespace MetadataTool
 		/// <param name="PreviousChange">The last changelist that was built before this one</param>
 		/// <param name="InputJob">Job containing the step to add</param>
 		/// <param name="InputJobStep">The job step to add</param>
-		void AddStep(PerforceConnection Perforce, TrackedState State, InputJob InputJob, InputJobStep InputJobStep)
+		void AddStep(PerforceConnection Perforce, BuildHealthState State, InputJob InputJob, InputJobStep InputJobStep)
 		{
 			// Create all the fingerprints for failures in this step
-			List<TrackedIssue> InputIssues = new List<TrackedIssue>();
+			List<BuildHealthIssue> InputIssues = new List<BuildHealthIssue>();
 			if(InputJobStep.Diagnostics != null)
 			{
 				List<InputDiagnostic> Diagnostics = new List<InputDiagnostic>(InputJobStep.Diagnostics); 
@@ -446,9 +446,9 @@ namespace MetadataTool
 			}
 
 			// Add all the fingerprints to issues
-			foreach(TrackedIssue InputIssue in InputIssues)
+			foreach(BuildHealthIssue InputIssue in InputIssues)
 			{
-				TrackedIssue Issue = FindOrAddIssueForFingerprint(Perforce, State, InputJob, InputJobStep, InputIssue);
+				BuildHealthIssue Issue = FindOrAddIssueForFingerprint(Perforce, State, InputJob, InputJobStep, InputIssue);
 				AddFailureToIssue(Issue, InputJob, InputJobStep, InputIssue.ErrorUrl, State);
 			}
 		}
@@ -462,23 +462,23 @@ namespace MetadataTool
 		/// <param name="PreviousChange">The last changelist that was built before this one</param>
 		/// <param name="InputJob">Job containing the step to add</param>
 		/// <param name="InputJobStep">The job step to add</param>
-		TrackedIssue FindOrAddIssueForFingerprint(PerforceConnection Perforce, TrackedState State, InputJob InputJob, InputJobStep InputJobStep, TrackedIssue InputIssue)
+		BuildHealthIssue FindOrAddIssueForFingerprint(PerforceConnection Perforce, BuildHealthState State, InputJob InputJob, InputJobStep InputJobStep, BuildHealthIssue InputIssue)
 		{
 			// Find the pattern matcher for this fingerprint
 			PatternMatcher Matcher = CategoryToMatcher[InputIssue.Category];
 
 			// Check if it can be added to an existing open issue
-			foreach (TrackedIssue Issue in State.Issues)
+			foreach (BuildHealthIssue Issue in State.Issues)
 			{
 				// Check this issue already exists in the current stream
-				Dictionary<string, TrackedIssueHistory> StepNameToHistory;
+				Dictionary<string, BuildHealthJobHistory> StepNameToHistory;
 				if(!Issue.Streams.TryGetValue(InputJob.Stream, out StepNameToHistory))
 				{
 					continue;
 				}
 
 				// Check that this issue has not already been closed
-				TrackedIssueHistory History;
+				BuildHealthJobHistory History;
 				if (StepNameToHistory.TryGetValue(InputJobStep.Name, out History))
 				{
 					if(!History.CanAddFailedBuild(InputJob.Change))
@@ -509,7 +509,7 @@ namespace MetadataTool
 			IReadOnlyList<ChangeInfo> Changes = null;
 
 			// Find the previous changelist that was built in this stream
-			List<TrackedBuild> StreamBuilds;
+			List<BuildHealthJobStep> StreamBuilds;
 			if(State.Streams.TryGetValue(InputJob.Stream, out StreamBuilds))
 			{
 				// Find the last change submitted to this stream before it started failing
@@ -527,7 +527,7 @@ namespace MetadataTool
 					if (Changes.Count > 0)
 					{
 						SortedSet<int> SourceChanges = new SortedSet<int>(Changes.SelectMany(x => x.SourceChanges));
-						foreach (TrackedIssue Issue in State.Issues)
+						foreach (BuildHealthIssue Issue in State.Issues)
 						{
 							// Check if this issue does not already contain this stream, but contains one of the causing changes
 							if (Issue.Streams.ContainsKey(InputJob.Stream))
@@ -572,9 +572,9 @@ namespace MetadataTool
 		/// <param name="InputJobStep">The step to create a build for</param>
 		/// <param name="InputErrorUrl">The error Url</param>
 		/// <returns>New build instance</returns>
-		TrackedBuild CreateBuildForJobStep(InputJob InputJob, InputJobStep InputJobStep, string InputErrorUrl)
+		BuildHealthJobStep CreateBuildForJobStep(InputJob InputJob, InputJobStep InputJobStep, string InputErrorUrl)
 		{
-			return new TrackedBuild(InputJob.Change, InputJob.Name, InputJob.Url, InputJobStep.Name, InputJobStep.Url, InputErrorUrl);
+			return new BuildHealthJobStep(InputJob.Change, InputJob.Name, InputJob.Url, InputJobStep.Name, InputJobStep.Url, InputErrorUrl);
 		}
 
 		/// <summary>
@@ -585,21 +585,21 @@ namespace MetadataTool
 		/// <param name="InputJobStep">The job step containing the error</param>
 		/// <param name="InputErrorUrl">Url of the error</param>
 		/// <param name="State">Current persistent state. Used to find previous build history.</param>
-		void AddFailureToIssue(TrackedIssue Issue, InputJob InputJob, InputJobStep InputJobStep, string InputErrorUrl, TrackedState State)
+		void AddFailureToIssue(BuildHealthIssue Issue, InputJob InputJob, InputJobStep InputJobStep, string InputErrorUrl, BuildHealthState State)
 		{
 			// Find or add a step name to history mapping
-			Dictionary<string, TrackedIssueHistory> StepNameToHistory;
+			Dictionary<string, BuildHealthJobHistory> StepNameToHistory;
 			if(!Issue.Streams.TryGetValue(InputJob.Stream, out StepNameToHistory))
 			{
-				StepNameToHistory = new Dictionary<string, TrackedIssueHistory>();
+				StepNameToHistory = new Dictionary<string, BuildHealthJobHistory>();
 				Issue.Streams.Add(InputJob.Stream, StepNameToHistory);
 			}
 
 			// Find or add a history for this step
-			TrackedIssueHistory History;
+			BuildHealthJobHistory History;
 			if(!StepNameToHistory.TryGetValue(InputJobStep.Name, out History))
 			{
-				History = new TrackedIssueHistory(State.FindBuildBefore(InputJob.Stream, InputJob.Change, InputJobStep.Name));
+				History = new BuildHealthJobHistory(State.FindBuildBefore(InputJob.Stream, InputJob.Change, InputJobStep.Name));
 				StepNameToHistory.Add(InputJobStep.Name, History);
 			}
 
@@ -686,7 +686,7 @@ namespace MetadataTool
 		/// </summary>
 		/// <param name="StateFile">The file to write to</param>
 		/// <param name="State">The state object</param>
-		static void WriteState(FileReference StateFile, TrackedState State)
+		static void WriteState(FileReference StateFile, BuildHealthState State)
 		{
 			// Write out the state to the transaction file
 			FileReference StateTransactionFile = GetStateTransactionFile(StateFile);
@@ -708,9 +708,7 @@ namespace MetadataTool
 			{
 				DataContractJsonSerializer InputFileDataSerializer = new DataContractJsonSerializer(Object.GetType());
 				InputFileDataSerializer.WriteObject(Stream, Object);
-
-				string Text = Encoding.UTF8.GetString(Stream.ToArray());
-				FileReference.WriteAllText(Location, Json.Format(Text));
+				FileReference.WriteAllBytes(Location, Stream.ToArray());
 			}
 		}
 
