UGS: Renaming some source files.

#rnx

diff --git a/Engine/Source/Programs/UnrealGameSync/MetadataTool/BuildHealth/BuildHealthIssue.cs b/Engine/Source/Programs/UnrealGameSync/MetadataTool/BuildHealth/BuildHealthIssue.cs
new file mode 100644
--- /dev/null
+++ b/Engine/Source/Programs/UnrealGameSync/MetadataTool/BuildHealth/BuildHealthIssue.cs
@@ -0,0 +1,159 @@
+// Copyright 1998-2019 Epic Games, Inc. All Rights Reserved.
+
+using System;
+using System.Collections.Generic;
+using System.Diagnostics;
+using System.Linq;
+using System.Runtime.Serialization;
+using System.Text;
+using System.Threading.Tasks;
+
+namespace MetadataTool
+{
+	/// <summary>
+	/// Information about a particular issue
+	/// </summary>
+	[DataContract]
+	class BuildHealthIssue
+	{
+		/// <summary>
+		/// The issue id in the database. -1 for issues that have not been posted yet.
+		/// </summary>
+		[DataMember(Order = 0, IsRequired = true)]
+		public long Id = -1;
+
+		/// <summary>
+		/// Type common to all diagnostics within this issue.
+		/// </summary>
+		[DataMember(Order = 1, IsRequired = true)]
+		public string Category;
+
+		/// <summary>
+		/// The initial job that this error was seen on. Allows other issues from the same job to be merged.
+		/// </summary>
+		[DataMember(Order = 2, IsRequired = true)]
+		public string InitialJobUrl;
+
+		/// <summary>
+		/// Url for this issue
+		/// </summary>
+		[DataMember(Order = 3, IsRequired = true)]
+		public string ErrorUrl;
+
+		/// <summary>
+		/// List of strings to display in the details panel for this job
+		/// </summary>
+		[DataMember(Order = 4, IsRequired = true)]
+		public List<string> Details = new List<string>();
+
+		/// <summary>
+		/// List of files associated with this issue
+		/// </summary>
+		[DataMember(Order = 5)]
+		public HashSet<string> FileNames = new HashSet<string>();
+
+		/// <summary>
+		/// List of messages associated with this issue
+		/// </summary>
+		[DataMember(Order = 6)]
+		public SortedSet<string> Identifiers = new SortedSet<string>();
+
+		/// <summary>
+		/// The last posted issue summary. Will be updated if it changes.
+		/// </summary>
+		[DataMember(Order = 20)]
+		public string PostedSummary;
+
+		/// <summary>
+		/// The last posted issue details. Will be updated if it changes.
+		/// </summary>
+		[DataMember(Order = 21)]
+		public string PostedDetails;
+
+		/// <summary>
+		/// Whether we've posted an update to the resolved flag to the server
+		/// </summary>
+		[DataMember(Order = 22)]
+		public bool bPostedResolved;
+
+		/// <summary>
+		/// The time at which the issue was closed. Issues will be retained for 24 hours after they are closed, in case the same issue appears in another stream and to prevent the issue being added again. 
+		/// </summary>
+		[DataMember(Order = 23)]
+		public DateTime? ResolvedAt;
+
+		/// <summary>
+		/// Map of stream name -> step name -> history for builds exhibiting this issue
+		/// </summary>
+		[DataMember(Order = 30, IsRequired = true)]
+		public Dictionary<string, Dictionary<string, BuildHealthJobHistory>> Streams = new Dictionary<string, Dictionary<string, BuildHealthJobHistory>>();
+
+		/// <summary>
+		/// Set of changes which may have caused this issue. Used to de-duplicate issues between streams.
+		/// </summary>
+		[DataMember(Order = 31)]
+		public HashSet<int> SourceChanges = new HashSet<int>();
+
+		/// <summary>
+		/// List of possible causers 
+		/// </summary>
+		[DataMember(Order = 32)]
+		public HashSet<string> Watchers = new HashSet<string>();
+
+		/// <summary>
+		/// Set of causers that have yet to be added to the possible causers list
+		/// </summary>
+		[DataMember(Order = 33)]
+		public HashSet<string> PendingWatchers = new HashSet<string>();
+
+		/// <summary>
+		/// Constructor
+		/// </summary>
+		/// <param name="Category">Category of this issue</param>
+		/// <param name="InitialJobUrl">Url of the initial job that this error was seen with</param>
+		/// <param name="ErrorUrl"></param>
+		public BuildHealthIssue(string Category, string InitialJobUrl, string ErrorUrl)
+		{
+			this.Category = Category;
+			this.InitialJobUrl = InitialJobUrl;
+			this.ErrorUrl = ErrorUrl;
+		}
+
+		/// <summary>
+		/// Determines whether the issue can be closed
+		/// </summary>
+		/// <returns>True if the issue can be closed</returns>
+		public bool IsResolved()
+		{
+			return Streams.Values.All(x => x.Values.All(y => y.NextSuccessfulBuild != null));
+		}
+
+		/// <summary>
+		/// Finds all the steps which are related to this issue
+		/// </summary>
+		/// <returns>Set of step names</returns>
+		public SortedSet<string> GetStepNames()
+		{
+			return new SortedSet<string>(Streams.Values.SelectMany(x => x.Keys));
+		}
+
+		/// <summary>
+		/// Format the issue for the debugger
+		/// </summary>
+		/// <returns>String representation of the issue</returns>
+		public override string ToString()
+		{
+			StringBuilder Result = new StringBuilder();
+			if(Id == -1)
+			{
+				Result.Append("[New] ");
+			}
+			else
+			{
+				Result.AppendFormat("[{0}] ", Id);
+			}
+			Result.AppendFormat("Errors in {0}", String.Join(", ", GetStepNames()));
+			return Result.ToString();
+		}
+	}
+}
diff --git a/Engine/Source/Programs/UnrealGameSync/MetadataTool/BuildHealth/BuildHealthJobStep.cs b/Engine/Source/Programs/UnrealGameSync/MetadataTool/BuildHealth/BuildHealthJobStep.cs
new file mode 100644
--- /dev/null
+++ b/Engine/Source/Programs/UnrealGameSync/MetadataTool/BuildHealth/BuildHealthJobStep.cs
@@ -0,0 +1,94 @@
+// Copyright 1998-2019 Epic Games, Inc. All Rights Reserved.
+
+using System;
+using System.Collections.Generic;
+using System.Diagnostics;
+using System.Linq;
+using System.Runtime.Serialization;
+using System.Text;
+using System.Threading.Tasks;
+
+namespace MetadataTool
+{
+	/// <summary>
+	/// Represents information about a build associated with a particular issue
+	/// </summary>
+	[DataContract]
+	[DebuggerDisplay("{Change}: {JobName}")]
+	class BuildHealthJobStep : IComparable<BuildHealthJobStep>
+	{
+		/// <summary>
+		/// The changelist that this build was run at.
+		/// </summary>
+		[DataMember(IsRequired = true)]
+		public int Change;
+
+		/// <summary>
+		/// Name of this job
+		/// </summary>
+		[DataMember(IsRequired = true)]
+		public string JobName;
+
+		/// <summary>
+		/// Url for this job
+		/// </summary>
+		[DataMember(IsRequired = true)]
+		public string JobUrl;
+
+		/// <summary>
+		/// Name of this job step (typically null for sentinel builds)
+		/// </summary>
+		[DataMember(IsRequired = true)]
+		public string JobStepName;
+
+		/// <summary>
+		/// Url for this job step (typically null for sentinel builds)
+		/// </summary>
+		[DataMember(IsRequired = true)]
+		public string JobStepUrl;
+
+		/// <summary>
+		/// Url of the first error within this job
+		/// </summary>
+		[DataMember]
+		public string ErrorUrl;
+
+		/// <summary>
+		/// Whether this build has been posted to the server or not
+		/// </summary>
+		[DataMember]
+		public bool bPostedToServer = false;
+
+		/// <summary>
+		/// Constructor
+		/// </summary>
+		public BuildHealthJobStep(int Change, string JobName, string JobUrl, string JobStepName, string JobStepUrl, string ErrorUrl)
+		{
+			this.Change = Change;
+			this.JobName = JobName;
+			this.JobUrl = JobUrl;
+			this.JobStepName = JobStepName;
+			this.JobStepUrl = JobStepUrl;
+			this.ErrorUrl = ErrorUrl;
+		}
+
+		/// <summary>
+		/// Compares this build to another build
+		/// </summary>
+		/// <param name="Other">Build to compare to</param>
+		/// <returns>Value indicating how the two builds should be ordered</returns>
+		public int CompareTo(BuildHealthJobStep Other)
+		{
+			int Delta = Change - Other.Change;
+			if (Delta == 0)
+			{
+				Delta = JobUrl.CompareTo(Other.JobUrl);
+				if(Delta == 0)
+				{
+					Delta = JobStepUrl.CompareTo(Other.JobStepUrl);
+				}
+			}
+			return Delta;
+		}
+	}
+}
diff --git a/Engine/Source/Programs/UnrealGameSync/MetadataTool/BuildHealth/BuildHealthJobStepHistory.cs b/Engine/Source/Programs/UnrealGameSync/MetadataTool/BuildHealth/BuildHealthJobStepHistory.cs
new file mode 100644
--- /dev/null
+++ b/Engine/Source/Programs/UnrealGameSync/MetadataTool/BuildHealth/BuildHealthJobStepHistory.cs
@@ -0,0 +1,101 @@
+// Copyright 1998-2019 Epic Games, Inc. All Rights Reserved.
+
+using System;
+using System.Collections.Generic;
+using System.Linq;
+using System.Runtime.Serialization;
+using System.Text;
+using System.Threading.Tasks;
+
+namespace MetadataTool
+{
+	/// <summary>
+	/// History of builds within a particular stream that contribute to an issue
+	/// </summary>
+	[DataContract]
+	class BuildHealthJobHistory
+	{
+		/// <summary>
+		/// The previous build before it started failing. This should be updated as new builds come in.
+		/// </summary>
+		[DataMember]
+		public BuildHealthJobStep PrevSuccessfulBuild;
+
+		/// <summary>
+		/// List of failing builds contributing to this issue
+		/// </summary>
+		[DataMember]
+		public List<BuildHealthJobStep> FailedBuilds = new List<BuildHealthJobStep>();
+
+		/// <summary>
+		/// The first successful build after the failures.
+		/// </summary>
+		[DataMember]
+		public BuildHealthJobStep NextSuccessfulBuild;
+
+		/// <summary>
+		/// Constructs a new history for a particular stream
+		/// </summary>
+		public BuildHealthJobHistory(BuildHealthJobStep PrevSuccessfulBuild)
+		{
+			this.PrevSuccessfulBuild  = PrevSuccessfulBuild;
+		}
+
+		/// <summary>
+		/// Adds a failed build to this object
+		/// </summary>
+		/// <param name="Build">The failed build</param>
+		public void AddFailedBuild(BuildHealthJobStep Build)
+		{
+			int Index = FailedBuilds.BinarySearch(Build);
+			if (Index < 0)
+			{
+				FailedBuilds.Insert(~Index, Build);
+			}
+		}
+
+		/// <summary>
+		/// Determines whether a given build can be added to an issue. This filters cases where an issue does not already have builds for the given stream, or where there is a successful build between the new build and known failures for this issue.
+		/// </summary>
+		/// <param name="Build">The build to add</param>
+		public bool CanAddFailedBuild(int Change)
+		{
+			// Check that this build is not after a succesful build
+			if(NextSuccessfulBuild != null && Change >= NextSuccessfulBuild.Change)
+			{
+				return false;
+			}
+
+			// Check that this build is not before the last known successful build
+			if(PrevSuccessfulBuild != null && Change <= PrevSuccessfulBuild.Change)
+			{
+				return false;
+			}
+
+			// Otherwise allow it
+			return true;
+		}
+
+		/// <summary>
+		/// Enumerates all the builds in this stream
+		/// </summary>
+		public IEnumerable<BuildHealthJobStep> Builds
+		{
+			get
+			{
+				if(PrevSuccessfulBuild != null)
+				{
+					yield return PrevSuccessfulBuild;
+				}
+				foreach(BuildHealthJobStep FailedBuild in FailedBuilds)
+				{
+					yield return FailedBuild;
+				}
+				if(NextSuccessfulBuild != null)
+				{
+					yield return NextSuccessfulBuild;
+				}
+			}
+		}
+	}
+}
diff --git a/Engine/Source/Programs/UnrealGameSync/MetadataTool/BuildHealth/BuildHealthState.cs b/Engine/Source/Programs/UnrealGameSync/MetadataTool/BuildHealth/BuildHealthState.cs
new file mode 100644
--- /dev/null
+++ b/Engine/Source/Programs/UnrealGameSync/MetadataTool/BuildHealth/BuildHealthState.cs
@@ -0,0 +1,106 @@
+// Copyright 1998-2019 Epic Games, Inc. All Rights Reserved.
+
+using System;
+using System.Collections.Generic;
+using System.Linq;
+using System.Runtime.Serialization;
+using System.Text;
+using System.Threading.Tasks;
+
+namespace MetadataTool
+{
+	/// <summary>
+	/// The persistent state tracked by the program between runs.
+	/// </summary>
+	[DataContract]
+	class BuildHealthState
+	{
+		/// <summary>
+		/// All open issues.
+		/// </summary>
+		[DataMember]
+		public List<BuildHealthIssue> Issues = new List<BuildHealthIssue>();
+
+		/// <summary>
+		/// Map of stream name to list of builds within it. Only retains a few builds in order to determine the last succesful CL for a build.
+		/// </summary>
+		[DataMember]
+		public Dictionary<string, List<BuildHealthJobStep>> Streams = new Dictionary<string, List<BuildHealthJobStep>>();
+
+		/// <summary>
+		/// Adds a new build for tracking
+		/// </summary>
+		/// <param name="Stream">The stream containing the build</param>
+		/// <param name="Build">The build to add</param>
+		public void AddBuild(string Stream, BuildHealthJobStep Build)
+		{
+			// Get the list of tracked builds for this stream
+			List<BuildHealthJobStep> Builds;
+			if (!Streams.TryGetValue(Stream, out Builds))
+			{
+				Builds = new List<BuildHealthJobStep>();
+				Streams.Add(Stream, Builds);
+			}
+
+			// Add this build to the tracked state data
+			int BuildIdx = Builds.BinarySearch(Build);
+			if (BuildIdx < 0)
+			{
+				BuildIdx = ~BuildIdx;
+				Builds.Insert(BuildIdx, Build);
+			}
+		}
+
+		/// <summary>
+		/// Finds the last build before the one given which executes the same steps. Used to determine the last succesful build before a failure.
+		/// </summary>
+		/// <param name="Stream">The stream to search for</param>
+		/// <param name="Change">The change to search for</param>
+		/// <param name="StepNames">The step names which the job step needs to have</param>
+		/// <returns>The last build known before the given changelist</returns>
+		public BuildHealthJobStep FindBuildBefore(string Stream, int Change, string StepName)
+		{
+			BuildHealthJobStep Result = null;
+
+			List<BuildHealthJobStep> Builds;
+			if(Streams.TryGetValue(Stream, out Builds))
+			{
+				for(int Idx = 0; Idx < Builds.Count && Builds[Idx].Change < Change; Idx++)
+				{
+					if(Builds[Idx].JobStepName == StepName)
+					{
+						Result = Builds[Idx];
+					}
+				}
+			}
+
+			return Result;
+		}
+
+		/// <summary>
+		/// Finds the first build after the one given which executes the same steps. Used to determine the last succesful build before a failure.
+		/// </summary>
+		/// <param name="Stream">The stream to search for</param>
+		/// <param name="Change">The change to search for</param>
+		/// <param name="StepName">Name of the job step to find</param>
+		/// <returns>The last build known before the given changelist</returns>
+		public BuildHealthJobStep FindBuildAfter(string Stream, int Change, string StepName)
+		{
+			BuildHealthJobStep Result = null;
+
+			List<BuildHealthJobStep> Builds;
+			if (Streams.TryGetValue(Stream, out Builds))
+			{
+				for (int Idx = Builds.Count - 1; Idx >= 0 && Builds[Idx].Change > Change; Idx--)
+				{
+					if (Builds[Idx].JobStepName == StepName)
+					{
+						Result = Builds[Idx];
+					}
+				}
+			}
+
+			return Result;
+		}
+	}
+}
diff --git a/Engine/Source/Programs/UnrealGameSync/MetadataTool/BuildHealth/TrackedBuild.cs b/Engine/Source/Programs/UnrealGameSync/MetadataTool/BuildHealth/TrackedBuild.cs
deleted file mode 100644
diff --git a/Engine/Source/Programs/UnrealGameSync/MetadataTool/BuildHealth/TrackedIssue.cs b/Engine/Source/Programs/UnrealGameSync/MetadataTool/BuildHealth/TrackedIssue.cs
deleted file mode 100644
diff --git a/Engine/Source/Programs/UnrealGameSync/MetadataTool/BuildHealth/TrackedIssueHistory.cs b/Engine/Source/Programs/UnrealGameSync/MetadataTool/BuildHealth/TrackedIssueHistory.cs
deleted file mode 100644
diff --git a/Engine/Source/Programs/UnrealGameSync/MetadataTool/BuildHealth/TrackedState.cs b/Engine/Source/Programs/UnrealGameSync/MetadataTool/BuildHealth/TrackedState.cs
deleted file mode 100644
diff --git a/Engine/Source/Programs/UnrealGameSync/MetadataTool/MetadataTool.csproj b/Engine/Source/Programs/UnrealGameSync/MetadataTool/MetadataTool.csproj
--- a/Engine/Source/Programs/UnrealGameSync/MetadataTool/MetadataTool.csproj
+++ b/Engine/Source/Programs/UnrealGameSync/MetadataTool/MetadataTool.csproj
@@ -62,10 +62,10 @@
     <Compile Include="BuildHealth\InputData.cs" />
     <Compile Include="Program.cs" />
     <Compile Include="Properties\AssemblyInfo.cs" />
-    <Compile Include="BuildHealth\TrackedBuild.cs" />
-    <Compile Include="BuildHealth\TrackedIssueHistory.cs" />
-    <Compile Include="BuildHealth\TrackedIssue.cs" />
-    <Compile Include="BuildHealth\TrackedState.cs" />
+    <Compile Include="BuildHealth\BuildHealthIssue.cs" />
+    <Compile Include="BuildHealth\BuildHealthJobStep.cs" />
+    <Compile Include="BuildHealth\BuildHealthJobStepHistory.cs" />
+    <Compile Include="BuildHealth\BuildHealthState.cs" />
   </ItemGroup>
   <ItemGroup>
     <None Include="App.config" />
