UGS: Fixed a bug where closed issues wouldn't be saved out to disk, and added some extra timing info into the metadata tool.

#rnx

diff --git a/Engine/Source/Programs/UnrealGameSync/MetadataTool/CommandHandler_BuildHealth.cs b/Engine/Source/Programs/UnrealGameSync/MetadataTool/CommandHandler_BuildHealth.cs
--- a/Engine/Source/Programs/UnrealGameSync/MetadataTool/CommandHandler_BuildHealth.cs
+++ b/Engine/Source/Programs/UnrealGameSync/MetadataTool/CommandHandler_BuildHealth.cs
@@ -2,6 +2,7 @@
 
 using System;
 using System.Collections.Generic;
+using System.Diagnostics;
 using System.IO;
 using System.Linq;
 using System.Net;
@@ -93,7 +94,8 @@ namespace MetadataTool
 
 			// Parse all the builds and add them to the persistent data
 			List<InputJob> InputJobs = InputData.Jobs.OrderBy(x => x.Change).ThenBy(x => x.Stream).ToList();
-			foreach(InputJob InputJob in InputJobs)
+			Stopwatch Timer = Stopwatch.StartNew();
+			foreach (InputJob InputJob in InputJobs)
 			{
 				// Add a new build for each job step
 				foreach(InputJobStep InputJobStep in InputJob.Steps)
@@ -109,6 +111,7 @@ namespace MetadataTool
 					AddStep(Perforce, State, InputJob, InputJobStep);
 				}
 			}
+			Log.TraceInformation("Added jobs in {0}s", Timer.Elapsed.TotalSeconds);
 
 			// Try to find the next successful build for each stream, so we can close it as part of updating the server
 			for (int Idx = 0; Idx < State.Issues.Count; Idx++)
@@ -372,6 +375,7 @@ namespace MetadataTool
 				if(Issue.ResolvedAt.HasValue && Issue.ResolvedAt.Value < RemoveIssueTime)
 				{
 					State.Issues.RemoveAt(Idx--);
+					WriteState(StateFile, State);
 					continue;
 				}
 			}
@@ -698,6 +702,8 @@ namespace MetadataTool
 		/// <param name="State">The state object</param>
 		static void WriteState(FileReference StateFile, BuildHealthState State)
 		{
+			Stopwatch Timer = Stopwatch.StartNew();
+
 			// Write out the state to the transaction file
 			FileReference StateTransactionFile = GetStateTransactionFile(StateFile);
 			SerializeJson(StateTransactionFile, State);
@@ -705,6 +711,8 @@ namespace MetadataTool
 			// Remove the original file, then move the transaction file into place
 			FileReference.Delete(StateFile);
 			FileReference.Move(StateTransactionFile, StateFile);
+
+			Log.TraceInformation("Took {0}s to write state", Timer.Elapsed.TotalSeconds);
 		}
 
 		/// <summary>
