UGS: Prettify multi-line errors.

#rnx

diff --git a/Engine/Source/Programs/UnrealGameSync/UnrealGameSync/Forms/IssueDetailsWindow.cs b/Engine/Source/Programs/UnrealGameSync/UnrealGameSync/Forms/IssueDetailsWindow.cs
--- a/Engine/Source/Programs/UnrealGameSync/UnrealGameSync/Forms/IssueDetailsWindow.cs
+++ b/Engine/Source/Programs/UnrealGameSync/UnrealGameSync/Forms/IssueDetailsWindow.cs
@@ -396,58 +396,78 @@ namespace UnrealGameSync
 
 		string CreateRichTextErrors(string ErrorText)
 		{
-			List<List<string>> Errors = new List<List<string>>();
-			List<string> CurrentError = null;
-			foreach (string Line in ErrorText.Split('\n'))
+			string[] Lines = ErrorText.Split('\n');
+			
+			Regex StepPattern = new Regex("^##([^#].*[^#])##$", RegexOptions.None);
+
+			Dictionary<string, int> StepNameToCount = new Dictionary<string, int>(StringComparer.Ordinal);
+			foreach(string Line in Lines)
 			{
-				if (Line == "***")
+				Match Match = StepPattern.Match(Line);
+				if(Match.Success)
 				{
-					CurrentError = null;
-				}
-				else
-				{
-					if (CurrentError == null)
-					{
-						CurrentError = new List<string>();
-						Errors.Add(CurrentError);
-					}
-					CurrentError.Add(Line);
+					string StepName = Match.Groups[1].Value;
+					int Count;
+					StepNameToCount.TryGetValue(StepName, out Count);
+					StepNameToCount[StepName] = Count + 1;
 				}
 			}
 
+			Dictionary<string, int> StepNameToIndex = StepNameToCount.Keys.ToDictionary(x => x, x => 0, StringComparer.Ordinal);
+
 			StringBuilder RichText = new StringBuilder();
 
 			RichText.AppendLine(@"{\rtf1\ansi");
 			RichText.AppendLine(@"{\fonttbl{\f0\fnil\fcharset0 Arial;}{\f1\fnil\fcharset0 Courier New;}{\f2\fnil\fcharset0 Calibri;}}");
 			RichText.AppendLine(@"{\colortbl;\red192\green80\blue77;}");
 
-			for(int ErrorIdx = 0; ErrorIdx < Errors.Count; ErrorIdx++)
+			string PrevStepName = null;
+			for(int Idx = 0; Idx < Lines.Length; Idx++)
 			{
-				// Error X/Y:
-				RichText.Append(@"\pard");   // Paragraph default
-				RichText.Append(@"\cf1");    // Foreground color 1
-				RichText.Append(@"\b1");     // Bold on
-				RichText.Append(@"\f0");     // Font 0
-				RichText.Append(@"\fs16");   // Font size 16 (8pt * 2)
-				if(ErrorIdx > 0)
+				Match Match = StepPattern.Match(Lines[Idx]);
+				if(Match.Success)
 				{
-					RichText.Append(@"\sb300"); // Space before 200
-				}
-				RichText.Append(@"\sa100");  // Space after 100
-				RichText.AppendFormat(@"Error {0}/{1}:", ErrorIdx + 1, Errors.Count);
-				RichText.AppendLine(@"\par");
+					// Step 'Foo'
+					string StepName = Match.Groups[1].Value;
+					if(StepName != PrevStepName)
+					{
+						RichText.Append(@"\pard");   // Paragraph default
+						RichText.Append(@"\cf0");    // Foreground color
+						RichText.Append(@"\b1");     // Bold
+						RichText.Append(@"\f0");     // Font
+						RichText.Append(@"\fs16");   // Font size
+						if (Idx > 0)
+						{
+							RichText.Append(@"\sb200"); // Space before
+						}
+						RichText.AppendFormat(@"In step '{0}':", StepName);
+						RichText.AppendLine(@"\par");
+
+						PrevStepName = StepName;
+					}
 
-				// Error text
-				foreach(string ErrorLine in Errors[ErrorIdx])
+					// Error X/Y:
+					RichText.Append(@"\pard");   // Paragraph default
+					RichText.Append(@"\cf1");    // Foreground color
+					RichText.Append(@"\b1");     // Bold
+					RichText.Append(@"\f0");     // Font
+					RichText.Append(@"\fs16");   // Font size
+					RichText.Append(@"\sb100");  // Space before
+					RichText.Append(@"\sa50");   // Space after
+					RichText.AppendFormat(@"Error {0}/{1}", ++StepNameToIndex[StepName], StepNameToCount[StepName]);
+					RichText.AppendLine(@"\par");
+				}
+				else
 				{
+					// Error text
 					RichText.Append(@"\pard");   // Paragraph default
-					RichText.Append(@"\cf0");    // Foreground color 0
-					RichText.Append(@"\b0");     // Bold off
-					RichText.Append(@"\f1");     // Font 1
-					RichText.Append(@"\fs16");   // Font size 16 (8pt * 2)
-					RichText.Append(@"\fi150");  // First line indent 150
-					RichText.Append(@"\li150");  // Other line indent 150
-					AppendEscapedRtf(RichText, ErrorLine);
+					RichText.Append(@"\cf0");    // Foreground color
+					RichText.Append(@"\b0");     // Bold
+					RichText.Append(@"\f1");     // Font
+					RichText.Append(@"\fs16");   // Font size 16
+					RichText.Append(@"\fi150");  // First line indent
+					RichText.Append(@"\li150");  // Other line indent
+					AppendEscapedRtf(RichText, Lines[Idx]);
 					RichText.Append(@"\par");
 				}
 			}
