UGS: Show issues with a fix changelist as green in the browser window.

#rnx

diff --git a/Engine/Source/Programs/UnrealGameSync/UnrealGameSync/Forms/IssueBrowserWindow.cs b/Engine/Source/Programs/UnrealGameSync/UnrealGameSync/Forms/IssueBrowserWindow.cs
--- a/Engine/Source/Programs/UnrealGameSync/UnrealGameSync/Forms/IssueBrowserWindow.cs
+++ b/Engine/Source/Programs/UnrealGameSync/UnrealGameSync/Forms/IssueBrowserWindow.cs
@@ -250,14 +250,10 @@ namespace UnrealGameSync
 			IssueData Issue = (IssueData)e.Item.Tag;
 			if (e.ColumnIndex == IconHeader.Index)
 			{
-				if(!Issue.ResolvedAt.HasValue)
+				if(!Issue.ResolvedAt.HasValue && Issue.FixChange == 0)
 				{
 					IssueListView.DrawIcon(e.Graphics, e.Bounds, WorkspaceControl.BadBuildIcon);
 				}
-				else if (Issue.FixChange > 0)
-				{
-					IssueListView.DrawIcon(e.Graphics, e.Bounds, WorkspaceControl.MixedBuildIcon);
-				}
 				else
 				{
 					IssueListView.DrawIcon(e.Graphics, e.Bounds, WorkspaceControl.GoodBuildIcon);
