Increase the amount of data that can be stored in an issue's details field.

#rnx

diff --git a/Engine/Source/Programs/UnrealGameSync/MetadataServer/Connectors/SqlConnector.cs b/Engine/Source/Programs/UnrealGameSync/MetadataServer/Connectors/SqlConnector.cs
--- a/Engine/Source/Programs/UnrealGameSync/MetadataServer/Connectors/SqlConnector.cs
+++ b/Engine/Source/Programs/UnrealGameSync/MetadataServer/Connectors/SqlConnector.cs
@@ -352,7 +352,7 @@ namespace MetadataServer.Connectors
 		}
 
 		const int IssueSummaryMaxLength = 200;
-		const int IssueDetailsMaxLength = 1000;
+		const int IssueDetailsMaxLength = 8000;
 
 		public static long AddIssue(IssueData Issue)
 		{
@@ -546,7 +546,15 @@ namespace MetadataServer.Connectors
 		{
 			if(Text.Length > Length)
 			{
-				Text = Text.Substring(0, Length - 3).TrimEnd() + "...";
+				int NewlineIdx = Text.LastIndexOf('\n', Length);
+				if(NewlineIdx == -1)
+				{
+					Text = Text.Substring(0, Length - 3).TrimEnd() + "...";
+				}
+				else
+				{
+					Text = Text.Substring(0, NewlineIdx + 1) + "...";
+				}
 			}
 			return Text;
 		}
