UGSMetadataSever:
- Add Projects table
- Set up foreign keys to projects table in all other tables
- New badges table replacing old CIS one
- set up queries to use the projectID index rather than a string search

diff --git a/Engine/Source/Programs/UnrealGameSync/MetadataServer/Connectors/SqlConnector.cs b/Engine/Source/Programs/UnrealGameSync/MetadataServer/Connectors/SqlConnector.cs
--- a/Engine/Source/Programs/UnrealGameSync/MetadataServer/Connectors/SqlConnector.cs
+++ b/Engine/Source/Programs/UnrealGameSync/MetadataServer/Connectors/SqlConnector.cs
@@ -28,7 +28,9 @@ namespace MetadataServer.Connectors
 			using (SQLiteConnection Connection = new SQLiteConnection(ConnectionString))
 			{
 				Connection.Open();
-				using (SQLiteCommand Command = new SQLiteCommand("SELECT Id FROM [UserVotes] WHERE Project LIKE @param1 GROUP BY Changelist ORDER BY Changelist DESC LIMIT 1 OFFSET 432;", Connection))
+				using (SQLiteCommand Command = new SQLiteCommand("SELECT UserVotes.Id FROM [UserVotes] " +
+																 "INNER JOIN [Projects] ON Projects.Id = UserVotes.ProjectId " +
+																 "WHERE Projects.Name LIKE @param1 GROUP BY Changelist ORDER BY Changelist DESC LIMIT 1 OFFSET 432", Connection))
 				{
 					Command.Parameters.AddWithValue("@param1", ProjectLikeString);
 					using (SQLiteDataReader Reader = Command.ExecuteReader())
@@ -40,7 +42,10 @@ namespace MetadataServer.Connectors
 						}
 					}
 				}
-				using (SQLiteCommand Command = new SQLiteCommand("SELECT Id FROM [Comments] WHERE Project LIKE @param1 GROUP BY ChangeNumber ORDER BY ChangeNumber DESC LIMIT 1 OFFSET 432;", Connection))
+				
+				using (SQLiteCommand Command = new SQLiteCommand("SELECT Comments.Id FROM [Comments] " +
+																 "INNER JOIN [Projects] ON Projects.Id = Comments.ProjectId " +
+																 "WHERE Projects.Name LIKE @param1 GROUP BY ChangeNumber ORDER BY ChangeNumber DESC LIMIT 1 OFFSET 432", Connection))
 				{
 					Command.Parameters.AddWithValue("@param1", ProjectLikeString);
 					using (SQLiteDataReader Reader = Command.ExecuteReader())
@@ -52,14 +57,17 @@ namespace MetadataServer.Connectors
 						}
 					}
 				}
-				using (SQLiteCommand Command = new SQLiteCommand("SELECT Id FROM [CIS] WHERE Project LIKE @param1 GROUP BY ChangeNumber ORDER BY ChangeNumber DESC LIMIT 1 OFFSET 432", Connection))
+				using (SQLiteCommand Command = new SQLiteCommand("SELECT Badges.Id FROM [Badges] " +
+																 "INNER JOIN [Projects] ON Projects.Id = Badges.ProjectId " +
+																 "WHERE Projects.Name LIKE @param1 GROUP BY ChangeNumber ORDER BY ChangeNumber DESC LIMIT 1 OFFSET 432", Connection))
 				{
+					//Command.Parameters.AddWithValue("@param1", ProjectId);
 					Command.Parameters.AddWithValue("@param1", ProjectLikeString);
 					using (SQLiteDataReader Reader = Command.ExecuteReader())
 					{
 						while (Reader.Read())
 						{
-							LastBuildId = Reader.GetInt32(0);
+							LastBuildId = Math.Max(LastBuildId, Reader.GetInt32(0));
 							break;
 						}
 					}
@@ -74,7 +82,8 @@ namespace MetadataServer.Connectors
 			using (SQLiteConnection Connection = new SQLiteConnection(ConnectionString))
 			{
 				Connection.Open();
-				using (SQLiteCommand Command = new SQLiteCommand("SELECT Id, Changelist, UserName, Verdict, Project FROM [UserVotes] WHERE Id > @param1 AND Project LIKE @param2 ORDER BY Id", Connection))
+				using (SQLiteCommand Command = new SQLiteCommand("SELECT UserVotes.Id, UserVotes.Changelist, UserVotes.UserName, UserVotes.Verdict, UserVotes.Project FROM [UserVotes] " +
+																 "INNER JOIN [Projects] ON Projects.Id = UserVotes.ProjectId WHERE UserVotes.Id > @param1 AND Projects.Name LIKE @param2 ORDER BY UserVotes.Id", Connection))
 				{
 					Command.Parameters.AddWithValue("@param1", LastEventId);
 					Command.Parameters.AddWithValue("@param2", ProjectLikeString);
@@ -107,7 +116,8 @@ namespace MetadataServer.Connectors
 			using (SQLiteConnection Connection = new SQLiteConnection(ConnectionString))
 			{
 				Connection.Open();
-				using (SQLiteCommand Command = new SQLiteCommand("SELECT Id, ChangeNumber, UserName, Text, Project FROM [Comments] WHERE Id > @param1 AND Project LIKE @param2 ORDER BY Id", Connection))
+				using (SQLiteCommand Command = new SQLiteCommand("SELECT Comments.Id, Comments.ChangeNumber, Comments.UserName, Comments.Text, Comments.Project FROM [Comments] " +
+															     "INNER JOIN [Projects] ON Projects.Id = Comments.ProjectId WHERE Comments.Id > @param1 AND Projects.Name LIKE @param2 ORDER BY Comments.Id", Connection))
 				{
 					Command.Parameters.AddWithValue("@param1", LastCommentId);
 					Command.Parameters.AddWithValue("@param2", ProjectLikeString);
@@ -139,8 +149,9 @@ namespace MetadataServer.Connectors
 			using (SQLiteConnection Connection = new SQLiteConnection(ConnectionString))
 			{
 				Connection.Open();
-				using (SQLiteCommand Command = new SQLiteCommand("SELECT Id, ChangeNumber, BuildType, Result, Url, Project FROM [CIS] WHERE Id > @param1 AND Project LIKE @param2 ORDER BY Id", Connection))
-				{
+				using (SQLiteCommand Command = new SQLiteCommand("SELECT Badges.Id, Badges.ChangeNumber, Badges.BuildType, Badges.Result, Badges.Url, Projects.Name, Badges.ArchivePath FROM [Badges] " +
+																 "INNER JOIN [Projects] ON Projects.Id = Badges.ProjectId WHERE Badges.Id > @param1 AND Projects.Name LIKE @param2 ORDER BY Badges.Id", Connection))
+				{ 
 					Command.Parameters.AddWithValue("@param1", LastBuildId);
 					Command.Parameters.AddWithValue("@param2", ProjectLikeString);
 					using (SQLiteDataReader Reader = Command.ExecuteReader())
@@ -198,20 +209,29 @@ namespace MetadataServer.Connectors
 			return ReturnedErrors;
 		}
 
-
+		private static long TryInsertAndGetProject(SQLiteConnection Connection, string Project)
+		{
+			using (SQLiteCommand Command = new SQLiteCommand("INSERT OR IGNORE INTO [Projects] (Name) VALUES (@Project); SELECT [Id] FROM [Projects] WHERE Name = @Project", Connection))
+			{
+				Command.Parameters.AddWithValue("@Project", Project);
+				object ProjectId = Command.ExecuteScalar();
+				return (long)ProjectId;
+			}
+		}
 		public static void PostBuild(BuildData Build)
 		{
 			using (SQLiteConnection Connection = new SQLiteConnection(ConnectionString))
 			{
 				Connection.Open();
-				using (SQLiteCommand Command = new SQLiteCommand("INSERT INTO [CIS] (ChangeNumber, BuildType, Result, URL, Project, ArchivePath) VALUES (@ChangeNumber, @BuildType, @Result, @URL, @Project, @ArchivePath)", Connection))
+				long ProjectId = TryInsertAndGetProject(Connection, Build.Project);
+				using (SQLiteCommand Command = new SQLiteCommand("INSERT INTO [Badges] (ChangeNumber, BuildType, Result, URL, ArchivePath, ProjectId) VALUES (@ChangeNumber, @BuildType, @Result, @URL, @ArchivePath, @ProjectId)", Connection))
 				{
 					Command.Parameters.AddWithValue("@ChangeNumber", Build.ChangeNumber);
 					Command.Parameters.AddWithValue("@BuildType", Build.BuildType);
 					Command.Parameters.AddWithValue("@Result", Build.Result);
 					Command.Parameters.AddWithValue("@URL", Build.Url);
-					Command.Parameters.AddWithValue("@Project", Build.Project);
 					Command.Parameters.AddWithValue("@ArchivePath", Build.ArchivePath);
+					Command.Parameters.AddWithValue("@ProjectId", ProjectId);
 					Command.ExecuteNonQuery();
 				}
 			}
@@ -222,12 +242,14 @@ namespace MetadataServer.Connectors
 			using (SQLiteConnection Connection = new SQLiteConnection(ConnectionString))
 			{
 				Connection.Open();
-				using (SQLiteCommand Command = new SQLiteCommand("INSERT INTO [UserVotes] (Changelist, UserName, Verdict, Project) VALUES (@Changelist, @UserName, @Verdict, @Project)", Connection))
+				long ProjectId = TryInsertAndGetProject(Connection, Event.Project);
+				using (SQLiteCommand Command = new SQLiteCommand("INSERT INTO [UserVotes] (Changelist, UserName, Verdict, Project, ProjectId) VALUES (@Changelist, @UserName, @Verdict, @Project, @ProjectId)", Connection))
 				{
 					Command.Parameters.AddWithValue("@Changelist", Event.Change);
 					Command.Parameters.AddWithValue("@UserName", Event.UserName.ToString());
 					Command.Parameters.AddWithValue("@Verdict", Event.Type.ToString());
 					Command.Parameters.AddWithValue("@Project", Event.Project);
+					Command.Parameters.AddWithValue("@ProjectId", ProjectId);
 					Command.ExecuteNonQuery();
 				}
 			}
@@ -238,12 +260,14 @@ namespace MetadataServer.Connectors
 			using (SQLiteConnection Connection = new SQLiteConnection(ConnectionString))
 			{
 				Connection.Open();
-				using (SQLiteCommand Command = new SQLiteCommand("INSERT INTO [Comments] (ChangeNumber, UserName, Text, Project) VALUES (@ChangeNumber, @UserName, @Text, @Project)", Connection))
+				long ProjectId = TryInsertAndGetProject(Connection, Comment.Project);
+				using (SQLiteCommand Command = new SQLiteCommand("INSERT INTO [Comments] (ChangeNumber, UserName, Text, Project, ProjectId) VALUES (@ChangeNumber, @UserName, @Text, @Project, @ProjectId)", Connection))
 				{
 					Command.Parameters.AddWithValue("@ChangeNumber", Comment.ChangeNumber);
 					Command.Parameters.AddWithValue("@UserName", Comment.UserName);
 					Command.Parameters.AddWithValue("@Text", Comment.Text);
 					Command.Parameters.AddWithValue("@Project", Comment.Project);
+					Command.Parameters.AddWithValue("@ProjectId", ProjectId);
 					Command.ExecuteNonQuery();
 				}
 			}
@@ -254,7 +278,8 @@ namespace MetadataServer.Connectors
 			using (SQLiteConnection Connection = new SQLiteConnection(ConnectionString))
 			{
 				Connection.Open();
-				using (SQLiteCommand Command = new SQLiteCommand("INSERT INTO [Telemetry.v2] (Action, Result, UserName, Project, Timestamp, Duration, Version, IpAddress) VALUES (@Action, @Result, @UserName, @Project, @Timestamp, @Duration, @Version, @IpAddress)", Connection))
+				long ProjectId = TryInsertAndGetProject(Connection, Data.Project);
+				using (SQLiteCommand Command = new SQLiteCommand("INSERT INTO [Telemetry.v2] (Action, Result, UserName, Project, Timestamp, Duration, Version, IpAddress, ProjectId) VALUES (@Action, @Result, @UserName, @Project, @Timestamp, @Duration, @Version, @IpAddress, @ProjectId)", Connection))
 				{
 					Command.Parameters.AddWithValue("@Action", Data.Action);
 					Command.Parameters.AddWithValue("@Result", Data.Result);
@@ -264,6 +289,7 @@ namespace MetadataServer.Connectors
 					Command.Parameters.AddWithValue("@Duration", Data.Duration);
 					Command.Parameters.AddWithValue("@Version", Version);
 					Command.Parameters.AddWithValue("@IPAddress", IpAddress);
+					Command.Parameters.AddWithValue("@ProjectId", ProjectId);
 					Command.ExecuteNonQuery();
 				}
 			}
@@ -274,7 +300,12 @@ namespace MetadataServer.Connectors
 			using (SQLiteConnection Connection = new SQLiteConnection(ConnectionString))
 			{
 				Connection.Open();
-				using (SQLiteCommand Command = new SQLiteCommand("INSERT INTO [Errors] (Type, Text, UserName, Project, Timestamp, Version, IpAddress) VALUES (@Type, @Text, @UserName, @Project, @Timestamp, @Version, @IpAddress)", Connection))
+				long? ProjectId = null;
+				if(Data.Project != null)
+				{
+					ProjectId = TryInsertAndGetProject(Connection, Data.Project);
+				}
+				using (SQLiteCommand Command = new SQLiteCommand("INSERT INTO [Errors] (Type, Text, UserName, Project, Timestamp, Version, IpAddress, ProjectId) VALUES (@Type, @Text, @UserName, @Project, @Timestamp, @Version, @IpAddress, @ProjectId)", Connection))
 				{
 					Command.Parameters.AddWithValue("@Type", Data.Type.ToString());
 					Command.Parameters.AddWithValue("@Text", Data.Text);
@@ -282,10 +313,12 @@ namespace MetadataServer.Connectors
 					if (Data.Project == null)
 					{
 						Command.Parameters.AddWithValue("@Project", DBNull.Value);
+						Command.Parameters.AddWithValue("@ProjectId", DBNull.Value);
 					}
 					else
 					{
 						Command.Parameters.AddWithValue("@Project", Data.Project);
+						Command.Parameters.AddWithValue("@ProjectId", ProjectId.Value);
 					}
 					Command.Parameters.AddWithValue("@Timestamp", Data.Timestamp);
 					Command.Parameters.AddWithValue("@Version", Version);
diff --git a/Engine/Source/Programs/UnrealGameSync/MetadataServer/Global.asax.cs b/Engine/Source/Programs/UnrealGameSync/MetadataServer/Global.asax.cs
--- a/Engine/Source/Programs/UnrealGameSync/MetadataServer/Global.asax.cs
+++ b/Engine/Source/Programs/UnrealGameSync/MetadataServer/Global.asax.cs
@@ -1,7 +1,9 @@
 ﻿// Copyright 1998-2019 Epic Games, Inc. All Rights Reserved.
 using System;
+using System.Collections.Generic;
 using System.Data.SQLite;
 using System.IO;
+using System.Linq;
 using System.Text.RegularExpressions;
 using System.Web;
 using System.Web.Http;
@@ -10,34 +12,146 @@ namespace MetadataServer
 {
 	public class WebApiApplication : System.Web.HttpApplication
 	{
-		protected void Application_Start()
-		{
-			try
-			{
-				string ConnectionString = System.Configuration.ConfigurationManager.ConnectionStrings["ConnectionString"].ConnectionString;
-				string FileName = Regex.Match(ConnectionString, "Data Source=(.+);.+").Groups[1].Value;
-				Directory.CreateDirectory(new FileInfo(FileName).Directory.FullName);
-				if (!File.Exists(FileName))
-				{
-					SQLiteConnection.CreateFile(FileName);
-				}
-				// initialize the db if it doesn't exist
-				using (SQLiteConnection Connection = new SQLiteConnection(ConnectionString, true))
-				{
-					Connection.Open();
-					using (SQLiteCommand Command = new SQLiteCommand(Properties.Resources.Setup, Connection))
+        protected void Application_Start()
+        {
+            try
+            {
+                string ConnectionString = System.Configuration.ConfigurationManager.ConnectionStrings["ConnectionString"].ConnectionString;
+                string FileName = Regex.Match(ConnectionString, "Data Source=(.+);.+").Groups[1].Value;
+                Directory.CreateDirectory(new FileInfo(FileName).Directory.FullName);
+                if (!File.Exists(FileName))
+                {
+                    SQLiteConnection.CreateFile(FileName);
+                }
+                // initialize the db if it doesn't exist
+                using (SQLiteConnection Connection = new SQLiteConnection(ConnectionString, true))
+                {
+                    Connection.Open();
+                    using (SQLiteCommand Command = new SQLiteCommand(Properties.Resources.Setup, Connection))
+                    {
+                        Command.ExecuteNonQuery();
+                    }
+
+					////////////////////////////////////////////////////////
+					// DB change for indexing of project IDs by foreign key
+					////////////////////////////////////////////////////////
+					
+					// check if we've seeded the projects table already by looking at the raw number of rows
+					// if we're a brand new database this will techincally trigger with a no-op every time the server is started
+					// but will stop as soon as a badge is added.
+					bool ProjectsTableSeeded = false;
+					using (SQLiteCommand Command = new SQLiteCommand("SELECT COALESCE(MAX(id)+1, 0) FROM [Projects]", Connection))
+					{
+						using (SQLiteDataReader Reader = Command.ExecuteReader())
+						{
+							while (Reader.Read())
+							{
+								int ProjectRows = Reader.GetInt32(0);
+								ProjectsTableSeeded = ProjectRows == 0 ? false : true;
+							}
+						}
+					}
+					// if not found, get all the existing projects fromm all the other tables.
+					if (!ProjectsTableSeeded)
 					{
-						Command.ExecuteNonQuery();
+						bool AlterTableRun = false;
+						// check and see if we've run the alter already.
+						using (SQLiteCommand Command = new SQLiteCommand("PRAGMA table_info(CIS)", Connection))
+						{
+							using (SQLiteDataReader Reader = Command.ExecuteReader())
+							{
+								while (Reader.Read())
+								{
+									string ColumnName = Reader.GetString(1);
+									if(ColumnName == "ProjectID")
+									{
+										AlterTableRun = true;
+									}
+								}
+							}
+						}
+						if (!AlterTableRun)
+						{
+							using (SQLiteCommand Command = new SQLiteCommand("ALTER TABLE [CIS] ADD COLUMN ProjectID INTEGER REFERENCES Projects(Id);" +
+																			 "ALTER TABLE [Comments] ADD COLUMN ProjectID INTEGER REFERENCES Projects(Id);" +
+																			 "ALTER TABLE [Errors] ADD COLUMN ProjectID INTEGER REFERENCES Projects(Id);" +
+																			 "ALTER TABLE [Telemetry.v2] ADD COLUMN ProjectID INTEGER REFERENCES Projects(Id);" +
+																			 "ALTER TABLE [UserVotes] ADD COLUMN ProjectID INTEGER REFERENCES Projects(Id);", Connection))
+							{
+								Command.ExecuteNonQuery();
+							}
+
+							// get all the distinct project names.
+							Dictionary<string, long> Projects = new Dictionary<string, long>();
+							using (SQLiteCommand Command = new SQLiteCommand("SELECT Project FROM [CIS] UNION Select Project FROM [Comments] WHERE Project IS NOT NULL " +
+																			 "UNION " +
+																			 "SELECT Project FROM [Comments] WHERE Project IS NOT NULL " +
+																			 "UNION " +
+																			 "SELECT Project FROM [Errors] WHERE Project IS NOT NULL " +
+																			 "UNION " +
+																			 "SELECT Project FROM [Telemetry.v2] WHERE Project is not NULL " +
+																			 "UNION " +
+																			 "SELECT Project FROM [UserVotes] WHERE Project is not NULL", Connection))
+							{
+								using (SQLiteDataReader Reader = Command.ExecuteReader())
+								{
+									while (Reader.Read())
+									{
+										object ProjectName = Reader.GetValue(0);
+										if (ProjectName != null)
+										{
+											Projects.Add((string)ProjectName, 0);
+										}
+
+									}
+								}
+							}
+							// insert all the projects into the new table
+							using (SQLiteCommand Command = new SQLiteCommand(Connection))
+							{
+								using (SQLiteTransaction Transaction = Connection.BeginTransaction())
+								{
+									foreach (string ProjectName in Projects.Keys.ToList())
+									{
+										Command.CommandText = "INSERT INTO [Projects] (Name) VALUES (@ProjectName)";
+										Command.Parameters.AddWithValue("@ProjectName", ProjectName);
+										Command.ExecuteNonQuery();
+										Projects[ProjectName] = Connection.LastInsertRowId;
+									}
+									Transaction.Commit();
+								}
+							}
+							// update all the foreign keys
+							using (SQLiteCommand Command = new SQLiteCommand(Connection))
+							{
+								using (SQLiteTransaction Transaction = Connection.BeginTransaction())
+								{
+									foreach (KeyValuePair<string, long> Project in Projects)
+									{
+										Command.CommandText = "UPDATE [CIS] SET ProjectID = @ProjectID WHERE Project = @ProjectName;" +
+															  "UPDATE [Comments] SET ProjectID = @ProjectID WHERE Project = @ProjectName;" +
+															  "UPDATE [Errors] SET ProjectID = @ProjectID WHERE Project = @ProjectName;" +
+															  "UPDATE [Telemetry.v2] SET ProjectID = @ProjectID WHERE Project = @ProjectName;" +
+															  "UPDATE [UserVotes] SET ProjectID = @ProjectID WHERE Project = @ProjectName;";
+										Command.Parameters.AddWithValue("@ProjectName", Project.Key);
+										Command.Parameters.AddWithValue("@ProjectID", Project.Value);
+										Command.ExecuteNonQuery();
+									}
+									Transaction.Commit();
+								}
+							}
+						}
 					}
-				}
+                }
 
-			}
-			catch(Exception)
-			{
-				HttpRuntime.UnloadAppDomain();
-				return;
-			}
-			GlobalConfiguration.Configure(WebApiConfig.Register);
-		}
+            }
+            catch (Exception)
+            {
+                HttpRuntime.UnloadAppDomain();
+                return;
+            }
+            GlobalConfiguration.Configure(WebApiConfig.Register);
+        }
 	}
+    
 }
