UGS: Add support for different projects in the build health system. Alerts will be shown for any project, but each tab will only show issues for the current project by default. Going to the "Show All" menu will show other projects, and also allows filtering by them.

#rnx

diff --git a/Engine/Programs/UnrealGameSync/UnrealGameSync.ini b/Engine/Programs/UnrealGameSync/UnrealGameSync.ini
--- a/Engine/Programs/UnrealGameSync/UnrealGameSync.ini
+++ b/Engine/Programs/UnrealGameSync/UnrealGameSync.ini
@@ -2,7 +2,6 @@
 UseFastModularVersioning=1
 UseFastModularVersioningV2=1
 DebugGameHasSeparateExecutable=1
-ShowBuildHealth=1
 +ExcludeChanges=^Rebuild lightmaps\.$
 +ExcludeChanges=^Rebuild lightmaps & texture streaming\.$
 +ExcludeChanges=^Added collection '.*'\s*
diff --git a/Engine/Source/Programs/UnrealGameSync/MetadataTool/CommandTypes.cs b/Engine/Source/Programs/UnrealGameSync/MetadataTool/CommandTypes.cs
--- a/Engine/Source/Programs/UnrealGameSync/MetadataTool/CommandTypes.cs
+++ b/Engine/Source/Programs/UnrealGameSync/MetadataTool/CommandTypes.cs
@@ -32,7 +32,7 @@ namespace MetadataTool
 			public string Project;
 			public string Summary;
 			public string Details;
-			public string Owner;
+			[Optional] public string Owner;
 		}
 
 		public class UpdateBuild
diff --git a/Engine/Source/Programs/UnrealGameSync/UnrealGameSync/Controls/WorkspaceControl.cs b/Engine/Source/Programs/UnrealGameSync/UnrealGameSync/Controls/WorkspaceControl.cs
--- a/Engine/Source/Programs/UnrealGameSync/UnrealGameSync/Controls/WorkspaceControl.cs
+++ b/Engine/Source/Programs/UnrealGameSync/UnrealGameSync/Controls/WorkspaceControl.cs
@@ -1522,6 +1522,13 @@ namespace UnrealGameSync
 
 		void UpdateBuildFailureNotification()
 		{
+			// Ignore this if we're using the new build health system
+			string BuildHealthProject;
+			if (TryGetProjectSetting(PerforceMonitor.LatestProjectConfigFile, "BuildHealthProject", out BuildHealthProject))
+			{
+				return;
+			}
+
 			int LastChangeByCurrentUser = PerforceMonitor.LastChangeByCurrentUser;
 			int LastCodeChangeByCurrentUser = PerforceMonitor.LastCodeChangeByCurrentUser;
 
@@ -2941,16 +2948,17 @@ namespace UnrealGameSync
 				ProgramsLine.AddText("  |  ");
 				ProgramsLine.AddLink("Windows Explorer", FontStyle.Regular, () => { Process.Start("explorer.exe", String.Format("\"{0}\"", Path.GetDirectoryName(SelectedFileName))); });
 
-				if(Workspace.ProjectConfigFile.GetValue("Options.ShowBuildHealth", false))
+				string BuildHealthProject;
+				if(TryGetProjectSetting(PerforceMonitor.LatestProjectConfigFile, "BuildHealthProject", out BuildHealthProject))
 				{
 					ProgramsLine.AddText("  |  ");
 					if(bUserHasOpenIssues)
 					{
-						ProgramsLine.AddBadge("Build Health", GetBuildBadgeColor(BadgeResult.Failure), (P, R) => ShowBuildHealthMenu(R));
+						ProgramsLine.AddBadge("Build Health", GetBuildBadgeColor(BadgeResult.Failure), (P, R) => ShowBuildHealthMenu(R, BuildHealthProject));
 					}
 					else
 					{
-						ProgramsLine.AddLink("Build Health", FontStyle.Regular, (P, R) => { ShowBuildHealthMenu(R); });
+						ProgramsLine.AddLink("Build Health", FontStyle.Regular, (P, R) => { ShowBuildHealthMenu(R, BuildHealthProject); });
 					}
 				}
 
@@ -3055,7 +3063,7 @@ namespace UnrealGameSync
 			StatusPanel.Set(Lines, Caption, Alert, TintColor);
 		}
 
-		private void ShowBuildHealthMenu(Rectangle Bounds)
+		private void ShowBuildHealthMenu(Rectangle Bounds, string BuildHealthProject)
 		{
 			int MinSeparatorIdx = BuildHealthContextMenu.Items.IndexOf(BuildHealthContextMenu_MinSeparator);
 
@@ -3064,7 +3072,22 @@ namespace UnrealGameSync
 				BuildHealthContextMenu.Items.RemoveAt(MinSeparatorIdx + 1);
 			}
 
-			List<IssueData> Issues = IssueMonitor.GetIssues();
+			List<IssueData> Issues = new List<IssueData>();
+			bool bHasOtherAssignedIssue = false;
+			foreach(IssueData Issue in IssueMonitor.GetIssues())
+			{
+				if(Issue.Project == BuildHealthProject)
+				{
+					Issues.Add(Issue);
+				}
+				else if(Issue.FixChange == 0 && Issue.Owner != null && String.Compare(Issue.Owner, Workspace.Perforce.UserName, StringComparison.OrdinalIgnoreCase) == 0)
+				{
+					bHasOtherAssignedIssue = true;
+				}
+			}
+			
+			BuildHealthContextMenu_Browse.Font = new Font(BuildHealthContextMenu_Browse.Font, bHasOtherAssignedIssue? FontStyle.Bold : FontStyle.Regular);
+
 			if(Issues.Count == 0)
 			{
 				BuildHealthContextMenu_MaxSeparator.Visible = false;
diff --git a/Engine/Source/Programs/UnrealGameSync/UnrealGameSync/Forms/IssueBrowserWindow.cs b/Engine/Source/Programs/UnrealGameSync/UnrealGameSync/Forms/IssueBrowserWindow.cs
--- a/Engine/Source/Programs/UnrealGameSync/UnrealGameSync/Forms/IssueBrowserWindow.cs
+++ b/Engine/Source/Programs/UnrealGameSync/UnrealGameSync/Forms/IssueBrowserWindow.cs
@@ -21,6 +21,9 @@ namespace UnrealGameSync
 		TextWriter Log;
 		string CurrentStream;
 		int MaxResults = 0;
+		string FilterProjectName;
+		List<string> ProjectNames = new List<string>();
+		List<IssueData> Issues = new List<IssueData>();
 
 		public IssueBrowserWindow(IssueMonitor IssueMonitor, string ServerAndPort, string UserName, TextWriter Log, string CurrentStream)
 		{
@@ -53,7 +56,6 @@ namespace UnrealGameSync
 		private void IssueBrowserWindow_Load(object sender, EventArgs e)
 		{
 			FetchMoreResults();
-			UpdateDetailsBtn();
 		}
 
 		class QueryIssuesTask : IModalTask
@@ -95,6 +97,17 @@ namespace UnrealGameSync
 				return;
 			}
 
+			// Update the list of project names
+			Issues = Task.Issues;
+			ProjectNames = Task.Issues.Select(x => x.Project).Distinct().OrderBy(x => x).ToList();
+
+			// Populate the list control
+			UpdateIssueList();
+			MaxResults = NewMaxResults;
+		}
+
+		void UpdateIssueList()
+		{
 			// Get the time at midnight
 			DateTime Now = DateTime.Now;
 			DateTime Midnight = (Now - Now.TimeOfDay).ToUniversalTime();
@@ -102,22 +115,25 @@ namespace UnrealGameSync
 			// Fetch the new issues
 			IssueListView.BeginUpdate();
 			IssueListView.Items.Clear();
-			foreach(IssueData Issue in Task.Issues)
+			foreach(IssueData Issue in Issues)
 			{
-				ListViewItem Item = new ListViewItem("");
-				Item.SubItems.Add(Issue.Id.ToString());
-				Item.SubItems.Add(FormatIssueDateTime(Issue.CreatedAt.ToLocalTime(), Midnight));
-				Item.SubItems.Add(Issue.ResolvedAt.HasValue? FormatIssueDateTime(Issue.ResolvedAt.Value.ToLocalTime(), Midnight) : "Unresolved");
-				Item.SubItems.Add((Issue.Owner == null)? "-" : Utility.FormatUserName(Issue.Owner));
-				Item.SubItems.Add(Issue.Summary);
-				Item.Tag = Issue;
-				IssueListView.Items.Add(Item);
+				if(FilterProjectName == null || Issue.Project == FilterProjectName)
+				{
+					ListViewItem Item = new ListViewItem("");
+					Item.SubItems.Add(Issue.Id.ToString());
+					Item.SubItems.Add(Issue.Project);
+					Item.SubItems.Add(FormatIssueDateTime(Issue.CreatedAt.ToLocalTime(), Midnight));
+					Item.SubItems.Add(Issue.ResolvedAt.HasValue? FormatIssueDateTime(Issue.ResolvedAt.Value.ToLocalTime(), Midnight) : "Unresolved");
+					Item.SubItems.Add((Issue.Owner == null)? "-" : Utility.FormatUserName(Issue.Owner));
+					Item.SubItems.Add(Issue.Summary);
+					Item.Tag = Issue;
+					IssueListView.Items.Add(Item);
+				}
 			}
 			IssueListView.EndUpdate();
 
 			// Update the maximum number of results
-			StatusLabel.Text = String.Format("Showing {0} results.", Task.Issues.Count);
-			MaxResults = NewMaxResults;
+			StatusLabel.Text = (IssueListView.Items.Count == Issues.Count)? String.Format("Showing {0} results.", Issues.Count) : String.Format("Showing {0}/{1} results.", IssueListView.Items.Count, Issues.Count);
 		}
 
 		static string FormatIssueDateTime(DateTime DateTime, DateTime Midnight)
@@ -165,16 +181,6 @@ namespace UnrealGameSync
 			Close();
 		}
 
-		private void IssueListView_SelectedIndexChanged(object sender, EventArgs e)
-		{
-			UpdateDetailsBtn();
-		}
-
-		void UpdateDetailsBtn()
-		{
-			DetailsBtn.Enabled = (IssueListView.SelectedItems.Count != 0);
-		}
-
 		private void IssueListView_MouseDoubleClick(object sender, MouseEventArgs e)
 		{
 			ListViewHitTestInfo HitTest = IssueListView.HitTest(e.Location);
@@ -185,14 +191,31 @@ namespace UnrealGameSync
 			}
 		}
 
-		private void DetailsBtn_Click(object sender, EventArgs e)
+		private void FilterBtn_Click(object sender, EventArgs e)
 		{
-			foreach(ListViewItem Item in IssueListView.SelectedItems)
+			int SeparatorIdx = FilterMenu.Items.IndexOf(FilterMenu_Separator);
+			while(FilterMenu.Items.Count > SeparatorIdx + 1)
 			{
-				IssueData Issue = (IssueData)Item.Tag;
-				ShowIssue(Issue);
-				break;
+				FilterMenu.Items.RemoveAt(SeparatorIdx + 1);
 			}
+
+			FilterMenu_ShowAll.Checked = (FilterProjectName == null);
+
+			foreach(string ProjectName in ProjectNames)
+			{
+				ToolStripMenuItem Item = new ToolStripMenuItem(ProjectName);
+				Item.Checked = (FilterProjectName == ProjectName);
+				Item.Click += (S, E) => { FilterProjectName = ProjectName; UpdateIssueList(); };
+				FilterMenu.Items.Add(Item);
+			}
+
+			FilterMenu.Show(FilterBtn, new Point(FilterBtn.Left, FilterBtn.Bottom));
+		}
+
+		private void FilterMenu_ShowAll_Click(object sender, EventArgs e)
+		{
+			FilterProjectName = null;
+			UpdateIssueList();
 		}
 
 		private void ShowIssue(IssueData Issue)
diff --git a/Engine/Source/Programs/UnrealGameSync/UnrealGameSync/Forms/IssueDetailsWindow.cs b/Engine/Source/Programs/UnrealGameSync/UnrealGameSync/Forms/IssueDetailsWindow.cs
--- a/Engine/Source/Programs/UnrealGameSync/UnrealGameSync/Forms/IssueDetailsWindow.cs
+++ b/Engine/Source/Programs/UnrealGameSync/UnrealGameSync/Forms/IssueDetailsWindow.cs
@@ -438,7 +438,7 @@ namespace UnrealGameSync
 						RichText.Append(@"\fs16");   // Font size
 						if (Idx > 0)
 						{
-							RichText.Append(@"\sb200"); // Space before
+							RichText.Append(@"\sb300"); // Space before
 						}
 						RichText.AppendFormat(@"In step '{0}':", StepName);
 						RichText.AppendLine(@"\par");
@@ -452,6 +452,8 @@ namespace UnrealGameSync
 					RichText.Append(@"\b1");     // Bold
 					RichText.Append(@"\f0");     // Font
 					RichText.Append(@"\fs16");   // Font size
+					RichText.Append(@"\fi50");   // First line indent
+					RichText.Append(@"\li50");   // Other line indent
 					RichText.Append(@"\sb100");  // Space before
 					RichText.Append(@"\sa50");   // Space after
 					RichText.AppendFormat(@"Error {0}/{1}", ++StepNameToIndex[StepName], StepNameToCount[StepName]);
