UGS: Prevent the same diagnostic being added to an issue more than once.

#rnx

diff --git a/Engine/Source/Programs/UnrealGameSync/MetadataTool/BuildHealth/PatternMatcher.cs b/Engine/Source/Programs/UnrealGameSync/MetadataTool/BuildHealth/PatternMatcher.cs
--- a/Engine/Source/Programs/UnrealGameSync/MetadataTool/BuildHealth/PatternMatcher.cs
+++ b/Engine/Source/Programs/UnrealGameSync/MetadataTool/BuildHealth/PatternMatcher.cs
@@ -84,7 +84,8 @@ namespace MetadataTool
 		/// <param name="Target">The fingerprint to merge into</param>
 		public virtual void Merge(BuildHealthIssue Source, BuildHealthIssue Target)
 		{
-			Target.Diagnostics.AddRange(Source.Diagnostics);
+			HashSet<string> TargetUrls = new HashSet<string>(Target.Diagnostics.Select(x => x.ErrorUrl), StringComparer.Ordinal);
+			Target.Diagnostics.AddRange(Source.Diagnostics.Where(x => !TargetUrls.Contains(x.ErrorUrl)));
 			Target.FileNames.UnionWith(Source.FileNames);
 			Target.Identifiers.UnionWith(Source.Identifiers);
 		}
