UGS: Fix order that items are added to the build list control. Seems to be something funky going on with list view groups that causes them to be drawn in a different order than either the ListView.Items or ListViewGroup.Items collections, which is fixed by adding a custom sorter to the class.

diff --git a/Engine/Source/Programs/UnrealGameSync/UnrealGameSync/Controls/WorkspaceControl.cs b/Engine/Source/Programs/UnrealGameSync/UnrealGameSync/Controls/WorkspaceControl.cs
--- a/Engine/Source/Programs/UnrealGameSync/UnrealGameSync/Controls/WorkspaceControl.cs
+++ b/Engine/Source/Programs/UnrealGameSync/UnrealGameSync/Controls/WorkspaceControl.cs
@@ -114,6 +114,27 @@ namespace UnrealGameSync
 			public Dictionary<string, List<BadgeInfo>> CustomBadges;
 		}
 
+		class BuildListItemComparer : System.Collections.IComparer
+		{
+			public int Compare(object A, object B)
+			{
+				PerforceChangeSummary ChangeA = ((ListViewItem)A).Tag as PerforceChangeSummary;
+				PerforceChangeSummary ChangeB = ((ListViewItem)B).Tag as PerforceChangeSummary;
+				if(ChangeA == null)
+				{
+					return +1;
+				}
+				else if(ChangeB == null)
+				{
+					return -1;
+				}
+				else
+				{
+					return ChangeB.Number - ChangeA.Number;
+				}
+			}
+		}
+
 		public static Rectangle GoodBuildIcon = new Rectangle(0, 0, 16, 16);
 		public static Rectangle MixedBuildIcon = new Rectangle(16, 0, 16, 16);
 		public static Rectangle BadBuildIcon = new Rectangle(32, 0, 16, 16);
@@ -274,6 +295,7 @@ namespace UnrealGameSync
 			BuildList.SmallImageList = new ImageList(){ ImageSize = new Size(1, 20) };
 			BuildList_FontChanged(null, null);
 			BuildList.OnScroll += BuildList_OnScroll;
+			BuildList.ListViewItemSorter = new BuildListItemComparer();
 
 			Splitter.OnVisibilityChanged += Splitter_OnVisibilityChanged;
 
@@ -1253,7 +1275,7 @@ namespace UnrealGameSync
 			ListViewGroup Group = BuildList_FindOrAddGroup(Change, DisplayTime.ToString("D"));
 
 			// Create the new item
-			ListViewItem Item = new ListViewItem(Group);
+			ListViewItem Item = new ListViewItem();
 			Item.Tag = Change;
 
 			// Get the new text for each column
