UGS: Fix issue with metadatatool where issues could accumulate way too many diagnostics.

diff --git a/Engine/Source/Programs/UnrealGameSync/MetadataTool/BuildHealth/PatternMatcher.cs b/Engine/Source/Programs/UnrealGameSync/MetadataTool/BuildHealth/PatternMatcher.cs
--- a/Engine/Source/Programs/UnrealGameSync/MetadataTool/BuildHealth/PatternMatcher.cs
+++ b/Engine/Source/Programs/UnrealGameSync/MetadataTool/BuildHealth/PatternMatcher.cs
@@ -85,7 +85,18 @@ namespace MetadataTool
 		public virtual void Merge(BuildHealthIssue Source, BuildHealthIssue Target)
 		{
 			HashSet<string> TargetUrls = new HashSet<string>(Target.Diagnostics.Select(x => x.ErrorUrl), StringComparer.Ordinal);
-			Target.Diagnostics.AddRange(Source.Diagnostics.Where(x => !TargetUrls.Contains(x.ErrorUrl)));
+			foreach(BuildHealthDiagnostic SourceDiagnostic in Source.Diagnostics)
+			{
+				if(Target.Diagnostics.Count >= 50)
+				{
+					break;
+				}
+				if(!TargetUrls.Contains(SourceDiagnostic.ErrorUrl))
+				{
+					Target.Diagnostics.Add(SourceDiagnostic);
+				}
+			}
+
 			Target.FileNames.UnionWith(Source.FileNames);
 			Target.Identifiers.UnionWith(Source.Identifiers);
 		}
diff --git a/Engine/Source/Programs/UnrealGameSync/MetadataTool/CommandHandler_BuildHealth.cs b/Engine/Source/Programs/UnrealGameSync/MetadataTool/CommandHandler_BuildHealth.cs
--- a/Engine/Source/Programs/UnrealGameSync/MetadataTool/CommandHandler_BuildHealth.cs
+++ b/Engine/Source/Programs/UnrealGameSync/MetadataTool/CommandHandler_BuildHealth.cs
@@ -56,7 +56,7 @@ namespace MetadataTool
 			bool bClean = Arguments.HasOption("-Clean");
 			string PerforcePort = Arguments.GetStringOrDefault("-P4Port=", null);
 			string PerforceUser = Arguments.GetStringOrDefault("-P4User=", null);
-			FileReference InputFile = Arguments.GetFileReference("-InputFile=");
+			FileReference InputFile = Arguments.GetFileReferenceOrDefault("-InputFile=", null);
 			FileReference StateFile = Arguments.GetFileReference("-StateFile=");
 			string ServerUrl = Arguments.GetStringOrDefault("-Server=", null);
 			bool bKeepHistory = Arguments.HasOption("-KeepHistory");
@@ -69,10 +69,6 @@ namespace MetadataTool
 				CategoryNameToMatcher[Matcher.Category] = Matcher;
 			}
 
-			// Parse the input file
-			Log.TraceInformation("Reading build results from {0}", InputFile);
-			InputData InputData = DeserializeJson<InputData>(InputFile);
-
 			// Complete any interrupted operation to update the state file
 			CompleteStateTransaction(StateFile);
 
@@ -92,77 +88,85 @@ namespace MetadataTool
 			// Create the Perforce connection
 			PerforceConnection Perforce = new PerforceConnection(PerforcePort, PerforceUser, null);
 
-			// Parse all the builds and add them to the persistent data
-			List<InputJob> InputJobs = InputData.Jobs.OrderBy(x => x.Change).ThenBy(x => x.Stream).ToList();
-			Stopwatch Timer = Stopwatch.StartNew();
-			foreach (InputJob InputJob in InputJobs)
+			// Process the input data
+			if(InputFile != null)
 			{
-				// Add a new build for each job step
-				foreach(InputJobStep InputJobStep in InputJob.Steps)
+				// Parse the input file
+				Log.TraceInformation("Reading build results from {0}", InputFile);
+				InputData InputData = DeserializeJson<InputData>(InputFile);
+
+				// Parse all the builds and add them to the persistent data
+				List<InputJob> InputJobs = InputData.Jobs.OrderBy(x => x.Change).ThenBy(x => x.Stream).ToList();
+				Stopwatch Timer = Stopwatch.StartNew();
+				foreach (InputJob InputJob in InputJobs)
 				{
-					BuildHealthJobStep NewBuild = new BuildHealthJobStep(InputJob.Change, InputJob.Name, InputJob.Url, InputJobStep.Name, InputJobStep.Url, null);
-					State.AddBuild(InputJob.Stream, NewBuild);
-				}
+					// Add a new build for each job step
+					foreach(InputJobStep InputJobStep in InputJob.Steps)
+					{
+						BuildHealthJobStep NewBuild = new BuildHealthJobStep(InputJob.Change, InputJob.Name, InputJob.Url, InputJobStep.Name, InputJobStep.Url, null);
+						State.AddBuild(InputJob.Stream, NewBuild);
+					}
 
-				// Add all the job steps
-				List<InputJobStep> InputJobSteps = InputJob.Steps.OrderBy(x => x.Name).ToList();
-				foreach (InputJobStep InputJobStep in InputJobSteps)
-				{
-					AddStep(Perforce, State, InputJob, InputJobStep);
+					// Add all the job steps
+					List<InputJobStep> InputJobSteps = InputJob.Steps.OrderBy(x => x.Name).ToList();
+					foreach (InputJobStep InputJobStep in InputJobSteps)
+					{
+						AddStep(Perforce, State, InputJob, InputJobStep);
+					}
 				}
-			}
-			Log.TraceInformation("Added jobs in {0}s", Timer.Elapsed.TotalSeconds);
+				Log.TraceInformation("Added jobs in {0}s", Timer.Elapsed.TotalSeconds);
 
-			// Try to find the next successful build for each stream, so we can close it as part of updating the server
-			for (int Idx = 0; Idx < State.Issues.Count; Idx++)
-			{
-				BuildHealthIssue Issue = State.Issues[Idx];
-				foreach(string Stream in Issue.Streams.Keys)
+				// Try to find the next successful build for each stream, so we can close it as part of updating the server
+				for (int Idx = 0; Idx < State.Issues.Count; Idx++)
 				{
-					Dictionary<string, BuildHealthJobHistory> StepNameToHistory = Issue.Streams[Stream];
-					foreach(string StepName in StepNameToHistory.Keys)
+					BuildHealthIssue Issue = State.Issues[Idx];
+					foreach(string Stream in Issue.Streams.Keys)
 					{
-						BuildHealthJobHistory IssueHistory = StepNameToHistory[StepName];
-						if(IssueHistory.FailedBuilds.Count > 0 && IssueHistory.NextSuccessfulBuild == null)
+						Dictionary<string, BuildHealthJobHistory> StepNameToHistory = Issue.Streams[Stream];
+						foreach(string StepName in StepNameToHistory.Keys)
 						{
-							// Find the successful build after this change
-							BuildHealthJobStep LastFailedBuild = IssueHistory.FailedBuilds[IssueHistory.FailedBuilds.Count - 1];
-							IssueHistory.NextSuccessfulBuild = State.FindBuildAfter(Stream, LastFailedBuild.Change, StepName);
+							BuildHealthJobHistory IssueHistory = StepNameToHistory[StepName];
+							if(IssueHistory.FailedBuilds.Count > 0 && IssueHistory.NextSuccessfulBuild == null)
+							{
+								// Find the successful build after this change
+								BuildHealthJobStep LastFailedBuild = IssueHistory.FailedBuilds[IssueHistory.FailedBuilds.Count - 1];
+								IssueHistory.NextSuccessfulBuild = State.FindBuildAfter(Stream, LastFailedBuild.Change, StepName);
+							}
 						}
 					}
 				}
-			}
 
-			// Find the change two days before the latest change being added
-			if(InputData.Jobs.Count > 0 && !bKeepHistory)
-			{
-				// Find all the unique change numbers for each stream
-				SortedSet<int> ChangeNumbers = new SortedSet<int>();
-				foreach (List<BuildHealthJobStep> Builds in State.Streams.Values)
+				// Find the change two days before the latest change being added
+				if(InputData.Jobs.Count > 0 && !bKeepHistory)
 				{
-					ChangeNumbers.UnionWith(Builds.Select(x => x.Change));
-				}
+					// Find all the unique change numbers for each stream
+					SortedSet<int> ChangeNumbers = new SortedSet<int>();
+					foreach (List<BuildHealthJobStep> Builds in State.Streams.Values)
+					{
+						ChangeNumbers.UnionWith(Builds.Select(x => x.Change));
+					}
 
-				// Get the latest change record
-				int LatestChangeNumber = InputData.Jobs.Min(x => x.Change);
-				ChangeRecord LatestChangeRecord = Perforce.GetChange(GetChangeOptions.None, LatestChangeNumber).Data;
+					// Get the latest change record
+					int LatestChangeNumber = InputData.Jobs.Min(x => x.Change);
+					ChangeRecord LatestChangeRecord = Perforce.GetChange(GetChangeOptions.None, LatestChangeNumber).Data;
 
-				// Step forward through all the changelists until we get to one we don't want to delete
-				int DeleteChangeNumber = -1;
-				foreach(int ChangeNumber in ChangeNumbers)
-				{
-					ChangeRecord ChangeRecord = Perforce.GetChange(GetChangeOptions.None, ChangeNumber).Data;
-					if (ChangeRecord.Date > LatestChangeRecord.Date - TimeSpan.FromDays(2))
+					// Step forward through all the changelists until we get to one we don't want to delete
+					int DeleteChangeNumber = -1;
+					foreach(int ChangeNumber in ChangeNumbers)
 					{
-						break;
+						ChangeRecord ChangeRecord = Perforce.GetChange(GetChangeOptions.None, ChangeNumber).Data;
+						if (ChangeRecord.Date > LatestChangeRecord.Date - TimeSpan.FromDays(2))
+						{
+							break;
+						}
+						DeleteChangeNumber = ChangeNumber;
 					}
-					DeleteChangeNumber = ChangeNumber;
-				}
 
-				// Remove any builds we no longer want to track
-				foreach (List<BuildHealthJobStep> Builds in State.Streams.Values)
-				{
-					Builds.RemoveAll(x => x.Change <= DeleteChangeNumber);
+					// Remove any builds we no longer want to track
+					foreach (List<BuildHealthJobStep> Builds in State.Streams.Values)
+					{
+						Builds.RemoveAll(x => x.Change <= DeleteChangeNumber);
+					}
 				}
 			}
 
@@ -207,6 +211,10 @@ namespace MetadataTool
 					{
 						DetailsBuilder.AppendFormat("##{0}##\n{1}\n", Diagnostic.JobStepName, Diagnostic.Message);
 					}
+					if(DetailsBuilder.Length > CommandTypes.MaxIssueDetailsLength)
+					{
+						DetailsBuilder.Remove(CommandTypes.MaxIssueDetailsLength, DetailsBuilder.Length - CommandTypes.MaxIssueDetailsLength);
+					}
 
 					string Summary = Matcher.GetSummary(Issue);
 					string Details = DetailsBuilder.ToString().TrimEnd();
diff --git a/Engine/Source/Programs/UnrealGameSync/MetadataTool/CommandTypes.cs b/Engine/Source/Programs/UnrealGameSync/MetadataTool/CommandTypes.cs
--- a/Engine/Source/Programs/UnrealGameSync/MetadataTool/CommandTypes.cs
+++ b/Engine/Source/Programs/UnrealGameSync/MetadataTool/CommandTypes.cs
@@ -15,6 +15,8 @@ namespace MetadataTool
 
 	static class CommandTypes
 	{
+		public const int MaxIssueDetailsLength = 8192;
+
 		public enum Outcome
 		{
 			Success = 1,
