UGS: Suppress duplicate messages from MetadataTool.

diff --git a/Engine/Source/Programs/UnrealGameSync/MetadataTool/BuildHealth/PatternMatcher.cs b/Engine/Source/Programs/UnrealGameSync/MetadataTool/BuildHealth/PatternMatcher.cs
--- a/Engine/Source/Programs/UnrealGameSync/MetadataTool/BuildHealth/PatternMatcher.cs
+++ b/Engine/Source/Programs/UnrealGameSync/MetadataTool/BuildHealth/PatternMatcher.cs
@@ -84,14 +84,14 @@ namespace MetadataTool
 		/// <param name="Target">The fingerprint to merge into</param>
 		public virtual void Merge(BuildHealthIssue Source, BuildHealthIssue Target)
 		{
-			HashSet<string> TargetUrls = new HashSet<string>(Target.Diagnostics.Select(x => x.ErrorUrl), StringComparer.Ordinal);
+			HashSet<string> TargetMessages = new HashSet<string>(Target.Diagnostics.Select(x => x.Message), StringComparer.Ordinal);
 			foreach(BuildHealthDiagnostic SourceDiagnostic in Source.Diagnostics)
 			{
 				if(Target.Diagnostics.Count >= 50)
 				{
 					break;
 				}
-				if(!TargetUrls.Contains(SourceDiagnostic.ErrorUrl))
+				if(!TargetMessages.Contains(SourceDiagnostic.Message))
 				{
 					Target.Diagnostics.Add(SourceDiagnostic);
 				}
