UGS: Sorten the entries in the issues context menu.

#rnx

diff --git a/Engine/Source/Programs/UnrealGameSync/UnrealGameSync/Controls/WorkspaceControl.cs b/Engine/Source/Programs/UnrealGameSync/UnrealGameSync/Controls/WorkspaceControl.cs
--- a/Engine/Source/Programs/UnrealGameSync/UnrealGameSync/Controls/WorkspaceControl.cs
+++ b/Engine/Source/Programs/UnrealGameSync/UnrealGameSync/Controls/WorkspaceControl.cs
@@ -3121,8 +3121,14 @@ namespace UnrealGameSync
 				{
 					IssueData Issue = Issues[Idx];
 
+					string Summary = Issue.Summary;
+					if(Summary.Length > 100)
+					{
+						Summary = Summary.Substring(0, 100).TrimEnd() + "...";
+					}
+
 					StringBuilder Description = new StringBuilder();
-					Description.AppendFormat("{0}: {1}", Issue.Id, Issue.Summary);
+					Description.AppendFormat("{0}: {1}", Issue.Id, Summary);
 					if(Issue.Owner == null)
 					{
 						Description.AppendFormat(" - Unassigned");
