UGS: Stop MetadataTool updating issues even when the details string hasn't changed.

#rnx

diff --git a/Engine/Source/Programs/UnrealGameSync/MetadataTool/CommandHandler_BuildHealth.cs b/Engine/Source/Programs/UnrealGameSync/MetadataTool/CommandHandler_BuildHealth.cs
--- a/Engine/Source/Programs/UnrealGameSync/MetadataTool/CommandHandler_BuildHealth.cs
+++ b/Engine/Source/Programs/UnrealGameSync/MetadataTool/CommandHandler_BuildHealth.cs
@@ -248,6 +248,7 @@ namespace MetadataTool
 						}
 
 						Issue.PostedSummary = Summary;
+						Issue.PostedDetails = Details;
 						WriteState(StateFile, State);
 					}
 					else if(Issue.PostedSummary == null 
@@ -271,6 +272,7 @@ namespace MetadataTool
 						}
 
 						Issue.PostedSummary = Summary;
+						Issue.PostedDetails = Details;
 						WriteState(StateFile, State);
 					}
 				}
