UGS: Add times into the issue details window.

#rnx

diff --git a/Engine/Source/Programs/UnrealGameSync/UnrealGameSync/Controls/WorkspaceControl.cs b/Engine/Source/Programs/UnrealGameSync/UnrealGameSync/Controls/WorkspaceControl.cs
--- a/Engine/Source/Programs/UnrealGameSync/UnrealGameSync/Controls/WorkspaceControl.cs
+++ b/Engine/Source/Programs/UnrealGameSync/UnrealGameSync/Controls/WorkspaceControl.cs
@@ -3166,14 +3166,24 @@ namespace UnrealGameSync
 			ShowIssueDetails(Issue);
 		}
 
+		private TimeSpan? GetServerTimeOffset()
+		{
+			TimeSpan? Offset = null;
+			if (Settings.bShowLocalTimes)
+			{
+				Offset = PerforceMonitor.ServerTimeZone;
+			}
+			return Offset;
+		}
+
 		public void ShowIssueDetails(IssueData Issue)
 		{
-			IssueDetailsWindow.Show(ParentForm, IssueMonitor, Workspace.Perforce.ServerAndPort, Workspace.Perforce.UserName, Issue, Log, StreamName);
+			IssueDetailsWindow.Show(ParentForm, IssueMonitor, Workspace.Perforce.ServerAndPort, Workspace.Perforce.UserName, GetServerTimeOffset(), Issue, Log, StreamName);
 		}
 
 		private void BuildHealthContextMenu_Browse_Click(object sender, EventArgs e)
 		{
-			IssueBrowserWindow.Show(ParentForm, IssueMonitor, Workspace.Perforce.ServerAndPort, Workspace.Perforce.UserName, Log, StreamName);
+			IssueBrowserWindow.Show(ParentForm, IssueMonitor, Workspace.Perforce.ServerAndPort, Workspace.Perforce.UserName, GetServerTimeOffset(), Log, StreamName);
 		}
 
 		private void BuildHealthContextMenu_Settings_Click(object sender, EventArgs e)
diff --git a/Engine/Source/Programs/UnrealGameSync/UnrealGameSync/Forms/IssueBrowserWindow.cs b/Engine/Source/Programs/UnrealGameSync/UnrealGameSync/Forms/IssueBrowserWindow.cs
--- a/Engine/Source/Programs/UnrealGameSync/UnrealGameSync/Forms/IssueBrowserWindow.cs
+++ b/Engine/Source/Programs/UnrealGameSync/UnrealGameSync/Forms/IssueBrowserWindow.cs
@@ -18,6 +18,7 @@ namespace UnrealGameSync
 		IssueMonitor IssueMonitor;
 		string ServerAndPort;
 		string UserName;
+		TimeSpan? ServerTimeOffset;
 		TextWriter Log;
 		string CurrentStream;
 		int MaxResults = 0;
@@ -25,11 +26,12 @@ namespace UnrealGameSync
 		List<string> ProjectNames = new List<string>();
 		List<IssueData> Issues = new List<IssueData>();
 
-		public IssueBrowserWindow(IssueMonitor IssueMonitor, string ServerAndPort, string UserName, TextWriter Log, string CurrentStream)
+		public IssueBrowserWindow(IssueMonitor IssueMonitor, string ServerAndPort, string UserName, TimeSpan? ServerTimeOffset, TextWriter Log, string CurrentStream)
 		{
 			this.IssueMonitor = IssueMonitor;
 			this.ServerAndPort = ServerAndPort;
 			this.UserName = UserName;
+			this.ServerTimeOffset = ServerTimeOffset;
 			this.Log = Log;
 			this.CurrentStream = CurrentStream;
 
@@ -150,12 +152,12 @@ namespace UnrealGameSync
 
 		static List<IssueBrowserWindow> ExistingWindows = new List<IssueBrowserWindow>();
 
-		public static void Show(Form Owner, IssueMonitor IssueMonitor, string ServerAndPort, string UserName, TextWriter Log, string CurrentStream)
+		public static void Show(Form Owner, IssueMonitor IssueMonitor, string ServerAndPort, string UserName, TimeSpan? ServerTimeOffset, TextWriter Log, string CurrentStream)
 		{
 			IssueBrowserWindow Window = ExistingWindows.FirstOrDefault(x => x.IssueMonitor == IssueMonitor);
 			if(Window == null)
 			{
-				Window = new IssueBrowserWindow(IssueMonitor, ServerAndPort, UserName, Log, CurrentStream);
+				Window = new IssueBrowserWindow(IssueMonitor, ServerAndPort, UserName, ServerTimeOffset, Log, CurrentStream);
 				Window.Owner = Owner;
 				Window.StartPosition = FormStartPosition.Manual;
 				Window.Location = new Point(Owner.Location.X + (Owner.Width - Window.Width) / 2, Owner.Location.Y + (Owner.Height - Window.Height) / 2);
@@ -221,7 +223,7 @@ namespace UnrealGameSync
 		private void ShowIssue(IssueData Issue)
 		{
 			Issue.Builds = RESTApi.GET<List<IssueBuildData>>(IssueMonitor.ApiUrl, String.Format("issues/{0}/builds", Issue.Id));
-			IssueDetailsWindow.Show(Owner, IssueMonitor, ServerAndPort, UserName, Issue, Log, CurrentStream);
+			IssueDetailsWindow.Show(Owner, IssueMonitor, ServerAndPort, UserName, ServerTimeOffset, Issue, Log, CurrentStream);
 		}
 
 		private void IssueListView_DrawColumnHeader(object sender, DrawListViewColumnHeaderEventArgs e)
diff --git a/Engine/Source/Programs/UnrealGameSync/UnrealGameSync/Forms/IssueDetailsWindow.cs b/Engine/Source/Programs/UnrealGameSync/UnrealGameSync/Forms/IssueDetailsWindow.cs
--- a/Engine/Source/Programs/UnrealGameSync/UnrealGameSync/Forms/IssueDetailsWindow.cs
+++ b/Engine/Source/Programs/UnrealGameSync/UnrealGameSync/Forms/IssueDetailsWindow.cs
@@ -242,6 +242,7 @@ namespace UnrealGameSync
 		IssueMonitor IssueMonitor;
 		IssueData Issue;
 		PerforceConnection Perforce;
+		TimeSpan? ServerTimeOffset;
 		PerforceWorkerThread PerforceWorker;
 		TextWriter Log;
 		SynchronizationContext MainThreadSynchronizationContext;
@@ -255,11 +256,12 @@ namespace UnrealGameSync
 		System.Windows.Forms.Timer UpdateTimer;
 		StatusElementResources StatusElementResources;
 
-		IssueDetailsWindow(IssueMonitor IssueMonitor, IssueData Issue, string ServerAndPort, string UserName, TextWriter Log, string CurrentStream)
+		IssueDetailsWindow(IssueMonitor IssueMonitor, IssueData Issue, string ServerAndPort, string UserName, TimeSpan? ServerTimeOffset, TextWriter Log, string CurrentStream)
 		{
 			this.IssueMonitor = IssueMonitor;
 			this.Issue = Issue;
 			this.Perforce = new PerforceConnection(UserName, null, ServerAndPort);
+			this.ServerTimeOffset = ServerTimeOffset;
 			this.Log = Log;
 
 			IssueMonitor.AddRef();
@@ -737,13 +739,22 @@ namespace UnrealGameSync
 
 							ListViewItem Item = new ListViewItem("");
 							Item.Tag = Change;
-							Item.SubItems.Add(Change.Number.ToString());
 
 							StatusLineListViewWidget TypeWidget = new StatusLineListViewWidget(Item, StatusElementResources);
 							UpdateChangeTypeWidget(TypeWidget, Details);
 							Item.SubItems.Add(new ListViewItem.ListViewSubItem(Item, "") { Tag = TypeWidget });
 
+							Item.SubItems.Add(Change.Number.ToString());
+
+							DateTime DisplayTime = Change.Date;
+							if (ServerTimeOffset.HasValue)
+							{
+								DisplayTime = (DisplayTime - ServerTimeOffset.Value).ToLocalTime();
+							}
+							Item.SubItems.Add(DisplayTime.ToString("h\\.mmtt"));
+
 							Item.SubItems.Add(WorkspaceControl.FormatUserName(Change.User));
+
 							Item.SubItems.Add(Change.Description);
 							BuildListView.Items.Insert(NumNewItems++, Item);
 						}
@@ -775,6 +786,7 @@ namespace UnrealGameSync
 				*/
 				ListViewItem BuildItem = new ListViewItem("");
 				BuildItem.Tag = Range.BuildGroup;
+				BuildItem.SubItems.Add("");
 				BuildItem.SubItems.Add(Range.BuildGroup.Change.ToString());
 				BuildItem.SubItems.Add("");
 
@@ -970,7 +982,7 @@ namespace UnrealGameSync
 						PerforceChangeDetailsWithDescribeRecord Details;
 						PerforceWorker.TryGetChangeDetails(Change.Number, out Details);
 
-						StatusLineListViewWidget TypeWidget = (StatusLineListViewWidget)Item.SubItems[2].Tag;
+						StatusLineListViewWidget TypeWidget = (StatusLineListViewWidget)Item.SubItems[TypeHeader.Index].Tag;
 						UpdateChangeTypeWidget(TypeWidget, Details);
 
 						BuildListView.RedrawItems(Item.Index, Item.Index, true);
@@ -987,12 +999,12 @@ namespace UnrealGameSync
 
 		static List<IssueDetailsWindow> ExistingWindows = new List<IssueDetailsWindow>();
 
-		public static void Show(Form Owner, IssueMonitor IssueMonitor, string ServerAndPort, string UserName, IssueData Issue, TextWriter Log, string CurrentStream)
+		public static void Show(Form Owner, IssueMonitor IssueMonitor, string ServerAndPort, string UserName, TimeSpan? ServerTimeOffset, IssueData Issue, TextWriter Log, string CurrentStream)
 		{
 			IssueDetailsWindow Window = ExistingWindows.FirstOrDefault(x => x.IssueMonitor == IssueMonitor && x.Issue.Id == Issue.Id);
 			if(Window == null)
 			{
-				Window = new IssueDetailsWindow(IssueMonitor, Issue, ServerAndPort, UserName, Log, CurrentStream);
+				Window = new IssueDetailsWindow(IssueMonitor, Issue, ServerAndPort, UserName, ServerTimeOffset, Log, CurrentStream);
 				Window.Owner = Owner;
 				Window.StartPosition = FormStartPosition.Manual;
 				Window.Location = new Point(Owner.Location.X + (Owner.Width - Window.Width) / 2, Owner.Location.Y + (Owner.Height - Window.Height) / 2);
@@ -1099,6 +1111,10 @@ namespace UnrealGameSync
 				{
 					TextRenderer.DrawText(e.Graphics, e.SubItem.Text, ChangeFont, e.Bounds, TextColor, TextFormatFlags.EndEllipsis | TextFormatFlags.SingleLine | TextFormatFlags.Left | TextFormatFlags.VerticalCenter | TextFormatFlags.NoPrefix);
 				}
+				else if(e.ColumnIndex == TimeHeader.Index)
+				{
+					TextRenderer.DrawText(e.Graphics, e.SubItem.Text, ChangeFont, e.Bounds, TextColor, TextFormatFlags.EndEllipsis | TextFormatFlags.SingleLine | TextFormatFlags.HorizontalCenter | TextFormatFlags.VerticalCenter | TextFormatFlags.NoPrefix);
+				}
 			}
 			else if(e.Item.Tag is IssueBuildData)
 			{
