UGS: Rely on the listview sorter for putting items in the correct order, and fix issue with the expand button not working.

#rnx

diff --git a/Engine/Source/Programs/UnrealGameSync/UnrealGameSync/Controls/WorkspaceControl.cs b/Engine/Source/Programs/UnrealGameSync/UnrealGameSync/Controls/WorkspaceControl.cs
--- a/Engine/Source/Programs/UnrealGameSync/UnrealGameSync/Controls/WorkspaceControl.cs
+++ b/Engine/Source/Programs/UnrealGameSync/UnrealGameSync/Controls/WorkspaceControl.cs
@@ -1082,8 +1082,6 @@ namespace UnrealGameSync
 				ChangeNumberToArchivePath.Clear();
 				ChangeNumberToLayoutInfo.Clear();
 
-				ExpandItem = null;
-
 				List<PerforceChangeSummary> Changes = PerforceMonitor.GetChanges();
 				EventMonitor.FilterChanges(Changes.Select(x => x.Number));
 
@@ -1152,82 +1150,55 @@ namespace UnrealGameSync
 		{
 			BuildList.BeginUpdate();
 
-			// Find the item before the expand row
-			ListViewItem BeforeExpandItem = null;
-			if(ExpandItem != null && ExpandItem.Index > 0)
+			// Remove any changes that no longer exist, and update the rest
+			Dictionary<int, PerforceChangeSummary> ChangeNumberToSummary = Changes.ToDictionary(x => x.Number, x => x);
+			for(int Idx = BuildList.Items.Count - 1; Idx >= 0; Idx--)
 			{
-				BeforeExpandItem = BuildList.Items[ExpandItem.Index - 1];
-			}
-
-			// Merge the new changes into the build list
-			int BuildListIndex = 0;
-			foreach (PerforceChangeSummary Change in Changes)
-			{
-				for (; ; )
+				PerforceChangeSummary Change = BuildList.Items[Idx].Tag as PerforceChangeSummary;
+				if (Change != null)
 				{
-					if (BuildListIndex == BuildList.Items.Count)
-					{
-						// Insert
-						BuildList_InsertItem(BuildListIndex, Change);
-						break;
-					}
-
-					ListViewItem NextItem = BuildList.Items[BuildListIndex];
-					if (NextItem.Tag == null)
-					{
-						// Insert
-						BuildList_InsertItem(BuildListIndex, Change);
-						break;
-					}
-
-					PerforceChangeSummary NextChange = (PerforceChangeSummary)NextItem.Tag;
-					if (Change.Number > NextChange.Number)
-					{
-						// Insert
-						BuildList_InsertItem(BuildListIndex, Change);
-						break;
-					}
-					else if (NextChange.Number == Change.Number)
+					PerforceChangeSummary Summary;
+					if(ChangeNumberToSummary.TryGetValue(Change.Number, out Summary))
 					{
 						// Update
-						NextItem.SubItems[DescriptionColumn.Index].Text = Change.Description.Replace('\n', ' ');
-						break;
+						BuildList.Items[Idx].SubItems[DescriptionColumn.Index].Text = Change.Description.Replace('\n', ' ');
+						ChangeNumberToSummary.Remove(Change.Number);
 					}
 					else
 					{
 						// Delete
-						BuildList.Items.RemoveAt(BuildListIndex);
-						continue;
+						BuildList.Items.RemoveAt(Idx);
 					}
 				}
+			}
 
-				// Move to the next item
-				BuildListIndex++;
+			// Add everything left over
+			foreach(PerforceChangeSummary Change in ChangeNumberToSummary.Values)
+			{
+				BuildList_AddItem(Change);
 			}
 
-			// Figure out if we need to invalidate the expand row
-			if(bShowExpandItem)
+			// Fixup the expand row
+			ListViewGroup Group = (BuildList.Groups.Count > 0) ? BuildList.Groups[BuildList.Groups.Count - 1] : null;
+			if (ExpandItem != null)
 			{
-				ListViewGroup Group = (BuildListIndex > 0)? BuildList.Items[BuildListIndex - 1].Group : null;
-				if (ExpandItem == null || ExpandItem.Index != BuildListIndex || ExpandItem.Group != Group)
+				if(!bShowExpandItem || ExpandItem.Group != Group)
 				{
-					ExpandItem = new ListViewItem(Group);
-					ExpandItem.Tag = null;
-					ExpandItem.Selected = false;
-					ExpandItem.Text = "";
-					for (int ColumnIdx = 1; ColumnIdx < BuildList.Columns.Count; ColumnIdx++)
-					{
-						ExpandItem.SubItems.Add(new ListViewItem.ListViewSubItem(ExpandItem, ""));
-					}
-					BuildList.Items.Insert(BuildListIndex, ExpandItem);
+					BuildList.Items.Remove(ExpandItem);
+					ExpandItem = null;
 				}
-				BuildListIndex++;
 			}
-
-			// Remove everything else
-			while (BuildList.Items.Count > BuildListIndex)
+			if(bShowExpandItem && ExpandItem == null)
 			{
-				BuildList.Items.RemoveAt(BuildList.Items.Count - 1);
+				ExpandItem = new ListViewItem(Group);
+				ExpandItem.Tag = null;
+				ExpandItem.Selected = false;
+				ExpandItem.Text = "";
+				for (int ColumnIdx = 1; ColumnIdx < BuildList.Columns.Count; ColumnIdx++)
+				{
+					ExpandItem.SubItems.Add(new ListViewItem.ListViewSubItem(ExpandItem, ""));
+				}
+				BuildList.Items.Add(ExpandItem);
 			}
 
 			BuildList.EndUpdate();
@@ -1262,7 +1233,7 @@ namespace UnrealGameSync
 			return Group;
 		}
 
-		void BuildList_InsertItem(int Index, PerforceChangeSummary Change)
+		void BuildList_AddItem(PerforceChangeSummary Change)
 		{
 			// Get the display time for this item
 			DateTime DisplayTime = Change.Date;
@@ -1303,7 +1274,7 @@ namespace UnrealGameSync
 			Group.Items.Insert(GroupInsertIdx, Item);
 
 			// Insert it into the build list
-			BuildList.Items.Insert(Index, Item);
+			BuildList.Items.Add(Item);
 		}
 
 		bool ShouldShowChange(PerforceChangeSummary Change, string[] ExcludeChanges)
