UGS did not format time according to local culture (meaning it always used 12 hour formats).

diff --git a/Engine/Source/Programs/UnrealGameSync/UnrealGameSync/Controls/WorkspaceControl.cs b/Engine/Source/Programs/UnrealGameSync/UnrealGameSync/Controls/WorkspaceControl.cs
--- a/Engine/Source/Programs/UnrealGameSync/UnrealGameSync/Controls/WorkspaceControl.cs
+++ b/Engine/Source/Programs/UnrealGameSync/UnrealGameSync/Controls/WorkspaceControl.cs
@@ -1252,7 +1252,8 @@ namespace UnrealGameSync
 			// Get the new text for each column
 			string[] Columns = new string[BuildList.Columns.Count];
 			Columns[ChangeColumn.Index] = String.Format("{0}", Change.Number);
-			Columns[TimeColumn.Index] = DisplayTime.ToString("h\\.mmtt");
+
+			Columns[TimeColumn.Index] = DisplayTime.ToShortTimeString();
 			Columns[AuthorColumn.Index] = FormatUserName(Change.User);
 			Columns[DescriptionColumn.Index] = Change.Description.Replace('\n', ' ');
 
diff --git a/Engine/Source/Programs/UnrealGameSync/UnrealGameSync/Forms/ScheduleWindow.cs b/Engine/Source/Programs/UnrealGameSync/UnrealGameSync/Forms/ScheduleWindow.cs
--- a/Engine/Source/Programs/UnrealGameSync/UnrealGameSync/Forms/ScheduleWindow.cs
+++ b/Engine/Source/Programs/UnrealGameSync/UnrealGameSync/Forms/ScheduleWindow.cs
@@ -5,6 +5,7 @@ using System.Collections.Generic;
 using System.ComponentModel;
 using System.Data;
 using System.Drawing;
+using System.Globalization;
 using System.IO;
 using System.Linq;
 using System.Text;
@@ -23,6 +24,7 @@ namespace UnrealGameSync
 			ChangeComboBox.SelectedIndex = (int)InChange;
 
 			DateTime CurrentTime = DateTime.Now;
+			TimePicker.CustomFormat = CultureInfo.CurrentCulture.DateTimeFormat.ShortTimePattern;
 			TimePicker.Value = new DateTime(CurrentTime.Year, CurrentTime.Month, CurrentTime.Day, InTime.Hours, InTime.Minutes, InTime.Seconds);
 
 			ProjectListBox.Items.Add("Any open projects", AnyOpenProject);
