UGS: Add additional diagnostic output to MetadataTool.

#rnx

diff --git a/Engine/Source/Programs/UnrealGameSync/MetadataTool/CommandHandler_BuildHealth.cs b/Engine/Source/Programs/UnrealGameSync/MetadataTool/CommandHandler_BuildHealth.cs
--- a/Engine/Source/Programs/UnrealGameSync/MetadataTool/CommandHandler_BuildHealth.cs
+++ b/Engine/Source/Programs/UnrealGameSync/MetadataTool/CommandHandler_BuildHealth.cs
@@ -406,9 +406,11 @@ namespace MetadataTool
 			HttpWebRequest Request = (HttpWebRequest)WebRequest.Create(Url);
 			Request.ContentType = "application/json";
 			Request.Method = Method;
+
+			string BodyText = null;
 			if (BodyObject != null)
 			{
-				string BodyText = new JavaScriptSerializer().Serialize(BodyObject);
+				BodyText = new JavaScriptSerializer().Serialize(BodyObject);
 				byte[] BodyData = Encoding.UTF8.GetBytes(BodyText);
 				using (Stream RequestStream = Request.GetRequestStream())
 				{
@@ -417,7 +419,17 @@ namespace MetadataTool
 			}
 
 			// Read the response
-			return (HttpWebResponse)Request.GetResponse();
+			try
+			{
+				return (HttpWebResponse)Request.GetResponse();
+			}
+			catch(Exception Ex)
+			{
+				ExceptionUtils.AddContext(Ex, String.Format("Url: {0}", Url));
+				ExceptionUtils.AddContext(Ex, String.Format("Method: {0}", Method));
+				ExceptionUtils.AddContext(Ex, String.Format("Body: {0}", BodyText));
+				throw;
+			}
 		}
 
 		/// <summary>
