UGS: Fix formatting of diagnostic message in MetadataTool.

#rnx

diff --git a/Engine/Source/Programs/UnrealGameSync/MetadataTool/Program.cs b/Engine/Source/Programs/UnrealGameSync/MetadataTool/Program.cs
--- a/Engine/Source/Programs/UnrealGameSync/MetadataTool/Program.cs
+++ b/Engine/Source/Programs/UnrealGameSync/MetadataTool/Program.cs
@@ -63,13 +63,14 @@ namespace MetadataTool
 							Response = new JavaScriptSerializer().Deserialize<Dictionary<string, object>>(ResponseContent);
 						}
 
-						StringBuilder Message = new StringBuilder(WebEx.Message);
+						StringBuilder Message = new StringBuilder("Additional context:");
+						Message.AppendFormat("\n  Summary:\n    {0}", WebEx.Message);
 						foreach(KeyValuePair<string, object> Pair in Response)
 						{
-							Message.AppendFormat("\n{0}:\n  {1}", Pair.Key, Pair.Value.ToString().Replace("\n", "\n  "));
+							Message.AppendFormat("\n  {0}:\n    {1}", Pair.Key, Pair.Value.ToString().Replace("\n", "\n    "));
 						}
 
-						Log.TraceError("{0}", Message.ToString());
+						Log.TraceInformation("{0}", Message.ToString());
 						return 1;
 					}
 					catch
