UGS: Ignore a few warnings to fix problems with Platforms/... folder not existing.

#rnx

diff --git a/Engine/Source/Programs/UnrealGameSync/UnrealGameSync/Perforce.cs b/Engine/Source/Programs/UnrealGameSync/UnrealGameSync/Perforce.cs
--- a/Engine/Source/Programs/UnrealGameSync/UnrealGameSync/Perforce.cs
+++ b/Engine/Source/Programs/UnrealGameSync/UnrealGameSync/Perforce.cs
@@ -822,7 +822,7 @@ namespace UnrealGameSync
 			}
 
 			List<string> Lines;
-			if(!RunCommand(Arguments, out Lines, CommandOptions.None, Log))
+			if(!RunCommand(Arguments, out Lines, CommandOptions.IgnoreFilesNotInClientViewError, Log))
 			{
 				Changes = null;
 				return false;
@@ -1161,7 +1161,7 @@ namespace UnrealGameSync
 
 		public bool Have(string Filter, out List<PerforceFileRecord> FileRecords, TextWriter Log)
 		{
-			return RunCommand(String.Format("have \"{0}\"", Filter), out FileRecords, CommandOptions.IgnoreFilesNotOnClientError, Log);
+			return RunCommand(String.Format("have \"{0}\"", Filter), out FileRecords, CommandOptions.IgnoreFilesNotOnClientError | CommandOptions.IgnoreFilesNotInClientViewError, Log);
 		}
 
 		public bool Have(string Filter, Action<PerforceFileRecord> HaveOutput, TextWriter Log)
@@ -1169,13 +1169,13 @@ namespace UnrealGameSync
 			using(PerforceTagRecordParser Parser = new PerforceTagRecordParser(x => HaveOutput(new PerforceFileRecord(x))))
 			{
 				string CommandLine = String.Format("-ztag have \"{0}\"", Filter);
-				return RunCommand(CommandLine.ToString(), null, Line => FilterTaggedOutput(Line, Parser, Log), CommandOptions.IgnoreFilesNotOnClientError, Log);
+				return RunCommand(CommandLine.ToString(), null, Line => FilterTaggedOutput(Line, Parser, Log), CommandOptions.IgnoreFilesNotOnClientError | CommandOptions.IgnoreFilesNotInClientViewError, Log);
 			}
 		}
 
 		public bool Stat(string Filter, out List<PerforceFileRecord> FileRecords, TextWriter Log)
 		{
-			return RunCommand(String.Format("fstat \"{0}\"", Filter), out FileRecords, CommandOptions.IgnoreFilesNotOnClientError | CommandOptions.IgnoreNoSuchFilesError | CommandOptions.IgnoreProtectedNamespaceError, Log);
+			return RunCommand(String.Format("fstat \"{0}\"", Filter), out FileRecords, CommandOptions.IgnoreFilesNotOnClientError | CommandOptions.IgnoreNoSuchFilesError | CommandOptions.IgnoreProtectedNamespaceError | CommandOptions.IgnoreFilesNotInClientViewError, Log);
 		}
 
 		public bool Stat(string Options, List<string> Files, out List<PerforceFileRecord> FileRecords, TextWriter Log)
@@ -1189,7 +1189,7 @@ namespace UnrealGameSync
 			{
 				Arguments.AppendFormat(" \"{0}\"", File);
 			}
-			return RunCommand(Arguments.ToString(), out FileRecords, CommandOptions.IgnoreFilesNotOnClientError | CommandOptions.IgnoreNoSuchFilesError | CommandOptions.IgnoreProtectedNamespaceError, Log);
+			return RunCommand(Arguments.ToString(), out FileRecords, CommandOptions.IgnoreFilesNotOnClientError | CommandOptions.IgnoreNoSuchFilesError | CommandOptions.IgnoreProtectedNamespaceError | CommandOptions.IgnoreFilesNotInClientViewError, Log);
 		}
 
 		public bool Sync(string Filter, TextWriter Log)
@@ -1299,7 +1299,7 @@ namespace UnrealGameSync
 			using(PerforceTagRecordParser Parser = new PerforceTagRecordParser(x => SyncOutput(new PerforceFileRecord(x))))
 			{
 				string CommandLine = String.Format("-ztag sync -n {0}@{1}{2}", Filter, bOnlyFilesInThisChange? "=" : "", ChangeNumber);
-				return RunCommand(CommandLine.ToString(), null, Line => FilterTaggedOutput(Line, Parser, Log), CommandOptions.NoFailOnErrors | CommandOptions.IgnoreFilesUpToDateError | CommandOptions.IgnoreExitCode | CommandOptions.IgnoreNoSuchFilesError, Log);
+				return RunCommand(CommandLine.ToString(), null, Line => FilterTaggedOutput(Line, Parser, Log), CommandOptions.NoFailOnErrors | CommandOptions.IgnoreFilesUpToDateError | CommandOptions.IgnoreExitCode | CommandOptions.IgnoreNoSuchFilesError | CommandOptions.IgnoreFilesNotInClientViewError, Log);
 			}
 		}
 
@@ -1320,7 +1320,7 @@ namespace UnrealGameSync
 
 		public bool GetUnresolvedFiles(string Filter, out List<PerforceFileRecord> FileRecords, TextWriter Log)
 		{
-			return RunCommand(String.Format("fstat -Ru \"{0}\"", Filter), out FileRecords, CommandOptions.IgnoreNoSuchFilesError | CommandOptions.IgnoreFilesNotOpenedOnThisClientError, Log);
+			return RunCommand(String.Format("fstat -Ru \"{0}\"", Filter), out FileRecords, CommandOptions.IgnoreNoSuchFilesError | CommandOptions.IgnoreFilesNotOpenedOnThisClientError | CommandOptions.IgnoreFilesNotInClientViewError, Log);
 		}
 
 		public bool AutoResolveFile(string File, TextWriter Log)
