UGS: Fix expand item not correctly being added to the correct group.

#rnx

diff --git a/Engine/Source/Programs/UnrealGameSync/UnrealGameSync/Controls/WorkspaceControl.cs b/Engine/Source/Programs/UnrealGameSync/UnrealGameSync/Controls/WorkspaceControl.cs
--- a/Engine/Source/Programs/UnrealGameSync/UnrealGameSync/Controls/WorkspaceControl.cs
+++ b/Engine/Source/Programs/UnrealGameSync/UnrealGameSync/Controls/WorkspaceControl.cs
@@ -1178,19 +1178,41 @@ namespace UnrealGameSync
 				BuildList_AddItem(Change);
 			}
 
-			// Fixup the expand row
-			ListViewGroup Group = (BuildList.Groups.Count > 0) ? BuildList.Groups[BuildList.Groups.Count - 1] : null;
+			// Figure out which group the expand item should be in
+			ListViewGroup NewExpandItemGroup = null;
+			for (int Idx = BuildList.Items.Count - 1; Idx >= 0; Idx--)
+			{
+				ListViewItem Item = BuildList.Items[Idx];
+				if(Item != ExpandItem)
+				{
+					NewExpandItemGroup = Item.Group;
+					break;
+				}
+			}
+
+			// Remove the expand row if it's in the wrong place
 			if (ExpandItem != null)
 			{
-				if(!bShowExpandItem || ExpandItem.Group != Group)
+				if(!bShowExpandItem || ExpandItem.Group != NewExpandItemGroup)
 				{
 					BuildList.Items.Remove(ExpandItem);
 					ExpandItem = null;
 				}
 			}
-			if(bShowExpandItem && ExpandItem == null)
+
+			// Remove any empty groups
+			for (int Idx = BuildList.Groups.Count - 1; Idx >= 0; Idx--)
+			{
+				if (BuildList.Groups[Idx].Items.Count == 0)
+				{
+					BuildList.Groups.RemoveAt(Idx);
+				}
+			}
+
+			// Add the expand row back in
+			if (bShowExpandItem && ExpandItem == null)
 			{
-				ExpandItem = new ListViewItem(Group);
+				ExpandItem = new ListViewItem(NewExpandItemGroup);
 				ExpandItem.Tag = null;
 				ExpandItem.Selected = false;
 				ExpandItem.Text = "";
