UGS: Truncate the length of the test on the issue alert window.

#rnx

diff --git a/Engine/Source/Programs/UnrealGameSync/UnrealGameSync/Forms/IssueAlertWindow.cs b/Engine/Source/Programs/UnrealGameSync/UnrealGameSync/Forms/IssueAlertWindow.cs
--- a/Engine/Source/Programs/UnrealGameSync/UnrealGameSync/Forms/IssueAlertWindow.cs
+++ b/Engine/Source/Programs/UnrealGameSync/UnrealGameSync/Forms/IssueAlertWindow.cs
@@ -117,12 +117,20 @@ namespace UnrealGameSync
 
 			Issue = NewIssue;
 
-			if(NewIssue.Summary != SummaryLabel.Text || OwnerText != OwnerLabel.Text || Reason != NewReason || bIsWarning != bNewIsWarning || bStrongAlert != bNewStrongAlert)
+			string Summary = NewIssue.Summary;
+
+			int MaxLength = 128;
+			if(Summary.Length > MaxLength)
+			{
+				Summary = Summary.Substring(0, MaxLength).TrimEnd() + "...";
+			}
+
+			if(Summary != SummaryLabel.Text || OwnerText != OwnerLabel.Text || Reason != NewReason || bIsWarning != bNewIsWarning || bStrongAlert != bNewStrongAlert)
 			{
 				Rectangle PrevBounds = Bounds;
 				SuspendLayout();
 
-				SummaryLabel.Text = NewIssue.Summary;
+				SummaryLabel.Text = Summary;
 				OwnerLabel.Text = OwnerText;
 
 				bool bForceUpdateButtons = false;
