UGS: Filter the issues list by default.

#rnx

diff --git a/Engine/Source/Programs/UnrealGameSync/UnrealGameSync/Controls/WorkspaceControl.cs b/Engine/Source/Programs/UnrealGameSync/UnrealGameSync/Controls/WorkspaceControl.cs
--- a/Engine/Source/Programs/UnrealGameSync/UnrealGameSync/Controls/WorkspaceControl.cs
+++ b/Engine/Source/Programs/UnrealGameSync/UnrealGameSync/Controls/WorkspaceControl.cs
@@ -3177,7 +3177,9 @@ namespace UnrealGameSync
 
 		private void BuildHealthContextMenu_Browse_Click(object sender, EventArgs e)
 		{
-			IssueBrowserWindow.Show(ParentForm, IssueMonitor, Workspace.Perforce.ServerAndPort, Workspace.Perforce.UserName, GetServerTimeOffset(), Log, StreamName);
+			string BuildHealthProject;
+			TryGetProjectSetting(PerforceMonitor.LatestProjectConfigFile, "BuildHealthProject", out BuildHealthProject);
+			IssueBrowserWindow.Show(ParentForm, IssueMonitor, Workspace.Perforce.ServerAndPort, Workspace.Perforce.UserName, GetServerTimeOffset(), Log, StreamName, BuildHealthProject);
 		}
 
 		private void BuildHealthContextMenu_Settings_Click(object sender, EventArgs e)
diff --git a/Engine/Source/Programs/UnrealGameSync/UnrealGameSync/Forms/IssueBrowserWindow.cs b/Engine/Source/Programs/UnrealGameSync/UnrealGameSync/Forms/IssueBrowserWindow.cs
--- a/Engine/Source/Programs/UnrealGameSync/UnrealGameSync/Forms/IssueBrowserWindow.cs
+++ b/Engine/Source/Programs/UnrealGameSync/UnrealGameSync/Forms/IssueBrowserWindow.cs
@@ -26,7 +26,7 @@ namespace UnrealGameSync
 		List<string> ProjectNames = new List<string>();
 		List<IssueData> Issues = new List<IssueData>();
 
-		public IssueBrowserWindow(IssueMonitor IssueMonitor, string ServerAndPort, string UserName, TimeSpan? ServerTimeOffset, TextWriter Log, string CurrentStream)
+		public IssueBrowserWindow(IssueMonitor IssueMonitor, string ServerAndPort, string UserName, TimeSpan? ServerTimeOffset, TextWriter Log, string CurrentStream, string FilterProjectName)
 		{
 			this.IssueMonitor = IssueMonitor;
 			this.ServerAndPort = ServerAndPort;
@@ -34,6 +34,7 @@ namespace UnrealGameSync
 			this.ServerTimeOffset = ServerTimeOffset;
 			this.Log = Log;
 			this.CurrentStream = CurrentStream;
+			this.FilterProjectName = FilterProjectName;
 
 			IssueMonitor.AddRef();
 
@@ -135,7 +136,7 @@ namespace UnrealGameSync
 			IssueListView.EndUpdate();
 
 			// Update the maximum number of results
-			StatusLabel.Text = (IssueListView.Items.Count == Issues.Count)? String.Format("Showing {0} results.", Issues.Count) : String.Format("Showing {0}/{1} results.", IssueListView.Items.Count, Issues.Count);
+			StatusLabel.Text = (IssueListView.Items.Count == Issues.Count)? String.Format("Showing {0} results matching project '{1}'.", Issues.Count, FilterProjectName) : String.Format("Showing {0}/{1} results matching project '{2}'.", IssueListView.Items.Count, Issues.Count, FilterProjectName);
 		}
 
 		static string FormatIssueDateTime(DateTime DateTime, DateTime Midnight)
@@ -152,12 +153,12 @@ namespace UnrealGameSync
 
 		static List<IssueBrowserWindow> ExistingWindows = new List<IssueBrowserWindow>();
 
-		public static void Show(Form Owner, IssueMonitor IssueMonitor, string ServerAndPort, string UserName, TimeSpan? ServerTimeOffset, TextWriter Log, string CurrentStream)
+		public static void Show(Form Owner, IssueMonitor IssueMonitor, string ServerAndPort, string UserName, TimeSpan? ServerTimeOffset, TextWriter Log, string CurrentStream, string FilterProjectName)
 		{
 			IssueBrowserWindow Window = ExistingWindows.FirstOrDefault(x => x.IssueMonitor == IssueMonitor);
 			if(Window == null)
 			{
-				Window = new IssueBrowserWindow(IssueMonitor, ServerAndPort, UserName, ServerTimeOffset, Log, CurrentStream);
+				Window = new IssueBrowserWindow(IssueMonitor, ServerAndPort, UserName, ServerTimeOffset, Log, CurrentStream, FilterProjectName);
 				Window.Owner = Owner;
 				Window.StartPosition = FormStartPosition.Manual;
 				Window.Location = new Point(Owner.Location.X + (Owner.Width - Window.Width) / 2, Owner.Location.Y + (Owner.Height - Window.Height) / 2);
