UGS: Periodically update the issue browser window.

diff --git a/Engine/Source/Programs/UnrealGameSync/UnrealGameSync/Forms/IssueBrowserWindow.cs b/Engine/Source/Programs/UnrealGameSync/UnrealGameSync/Forms/IssueBrowserWindow.cs
--- a/Engine/Source/Programs/UnrealGameSync/UnrealGameSync/Forms/IssueBrowserWindow.cs
+++ b/Engine/Source/Programs/UnrealGameSync/UnrealGameSync/Forms/IssueBrowserWindow.cs
@@ -8,6 +8,7 @@ using System.Drawing;
 using System.IO;
 using System.Linq;
 using System.Text;
+using System.Threading;
 using System.Threading.Tasks;
 using System.Windows.Forms;
 
@@ -22,9 +23,13 @@ namespace UnrealGameSync
 		TextWriter Log;
 		string CurrentStream;
 		int MaxResults = 0;
+		int PendingMaxResults = 0;
 		string FilterProjectName;
 		List<string> ProjectNames = new List<string>();
 		List<IssueData> Issues = new List<IssueData>();
+		SynchronizationContext MainThreadSynchronizationContext;
+		Thread BackgroundThread;
+		bool bDisposed;
 
 		public IssueBrowserWindow(IssueMonitor IssueMonitor, string ServerAndPort, string UserName, TimeSpan? ServerTimeOffset, TextWriter Log, string CurrentStream, string FilterProjectName)
 		{
@@ -35,6 +40,7 @@ namespace UnrealGameSync
 			this.Log = Log;
 			this.CurrentStream = CurrentStream;
 			this.FilterProjectName = FilterProjectName;
+			this.MainThreadSynchronizationContext = SynchronizationContext.Current;
 
 			IssueMonitor.AddRef();
 
@@ -53,6 +59,8 @@ namespace UnrealGameSync
 
 			IssueMonitor.Release();
 
+			bDisposed = true;
+
 			base.Dispose(disposing);
 		}
 
@@ -61,52 +69,61 @@ namespace UnrealGameSync
 			FetchMoreResults();
 		}
 
-		class QueryIssuesTask : IModalTask
+		private void FetchMoreResults()
 		{
-			string ApiUrl;
-			int MaxResults;
-			public List<IssueData> Issues;
+			PendingMaxResults = MaxResults + 100;
+			CheckToStartBackgroundThread();
+		}
 
-			public QueryIssuesTask(string ApiUrl, int MaxResults)
+		void FetchIssues(int NewMaxResults)
+		{
+			try
 			{
-				this.ApiUrl = ApiUrl;
-				this.MaxResults = MaxResults;
+				List<IssueData> NewIssues = RESTApi.GET<List<IssueData>>(IssueMonitor.ApiUrl, String.Format("issues?includeresolved=true&maxresults={0}", NewMaxResults));
+				MainThreadSynchronizationContext.Post((o) => { if (!bDisposed) { FetchIssuesSuccess(NewMaxResults, NewIssues); } }, null);
 			}
-
-			public bool Run(out string ErrorMessage)
+			catch(Exception Ex)
 			{
-				Issues = RESTApi.GET<List<IssueData>>(ApiUrl, String.Format("issues?includeresolved=true&maxresults={0}", MaxResults));
-
-				ErrorMessage = null;
-				return true;
+				MainThreadSynchronizationContext.Post((o) => { if (!bDisposed) { FetchIssuesFailure(Ex); } }, null);
 			}
 		}
 
-		private void FetchMoreResults()
+		void FetchIssuesSuccess(int NewMaxResults, List<IssueData> NewIssues)
+		{
+			MaxResults = NewMaxResults;
+
+			// Update the list of project names
+			Issues = NewIssues;
+			ProjectNames = NewIssues.Select(x => x.Project).Distinct().OrderBy(x => x).ToList();
+			UpdateIssueList();
+
+			BackgroundThread = null;
+			CheckToStartBackgroundThread();
+		}
+
+		void FetchIssuesFailure(Exception Ex)
 		{
-			int NewMaxResults = MaxResults + 100;
-			QueryIssuesTask Task = new QueryIssuesTask(IssueMonitor.ApiUrl, NewMaxResults);
+			StatusLabel.Text = String.Format("Unable to fetch issues ({0})", Ex.Message);
+			BackgroundThread = null;
+			CheckToStartBackgroundThread();
+		}
 
-			// Execute the task
-			string ErrorMessage;
-			ModalTaskResult Result = ModalTask.Execute(this, Task, "Querying issues", "Fetching data, please wait...", out ErrorMessage);
-			if(Result != ModalTaskResult.Succeeded)
+		void CheckToStartBackgroundThread()
+		{
+			if(PendingMaxResults != MaxResults)
 			{
-				if(Result != ModalTaskResult.Aborted)
-				{
-					StatusLabel.Text = "Unable to fetch issues.";
-					MessageBox.Show(ErrorMessage);
-				}
-				return;
+				StartBackgroundThread();
 			}
+		}
 
-			// Update the list of project names
-			Issues = Task.Issues;
-			ProjectNames = Task.Issues.Select(x => x.Project).Distinct().OrderBy(x => x).ToList();
-
-			// Populate the list control
-			UpdateIssueList();
-			MaxResults = NewMaxResults;
+		void StartBackgroundThread()
+		{
+			if(BackgroundThread == null)
+			{
+				int NewMaxResultsCopy = PendingMaxResults;
+				BackgroundThread = new Thread(() => FetchIssues(NewMaxResultsCopy));
+				BackgroundThread.Start();
+			}
 		}
 
 		void UpdateIssueList()
@@ -115,30 +132,74 @@ namespace UnrealGameSync
 			DateTime Now = DateTime.Now;
 			DateTime Midnight = (Now - Now.TimeOfDay).ToUniversalTime();
 
-			// Fetch the new issues
+			// Update the table
+			int ItemIdx = 0;
 			IssueListView.BeginUpdate();
-			IssueListView.Items.Clear();
 			foreach(IssueData Issue in Issues)
 			{
 				if(FilterProjectName == null || Issue.Project == FilterProjectName)
 				{
-					ListViewItem Item = new ListViewItem("");
-					Item.SubItems.Add(Issue.Id.ToString());
-					Item.SubItems.Add(Issue.Project);
-					Item.SubItems.Add(FormatIssueDateTime(Issue.CreatedAt.ToLocalTime(), Midnight));
-					Item.SubItems.Add(Issue.ResolvedAt.HasValue? FormatIssueDateTime(Issue.ResolvedAt.Value.ToLocalTime(), Midnight) : "Unresolved");
-					Item.SubItems.Add((Issue.Owner == null)? "-" : Utility.FormatUserName(Issue.Owner));
-					Item.SubItems.Add(Issue.Summary);
-					Item.Tag = Issue;
-					IssueListView.Items.Add(Item);
+					for(;;)
+					{
+						if(ItemIdx == IssueListView.Items.Count)
+						{
+							IssueList_InsertItem(ItemIdx, Issue, Midnight);
+							break;
+						}
+
+						ListViewItem ExistingItem = IssueListView.Items[ItemIdx];
+						IssueData ExistingIssue = (IssueData)ExistingItem.Tag;
+						if(ExistingIssue == null || ExistingIssue.Id < Issue.Id)
+						{
+							IssueList_InsertItem(ItemIdx, Issue, Midnight);
+							break;
+						}
+						else if(ExistingIssue.Id == Issue.Id)
+						{
+							IssueList_UpdateItem(ExistingItem, Issue, Midnight);
+							break;
+						}
+						else
+						{
+							IssueListView.Items.RemoveAt(ItemIdx);
+							continue;
+						}
+					}
+					ItemIdx++;
 				}
 			}
+			while(ItemIdx < IssueListView.Items.Count)
+			{
+				IssueListView.Items.RemoveAt(ItemIdx);
+			}
 			IssueListView.EndUpdate();
 
 			// Update the maximum number of results
 			StatusLabel.Text = (IssueListView.Items.Count == Issues.Count)? String.Format("Showing {0} results matching project '{1}'.", Issues.Count, FilterProjectName) : String.Format("Showing {0}/{1} results matching project '{2}'.", IssueListView.Items.Count, Issues.Count, FilterProjectName);
 		}
 
+		void IssueList_InsertItem(int ItemIdx, IssueData Issue, DateTime Midnight)
+		{
+			ListViewItem Item = new ListViewItem("");
+			for(int Idx = 0; Idx < IssueListView.Columns.Count - 1; Idx++)
+			{
+				Item.SubItems.Add("");
+			}
+			Item.Tag = Issue;
+			IssueList_UpdateItem(Item, Issue, Midnight);
+			IssueListView.Items.Insert(ItemIdx, Item);
+		}
+
+		void IssueList_UpdateItem(ListViewItem Item, IssueData Issue, DateTime Midnight)
+		{
+			Item.SubItems[IdHeader.Index].Text = Issue.Id.ToString();
+			Item.SubItems[ProjectHeader.Index].Text = Issue.Project;
+			Item.SubItems[CreatedHeader.Index].Text = FormatIssueDateTime(Issue.CreatedAt.ToLocalTime(), Midnight);
+			Item.SubItems[ResolvedHeader.Index].Text = Issue.ResolvedAt.HasValue ? FormatIssueDateTime(Issue.ResolvedAt.Value.ToLocalTime(), Midnight) : "Unresolved";
+			Item.SubItems[OwnerHeader.Index].Text = (Issue.Owner == null) ? "-" : Utility.FormatUserName(Issue.Owner);
+			Item.SubItems[DescriptionHeader.Index].Text = Issue.Summary;
+		}
+
 		static string FormatIssueDateTime(DateTime DateTime, DateTime Midnight)
 		{
 			if(DateTime > Midnight)
@@ -290,5 +351,10 @@ namespace UnrealGameSync
 				IssueListView.DrawNormalSubItem(e);
 			}
 		}
+
+		private void RefreshIssuesTimer_Tick(object sender, EventArgs e)
+		{
+			StartBackgroundThread();
+		}
 	}
 }
