UGS: Increase update interval to 60s.

diff --git a/Engine/Source/Programs/UnrealGameSync/UnrealGameSync/Forms/MainWindow.cs b/Engine/Source/Programs/UnrealGameSync/UnrealGameSync/Forms/MainWindow.cs
--- a/Engine/Source/Programs/UnrealGameSync/UnrealGameSync/Forms/MainWindow.cs
+++ b/Engine/Source/Programs/UnrealGameSync/UnrealGameSync/Forms/MainWindow.cs
@@ -1012,7 +1012,7 @@ namespace UnrealGameSync
 			if(IssueMonitor == null)
 			{
 				string LogFileName = Path.Combine(DataFolder, String.Format("IssueMonitor-{0}.log", ProjectSettings.PerforceClient.UserName));
-				IssueMonitor = new IssueMonitor(ApiUrl, ProjectSettings.PerforceClient.UserName, TimeSpan.FromSeconds(5.0), LogFileName);
+				IssueMonitor = new IssueMonitor(ApiUrl, ProjectSettings.PerforceClient.UserName, TimeSpan.FromSeconds(60.0), LogFileName);
 				IssueMonitor.OnIssuesChanged += IssueMonitor_OnUpdateAsync;
 				IssueMonitors.Add(IssueMonitor);
 			}
