UGS: Add a base class for code pattern matchers, which filters out content changes from the list of suspected causers.

#rnx

diff --git a/Engine/Source/Programs/UnrealGameSync/MetadataTool/BuildHealth/ChangeInfo.cs b/Engine/Source/Programs/UnrealGameSync/MetadataTool/BuildHealth/ChangeInfo.cs
--- a/Engine/Source/Programs/UnrealGameSync/MetadataTool/BuildHealth/ChangeInfo.cs
+++ b/Engine/Source/Programs/UnrealGameSync/MetadataTool/BuildHealth/ChangeInfo.cs
@@ -16,5 +16,6 @@ namespace MetadataTool
 	{
 		public ChangesRecord Record;
 		public List<int> SourceChanges = new List<int>();
+		public DescribeRecord CachedDescribeRecord;
 	}
 }
diff --git a/Engine/Source/Programs/UnrealGameSync/MetadataTool/BuildHealth/Matchers/CompilePatternMatcher.cs b/Engine/Source/Programs/UnrealGameSync/MetadataTool/BuildHealth/Matchers/CompilePatternMatcher.cs
--- a/Engine/Source/Programs/UnrealGameSync/MetadataTool/BuildHealth/Matchers/CompilePatternMatcher.cs
+++ b/Engine/Source/Programs/UnrealGameSync/MetadataTool/BuildHealth/Matchers/CompilePatternMatcher.cs
@@ -9,7 +9,7 @@ using System.Threading.Tasks;
 
 namespace MetadataTool
 {
-	class CompilePatternMatcher : PatternMatcher
+	class CompilePatternMatcher : GenericCodePatternMatcher
 	{
 		static string[] SourceFileExtensions =
 		{
diff --git a/Engine/Source/Programs/UnrealGameSync/MetadataTool/BuildHealth/Matchers/CopyrightNoticeMatcher.cs b/Engine/Source/Programs/UnrealGameSync/MetadataTool/BuildHealth/Matchers/CopyrightNoticeMatcher.cs
--- a/Engine/Source/Programs/UnrealGameSync/MetadataTool/BuildHealth/Matchers/CopyrightNoticeMatcher.cs
+++ b/Engine/Source/Programs/UnrealGameSync/MetadataTool/BuildHealth/Matchers/CopyrightNoticeMatcher.cs
@@ -9,7 +9,7 @@ using System.Threading.Tasks;
 
 namespace MetadataTool
 {
-	class CopyrightNoticeMatcher : PatternMatcher
+	class CopyrightNoticeMatcher : GenericCodePatternMatcher
 	{
 		public override string Category => "Copyright";
 
diff --git a/Engine/Source/Programs/UnrealGameSync/MetadataTool/BuildHealth/Matchers/GenericCodePatternMatcher.cs b/Engine/Source/Programs/UnrealGameSync/MetadataTool/BuildHealth/Matchers/GenericCodePatternMatcher.cs
new file mode 100644
--- /dev/null
+++ b/Engine/Source/Programs/UnrealGameSync/MetadataTool/BuildHealth/Matchers/GenericCodePatternMatcher.cs
@@ -0,0 +1,51 @@
+// Copyright 1998-2019 Epic Games, Inc. All Rights Reserved.
+
+using System;
+using System.Collections.Generic;
+using System.Linq;
+using System.Text;
+using System.Threading.Tasks;
+using Tools.DotNETCommon.Perforce;
+
+namespace MetadataTool
+{
+	/// <summary>
+	/// Base class for pattern matchers that match errors in source code
+	/// </summary>
+	abstract class GenericCodePatternMatcher : PatternMatcher
+	{
+		/// <summary>
+		/// Set of extensions to treat as code
+		/// </summary>
+		static readonly string[] CodeExtensions =
+		{
+			".c",
+			".cc",
+			".cpp",
+			".m",
+			".mm",
+			".rc",
+			".cs",
+			".h",
+			".hpp",
+			".inl",
+			".uproject",
+			".uplugin",
+			".ini",
+		};
+
+		/// <summary>
+		/// Finds any suspected causers for a particular failure. Excludes any changes that don't contain code.
+		/// </summary>
+		/// <param name="Perforce">The perforce connection</param>
+		/// <param name="Issue">The build issue</param>
+		/// <param name="Changes">List of changes since the issue first occurred.</param>
+		/// <returns>List of changes which are causers for the issue</returns>
+		public override List<ChangeInfo> FindCausers(PerforceConnection Perforce, BuildHealthIssue Issue, IReadOnlyList<ChangeInfo> Changes)
+		{
+			List<ChangeInfo> Causers = base.FindCausers(Perforce, Issue, Changes);
+			Causers.RemoveAll(x => !ContainsAnyFileWithExtension(Perforce, x, CodeExtensions));
+			return Causers;
+		}
+	}
+}
diff --git a/Engine/Source/Programs/UnrealGameSync/MetadataTool/BuildHealth/Matchers/UndefinedSymbolPatternMatcher.cs b/Engine/Source/Programs/UnrealGameSync/MetadataTool/BuildHealth/Matchers/UndefinedSymbolPatternMatcher.cs
--- a/Engine/Source/Programs/UnrealGameSync/MetadataTool/BuildHealth/Matchers/UndefinedSymbolPatternMatcher.cs
+++ b/Engine/Source/Programs/UnrealGameSync/MetadataTool/BuildHealth/Matchers/UndefinedSymbolPatternMatcher.cs
@@ -9,7 +9,7 @@ using System.Threading.Tasks;
 
 namespace MetadataTool
 {
-	class UndefinedSymbolPatternMatcher : PatternMatcher
+	class UndefinedSymbolPatternMatcher : GenericCodePatternMatcher
 	{
 		public override string Category => "UndefinedSymbol";
 
diff --git a/Engine/Source/Programs/UnrealGameSync/MetadataTool/BuildHealth/PatternMatcher.cs b/Engine/Source/Programs/UnrealGameSync/MetadataTool/BuildHealth/PatternMatcher.cs
--- a/Engine/Source/Programs/UnrealGameSync/MetadataTool/BuildHealth/PatternMatcher.cs
+++ b/Engine/Source/Programs/UnrealGameSync/MetadataTool/BuildHealth/PatternMatcher.cs
@@ -105,20 +105,23 @@ namespace MetadataTool
 		/// Filters all the likely causers from the list of changes since an issue was created
 		/// </summary>
 		/// <param name="Perforce">The perforce connection</param>
-		/// <param name="Fingerprint">Fingerprint for the issue</param>
+		/// <param name="Issue">The build issue</param>
 		/// <param name="Changes">List of changes since the issue first occurred.</param>
 		/// <returns>List of changes which are causers for the issue</returns>
 		public virtual List<ChangeInfo> FindCausers(PerforceConnection Perforce, BuildHealthIssue Issue, IReadOnlyList<ChangeInfo> Changes)
 		{
-			SortedSet<string> FileNamesWithoutPath = GetFileNamesWithoutPath(Issue.FileNames);
-
 			List<ChangeInfo> Causers = new List<ChangeInfo>();
-			foreach (ChangeInfo Change in Changes)
+
+			SortedSet<string> FileNamesWithoutPath = GetFileNamesWithoutPath(Issue.FileNames);
+			if (FileNamesWithoutPath.Count > 0)
 			{
-				DescribeRecord Description = Perforce.Describe(Change.Record.Number).Data;
-				if (ContainsFileNames(Description, FileNamesWithoutPath))
+				foreach (ChangeInfo Change in Changes)
 				{
-					Causers.Add(Change);
+					DescribeRecord DescribeRecord = GetDescribeRecord(Perforce, Change);
+					if (ContainsFileNames(DescribeRecord, FileNamesWithoutPath))
+					{
+						Causers.Add(Change);
+					}
 				}
 			}
 
@@ -132,6 +135,42 @@ namespace MetadataTool
 			}
 		}
 
+		/// <summary>
+		/// Utility method to get the describe record for a change. Caches it on the ChangeInfo object as necessary.
+		/// </summary>
+		/// <param name="Perforce">The Perforce connection</param>
+		/// <param name="Change">The change to query</param>
+		public DescribeRecord GetDescribeRecord(PerforceConnection Perforce, ChangeInfo Change)
+		{
+			if(Change.CachedDescribeRecord == null)
+			{
+				Change.CachedDescribeRecord = Perforce.Describe(Change.Record.Number).Data;
+			}
+			return Change.CachedDescribeRecord;
+		}
+
+		/// <summary>
+		/// Tests whether a change is a code change
+		/// </summary>
+		/// <param name="Perforce">The Perforce connection</param>
+		/// <param name="Change">The change to query</param>
+		/// <returns>True if the change is a code change</returns>
+		public bool ContainsAnyFileWithExtension(PerforceConnection Perforce, ChangeInfo Change, string[] Extensions)
+		{
+			DescribeRecord Record = GetDescribeRecord(Perforce, Change);
+			foreach(DescribeFileRecord File in Record.Files)
+			{
+				foreach(string Extension in Extensions)
+				{
+					if(File.DepotFile.EndsWith(Extension, StringComparison.OrdinalIgnoreCase))
+					{
+						return true;
+					}
+				}
+			}
+			return false;
+		}
+
 		/// <summary>
 		/// Determines if this change is a likely causer for an issue
 		/// </summary>
diff --git a/Engine/Source/Programs/UnrealGameSync/MetadataTool/MetadataTool.csproj b/Engine/Source/Programs/UnrealGameSync/MetadataTool/MetadataTool.csproj
--- a/Engine/Source/Programs/UnrealGameSync/MetadataTool/MetadataTool.csproj
+++ b/Engine/Source/Programs/UnrealGameSync/MetadataTool/MetadataTool.csproj
@@ -54,6 +54,7 @@
     <Compile Include="BuildHealth\ChangeInfo.cs" />
     <Compile Include="BuildHealth\Matchers\CompilePatternMatcher.cs" />
     <Compile Include="BuildHealth\Matchers\CopyrightNoticeMatcher.cs" />
+    <Compile Include="BuildHealth\Matchers\GenericCodePatternMatcher.cs" />
     <Compile Include="BuildHealth\Matchers\UndefinedSymbolPatternMatcher.cs" />
     <Compile Include="BuildHealth\PatternMatcher.cs" />
     <Compile Include="CommandHandler.cs" />
