UGS: Reduce the number of files listed in the summary for compile errors. Only show the actual files containing errors.

#rnx

diff --git a/Engine/Source/Programs/UnrealGameSync/MetadataTool/BuildHealth/Matchers/CompilePatternMatcher.cs b/Engine/Source/Programs/UnrealGameSync/MetadataTool/BuildHealth/Matchers/CompilePatternMatcher.cs
--- a/Engine/Source/Programs/UnrealGameSync/MetadataTool/BuildHealth/Matchers/CompilePatternMatcher.cs
+++ b/Engine/Source/Programs/UnrealGameSync/MetadataTool/BuildHealth/Matchers/CompilePatternMatcher.cs
@@ -28,8 +28,22 @@ namespace MetadataTool
 
 		public override bool TryMatch(InputJob Job, InputJobStep JobStep, InputDiagnostic Diagnostic, List<BuildHealthIssue> Issues)
 		{
-			// Find any files in compiler output format
-			HashSet<string> SourceFileNames = new HashSet<string>(StringComparer.OrdinalIgnoreCase);
+			// Find a list of source files with errors
+			HashSet<string> ErrorFileNames = new HashSet<string>(StringComparer.OrdinalIgnoreCase);
+			foreach (Match FileMatch in Regex.Matches(Diagnostic.Message, @"^\s*((?:[A-Za-z]:)?[^\s(:]+)[\(:]\d[\s\d:\)]+(?:warning|error)", RegexOptions.Multiline))
+			{
+				if (FileMatch.Success)
+				{
+					string FileName = GetNormalizedFileName(FileMatch.Groups[1].Value, JobStep.BaseDirectory);
+					if (SourceFileExtensions.Any(x => FileName.EndsWith(x, StringComparison.OrdinalIgnoreCase)))
+					{
+						ErrorFileNames.Add(FileName);
+					}
+				}
+			}
+
+			// Find any referenced files in compiler output format
+			HashSet<string> ReferencedFileNames = new HashSet<string>(StringComparer.OrdinalIgnoreCase);
 			foreach (Match FileMatch in Regex.Matches(Diagnostic.Message, @"^\s*(?:In file included from\s*)?((?:[A-Za-z]:)?[^\s(:]+)[\(:]\d", RegexOptions.Multiline))
 			{
 				if (FileMatch.Success)
@@ -37,16 +51,17 @@ namespace MetadataTool
 					string FileName = GetNormalizedFileName(FileMatch.Groups[1].Value, JobStep.BaseDirectory);
 					if (SourceFileExtensions.Any(x => FileName.EndsWith(x, StringComparison.OrdinalIgnoreCase)))
 					{
-						SourceFileNames.Add(FileName);
+						ReferencedFileNames.Add(FileName);
 					}
 				}
 			}
 
 			// If we found any source files, create a diagnostic category for them
-			if (SourceFileNames.Count > 0)
+			if (ReferencedFileNames.Count > 0)
 			{
 				BuildHealthIssue Issue = new BuildHealthIssue(Category, Job.Url, new BuildHealthDiagnostic(JobStep.Name, ShortenPaths(Diagnostic.Message), Diagnostic.Url));
-				Issue.FileNames.UnionWith(SourceFileNames);
+				Issue.FileNames.UnionWith(ReferencedFileNames);
+				Issue.Identifiers.UnionWith(GetSourceFileNames(ErrorFileNames));
 				Issues.Add(Issue);
 				return true;
 			}
@@ -63,14 +78,13 @@ namespace MetadataTool
 
 		public override string GetSummary(BuildHealthIssue Issue)
 		{
-			SortedSet<string> ShortFileNames = GetSourceFileNames(Issue.FileNames);
-			if (ShortFileNames.Count == 0)
+			if (Issue.Identifiers.Count == 0)
 			{
 				return "Compile errors";
 			}
 			else
 			{
-				return String.Format("Compile errors in {0}", String.Join(", ", ShortFileNames));
+				return String.Format("Compile errors in {0}", String.Join(", ", Issue.Identifiers));
 			}
 		}
 	}
