UGS: Move unmatched patterns into a different project.

#rnx

diff --git a/Engine/Source/Programs/UnrealGameSync/MetadataTool/BuildHealth/Matchers/DefaultPatternMatcher.cs b/Engine/Source/Programs/UnrealGameSync/MetadataTool/BuildHealth/Matchers/DefaultPatternMatcher.cs
--- a/Engine/Source/Programs/UnrealGameSync/MetadataTool/BuildHealth/Matchers/DefaultPatternMatcher.cs
+++ b/Engine/Source/Programs/UnrealGameSync/MetadataTool/BuildHealth/Matchers/DefaultPatternMatcher.cs
@@ -16,8 +16,12 @@ namespace MetadataTool
 
 		public override bool TryMatch(InputJob Job, InputJobStep JobStep, InputDiagnostic Diagnostic, List<BuildHealthIssue> Issues)
 		{
-			BuildHealthIssue Issue = new BuildHealthIssue(Job.Project, Category, Job.Url, new BuildHealthDiagnostic(JobStep.Name, Diagnostic.Message, Diagnostic.Url));
+			string DefaultProject = String.Format("{0} (Unmatched)", Job.Project);
+
+			BuildHealthIssue Issue = new BuildHealthIssue(DefaultProject, Category, Job.Url, new BuildHealthDiagnostic(JobStep.Name, Diagnostic.Message, Diagnostic.Url));
+			Issue.Identifiers.Add(Diagnostic.Message);
 			Issues.Add(Issue);
+			
 			return true;
 		}
 
@@ -26,6 +30,11 @@ namespace MetadataTool
 			return new List<ChangeInfo>();
 		}
 
+		public override bool CanMerge(BuildHealthIssue Source, BuildHealthIssue Target)
+		{
+			return Source.Category == Target.Category && Source.Identifiers.Any(x => Target.Identifiers.Contains(x));
+		}
+
 		public override string GetSummary(BuildHealthIssue Issue)
 		{
 			string Message = Issue.Diagnostics[0].Message;
