UGS: Add an option to alert when an issue is unacknowledged for a certain period of time.

#rnx

diff --git a/Engine/Source/Programs/UnrealGameSync/UnrealGameSync/Forms/IssueAlertWindow.cs b/Engine/Source/Programs/UnrealGameSync/UnrealGameSync/Forms/IssueAlertWindow.cs
--- a/Engine/Source/Programs/UnrealGameSync/UnrealGameSync/Forms/IssueAlertWindow.cs
+++ b/Engine/Source/Programs/UnrealGameSync/UnrealGameSync/Forms/IssueAlertWindow.cs
@@ -107,6 +107,10 @@ namespace UnrealGameSync
 					{
 						OwnerTextBuilder.AppendFormat(" by {0}", Utility.FormatUserName(NewIssue.NominatedBy));
 					}
+					if(!NewIssue.AcknowledgedAt.HasValue && (NewReason & IssueAlertReason.UnacknowledgedTimer) != 0)
+					{
+						OwnerTextBuilder.Append(" (not acknowledged)");
+					}
 					OwnerTextBuilder.Append(".");
 				}
 			}
diff --git a/Engine/Source/Programs/UnrealGameSync/UnrealGameSync/Forms/IssueSettingsWindow.cs b/Engine/Source/Programs/UnrealGameSync/UnrealGameSync/Forms/IssueSettingsWindow.cs
--- a/Engine/Source/Programs/UnrealGameSync/UnrealGameSync/Forms/IssueSettingsWindow.cs
+++ b/Engine/Source/Programs/UnrealGameSync/UnrealGameSync/Forms/IssueSettingsWindow.cs
@@ -33,7 +33,18 @@ namespace UnrealGameSync
 				NotifyUnassignedTextBox.Text = Settings.NotifyUnassignedMinutes.ToString();
 			}
 
-			if(Settings.NotifyUnresolvedMinutes < 0)
+			if (Settings.NotifyUnacknowledgedMinutes < 0)
+			{
+				NotifyUnacknowledgedCheckBox.Checked = false;
+				NotifyUnacknowledgedTextBox.Text = "5";
+			}
+			else
+			{
+				NotifyUnacknowledgedCheckBox.Checked = true;
+				NotifyUnacknowledgedTextBox.Text = Settings.NotifyUnacknowledgedMinutes.ToString();
+			}
+
+			if (Settings.NotifyUnresolvedMinutes < 0)
 			{
 				NotifyUnresolvedCheckBox.Checked = false;
 				NotifyUnresolvedTextBox.Text = "20";
@@ -50,6 +61,7 @@ namespace UnrealGameSync
 		private void UpdateEnabledTextBoxes()
 		{
 			NotifyUnassignedTextBox.Enabled = NotifyUnassignedCheckBox.Checked;
+			NotifyUnacknowledgedTextBox.Enabled = NotifyUnacknowledgedCheckBox.Checked;
 			NotifyUnresolvedTextBox.Enabled = NotifyUnresolvedCheckBox.Checked;
 		}
 
@@ -77,6 +89,18 @@ namespace UnrealGameSync
 				NewNotifyUnresolvedMinutes = NewNotifyUnresolvedMinutesValue;
 			}
 
+			int NewNotifyUnacknowledgedMinutes = -1;
+			if (NotifyUnacknowledgedCheckBox.Checked)
+			{
+				ushort NewNotifyUnacknowledgedMinutesValue;
+				if (!ushort.TryParse(NotifyUnacknowledgedTextBox.Text, out NewNotifyUnacknowledgedMinutesValue))
+				{
+					MessageBox.Show("Invalid time");
+					return;
+				}
+				NewNotifyUnacknowledgedMinutes = NewNotifyUnacknowledgedMinutesValue;
+			}
+
 			int NewNotifyUnassignedMinutes = -1;
 			if(NotifyUnassignedCheckBox.Checked)
 			{
@@ -90,6 +114,7 @@ namespace UnrealGameSync
 			}
 
 			Settings.NotifyUnresolvedMinutes = NewNotifyUnresolvedMinutes;
+			Settings.NotifyUnacknowledgedMinutes = NewNotifyUnacknowledgedMinutes;
 			Settings.NotifyUnassignedMinutes = NewNotifyUnassignedMinutes;
 			Settings.Save();
 
diff --git a/Engine/Source/Programs/UnrealGameSync/UnrealGameSync/Forms/MainWindow.cs b/Engine/Source/Programs/UnrealGameSync/UnrealGameSync/Forms/MainWindow.cs
--- a/Engine/Source/Programs/UnrealGameSync/UnrealGameSync/Forms/MainWindow.cs
+++ b/Engine/Source/Programs/UnrealGameSync/UnrealGameSync/Forms/MainWindow.cs
@@ -1248,7 +1248,7 @@ namespace UnrealGameSync
 				foreach(IssueData Issue in Issues)
 				{
 					IssueAlertReason Reason = 0;
-					if(Issue.FixChange <= 0)
+					if(Issue.FixChange <= 0 && !Issue.ResolvedAt.HasValue)
 					{
 						if(Issue.Owner == null)
 						{
@@ -1261,12 +1261,16 @@ namespace UnrealGameSync
 								Reason |= IssueAlertReason.UnassignedTimer;
 							}
 						}
-						else
+						else if(!Issue.AcknowledgedAt.HasValue)
 						{
-							if(String.Compare(Issue.Owner, IssueMonitor.UserName, StringComparison.OrdinalIgnoreCase) == 0 && !Issue.AcknowledgedAt.HasValue)
+							if(String.Compare(Issue.Owner, IssueMonitor.UserName, StringComparison.OrdinalIgnoreCase) == 0)
 							{
 								Reason |= IssueAlertReason.Owner;
 							}
+							else if(Settings.NotifyUnacknowledgedMinutes >= 0 && Issue.RetrievedAt - Issue.CreatedAt >= TimeSpan.FromMinutes(Settings.NotifyUnacknowledgedMinutes))
+							{
+								Reason |= IssueAlertReason.UnacknowledgedTimer;
+							}
 						}
 						if(Settings.NotifyUnresolvedMinutes >= 0 && Issue.RetrievedAt - Issue.CreatedAt >= TimeSpan.FromMinutes(Settings.NotifyUnresolvedMinutes))
 						{
diff --git a/Engine/Source/Programs/UnrealGameSync/UnrealGameSync/IssueMonitor.cs b/Engine/Source/Programs/UnrealGameSync/UnrealGameSync/IssueMonitor.cs
--- a/Engine/Source/Programs/UnrealGameSync/UnrealGameSync/IssueMonitor.cs
+++ b/Engine/Source/Programs/UnrealGameSync/UnrealGameSync/IssueMonitor.cs
@@ -64,7 +64,8 @@ namespace UnrealGameSync
 		Normal = 1,
 		Owner = 2,
 		UnassignedTimer = 4,
-		UnresolvedTimer = 8,
+		UnacknowledgedTimer = 8,
+		UnresolvedTimer = 16,
 	}
 
 	class IssueMonitor : IDisposable
diff --git a/Engine/Source/Programs/UnrealGameSync/UnrealGameSync/UserSettings.cs b/Engine/Source/Programs/UnrealGameSync/UnrealGameSync/UserSettings.cs
--- a/Engine/Source/Programs/UnrealGameSync/UnrealGameSync/UserSettings.cs
+++ b/Engine/Source/Programs/UnrealGameSync/UnrealGameSync/UserSettings.cs
@@ -231,6 +231,7 @@ namespace UnrealGameSync
 
 		// Notification settings
 		public int NotifyUnassignedMinutes;
+		public int NotifyUnacknowledgedMinutes;
 		public int NotifyUnresolvedMinutes;
 
 		// Project settings
@@ -378,6 +379,7 @@ namespace UnrealGameSync
 
 			// Notification settings
 			NotifyUnassignedMinutes = ConfigFile.GetValue("Notifications.NotifyUnassignedMinutes", -1);
+			NotifyUnacknowledgedMinutes = ConfigFile.GetValue("Notifications.NotifyUnacknowledgedMinutes", -1);
 			NotifyUnresolvedMinutes = ConfigFile.GetValue("Notifications.NotifyUnresolvedMinutes", -1);
 
 			// Perforce settings
@@ -633,11 +635,15 @@ namespace UnrealGameSync
 			// Notification settings
 			ConfigSection NotificationSection = ConfigFile.FindOrAddSection("Notifications");
 			NotificationSection.Clear();
-			if(NotifyUnassignedMinutes != -1)
+			if (NotifyUnassignedMinutes != -1)
 			{
 				NotificationSection.SetValue("NotifyUnassignedMinutes", NotifyUnassignedMinutes);
 			}
-			if(NotifyUnresolvedMinutes != -1)
+			if (NotifyUnacknowledgedMinutes != -1)
+			{
+				NotificationSection.SetValue("NotifyUnacknowledgedMinutes", NotifyUnacknowledgedMinutes);
+			}
+			if (NotifyUnresolvedMinutes != -1)
 			{
 				NotificationSection.SetValue("NotifyUnresolvedMinutes", NotifyUnresolvedMinutes);
 			}
