UGS: Add a matcher for content errors in jobs called "Check Asset References".

#rnx

diff --git a/Engine/Source/Programs/UnrealGameSync/MetadataTool/BuildHealth/BuildHealthIssue.cs b/Engine/Source/Programs/UnrealGameSync/MetadataTool/BuildHealth/BuildHealthIssue.cs
--- a/Engine/Source/Programs/UnrealGameSync/MetadataTool/BuildHealth/BuildHealthIssue.cs
+++ b/Engine/Source/Programs/UnrealGameSync/MetadataTool/BuildHealth/BuildHealthIssue.cs
@@ -97,6 +97,12 @@ namespace MetadataTool
 		[DataMember(Order = 6)]
 		public SortedSet<string> Identifiers = new SortedSet<string>();
 
+		/// <summary>
+		/// Other arbitrary references to assets or files
+		/// </summary>
+		[DataMember(Order = 7)]
+		public SortedSet<string> References = new SortedSet<string>();
+
 		/// <summary>
 		/// The last posted issue summary. Will be updated if it changes.
 		/// </summary>
diff --git a/Engine/Source/Programs/UnrealGameSync/MetadataTool/BuildHealth/Matchers/ContentPatternMatcher.cs b/Engine/Source/Programs/UnrealGameSync/MetadataTool/BuildHealth/Matchers/ContentPatternMatcher.cs
new file mode 100644
--- /dev/null
+++ b/Engine/Source/Programs/UnrealGameSync/MetadataTool/BuildHealth/Matchers/ContentPatternMatcher.cs
@@ -0,0 +1,80 @@
+// Copyright 1998-2019 Epic Games, Inc. All Rights Reserved.
+
+using System;
+using System.Collections.Generic;
+using System.Linq;
+using System.Text;
+using System.Text.RegularExpressions;
+using System.Threading.Tasks;
+using Tools.DotNETCommon.Perforce;
+
+namespace MetadataTool
+{
+	class ContentPatternMatcher : GenericContentPatternMatcher
+	{
+		public override string Category => "Content";
+
+		public override bool TryMatch(InputJob Job, InputJobStep JobStep, InputDiagnostic Diagnostic, List<BuildHealthIssue> Issues)
+		{
+			if(JobStep.Name.Contains("Check Asset References") && Regex.IsMatch(Diagnostic.Message, @"^\s*Log[a-zA-Z0-9]+: (?:Error|Warning):"))
+			{
+				BuildHealthIssue Issue = new BuildHealthIssue(Job.Project, Category, Job.Url, new BuildHealthDiagnostic(JobStep.Name, Diagnostic.Message, Diagnostic.Url));
+
+				// Use the error text to identify this error
+				Issue.Identifiers.Add(Diagnostic.Message);
+
+				// Find any asset references
+				foreach (Match FileMatch in Regex.Matches(Diagnostic.Message, @"[^a-zA-Z0-9_](?:/Game/|/Engine/)([^(). ]+)"))
+				{
+					if (FileMatch.Success)
+					{
+						Issue.References.Add(FileMatch.Groups[1].Value);
+					}
+				}
+
+				Issues.Add(Issue);
+				return true;
+			}
+			return false;
+		}
+
+		public override List<ChangeInfo> FindCausers(PerforceConnection Perforce, BuildHealthIssue Issue, IReadOnlyList<ChangeInfo> Changes)
+		{
+			// Check for any changes that modify those files
+			List<ChangeInfo> Causers = new List<ChangeInfo>();
+			foreach(ChangeInfo Change in Changes)
+			{
+				DescribeRecord DescribeRecord = GetDescribeRecord(Perforce, Change);
+				foreach(DescribeFileRecord File in DescribeRecord.Files)
+				{
+					if(Issue.References.Any(Reference => File.DepotFile.IndexOf(Reference, StringComparison.OrdinalIgnoreCase) != -1))
+					{
+						Causers.Add(Change);
+						break;
+					}
+				}
+			}
+
+			if(Causers.Count > 0)
+			{
+				return Causers;
+			}
+			else
+			{
+				return base.FindCausers(Perforce, Issue, Changes);
+			}
+		}
+
+		public override string GetSummary(BuildHealthIssue Issue)
+		{
+			if (Issue.References.Count == 0)
+			{
+				return "Content errors";
+			}
+			else
+			{
+				return String.Format("Content errors in {0}", String.Join(", ", Issue.References));
+			}
+		}
+	}
+}
diff --git a/Engine/Source/Programs/UnrealGameSync/MetadataTool/BuildHealth/Matchers/DefaultPatternMatcher.cs b/Engine/Source/Programs/UnrealGameSync/MetadataTool/BuildHealth/Matchers/DefaultPatternMatcher.cs
--- a/Engine/Source/Programs/UnrealGameSync/MetadataTool/BuildHealth/Matchers/DefaultPatternMatcher.cs
+++ b/Engine/Source/Programs/UnrealGameSync/MetadataTool/BuildHealth/Matchers/DefaultPatternMatcher.cs
@@ -30,11 +30,6 @@ namespace MetadataTool
 			return new List<ChangeInfo>();
 		}
 
-		public override bool CanMerge(BuildHealthIssue Source, BuildHealthIssue Target)
-		{
-			return Source.Category == Target.Category && Source.Identifiers.Any(x => Target.Identifiers.Contains(x));
-		}
-
 		public override string GetSummary(BuildHealthIssue Issue)
 		{
 			string Message = Issue.Diagnostics[0].Message;
diff --git a/Engine/Source/Programs/UnrealGameSync/MetadataTool/BuildHealth/Matchers/GenericContentPatternMatcher.cs b/Engine/Source/Programs/UnrealGameSync/MetadataTool/BuildHealth/Matchers/GenericContentPatternMatcher.cs
new file mode 100644
--- /dev/null
+++ b/Engine/Source/Programs/UnrealGameSync/MetadataTool/BuildHealth/Matchers/GenericContentPatternMatcher.cs
@@ -0,0 +1,41 @@
+// Copyright 1998-2019 Epic Games, Inc. All Rights Reserved.
+
+using System;
+using System.Collections.Generic;
+using System.Linq;
+using System.Text;
+using System.Threading.Tasks;
+using Tools.DotNETCommon.Perforce;
+
+namespace MetadataTool
+{
+	/// <summary>
+	/// Base class for pattern matchers that match errors in source code
+	/// </summary>
+	abstract class GenericContentPatternMatcher : PatternMatcher
+	{
+		/// <summary>
+		/// Set of extensions to treat as code
+		/// </summary>
+		static readonly string[] ContentExtensions =
+		{
+			".uasset",
+			".umap",
+			".ini",
+		};
+
+		/// <summary>
+		/// Finds any suspected causers for a particular failure. Excludes any changes that don't contain code.
+		/// </summary>
+		/// <param name="Perforce">The perforce connection</param>
+		/// <param name="Issue">The build issue</param>
+		/// <param name="Changes">List of changes since the issue first occurred.</param>
+		/// <returns>List of changes which are causers for the issue</returns>
+		public override List<ChangeInfo> FindCausers(PerforceConnection Perforce, BuildHealthIssue Issue, IReadOnlyList<ChangeInfo> Changes)
+		{
+			List<ChangeInfo> Causers = base.FindCausers(Perforce, Issue, Changes);
+			Causers.RemoveAll(x => !ContainsAnyFileWithExtension(Perforce, x, ContentExtensions));
+			return Causers;
+		}
+	}
+}
diff --git a/Engine/Source/Programs/UnrealGameSync/MetadataTool/BuildHealth/PatternMatcher.cs b/Engine/Source/Programs/UnrealGameSync/MetadataTool/BuildHealth/PatternMatcher.cs
--- a/Engine/Source/Programs/UnrealGameSync/MetadataTool/BuildHealth/PatternMatcher.cs
+++ b/Engine/Source/Programs/UnrealGameSync/MetadataTool/BuildHealth/PatternMatcher.cs
@@ -99,6 +99,7 @@ namespace MetadataTool
 
 			Target.FileNames.UnionWith(Source.FileNames);
 			Target.Identifiers.UnionWith(Source.Identifiers);
+			Target.References.UnionWith(Source.References);
 		}
 
 		/// <summary>
diff --git a/Engine/Source/Programs/UnrealGameSync/MetadataTool/CommandHandler_BuildHealth.cs b/Engine/Source/Programs/UnrealGameSync/MetadataTool/CommandHandler_BuildHealth.cs
--- a/Engine/Source/Programs/UnrealGameSync/MetadataTool/CommandHandler_BuildHealth.cs
+++ b/Engine/Source/Programs/UnrealGameSync/MetadataTool/CommandHandler_BuildHealth.cs
@@ -25,6 +25,7 @@ namespace MetadataTool
 			new CompilePatternMatcher(),
 			new UndefinedSymbolPatternMatcher(),
 			new CopyrightNoticeMatcher(),
+			new ContentPatternMatcher(),
 			new DefaultPatternMatcher(),
 		};
 
diff --git a/Engine/Source/Programs/UnrealGameSync/MetadataTool/MetadataTool.csproj b/Engine/Source/Programs/UnrealGameSync/MetadataTool/MetadataTool.csproj
--- a/Engine/Source/Programs/UnrealGameSync/MetadataTool/MetadataTool.csproj
+++ b/Engine/Source/Programs/UnrealGameSync/MetadataTool/MetadataTool.csproj
@@ -53,9 +53,11 @@
   <ItemGroup>
     <Compile Include="BuildHealth\ChangeInfo.cs" />
     <Compile Include="BuildHealth\Matchers\CompilePatternMatcher.cs" />
+    <Compile Include="BuildHealth\Matchers\ContentPatternMatcher.cs" />
     <Compile Include="BuildHealth\Matchers\CopyrightNoticeMatcher.cs" />
     <Compile Include="BuildHealth\Matchers\DefaultPatternMatcher.cs" />
     <Compile Include="BuildHealth\Matchers\GenericCodePatternMatcher.cs" />
+    <Compile Include="BuildHealth\Matchers\GenericContentPatternMatcher.cs" />
     <Compile Include="BuildHealth\Matchers\UndefinedSymbolPatternMatcher.cs" />
     <Compile Include="BuildHealth\PatternMatcher.cs" />
     <Compile Include="CommandHandler.cs" />
