UGS: Update state data read from disk.

#rnx

diff --git a/Engine/Source/Programs/UnrealGameSync/MetadataTool/CommandHandler_BuildHealth.cs b/Engine/Source/Programs/UnrealGameSync/MetadataTool/CommandHandler_BuildHealth.cs
--- a/Engine/Source/Programs/UnrealGameSync/MetadataTool/CommandHandler_BuildHealth.cs
+++ b/Engine/Source/Programs/UnrealGameSync/MetadataTool/CommandHandler_BuildHealth.cs
@@ -88,6 +88,15 @@ namespace MetadataTool
 				State = new BuildHealthState();
 			}
 
+			// Fixup any issues loaded from disk
+			foreach(BuildHealthIssue Issue in State.Issues)
+			{
+				if(Issue.References == null)
+				{
+					Issue.References = new SortedSet<string>();
+				}
+			}
+
 			// Create the Perforce connection
 			PerforceConnection Perforce = new PerforceConnection(PerforcePort, PerforceUser, null);
 
