Fix copyright notice matcher for source files.

#rnx

diff --git a/Engine/Source/Programs/UnrealGameSync/MetadataTool/BuildHealth/Matchers/CopyrightNoticeMatcher.cs b/Engine/Source/Programs/UnrealGameSync/MetadataTool/BuildHealth/Matchers/CopyrightNoticeMatcher.cs
--- a/Engine/Source/Programs/UnrealGameSync/MetadataTool/BuildHealth/Matchers/CopyrightNoticeMatcher.cs
+++ b/Engine/Source/Programs/UnrealGameSync/MetadataTool/BuildHealth/Matchers/CopyrightNoticeMatcher.cs
@@ -27,7 +27,8 @@ namespace MetadataTool
 			{
 				if (FileMatch.Success)
 				{
-					SourceFileNames.Add(FileMatch.Groups[1].Value);
+					string SourceFileName = FileMatch.Groups[1].Value.Replace('\\', '/');
+					SourceFileNames.Add(SourceFileName);
 				}
 			}
 
