UGS: Only filter suspected causers for an issue after all related issues for a step have been merged.

#rnx

diff --git a/Engine/Source/Programs/UnrealGameSync/MetadataTool/CommandHandler_BuildHealth.cs b/Engine/Source/Programs/UnrealGameSync/MetadataTool/CommandHandler_BuildHealth.cs
--- a/Engine/Source/Programs/UnrealGameSync/MetadataTool/CommandHandler_BuildHealth.cs
+++ b/Engine/Source/Programs/UnrealGameSync/MetadataTool/CommandHandler_BuildHealth.cs
@@ -482,27 +482,56 @@ namespace MetadataTool
 		/// <param name="InputJobStep">The job step to add</param>
 		void AddStep(PerforceConnection Perforce, BuildHealthState State, InputJob InputJob, InputJobStep InputJobStep)
 		{
-			// Create all the fingerprints for failures in this step
-			List<BuildHealthIssue> InputIssues = new List<BuildHealthIssue>();
-			if(InputJobStep.Diagnostics != null)
+			if(InputJobStep.Diagnostics != null && InputJobStep.Diagnostics.Count > 0)
 			{
+				// Create a lazily evaluated list of changes that are responsible for any errors
+				Lazy<IReadOnlyList<ChangeInfo>> LazyChanges = new Lazy<IReadOnlyList<ChangeInfo>>(() => FindChanges(Perforce, State, InputJob));
+
+				// List of created issues
+				List<BuildHealthIssue> InputIssues = new List<BuildHealthIssue>();
+
+				// Create issues for any diagnostics in this step
 				List<InputDiagnostic> Diagnostics = new List<InputDiagnostic>(InputJobStep.Diagnostics); 
-				foreach(PatternMatcher PatternMatcher in Matchers)
+				foreach(PatternMatcher Matcher in Matchers)
 				{
-					PatternMatcher.Match(InputJob, InputJobStep, Diagnostics, InputIssues);
+					Matcher.Match(InputJob, InputJobStep, Diagnostics, InputIssues);
 				}
-			}
 
-			// Add all the fingerprints to issues
-			foreach(BuildHealthIssue InputIssue in InputIssues)
-			{
-				BuildHealthIssue Issue = FindOrAddIssueForFingerprint(Perforce, State, InputJob, InputJobStep, InputIssue);
-				AddFailureToIssue(Issue, InputJob, InputJobStep, InputIssue.Diagnostics[0].ErrorUrl, State);
+				// Add all the fingerprints to issues
+				List<BuildHealthIssue> NewIssues = new List<BuildHealthIssue>();
+				foreach(BuildHealthIssue InputIssue in InputIssues)
+				{
+					if(!MergeIntoExistingIssue(Perforce, State, InputJob, InputJobStep, InputIssue, LazyChanges))
+					{
+						NewIssues.Add(InputIssue);
+						State.Issues.Add(InputIssue);
+					}
+					AddFailureToIssue(InputIssue, InputJob, InputJobStep, InputIssue.Diagnostics[0].ErrorUrl, State);
+				}
+
+				// Update the watchers for any new issues
+				foreach(BuildHealthIssue NewIssue in NewIssues)
+				{
+					IReadOnlyList<ChangeInfo> Changes = LazyChanges.Value;
+					if (Changes != null)
+					{
+						// Find the pattern matcher for this issue
+						PatternMatcher Matcher = CategoryToMatcher[NewIssue.Category];
+
+						// Update the causers
+						List<ChangeInfo> Causers = Matcher.FindCausers(Perforce, NewIssue, Changes);
+						foreach (ChangeInfo Causer in Causers)
+						{
+							NewIssue.SourceChanges.UnionWith(Causer.SourceChanges);
+							NewIssue.PendingWatchers.Add(Causer.Record.User);
+						}
+					}
+				}
 			}
 		}
 
 		/// <summary>
-		/// Finds or adds an issue for a particular fingerprint
+		/// Finds or adds an issue for a particular issue
 		/// </summary>
 		/// <param name="Perforce">Perforce connection used to find possible causers</param>
 		/// <param name="State">The current set of tracked issues</param>
@@ -510,7 +539,7 @@ namespace MetadataTool
 		/// <param name="PreviousChange">The last changelist that was built before this one</param>
 		/// <param name="InputJob">Job containing the step to add</param>
 		/// <param name="InputJobStep">The job step to add</param>
-		BuildHealthIssue FindOrAddIssueForFingerprint(PerforceConnection Perforce, BuildHealthState State, InputJob InputJob, InputJobStep InputJobStep, BuildHealthIssue InputIssue)
+		bool MergeIntoExistingIssue(PerforceConnection Perforce, BuildHealthState State, InputJob InputJob, InputJobStep InputJobStep, BuildHealthIssue InputIssue, Lazy<IReadOnlyList<ChangeInfo>> LazyChanges)
 		{
 			// Find the pattern matcher for this fingerprint
 			PatternMatcher Matcher = CategoryToMatcher[InputIssue.Category];
@@ -550,67 +579,37 @@ namespace MetadataTool
 
 				// Add the new build
 				Matcher.Merge(InputIssue, Issue);
-				return Issue;
+				return true;
 			}
 
-			// List of changes since the last successful build in this stream
-			IReadOnlyList<ChangeInfo> Changes = null;
-
-			// Find the previous changelist that was built in this stream
-			List<BuildHealthJobStep> StreamBuilds;
-			if(State.Streams.TryGetValue(InputJob.Stream, out StreamBuilds))
+			// Check if this issue can be merged with an issue built in another stream
+			IReadOnlyList<ChangeInfo> Changes = LazyChanges.Value;
+			if(Changes != null && Changes.Count > 0)
 			{
-				// Find the last change submitted to this stream before it started failing
-				int LastChange = -1;
-				for (int Idx = 0; Idx < StreamBuilds.Count && StreamBuilds[Idx].Change < InputJob.Change; Idx++)
-				{
-					LastChange = StreamBuilds[Idx].Change;
-				}
-
-				// Allow adding to any open issue that contains changes merged from other branches
-				if (LastChange != -1)
+				SortedSet<int> SourceChanges = new SortedSet<int>(Changes.SelectMany(x => x.SourceChanges));
+				foreach (BuildHealthIssue Issue in State.Issues)
 				{
-					// Query for all the changes since then
-					Changes = FindChanges(Perforce, InputJob.Stream, LastChange, InputJob.Change);
-					if (Changes.Count > 0)
+					// Check if this issue does not already contain this stream, but contains one of the causing changes
+					if (Issue.Streams.ContainsKey(InputJob.Stream))
 					{
-						SortedSet<int> SourceChanges = new SortedSet<int>(Changes.SelectMany(x => x.SourceChanges));
-						foreach (BuildHealthIssue Issue in State.Issues)
-						{
-							// Check if this issue does not already contain this stream, but contains one of the causing changes
-							if (Issue.Streams.ContainsKey(InputJob.Stream))
-							{
-								continue;
-							}
-							if(!SourceChanges.Any(x => Issue.SourceChanges.Contains(x)))
-							{
-								continue;
-							}
-							if (!Matcher.CanMerge(InputIssue, Issue))
-							{
-								continue;
-							}
-
-							// Merge the issue
-							Matcher.Merge(InputIssue, Issue);
-							return Issue;
-						}
+						continue;
+					}
+					if(!SourceChanges.Any(x => Issue.SourceChanges.Contains(x)))
+					{
+						continue;
+					}
+					if (!Matcher.CanMerge(InputIssue, Issue))
+					{
+						continue;
 					}
-				}
-			}
 
-			// Create new issues for everything else in this stream
-			if (Changes != null)
-			{
-				List<ChangeInfo> Causers = Matcher.FindCausers(Perforce, InputIssue, Changes);
-				foreach(ChangeInfo Causer in Causers)
-				{
-					InputIssue.SourceChanges.UnionWith(Causer.SourceChanges);
-					InputIssue.PendingWatchers.Add(Causer.Record.User);
+					// Merge the issue
+					Matcher.Merge(InputIssue, Issue);
+					return true;
 				}
 			}
-			State.Issues.Add(InputIssue);
-			return InputIssue;
+
+			return false;
 		}
 
 		/// <summary>
@@ -655,6 +654,39 @@ namespace MetadataTool
 			History.AddFailedBuild(CreateBuildForJobStep(InputJob, InputJobStep, InputErrorUrl));
 		}
 
+		/// <summary>
+		/// Find all changes PLUS all robomerge source changes
+		/// </summary>
+		/// <param name="Perforce">The Perforce connection to use</param>
+		/// <param name="State">State of </param>
+		/// <param name="InputJob">The job that failed</param>
+		/// <returns>Set of changelist numbers</returns>
+		IReadOnlyList<ChangeInfo> FindChanges(PerforceConnection Perforce, BuildHealthState State, InputJob InputJob)
+		{
+			// List of changes since the last successful build in this stream
+			IReadOnlyList<ChangeInfo> Changes = null;
+
+			// Find the previous changelist that was built in this stream
+			List<BuildHealthJobStep> StreamBuilds;
+			if (State.Streams.TryGetValue(InputJob.Stream, out StreamBuilds))
+			{
+				// Find the last change submitted to this stream before it started failing
+				int LastChange = -1;
+				for (int Idx = 0; Idx < StreamBuilds.Count && StreamBuilds[Idx].Change < InputJob.Change; Idx++)
+				{
+					LastChange = StreamBuilds[Idx].Change;
+				}
+
+				// Allow adding to any open issue that contains changes merged from other branches
+				if (LastChange != -1)
+				{
+					// Query for all the changes since then
+					Changes = FindChanges(Perforce, InputJob.Stream, LastChange, InputJob.Change);
+				}
+			}
+			return Changes;
+		}
+
 		/// <summary>
 		/// Find all changes PLUS all robomerge source changes
 		/// </summary>
