UGS: Allow deleting issues from the metadata server.

#rnx

diff --git a/Engine/Source/Programs/UnrealGameSync/MetadataServer/Connectors/SqlConnector.cs b/Engine/Source/Programs/UnrealGameSync/MetadataServer/Connectors/SqlConnector.cs
--- a/Engine/Source/Programs/UnrealGameSync/MetadataServer/Connectors/SqlConnector.cs
+++ b/Engine/Source/Programs/UnrealGameSync/MetadataServer/Connectors/SqlConnector.cs
@@ -592,6 +592,36 @@ namespace MetadataServer.Connectors
 			return Text;
 		}
 
+		public static void DeleteIssue(long IssueId)
+		{
+			using (SQLiteConnection Connection = new SQLiteConnection(SqlConnector.ConnectionString))
+			{
+				Connection.Open();
+
+				using (SQLiteTransaction Transaction = Connection.BeginTransaction())
+				{
+					using (SQLiteCommand Command = new SQLiteCommand(Connection))
+					{
+						Command.Transaction = Transaction;
+
+						Command.CommandText = "DELETE FROM [IssueWatchers] WHERE IssueId = @IssueId";
+						Command.Parameters.AddWithValue("@IssueId", IssueId);
+						Command.ExecuteNonQuery();
+
+						Command.CommandText = "DELETE FROM [IssueBuilds] WHERE IssueId = @IssueId";
+						Command.Parameters.AddWithValue("@IssueId", IssueId);
+						Command.ExecuteNonQuery();
+
+						Command.CommandText = "DELETE FROM [Issues] WHERE Id = @IssueId";
+						Command.Parameters.AddWithValue("@IssueId", IssueId);
+						Command.ExecuteNonQuery();
+
+						Transaction.Commit();
+					}
+				}
+			}
+		}
+
 		public static void AddWatcher(long IssueId, string UserName)
 		{
 			using (SQLiteConnection Connection = new SQLiteConnection(SqlConnector.ConnectionString))
diff --git a/Engine/Source/Programs/UnrealGameSync/MetadataServer/Controllers/IssuesController.cs b/Engine/Source/Programs/UnrealGameSync/MetadataServer/Controllers/IssuesController.cs
--- a/Engine/Source/Programs/UnrealGameSync/MetadataServer/Controllers/IssuesController.cs
+++ b/Engine/Source/Programs/UnrealGameSync/MetadataServer/Controllers/IssuesController.cs
@@ -45,6 +45,12 @@ namespace MetadataServer.Controllers
 			long IssueId = SqlConnector.AddIssue(Issue);
 			return new { Id = IssueId };
 		}
+
+		[HttpDelete]
+		public void Delete(long id)
+		{
+			SqlConnector.DeleteIssue(id);
+		}
 	}
 
 	public class IssueBuildsSubController : ApiController
diff --git a/Engine/Source/Programs/UnrealGameSync/MetadataTool/Program.cs b/Engine/Source/Programs/UnrealGameSync/MetadataTool/Program.cs
--- a/Engine/Source/Programs/UnrealGameSync/MetadataTool/Program.cs
+++ b/Engine/Source/Programs/UnrealGameSync/MetadataTool/Program.cs
@@ -22,6 +22,7 @@ namespace MetadataTool
 			new CommandHandler_Http("GetIssue", "GET", "/api/issues/{id}"),
 			new CommandHandler_Http("GetIssues", "GET", "/api/issues", OptionalParams: new string[] { "includeresolved", "maxresults", "user" }),
 			new CommandHandler_Http("UpdateIssue", "PUT", "/api/issues/{Id}", typeof(CommandTypes.UpdateIssue)),
+			new CommandHandler_Http("DeleteIssue", "DELETE", "/api/issues/{Id}"),
 
 			// Builds
 			new CommandHandler_Http("AddBuild", "POST", "/api/issues/{Issue}/builds", typeof(CommandTypes.AddBuild)),
