UGS: Add arguments to allow saving off the state and input data whenever an issue cannot be matched.

#rnx

diff --git a/Engine/Source/Programs/UnrealGameSync/MetadataTool/CommandHandler_BuildHealth.cs b/Engine/Source/Programs/UnrealGameSync/MetadataTool/CommandHandler_BuildHealth.cs
--- a/Engine/Source/Programs/UnrealGameSync/MetadataTool/CommandHandler_BuildHealth.cs
+++ b/Engine/Source/Programs/UnrealGameSync/MetadataTool/CommandHandler_BuildHealth.cs
@@ -25,8 +25,7 @@ namespace MetadataTool
 			new CompilePatternMatcher(),
 			new UndefinedSymbolPatternMatcher(),
 			new CopyrightNoticeMatcher(),
-			new ContentPatternMatcher(),
-			new DefaultPatternMatcher(),
+			new ContentPatternMatcher()
 		};
 
 		static readonly Dictionary<string, PatternMatcher> CategoryToMatcher = Matchers.ToDictionary(x => x.Category, x => x);
@@ -63,6 +62,8 @@ namespace MetadataTool
 			FileReference StateFile = Arguments.GetFileReference("-StateFile=");
 			string ServerUrl = Arguments.GetStringOrDefault("-Server=", null);
 			bool bKeepHistory = Arguments.HasOption("-KeepHistory");
+			bool bReadOnly = Arguments.HasOption("-ReadOnly");
+			DirectoryReference SaveUnmatchedDir = Arguments.GetDirectoryReferenceOrDefault("-SaveUnmatched=", null);
 			Arguments.CheckAllArgumentsUsed();
 
 			// Build a mapping from category to matching
@@ -123,11 +124,26 @@ namespace MetadataTool
 					List<InputJobStep> InputJobSteps = InputJob.Steps.OrderBy(x => x.Name).ToList();
 					foreach (InputJobStep InputJobStep in InputJobSteps)
 					{
-						AddStep(Perforce, State, InputJob, InputJobStep);
+						if (InputJobStep.Diagnostics != null && InputJobStep.Diagnostics.Count > 0)
+						{
+							AddStep(Perforce, State, InputJob, InputJobStep);
+						}
 					}
+					
+					// Remove any steps which are empty
+					InputJob.Steps.RemoveAll(x => x.Diagnostics == null || x.Diagnostics.Count == 0);
 				}
+				InputJobs.RemoveAll(x => x.Steps.Count == 0);
 				Log.TraceInformation("Added jobs in {0}s", Timer.Elapsed.TotalSeconds);
 
+				// If there are any unmatched issues, save out the current state and remaining input
+				if(SaveUnmatchedDir != null && InputJobs.Count > 0)
+				{
+					DirectoryReference.CreateDirectory(SaveUnmatchedDir);
+					FileReference.Copy(StateFile, FileReference.Combine(SaveUnmatchedDir, "State.json"));
+					SerializeJson(FileReference.Combine(SaveUnmatchedDir, "Input.json"), InputData);
+				}
+
 				// Try to find the next successful build for each stream, so we can close it as part of updating the server
 				for (int Idx = 0; Idx < State.Issues.Count; Idx++)
 				{
@@ -201,6 +217,12 @@ namespace MetadataTool
 				}
 			}
 
+			// If we're in read-only mode, don't write anything out
+			if(bReadOnly)
+			{
+				return;
+			}
+
 			// Save the persistent data
 			Log.TraceInformation("Writing persistent data to {0}", StateFile);
 			DirectoryReference.CreateDirectory(StateFile.Directory);
@@ -482,49 +504,43 @@ namespace MetadataTool
 		/// <param name="InputJobStep">The job step to add</param>
 		void AddStep(PerforceConnection Perforce, BuildHealthState State, InputJob InputJob, InputJobStep InputJobStep)
 		{
-			if(InputJobStep.Diagnostics != null && InputJobStep.Diagnostics.Count > 0)
-			{
-				// Create a lazily evaluated list of changes that are responsible for any errors
-				Lazy<IReadOnlyList<ChangeInfo>> LazyChanges = new Lazy<IReadOnlyList<ChangeInfo>>(() => FindChanges(Perforce, State, InputJob));
+			// Create a lazily evaluated list of changes that are responsible for any errors
+			Lazy<IReadOnlyList<ChangeInfo>> LazyChanges = new Lazy<IReadOnlyList<ChangeInfo>>(() => FindChanges(Perforce, State, InputJob));
 
-				// List of created issues
-				List<BuildHealthIssue> InputIssues = new List<BuildHealthIssue>();
+			// Create issues for any diagnostics in this step
+			List<BuildHealthIssue> InputIssues = new List<BuildHealthIssue>();
+			foreach(PatternMatcher Matcher in Matchers)
+			{
+				Matcher.Match(InputJob, InputJobStep, InputJobStep.Diagnostics, InputIssues);
+			}
 
-				// Create issues for any diagnostics in this step
-				List<InputDiagnostic> Diagnostics = new List<InputDiagnostic>(InputJobStep.Diagnostics); 
-				foreach(PatternMatcher Matcher in Matchers)
+			// Add all the fingerprints to issues
+			List<BuildHealthIssue> NewIssues = new List<BuildHealthIssue>();
+			foreach(BuildHealthIssue InputIssue in InputIssues)
+			{
+				if(!MergeIntoExistingIssue(Perforce, State, InputJob, InputJobStep, InputIssue, LazyChanges))
 				{
-					Matcher.Match(InputJob, InputJobStep, Diagnostics, InputIssues);
+					NewIssues.Add(InputIssue);
+					State.Issues.Add(InputIssue);
 				}
+				AddFailureToIssue(InputIssue, InputJob, InputJobStep, InputIssue.Diagnostics[0].ErrorUrl, State);
+			}
 
-				// Add all the fingerprints to issues
-				List<BuildHealthIssue> NewIssues = new List<BuildHealthIssue>();
-				foreach(BuildHealthIssue InputIssue in InputIssues)
+			// Update the watchers for any new issues
+			foreach(BuildHealthIssue NewIssue in NewIssues)
+			{
+				IReadOnlyList<ChangeInfo> Changes = LazyChanges.Value;
+				if (Changes != null)
 				{
-					if(!MergeIntoExistingIssue(Perforce, State, InputJob, InputJobStep, InputIssue, LazyChanges))
-					{
-						NewIssues.Add(InputIssue);
-						State.Issues.Add(InputIssue);
-					}
-					AddFailureToIssue(InputIssue, InputJob, InputJobStep, InputIssue.Diagnostics[0].ErrorUrl, State);
-				}
+					// Find the pattern matcher for this issue
+					PatternMatcher Matcher = CategoryToMatcher[NewIssue.Category];
 
-				// Update the watchers for any new issues
-				foreach(BuildHealthIssue NewIssue in NewIssues)
-				{
-					IReadOnlyList<ChangeInfo> Changes = LazyChanges.Value;
-					if (Changes != null)
+					// Update the causers
+					List<ChangeInfo> Causers = Matcher.FindCausers(Perforce, NewIssue, Changes);
+					foreach (ChangeInfo Causer in Causers)
 					{
-						// Find the pattern matcher for this issue
-						PatternMatcher Matcher = CategoryToMatcher[NewIssue.Category];
-
-						// Update the causers
-						List<ChangeInfo> Causers = Matcher.FindCausers(Perforce, NewIssue, Changes);
-						foreach (ChangeInfo Causer in Causers)
-						{
-							NewIssue.SourceChanges.UnionWith(Causer.SourceChanges);
-							NewIssue.PendingWatchers.Add(Causer.Record.User);
-						}
+						NewIssue.SourceChanges.UnionWith(Causer.SourceChanges);
+						NewIssue.PendingWatchers.Add(Causer.Record.User);
 					}
 				}
 			}
