UGS: Allow overwriting existing state files in the unmatched directory.

#rnx

diff --git a/Engine/Source/Programs/UnrealGameSync/MetadataTool/CommandHandler_BuildHealth.cs b/Engine/Source/Programs/UnrealGameSync/MetadataTool/CommandHandler_BuildHealth.cs
--- a/Engine/Source/Programs/UnrealGameSync/MetadataTool/CommandHandler_BuildHealth.cs
+++ b/Engine/Source/Programs/UnrealGameSync/MetadataTool/CommandHandler_BuildHealth.cs
@@ -140,7 +140,7 @@ namespace MetadataTool
 				if(SaveUnmatchedDir != null && InputJobs.Count > 0)
 				{
 					DirectoryReference.CreateDirectory(SaveUnmatchedDir);
-					FileReference.Copy(StateFile, FileReference.Combine(SaveUnmatchedDir, "State.json"));
+					FileReference.Copy(StateFile, FileReference.Combine(SaveUnmatchedDir, "State.json"), true);
 					SerializeJson(FileReference.Combine(SaveUnmatchedDir, "Input.json"), InputData);
 				}
 
