UGS: Create a separate table for build health diagnostics, allowing diagnostics for issues to be cross-referenced with builds.

#rnx

diff --git a/Engine/Source/Programs/UnrealGameSync/MetadataServer/App_Start/WebApiConfig.cs b/Engine/Source/Programs/UnrealGameSync/MetadataServer/App_Start/WebApiConfig.cs
--- a/Engine/Source/Programs/UnrealGameSync/MetadataServer/App_Start/WebApiConfig.cs
+++ b/Engine/Source/Programs/UnrealGameSync/MetadataServer/App_Start/WebApiConfig.cs
@@ -24,9 +24,9 @@ namespace MetadataServer
 			);
 
 			config.Routes.MapHttpRoute(
-				name: "IssueEventsApi",
-				routeTemplate: "api/issues/{IssueId}/events",
-				defaults: new { controller = "IssueEvents" }
+				name: "IssueDiagnosticsSubApi",
+				routeTemplate: "api/issues/{IssueId}/diagnostics",
+				defaults: new { controller = "IssueDiagnosticsSub" }
 			);
 
 			config.Routes.MapHttpRoute(
diff --git a/Engine/Source/Programs/UnrealGameSync/MetadataServer/Connectors/SqlConnector.cs b/Engine/Source/Programs/UnrealGameSync/MetadataServer/Connectors/SqlConnector.cs
--- a/Engine/Source/Programs/UnrealGameSync/MetadataServer/Connectors/SqlConnector.cs
+++ b/Engine/Source/Programs/UnrealGameSync/MetadataServer/Connectors/SqlConnector.cs
@@ -385,7 +385,6 @@ namespace MetadataServer.Connectors
 		}
 
 		const int IssueSummaryMaxLength = 200;
-		const int IssueDetailsMaxLength = 8000;
 
 		public static long AddIssue(IssueData Issue)
 		{
@@ -398,7 +397,7 @@ namespace MetadataServer.Connectors
 				{
 					Command.Parameters.AddWithValue("@Project", Issue.Project);
 					Command.Parameters.AddWithValue("@Summary", SanitizeText(Issue.Summary, IssueSummaryMaxLength));
-					Command.Parameters.AddWithValue("@Details", SanitizeText(Issue.Details, IssueDetailsMaxLength));
+					Command.Parameters.AddWithValue("@Details", "");
 					if (Issue.Owner != null)
 					{
 						Command.Parameters.AddWithValue("OwnerId", FindOrAddUserId(Issue.Owner, Connection));
@@ -453,7 +452,7 @@ namespace MetadataServer.Connectors
 
 				StringBuilder CommandBuilder = new StringBuilder();
 				CommandBuilder.Append("SELECT");
-				CommandBuilder.Append(" Issues.Id, Issues.CreatedAt, DATETIME('now'), Issues.Project, Issues.Summary, Issues.Details, OwnerUsers.Name, NominatedByUsers.Name, Issues.AcknowledgedAt, Issues.FixChange, Issues.ResolvedAt");
+				CommandBuilder.Append(" Issues.Id, Issues.CreatedAt, DATETIME('now'), Issues.Project, Issues.Summary, OwnerUsers.Name, NominatedByUsers.Name, Issues.AcknowledgedAt, Issues.FixChange, Issues.ResolvedAt");
 				if(UserName != null)
 				{
 					CommandBuilder.Append(", IssueWatchers.UserId");
@@ -499,15 +498,14 @@ namespace MetadataServer.Connectors
 							Issue.RetrievedAt = Reader.GetDateTime(2);
 							Issue.Project = Reader.GetString(3);
 							Issue.Summary = Reader.GetString(4);
-							Issue.Details = Reader.GetString(5);
-							Issue.Owner = Reader.IsDBNull(6)? null : Reader.GetString(6);
-							Issue.NominatedBy = Reader.IsDBNull(7)? null : Reader.GetString(7);
-							Issue.AcknowledgedAt = Reader.IsDBNull(8)? (DateTime?)null : Reader.GetDateTime(8);
-							Issue.FixChange = Reader.GetInt32(9);
-							Issue.ResolvedAt = Reader.IsDBNull(10)? (DateTime?)null : Reader.GetDateTime(10);
+							Issue.Owner = Reader.IsDBNull(5)? null : Reader.GetString(5);
+							Issue.NominatedBy = Reader.IsDBNull(6)? null : Reader.GetString(6);
+							Issue.AcknowledgedAt = Reader.IsDBNull(7)? (DateTime?)null : Reader.GetDateTime(7);
+							Issue.FixChange = Reader.GetInt32(8);
+							Issue.ResolvedAt = Reader.IsDBNull(9)? (DateTime?)null : Reader.GetDateTime(9);
 							if(UserName != null)
 							{
-								Issue.bNotify = !Reader.IsDBNull(11);
+								Issue.bNotify = !Reader.IsDBNull(10);
 							}
 							Issues.Add(Issue);
 						}
@@ -533,12 +531,6 @@ namespace MetadataServer.Connectors
 						Values.Add("@Summary");
 						Command.Parameters.AddWithValue("@Summary", SanitizeText(Issue.Summary, IssueSummaryMaxLength));
 					}
-					if(Issue.Details != null)
-					{
-						Columns.Add("Details");
-						Values.Add("@Details");
-						Command.Parameters.AddWithValue("@Details", SanitizeText(Issue.Details, IssueDetailsMaxLength));
-					}
 					if (Issue.Owner != null)
 					{
 						Columns.Add("OwnerId");
@@ -622,6 +614,53 @@ namespace MetadataServer.Connectors
 			}
 		}
 
+		public static void AddDiagnostic(long IssueId, IssueDiagnosticData Diagnostic)
+		{
+			using (SQLiteConnection Connection = new SQLiteConnection(SqlConnector.ConnectionString))
+			{
+				Connection.Open();
+
+				using (SQLiteCommand Command = new SQLiteCommand("INSERT INTO [IssueDiagnostics] (IssueId, BuildId, Message, Url) VALUES (@IssueId, @BuildId, @Message, @Url)", Connection))
+				{
+					Command.Parameters.AddWithValue("@IssueId", IssueId);
+					Command.Parameters.AddWithValue("@BuildId", Diagnostic.BuildId);
+					Command.Parameters.AddWithValue("@Message", SanitizeText(Diagnostic.Message, 1000));
+					Command.Parameters.AddWithValue("@Url", Diagnostic.Url);
+					Command.ExecuteNonQuery();
+				}
+			}
+		}
+
+		public static List<IssueDiagnosticData> GetDiagnostics(long IssueId)
+		{
+			List<IssueDiagnosticData> Diagnostics = new List<IssueDiagnosticData>();
+			using (SQLiteConnection Connection = new SQLiteConnection(SqlConnector.ConnectionString))
+			{
+				Connection.Open();
+
+				StringBuilder CommandBuilder = new StringBuilder();
+				CommandBuilder.Append("SELECT BuildId, Message, Url FROM [IssueDiagnostics]");
+				CommandBuilder.Append(" WHERE IssueDiagnostics.IssueId = @IssueId");
+
+				using (SQLiteCommand Command = new SQLiteCommand(CommandBuilder.ToString(), Connection))
+				{
+					Command.Parameters.AddWithValue("@IssueId", IssueId);
+					using (SQLiteDataReader Reader = Command.ExecuteReader())
+					{
+						while (Reader.Read())
+						{
+							IssueDiagnosticData Diagnostic = new IssueDiagnosticData();
+							Diagnostic.BuildId = Reader.IsDBNull(0)? (long?)null : (long?)Reader.GetInt64(0);
+							Diagnostic.Message = Reader.GetString(1);
+							Diagnostic.Url = Reader.GetString(2);
+							Diagnostics.Add(Diagnostic);
+						}
+					}
+				}
+			}
+			return Diagnostics;
+		}
+
 		public static void AddWatcher(long IssueId, string UserName)
 		{
 			using (SQLiteConnection Connection = new SQLiteConnection(SqlConnector.ConnectionString))
diff --git a/Engine/Source/Programs/UnrealGameSync/MetadataServer/Controllers/IssuesController.cs b/Engine/Source/Programs/UnrealGameSync/MetadataServer/Controllers/IssuesController.cs
--- a/Engine/Source/Programs/UnrealGameSync/MetadataServer/Controllers/IssuesController.cs
+++ b/Engine/Source/Programs/UnrealGameSync/MetadataServer/Controllers/IssuesController.cs
@@ -62,9 +62,25 @@ namespace MetadataServer.Controllers
 		}
 
 		[HttpPost]
-		public void Post(long IssueId, [FromBody] IssueBuildData Data)
+		public object Post(long IssueId, [FromBody] IssueBuildData Data)
 		{
-			SqlConnector.AddBuild(IssueId, Data);
+			long BuildId = SqlConnector.AddBuild(IssueId, Data);
+			return new { Id = BuildId };
+		}
+	}
+
+	public class IssueDiagnosticsSubController : ApiController
+	{
+		[HttpGet]
+		public List<IssueDiagnosticData> Get(long IssueId)
+		{
+			return SqlConnector.GetDiagnostics(IssueId);
+		}
+
+		[HttpPost]
+		public void Post(long IssueId, [FromBody] IssueDiagnosticData Data)
+		{
+			SqlConnector.AddDiagnostic(IssueId, Data);
 		}
 	}
 
diff --git a/Engine/Source/Programs/UnrealGameSync/MetadataServer/Models/IssueData.cs b/Engine/Source/Programs/UnrealGameSync/MetadataServer/Models/IssueData.cs
--- a/Engine/Source/Programs/UnrealGameSync/MetadataServer/Models/IssueData.cs
+++ b/Engine/Source/Programs/UnrealGameSync/MetadataServer/Models/IssueData.cs
@@ -31,6 +31,13 @@ namespace MetadataServer.Models
 		public int Outcome;
 	}
 
+	public class IssueDiagnosticData
+	{
+		public long? BuildId;
+		public string Message;
+		public string Url;
+	}
+
 	public class IssueData
 	{
 		public long Id;
@@ -38,7 +45,6 @@ namespace MetadataServer.Models
 		public DateTime RetrievedAt;
 		public string Project;
 		public string Summary;
-		public string Details;
 		public string Owner;
 		public string NominatedBy;
 		public DateTime? AcknowledgedAt;
@@ -50,7 +56,6 @@ namespace MetadataServer.Models
 	public class IssueUpdateData
 	{
 		public string Summary;
-		public string Details;
 		public string Owner;
 		public string NominatedBy;
 		public bool? Acknowledged;
diff --git a/Engine/Source/Programs/UnrealGameSync/MetadataTool/BuildHealth/BuildHealthIssue.cs b/Engine/Source/Programs/UnrealGameSync/MetadataTool/BuildHealth/BuildHealthIssue.cs
--- a/Engine/Source/Programs/UnrealGameSync/MetadataTool/BuildHealth/BuildHealthIssue.cs
+++ b/Engine/Source/Programs/UnrealGameSync/MetadataTool/BuildHealth/BuildHealthIssue.cs
@@ -23,27 +23,41 @@ namespace MetadataTool
 		[DataMember(Order = 0, IsRequired = true)]
 		public string JobStepName;
 
+		/// <summary>
+		/// The url of this job step
+		/// </summary>
+		[DataMember(Order = 1)]
+		public string JobStepUrl;
+
 		/// <summary>
 		/// The diagnostic message
 		/// </summary>
-		[DataMember(Order = 1, IsRequired = true)]
+		[DataMember(Order = 2, IsRequired = true)]
 		public string Message;
 
 		/// <summary>
 		/// Url to the error
 		/// </summary>
-		[DataMember(Order = 2, IsRequired = true)]
+		[DataMember(Order = 3, IsRequired = true)]
 		public string ErrorUrl;
 
+		/// <summary>
+		/// Whether or not this diagnostic has been posted to the server
+		/// </summary>
+		[DataMember(Order = 4)]
+		public bool bPostedToServer = false;
+
 		/// <summary>
 		/// Constructor
 		/// </summary>
 		/// <param name="JobStepName">Name of the job step</param>
+		/// <param name="JobStepUrl">Url of the job step</param>
 		/// <param name="Message">Message to display</param>
 		/// <param name="ErrorUrl">Url to link to</param>
-		public BuildHealthDiagnostic(string JobStepName, string Message, string ErrorUrl)
+		public BuildHealthDiagnostic(string JobStepName, string JobStepUrl, string Message, string ErrorUrl)
 		{
 			this.JobStepName = JobStepName;
+			this.JobStepUrl = JobStepUrl;
 			this.Message = Message;
 			this.ErrorUrl = ErrorUrl;
 		}
@@ -109,12 +123,6 @@ namespace MetadataTool
 		[DataMember(Order = 20)]
 		public string PostedSummary;
 
-		/// <summary>
-		/// The last posted issue details. Will be updated if it changes.
-		/// </summary>
-		[DataMember(Order = 21)]
-		public string PostedDetails;
-
 		/// <summary>
 		/// Whether we've posted an update to the resolved flag to the server
 		/// </summary>
diff --git a/Engine/Source/Programs/UnrealGameSync/MetadataTool/BuildHealth/BuildHealthJobStep.cs b/Engine/Source/Programs/UnrealGameSync/MetadataTool/BuildHealth/BuildHealthJobStep.cs
--- a/Engine/Source/Programs/UnrealGameSync/MetadataTool/BuildHealth/BuildHealthJobStep.cs
+++ b/Engine/Source/Programs/UnrealGameSync/MetadataTool/BuildHealth/BuildHealthJobStep.cs
@@ -17,6 +17,12 @@ namespace MetadataTool
 	[DebuggerDisplay("{Change}: {JobName}")]
 	class BuildHealthJobStep : IComparable<BuildHealthJobStep>
 	{
+		/// <summary>
+		/// The id of this build on the server
+		/// </summary>
+		[DataMember(IsRequired = true)]
+		public long Id = -1;
+
 		/// <summary>
 		/// The changelist that this build was run at.
 		/// </summary>
diff --git a/Engine/Source/Programs/UnrealGameSync/MetadataTool/BuildHealth/Matchers/CompilePatternMatcher.cs b/Engine/Source/Programs/UnrealGameSync/MetadataTool/BuildHealth/Matchers/CompilePatternMatcher.cs
--- a/Engine/Source/Programs/UnrealGameSync/MetadataTool/BuildHealth/Matchers/CompilePatternMatcher.cs
+++ b/Engine/Source/Programs/UnrealGameSync/MetadataTool/BuildHealth/Matchers/CompilePatternMatcher.cs
@@ -59,7 +59,7 @@ namespace MetadataTool
 			// If we found any source files, create a diagnostic category for them
 			if (ReferencedFileNames.Count > 0)
 			{
-				BuildHealthIssue Issue = new BuildHealthIssue(Job.Project, Category, Job.Url, new BuildHealthDiagnostic(JobStep.Name, ShortenPaths(Diagnostic.Message), Diagnostic.Url));
+				BuildHealthIssue Issue = new BuildHealthIssue(Job.Project, Category, Job.Url, new BuildHealthDiagnostic(JobStep.Name, JobStep.Url, ShortenPaths(Diagnostic.Message), Diagnostic.Url));
 				Issue.FileNames.UnionWith(ReferencedFileNames);
 				Issue.Identifiers.UnionWith(GetSourceFileNames(ErrorFileNames));
 				Issues.Add(Issue);
diff --git a/Engine/Source/Programs/UnrealGameSync/MetadataTool/BuildHealth/Matchers/ContentPatternMatcher.cs b/Engine/Source/Programs/UnrealGameSync/MetadataTool/BuildHealth/Matchers/ContentPatternMatcher.cs
--- a/Engine/Source/Programs/UnrealGameSync/MetadataTool/BuildHealth/Matchers/ContentPatternMatcher.cs
+++ b/Engine/Source/Programs/UnrealGameSync/MetadataTool/BuildHealth/Matchers/ContentPatternMatcher.cs
@@ -18,7 +18,7 @@ namespace MetadataTool
 		{
 			if(JobStep.Name.Contains("Check Asset References") && Regex.IsMatch(Diagnostic.Message, @"^\s*Log[a-zA-Z0-9]+: (?:Error|Warning):"))
 			{
-				BuildHealthIssue Issue = new BuildHealthIssue(Job.Project, Category, Job.Url, new BuildHealthDiagnostic(JobStep.Name, Diagnostic.Message, Diagnostic.Url));
+				BuildHealthIssue Issue = new BuildHealthIssue(Job.Project, Category, Job.Url, new BuildHealthDiagnostic(JobStep.Name, JobStep.Url, Diagnostic.Message, Diagnostic.Url));
 
 				// Use the error text to identify this error
 				Issue.Identifiers.Add(Diagnostic.Message);
diff --git a/Engine/Source/Programs/UnrealGameSync/MetadataTool/BuildHealth/Matchers/CopyrightNoticeMatcher.cs b/Engine/Source/Programs/UnrealGameSync/MetadataTool/BuildHealth/Matchers/CopyrightNoticeMatcher.cs
--- a/Engine/Source/Programs/UnrealGameSync/MetadataTool/BuildHealth/Matchers/CopyrightNoticeMatcher.cs
+++ b/Engine/Source/Programs/UnrealGameSync/MetadataTool/BuildHealth/Matchers/CopyrightNoticeMatcher.cs
@@ -35,7 +35,7 @@ namespace MetadataTool
 			// If we found any source files, create a diagnostic category for them
 			if (SourceFileNames.Count > 0)
 			{
-				BuildHealthIssue Issue = new BuildHealthIssue(Job.Project, Category, Job.Url, new BuildHealthDiagnostic(JobStep.Name, Diagnostic.Message, Diagnostic.Url));
+				BuildHealthIssue Issue = new BuildHealthIssue(Job.Project, Category, Job.Url, new BuildHealthDiagnostic(JobStep.Name, JobStep.Url, Diagnostic.Message, Diagnostic.Url));
 				Issue.FileNames.UnionWith(SourceFileNames);
 				Issues.Add(Issue);
 				return true;
diff --git a/Engine/Source/Programs/UnrealGameSync/MetadataTool/BuildHealth/Matchers/DefaultPatternMatcher.cs b/Engine/Source/Programs/UnrealGameSync/MetadataTool/BuildHealth/Matchers/DefaultPatternMatcher.cs
--- a/Engine/Source/Programs/UnrealGameSync/MetadataTool/BuildHealth/Matchers/DefaultPatternMatcher.cs
+++ b/Engine/Source/Programs/UnrealGameSync/MetadataTool/BuildHealth/Matchers/DefaultPatternMatcher.cs
@@ -18,7 +18,7 @@ namespace MetadataTool
 		{
 			string DefaultProject = String.Format("{0} (Unmatched)", Job.Project);
 
-			BuildHealthIssue Issue = new BuildHealthIssue(DefaultProject, Category, Job.Url, new BuildHealthDiagnostic(JobStep.Name, Diagnostic.Message, Diagnostic.Url));
+			BuildHealthIssue Issue = new BuildHealthIssue(DefaultProject, Category, Job.Url, new BuildHealthDiagnostic(JobStep.Name, JobStep.Url, Diagnostic.Message, Diagnostic.Url));
 			Issue.Identifiers.Add(Diagnostic.Message);
 			Issues.Add(Issue);
 			
diff --git a/Engine/Source/Programs/UnrealGameSync/MetadataTool/BuildHealth/Matchers/UndefinedSymbolPatternMatcher.cs b/Engine/Source/Programs/UnrealGameSync/MetadataTool/BuildHealth/Matchers/UndefinedSymbolPatternMatcher.cs
--- a/Engine/Source/Programs/UnrealGameSync/MetadataTool/BuildHealth/Matchers/UndefinedSymbolPatternMatcher.cs
+++ b/Engine/Source/Programs/UnrealGameSync/MetadataTool/BuildHealth/Matchers/UndefinedSymbolPatternMatcher.cs
@@ -83,7 +83,7 @@ namespace MetadataTool
 			// If we found any symbol names, create a fingerprint for them
 			if (SymbolNames.Count > 0)
 			{
-				BuildHealthIssue Issue = new BuildHealthIssue(Job.Project, Category, Job.Url, new BuildHealthDiagnostic(JobStep.Name, Diagnostic.Message, Diagnostic.Url));
+				BuildHealthIssue Issue = new BuildHealthIssue(Job.Project, Category, Job.Url, new BuildHealthDiagnostic(JobStep.Name, JobStep.Url, Diagnostic.Message, Diagnostic.Url));
 				Issue.Identifiers.UnionWith(SymbolNames);
 				Issues.Add(Issue);
 				return true;
diff --git a/Engine/Source/Programs/UnrealGameSync/MetadataTool/CommandHandler_BuildHealth.cs b/Engine/Source/Programs/UnrealGameSync/MetadataTool/CommandHandler_BuildHealth.cs
--- a/Engine/Source/Programs/UnrealGameSync/MetadataTool/CommandHandler_BuildHealth.cs
+++ b/Engine/Source/Programs/UnrealGameSync/MetadataTool/CommandHandler_BuildHealth.cs
@@ -240,18 +240,7 @@ namespace MetadataTool
 						continue;
 					}
 
-					StringBuilder DetailsBuilder = new StringBuilder();
-					foreach (BuildHealthDiagnostic Diagnostic in Issue.Diagnostics)
-					{
-						DetailsBuilder.AppendFormat("##{0}##\n{1}\n", Diagnostic.JobStepName, Diagnostic.Message);
-					}
-					if(DetailsBuilder.Length > CommandTypes.MaxIssueDetailsLength)
-					{
-						DetailsBuilder.Remove(CommandTypes.MaxIssueDetailsLength, DetailsBuilder.Length - CommandTypes.MaxIssueDetailsLength);
-					}
-
 					string Summary = Matcher.GetSummary(Issue);
-					string Details = DetailsBuilder.ToString().TrimEnd();
 					if (Issue.Id == -1)
 					{
 						Log.TraceInformation("Adding issue: {0}", Issue);
@@ -264,7 +253,6 @@ namespace MetadataTool
 						CommandTypes.AddIssue IssueBody = new CommandTypes.AddIssue();
 						IssueBody.Project = Issue.Project;
 						IssueBody.Summary = Summary;
-						IssueBody.Details = Details;
 
 						if(Issue.PendingWatchers.Count == 1)
 						{
@@ -282,19 +270,14 @@ namespace MetadataTool
 						}
 
 						Issue.PostedSummary = Summary;
-						Issue.PostedDetails = Details;
 						WriteState(StateFile, State);
 					}
-					else if(Issue.PostedSummary == null 
-							|| Issue.PostedDetails == null
-							|| !String.Equals(Issue.PostedSummary, Summary, StringComparison.Ordinal)
-							|| !String.Equals(Issue.PostedDetails, Details, StringComparison.Ordinal))
+					else if(Issue.PostedSummary == null || !String.Equals(Issue.PostedSummary, Summary, StringComparison.Ordinal)
 					{
 						Log.TraceInformation("Updating issue {0}", Issue.Id);
 
 						CommandTypes.UpdateIssue IssueBody = new CommandTypes.UpdateIssue();
 						IssueBody.Summary = Summary;
-						IssueBody.Details = Details;
 
 						using (HttpWebResponse Response = SendHttpRequest(String.Format("{0}/api/issues/{1}", ServerUrl, Issue.Id), "PUT", IssueBody))
 						{
@@ -306,21 +289,12 @@ namespace MetadataTool
 						}
 
 						Issue.PostedSummary = Summary;
-						Issue.PostedDetails = Details;
 						WriteState(StateFile, State);
 					}
 				}
 
-				// Update the summary for any issues that are still open
-				foreach (BuildHealthIssue Issue in State.Issues)
-				{
-					if (Issue.Id == -1)
-					{
-						Log.TraceInformation("Adding issue: {0}", Issue);
-					}
-				}
-
 				// Add any new builds associated with issues
+				Dictionary<string, long> JobStepUrlToId = new Dictionary<string, long>(StringComparer.Ordinal);
 				foreach (BuildHealthIssue Issue in State.Issues)
 				{
 					foreach(KeyValuePair<string, Dictionary<string, BuildHealthJobHistory>> StreamPair in Issue.Streams)
@@ -350,12 +324,64 @@ namespace MetadataTool
 										{
 											throw new Exception("Unable to add build");
 										}
+										Build.Id = ParseHttpResponse<CommandTypes.AddBuildResponse>(Response).Id;
 									}
 
 									Build.bPostedToServer = true;
 									WriteState(StateFile, State);
 								}
+								if(Build.Id != -1)
+								{
+									JobStepUrlToId[Build.JobStepUrl] = Build.Id;
+								}
+							}
+						}
+					}
+				}
+
+				// Add any new diagnostics
+				foreach(BuildHealthIssue Issue in State.Issues)
+				{
+					foreach(BuildHealthDiagnostic Diagnostic in Issue.Diagnostics)
+					{
+						if(!Diagnostic.bPostedToServer)
+						{
+							string Summary = Diagnostic.Message;
+
+							const int MaxLength = 40;
+							if(Summary.Length > MaxLength)
+							{
+								Summary = Summary.Substring(0, MaxLength).TrimEnd();
+							}
+
+							Log.TraceInformation("Adding diagnostic '{0}' to issue {1}", Summary, Issue.Id);
+
+							CommandTypes.AddDiagnostic AddDiagnostic = new CommandTypes.AddDiagnostic();
+
+							long BuildId;
+							if(Diagnostic.JobStepUrl != null && JobStepUrlToId.TryGetValue(Diagnostic.JobStepUrl, out BuildId))
+							{
+								AddDiagnostic.BuildId = BuildId;
+							}
+							else
+							{
+								Console.WriteLine("ERROR");
+							}
+
+							AddDiagnostic.Message = Diagnostic.Message;
+							AddDiagnostic.Url = Diagnostic.ErrorUrl;
+
+							using (HttpWebResponse Response = SendHttpRequest(String.Format("{0}/api/issues/{1}/diagnostics", ServerUrl, Issue.Id), "POST", AddDiagnostic))
+							{
+								int ResponseCode = (int)Response.StatusCode;
+								if (!(ResponseCode >= 200 && ResponseCode <= 299))
+								{
+									throw new Exception("Unable to add build");
+								}
 							}
+
+							Diagnostic.bPostedToServer = true;
+							WriteState(StateFile, State);
 						}
 					}
 				}
@@ -514,16 +540,18 @@ namespace MetadataTool
 				Matcher.Match(InputJob, InputJobStep, InputJobStep.Diagnostics, InputIssues);
 			}
 
-			// Add all the fingerprints to issues
+			// Merge the issues together
 			List<BuildHealthIssue> NewIssues = new List<BuildHealthIssue>();
 			foreach(BuildHealthIssue InputIssue in InputIssues)
 			{
-				if(!MergeIntoExistingIssue(Perforce, State, InputJob, InputJobStep, InputIssue, LazyChanges))
+				BuildHealthIssue OutputIssue = MergeIntoExistingIssue(Perforce, State, InputJob, InputJobStep, InputIssue, LazyChanges);
+				if(OutputIssue == null)
 				{
 					NewIssues.Add(InputIssue);
 					State.Issues.Add(InputIssue);
+					OutputIssue = InputIssue;
 				}
-				AddFailureToIssue(InputIssue, InputJob, InputJobStep, InputIssue.Diagnostics[0].ErrorUrl, State);
+				AddFailureToIssue(OutputIssue, InputJob, InputJobStep, InputIssue.Diagnostics[0].ErrorUrl, State);
 			}
 
 			// Update the watchers for any new issues
@@ -555,7 +583,7 @@ namespace MetadataTool
 		/// <param name="PreviousChange">The last changelist that was built before this one</param>
 		/// <param name="InputJob">Job containing the step to add</param>
 		/// <param name="InputJobStep">The job step to add</param>
-		bool MergeIntoExistingIssue(PerforceConnection Perforce, BuildHealthState State, InputJob InputJob, InputJobStep InputJobStep, BuildHealthIssue InputIssue, Lazy<IReadOnlyList<ChangeInfo>> LazyChanges)
+		BuildHealthIssue MergeIntoExistingIssue(PerforceConnection Perforce, BuildHealthState State, InputJob InputJob, InputJobStep InputJobStep, BuildHealthIssue InputIssue, Lazy<IReadOnlyList<ChangeInfo>> LazyChanges)
 		{
 			// Find the pattern matcher for this fingerprint
 			PatternMatcher Matcher = CategoryToMatcher[InputIssue.Category];
@@ -595,7 +623,7 @@ namespace MetadataTool
 
 				// Add the new build
 				Matcher.Merge(InputIssue, Issue);
-				return true;
+				return Issue;
 			}
 
 			// Check if this issue can be merged with an issue built in another stream
@@ -621,11 +649,11 @@ namespace MetadataTool
 
 					// Merge the issue
 					Matcher.Merge(InputIssue, Issue);
-					return true;
+					return Issue;
 				}
 			}
 
-			return false;
+			return null;
 		}
 
 		/// <summary>
diff --git a/Engine/Source/Programs/UnrealGameSync/MetadataTool/CommandTypes.cs b/Engine/Source/Programs/UnrealGameSync/MetadataTool/CommandTypes.cs
--- a/Engine/Source/Programs/UnrealGameSync/MetadataTool/CommandTypes.cs
+++ b/Engine/Source/Programs/UnrealGameSync/MetadataTool/CommandTypes.cs
@@ -24,22 +24,16 @@ namespace MetadataTool
 			Warning = 3,
 		}
 
-		public class AddIssueResponse
-		{
-			public long Id;
-		}
-
 		public class AddIssue
 		{
 			public string Project;
 			public string Summary;
-			public string Details;
 			[Optional] public string Owner;
 		}
 
-		public class UpdateBuild
+		public class AddIssueResponse
 		{
-			[Optional] public int? Outcome;
+			public long Id;
 		}
 
 		public class UpdateIssue
@@ -50,7 +44,6 @@ namespace MetadataTool
 			[Optional] public string NominatedBy;
 			[Optional] public bool? Resolved;
 			[Optional] public string Summary;
-			[Optional] public string Details;
 			[Optional] public string Url;
 		}
 
@@ -66,6 +59,28 @@ namespace MetadataTool
 			public Outcome Outcome;
 		}
 
+		public class AddBuildResponse
+		{
+			public long Id;
+		}
+
+		public class UpdateBuild
+		{
+			[Optional] public int? Outcome;
+		}
+
+		public class AddDiagnostic
+		{
+			[Optional] public long? BuildId;
+			public string Message;
+			public string Url;
+		}
+
+		public class AddDiagnosticResponse
+		{
+			public long Id;
+		}
+
 		public class Watcher
 		{
 			public string UserName;
diff --git a/Engine/Source/Programs/UnrealGameSync/MetadataTool/Program.cs b/Engine/Source/Programs/UnrealGameSync/MetadataTool/Program.cs
--- a/Engine/Source/Programs/UnrealGameSync/MetadataTool/Program.cs
+++ b/Engine/Source/Programs/UnrealGameSync/MetadataTool/Program.cs
@@ -32,6 +32,10 @@ namespace MetadataTool
 			new CommandHandler_Http("GetBuild", "GET", "/api/issuebuilds/{Build}"),
 			new CommandHandler_Http("UpdateBuild", "PUT", "/api/issuebuilds/{Build}", typeof(CommandTypes.UpdateBuild)),
 
+			// Diagnostics
+			new CommandHandler_Http("AddDiagnostic", "POST", "/api/issues/{Issue}/diagnostics", typeof(CommandTypes.AddDiagnostic)),
+			new CommandHandler_Http("GetDiagnostics", "GET", "/api/issues/{Issue}/diagnostics"),
+
 			// Watchers
 			new CommandHandler_Http("GetWatchers", "GET", "/api/issues/{Issue}/watchers"),
 			new CommandHandler_Http("AddWatcher", "POST", "/api/issues/{Issue}/watchers", typeof(CommandTypes.Watcher)),
diff --git a/Engine/Source/Programs/UnrealGameSync/UnrealGameSync/Controls/RichTextBox41.cs b/Engine/Source/Programs/UnrealGameSync/UnrealGameSync/Controls/RichTextBox41.cs
new file mode 100644
--- /dev/null
+++ b/Engine/Source/Programs/UnrealGameSync/UnrealGameSync/Controls/RichTextBox41.cs
@@ -0,0 +1,29 @@
+using System;
+using System.Collections.Generic;
+using System.Linq;
+using System.Runtime.InteropServices;
+using System.Text;
+using System.Threading.Tasks;
+using System.Windows.Forms;
+
+namespace UnrealGameSync
+{
+	class RichTextBox41 : RichTextBox
+	{
+		[DllImport("kernel32.dll", CharSet = CharSet.Auto)]
+		static extern IntPtr LoadLibrary(string lpFileName);
+
+		protected override CreateParams CreateParams
+		{
+			get
+			{
+				CreateParams Params = base.CreateParams;
+				if (LoadLibrary("msftedit.dll") != IntPtr.Zero)
+				{
+					Params.ClassName = "RICHEDIT50W";
+				}
+				return Params;
+			}
+		}
+	}
+}
diff --git a/Engine/Source/Programs/UnrealGameSync/UnrealGameSync/Forms/IssueDetailsWindow.cs b/Engine/Source/Programs/UnrealGameSync/UnrealGameSync/Forms/IssueDetailsWindow.cs
--- a/Engine/Source/Programs/UnrealGameSync/UnrealGameSync/Forms/IssueDetailsWindow.cs
+++ b/Engine/Source/Programs/UnrealGameSync/UnrealGameSync/Forms/IssueDetailsWindow.cs
@@ -241,6 +241,7 @@ namespace UnrealGameSync
 
 		IssueMonitor IssueMonitor;
 		IssueData Issue;
+		List<IssueDiagnosticData> Diagnostics;
 		PerforceConnection Perforce;
 		TimeSpan? ServerTimeOffset;
 		PerforceWorkerThread PerforceWorker;
@@ -256,10 +257,11 @@ namespace UnrealGameSync
 		System.Windows.Forms.Timer UpdateTimer;
 		StatusElementResources StatusElementResources;
 
-		IssueDetailsWindow(IssueMonitor IssueMonitor, IssueData Issue, string ServerAndPort, string UserName, TimeSpan? ServerTimeOffset, TextWriter Log, string CurrentStream)
+		IssueDetailsWindow(IssueMonitor IssueMonitor, IssueData Issue, List<IssueDiagnosticData> Diagnostics, string ServerAndPort, string UserName, TimeSpan? ServerTimeOffset, TextWriter Log, string CurrentStream)
 		{
 			this.IssueMonitor = IssueMonitor;
 			this.Issue = Issue;
+			this.Diagnostics = Diagnostics;
 			this.Perforce = new PerforceConnection(UserName, null, ServerAndPort);
 			this.ServerTimeOffset = ServerTimeOffset;
 			this.Log = Log;
@@ -372,7 +374,7 @@ namespace UnrealGameSync
 			}
 		}
 
-		void AppendEscapedRtf(StringBuilder Result, string Text)
+		static void AppendEscapedRtf(StringBuilder Result, string Text)
 		{
 			for (int Idx = 0; Idx < Text.Length; Idx++)
 			{
@@ -396,57 +398,57 @@ namespace UnrealGameSync
 			}
 		}
 
-		string CreateRichTextErrors(string ErrorText)
+		static void AppendHyperlink(StringBuilder RichText, string Label, string Url)
 		{
-			string[] Lines = ErrorText.Split('\n');
-			
-			Regex StepPattern = new Regex("^##([^#].*[^#])##$", RegexOptions.None);
-
-			Dictionary<string, int> StepNameToCount = new Dictionary<string, int>(StringComparer.Ordinal);
-			foreach(string Line in Lines)
-			{
-				Match Match = StepPattern.Match(Line);
-				if(Match.Success)
-				{
-					string StepName = Match.Groups[1].Value;
-					int Count;
-					StepNameToCount.TryGetValue(StepName, out Count);
-					StepNameToCount[StepName] = Count + 1;
-				}
-			}
-
-			Dictionary<string, int> StepNameToIndex = StepNameToCount.Keys.ToDictionary(x => x, x => 0, StringComparer.Ordinal);
+			RichText.Append(@"{\field");
+			RichText.Append(@"{\*\fldinst");
+			RichText.AppendFormat("{{ HYPERLINK \"{0}\" }}", Url);
+			RichText.Append(@"}");
+			RichText.Append(@"{\fldrslt ");
+			AppendEscapedRtf(RichText, Label);
+			RichText.Append(@"}");
+			RichText.Append(@"}");
+		}
 
+		static string CreateRichTextErrors(IssueData Issue, List<IssueDiagnosticData> Diagnostics)
+		{
 			StringBuilder RichText = new StringBuilder();
 
 			RichText.AppendLine(@"{\rtf1\ansi");
 			RichText.AppendLine(@"{\fonttbl{\f0\fnil\fcharset0 Arial;}{\f1\fnil\fcharset0 Courier New;}{\f2\fnil\fcharset0 Calibri;}}");
-			RichText.AppendLine(@"{\colortbl;\red192\green80\blue77;}");
+			RichText.AppendLine(@"{\colortbl;\red192\green80\blue77;\red0\green0\blue0;}");
 
-			string PrevStepName = null;
-			for(int Idx = 0; Idx < Lines.Length; Idx++)
+			bool bFirst = true;
+			foreach (IGrouping<long, IssueDiagnosticData> Group in Diagnostics.GroupBy(x => x.BuildId ?? -1))
 			{
-				Match Match = StepPattern.Match(Lines[Idx]);
-				if(Match.Success)
+				// Step 'Foo'
+				IssueBuildData Build = Issue.Builds.FirstOrDefault(x => x.Id == Group.Key);
+				if (Build != null)
 				{
-					// Step 'Foo'
-					string StepName = Match.Groups[1].Value;
-					if(StepName != PrevStepName)
+					RichText.Append(@"\pard");   // Paragraph default
+					RichText.Append(@"\cf2");    // Foreground color
+					RichText.Append(@"\b1");     // Bold
+					RichText.Append(@"\f0");     // Font
+					RichText.Append(@"\fs16");   // Font size
+					if (bFirst)
 					{
-						RichText.Append(@"\pard");   // Paragraph default
-						RichText.Append(@"\cf0");    // Foreground color
-						RichText.Append(@"\b1");     // Bold
-						RichText.Append(@"\f0");     // Font
-						RichText.Append(@"\fs16");   // Font size
-						if (Idx > 0)
-						{
-							RichText.Append(@"\sb300"); // Space before
-						}
-						RichText.AppendFormat(@"In step '{0}':", StepName);
-						RichText.AppendLine(@"\par");
-
-						PrevStepName = StepName;
+						RichText.Append(@"\sb100"); // Space before
+					}
+					else
+					{
+						RichText.Append(@"\sb300"); // Space before
 					}
+					RichText.Append(@" In step '\uld1");
+					AppendHyperlink(RichText, Build.JobStepName, Build.JobStepUrl);
+					RichText.Append(@"\uld0':");
+
+					RichText.AppendLine(@"\par");
+				}
+
+				IssueDiagnosticData[] DiagnosticsArray = Group.ToArray();
+				for(int Idx = 0; Idx < DiagnosticsArray.Length; Idx++)
+				{
+					IssueDiagnosticData Diagnostic = DiagnosticsArray[Idx];
 
 					// Error X/Y:
 					RichText.Append(@"\pard");   // Paragraph default
@@ -458,22 +460,30 @@ namespace UnrealGameSync
 					RichText.Append(@"\li50");   // Other line indent
 					RichText.Append(@"\sb100");  // Space before
 					RichText.Append(@"\sa50");   // Space after
-					RichText.AppendFormat(@"Error {0}/{1}", ++StepNameToIndex[StepName], StepNameToCount[StepName]);
+					RichText.Append(@" ");
+
+					RichText.Append(@"\ul1");
+					AppendHyperlink(RichText, String.Format("Error {0}/{1}", Idx + 1, DiagnosticsArray.Length), Diagnostic.Url);
+					RichText.Append(@"\ul0");
+
 					RichText.AppendLine(@"\par");
-				}
-				else
-				{
+
 					// Error text
-					RichText.Append(@"\pard");   // Paragraph default
-					RichText.Append(@"\cf0");    // Foreground color
-					RichText.Append(@"\b0");     // Bold
-					RichText.Append(@"\f1");     // Font
-					RichText.Append(@"\fs16");   // Font size 16
-					RichText.Append(@"\fi150");  // First line indent
-					RichText.Append(@"\li150");  // Other line indent
-					AppendEscapedRtf(RichText, Lines[Idx]);
-					RichText.Append(@"\par");
+					foreach(string Line in Diagnostic.Message.TrimEnd().Split('\n'))
+					{
+						RichText.Append(@"\pard");   // Paragraph default
+						RichText.Append(@"\cf0");    // Foreground color
+						RichText.Append(@"\b0");     // Bold
+						RichText.Append(@"\f1");     // Font
+						RichText.Append(@"\fs16");   // Font size 16
+						RichText.Append(@"\fi150");  // First line indent
+						RichText.Append(@"\li150");  // Other line indent
+						AppendEscapedRtf(RichText, Line);
+						RichText.Append(@"\par");
+					}
 				}
+
+				bFirst = false;
 			}
 
 			RichText.AppendLine("}");
@@ -560,15 +570,15 @@ namespace UnrealGameSync
 			UpdateSummaryTextIfChanged(StepNamesTextBox, String.Join(", ", Issue.Builds.Select(x => x.JobStepName).Distinct().OrderBy(x => x)));
 			UpdateSummaryTextIfChanged(StreamNamesTextBox, String.Join(", ", Issue.Builds.Select(x => x.Stream).Distinct().OrderBy(x => x)));
 
-			if(LastDetailsText != Issue.Details)
+			string RtfText = CreateRichTextErrors(Issue, Diagnostics);
+			if (LastDetailsText != RtfText)
 			{
-				string RtfText = CreateRichTextErrors(Issue.Details);
 				using (MemoryStream Stream = new MemoryStream(Encoding.UTF8.GetBytes(RtfText), false))
 				{
 					DetailsTextBox.LoadFile(Stream, RichTextBoxStreamType.RichText);
 					DetailsTextBox.Select(0, 0);
 				}
-				LastDetailsText = Issue.Details;
+				LastDetailsText = RtfText;
 			}
 
 			if(Issue.FixChange == 0)
@@ -999,12 +1009,42 @@ namespace UnrealGameSync
 
 		static List<IssueDetailsWindow> ExistingWindows = new List<IssueDetailsWindow>();
 
+		class UpdateIssueDetailsTask : IModalTask
+		{
+			string ApiUrl;
+			long IssueId;
+			List<IssueDiagnosticData> Diagnostics;
+
+			public UpdateIssueDetailsTask(string ApiUrl, long IssueId, List<IssueDiagnosticData> Diagnostics)
+			{
+				this.ApiUrl = ApiUrl;
+				this.IssueId = IssueId;
+				this.Diagnostics = Diagnostics;
+			}
+
+			public bool Run(out string ErrorMessage)
+			{
+				Diagnostics.AddRange(RESTApi.GET<List<IssueDiagnosticData>>(ApiUrl, String.Format("issues/{0}/diagnostics", IssueId)));
+
+				ErrorMessage = null;
+				return true;
+			}
+		}
+
 		public static void Show(Form Owner, IssueMonitor IssueMonitor, string ServerAndPort, string UserName, TimeSpan? ServerTimeOffset, IssueData Issue, TextWriter Log, string CurrentStream)
 		{
 			IssueDetailsWindow Window = ExistingWindows.FirstOrDefault(x => x.IssueMonitor == IssueMonitor && x.Issue.Id == Issue.Id);
 			if(Window == null)
 			{
-				Window = new IssueDetailsWindow(IssueMonitor, Issue, ServerAndPort, UserName, ServerTimeOffset, Log, CurrentStream);
+				List<IssueDiagnosticData> Diagnostics = new List<IssueDiagnosticData>();
+
+				UpdateIssueDetailsTask Task = new UpdateIssueDetailsTask(IssueMonitor.ApiUrl, Issue.Id, Diagnostics);
+				if(!ModalTask.ExecuteAndShowError(Owner, Task, "Fetching data", "Fetchind data, please wait..."))
+				{
+					return;
+				}
+
+				Window = new IssueDetailsWindow(IssueMonitor, Issue, Diagnostics, ServerAndPort, UserName, ServerTimeOffset, Log, CurrentStream);
 				Window.Owner = Owner;
 				Window.StartPosition = FormStartPosition.Manual;
 				Window.Location = new Point(Owner.Location.X + (Owner.Width - Window.Width) / 2, Owner.Location.Y + (Owner.Height - Window.Height) / 2);
@@ -1357,5 +1397,10 @@ namespace UnrealGameSync
 				System.Diagnostics.Process.Start(Build.ErrorUrl);
 			}
 		}
+
+		private void DetailsTextBox_LinkClicked(object sender, LinkClickedEventArgs e)
+		{
+			System.Diagnostics.Process.Start(e.LinkText);
+		}
 	}
 }
diff --git a/Engine/Source/Programs/UnrealGameSync/UnrealGameSync/IssueMonitor.cs b/Engine/Source/Programs/UnrealGameSync/UnrealGameSync/IssueMonitor.cs
--- a/Engine/Source/Programs/UnrealGameSync/UnrealGameSync/IssueMonitor.cs
+++ b/Engine/Source/Programs/UnrealGameSync/UnrealGameSync/IssueMonitor.cs
@@ -21,6 +21,7 @@ namespace UnrealGameSync
 
 	public class IssueBuildData
 	{
+		public long Id;
 		public string Stream;
 		public int Change;
 		public string JobName;
@@ -31,6 +32,13 @@ namespace UnrealGameSync
 		public IssueBuildOutcome Outcome;
 	}
 
+	public class IssueDiagnosticData
+	{
+		public long? BuildId;
+		public string Message;
+		public string Url;
+	}
+
 	public class IssueData
 	{
 		public long Id;
@@ -363,7 +371,7 @@ namespace UnrealGameSync
 				}
 
 				// Apply any pending updates to this issue list, and update it
-				lock(LockObject)
+				lock (LockObject)
 				{
 					foreach(IssueUpdateData PendingUpdate in PendingUpdates)
 					{
diff --git a/Engine/Source/Programs/UnrealGameSync/UnrealGameSync/ModalTask.cs b/Engine/Source/Programs/UnrealGameSync/UnrealGameSync/ModalTask.cs
--- a/Engine/Source/Programs/UnrealGameSync/UnrealGameSync/ModalTask.cs
+++ b/Engine/Source/Programs/UnrealGameSync/UnrealGameSync/ModalTask.cs
@@ -31,5 +31,20 @@ namespace UnrealGameSync
 			ErrorMessage = Window.ErrorMessage;
 			return Window.Result;
 		}
+
+		public static bool ExecuteAndShowError(IWin32Window Owner, IModalTask Task, string InTitle, string InMessage)
+		{
+			string ErrorMessage;
+			ModalTaskResult Result = Execute(Owner, Task, InTitle, InMessage, out ErrorMessage);
+			if(Result != ModalTaskResult.Succeeded)
+			{
+				if (!String.IsNullOrEmpty(ErrorMessage))
+				{
+					MessageBox.Show(ErrorMessage);
+				}
+				return false;
+			}
+			return true;
+		}
 	}
 }
diff --git a/Engine/Source/Programs/UnrealGameSync/UnrealGameSync/UnrealGameSync.csproj b/Engine/Source/Programs/UnrealGameSync/UnrealGameSync/UnrealGameSync.csproj
--- a/Engine/Source/Programs/UnrealGameSync/UnrealGameSync/UnrealGameSync.csproj
+++ b/Engine/Source/Programs/UnrealGameSync/UnrealGameSync/UnrealGameSync.csproj
@@ -125,6 +125,9 @@
     <Compile Include="Controls\ErrorPanel.cs">
       <SubType>Component</SubType>
     </Compile>
+    <Compile Include="Controls\RichTextBox41.cs">
+      <SubType>Component</SubType>
+    </Compile>
     <Compile Include="Controls\SyncFilterControl.cs">
       <SubType>UserControl</SubType>
     </Compile>
