Fix typo

#rnx

diff --git a/Engine/Source/Programs/UnrealGameSync/MetadataTool/CommandHandler_BuildHealth.cs b/Engine/Source/Programs/UnrealGameSync/MetadataTool/CommandHandler_BuildHealth.cs
--- a/Engine/Source/Programs/UnrealGameSync/MetadataTool/CommandHandler_BuildHealth.cs
+++ b/Engine/Source/Programs/UnrealGameSync/MetadataTool/CommandHandler_BuildHealth.cs
@@ -272,7 +272,7 @@ namespace MetadataTool
 						Issue.PostedSummary = Summary;
 						WriteState(StateFile, State);
 					}
-					else if(Issue.PostedSummary == null || !String.Equals(Issue.PostedSummary, Summary, StringComparison.Ordinal)
+					else if(Issue.PostedSummary == null || !String.Equals(Issue.PostedSummary, Summary, StringComparison.Ordinal))
 					{
 						Log.TraceInformation("Updating issue {0}", Issue.Id);
 
