Fix MetadataTool exception when backing up unmatched error file.

#rnx

diff --git a/Engine/Source/Programs/UnrealGameSync/MetadataTool/CommandHandler_BuildHealth.cs b/Engine/Source/Programs/UnrealGameSync/MetadataTool/CommandHandler_BuildHealth.cs
--- a/Engine/Source/Programs/UnrealGameSync/MetadataTool/CommandHandler_BuildHealth.cs
+++ b/Engine/Source/Programs/UnrealGameSync/MetadataTool/CommandHandler_BuildHealth.cs
@@ -140,7 +140,10 @@ namespace MetadataTool
 				if(SaveUnmatchedDir != null && InputJobs.Count > 0)
 				{
 					DirectoryReference.CreateDirectory(SaveUnmatchedDir);
-					FileReference.Copy(StateFile, FileReference.Combine(SaveUnmatchedDir, "State.json"), true);
+					if(FileReference.Exists(StateFile))
+					{
+						FileReference.Copy(StateFile, FileReference.Combine(SaveUnmatchedDir, "State.json"), true);
+					}
 					SerializeJson(FileReference.Combine(SaveUnmatchedDir, "Input.json"), InputData);
 				}
 
