Fix MetadataTool server exception.

#rnx

diff --git a/Engine/Source/Programs/UnrealGameSync/MetadataServer/Connectors/SqlConnector.cs b/Engine/Source/Programs/UnrealGameSync/MetadataServer/Connectors/SqlConnector.cs
--- a/Engine/Source/Programs/UnrealGameSync/MetadataServer/Connectors/SqlConnector.cs
+++ b/Engine/Source/Programs/UnrealGameSync/MetadataServer/Connectors/SqlConnector.cs
@@ -393,11 +393,10 @@ namespace MetadataServer.Connectors
 			{
 				Connection.Open();
 
-				using (SQLiteCommand Command = new SQLiteCommand("INSERT INTO [Issues] (Project, Summary, Details, OwnerId, CreatedAt, FixChange) VALUES (@Project, @Summary, @Details, @OwnerId, DATETIME('now'), 0)", Connection))
+				using (SQLiteCommand Command = new SQLiteCommand("INSERT INTO [Issues] (Project, Summary, OwnerId, CreatedAt, FixChange) VALUES (@Project, @Summary, @OwnerId, DATETIME('now'), 0)", Connection))
 				{
 					Command.Parameters.AddWithValue("@Project", Issue.Project);
 					Command.Parameters.AddWithValue("@Summary", SanitizeText(Issue.Summary, IssueSummaryMaxLength));
-					Command.Parameters.AddWithValue("@Details", "");
 					if (Issue.Owner != null)
 					{
 						Command.Parameters.AddWithValue("OwnerId", FindOrAddUserId(Issue.Owner, Connection));
diff --git a/Engine/Source/Programs/UnrealGameSync/MetadataTool/CommandHandler_BuildHealth.cs b/Engine/Source/Programs/UnrealGameSync/MetadataTool/CommandHandler_BuildHealth.cs
--- a/Engine/Source/Programs/UnrealGameSync/MetadataTool/CommandHandler_BuildHealth.cs
+++ b/Engine/Source/Programs/UnrealGameSync/MetadataTool/CommandHandler_BuildHealth.cs
@@ -262,7 +262,7 @@ namespace MetadataTool
 							IssueBody.Owner = Issue.PendingWatchers.First();
 						}
 
-						using(HttpWebResponse Response = SendHttpRequest(String.Format("{0}/api/issues/{1}", ServerUrl, Issue.Id), "POST", IssueBody))
+						using(HttpWebResponse Response = SendHttpRequest(String.Format("{0}/api/issues", ServerUrl), "POST", IssueBody))
 						{
 							int ResponseCode = (int)Response.StatusCode;
 							if (!(ResponseCode >= 200 && ResponseCode <= 299))
