UGS: Always link to the job from the issue details window so it matches the label, rather than sometimes linking to the job step.

#rnx

diff --git a/Engine/Source/Programs/UnrealGameSync/UnrealGameSync/Forms/IssueDetailsWindow.cs b/Engine/Source/Programs/UnrealGameSync/UnrealGameSync/Forms/IssueDetailsWindow.cs
--- a/Engine/Source/Programs/UnrealGameSync/UnrealGameSync/Forms/IssueDetailsWindow.cs
+++ b/Engine/Source/Programs/UnrealGameSync/UnrealGameSync/Forms/IssueDetailsWindow.cs
@@ -438,9 +438,9 @@ namespace UnrealGameSync
 					{
 						RichText.Append(@"\sb300"); // Space before
 					}
-					RichText.Append(@" In step '\uld1");
+					RichText.Append(@" In step '\ul1");
 					AppendHyperlink(RichText, Build.JobStepName, Build.JobStepUrl);
-					RichText.Append(@"\uld0':");
+					RichText.Append(@"\ul0':");
 
 					RichText.AppendLine(@"\par");
 				}
@@ -803,8 +803,7 @@ namespace UnrealGameSync
 				StatusLineListViewWidget BuildWidget = new StatusLineListViewWidget(BuildItem, StatusElementResources);
 				BuildWidget.HorizontalAlignment = HorizontalAlignment.Left;
 
-				string DefaultUrl = Range.BuildGroup.Builds.Select(x => x.ErrorUrl).FirstOrDefault(x => x != null) ?? Range.BuildGroup.JobUrl;
-				BuildWidget.Line.AddLink(Range.BuildGroup.JobName, FontStyle.Underline, () => System.Diagnostics.Process.Start(DefaultUrl));
+				BuildWidget.Line.AddLink(Range.BuildGroup.JobName, FontStyle.Underline, () => System.Diagnostics.Process.Start(Range.BuildGroup.JobUrl));
 				BuildItem.SubItems.Add(new ListViewItem.ListViewSubItem(BuildItem, ""){ Tag = BuildWidget });
 				BuildItem.SubItems.Add(new ListViewItem.ListViewSubItem(BuildItem, ""){ Tag = BuildWidget });
 
