UGS: Prefer to merge issues according to matching filenames/identifiers where possible, and only fall back to merging with existing issues for the current job as a last resort.

#rnx

diff --git a/Engine/Source/Programs/UnrealGameSync/MetadataTool/BuildHealth/PatternMatcher.cs b/Engine/Source/Programs/UnrealGameSync/MetadataTool/BuildHealth/PatternMatcher.cs
--- a/Engine/Source/Programs/UnrealGameSync/MetadataTool/BuildHealth/PatternMatcher.cs
+++ b/Engine/Source/Programs/UnrealGameSync/MetadataTool/BuildHealth/PatternMatcher.cs
@@ -54,10 +54,10 @@ namespace MetadataTool
 		public abstract bool TryMatch(InputJob Job, InputJobStep JobStep, InputDiagnostic Diagnostic, List<BuildHealthIssue> Issues);
 
 		/// <summary>
-		/// Determines if one fingerprint can be merged with another one
+		/// Determines if one issue can be merged into another
 		/// </summary>
-		/// <param name="Source">The source fingerprint</param>
-		/// <param name="Target">The fingerprint to merge into</param>
+		/// <param name="Source">The source issue</param>
+		/// <param name="Target">The target issue</param>
 		public virtual bool CanMerge(BuildHealthIssue Source, BuildHealthIssue Target)
 		{
 			// Make sure the categories match
@@ -67,16 +67,25 @@ namespace MetadataTool
 			}
 
 			// Check that a filename or message matches
-			if (Source.InitialJobUrl != Target.InitialJobUrl)
+			if (!Source.FileNames.Any(x => Target.FileNames.Contains(x)) && !Source.Identifiers.Any(x => Target.Identifiers.Contains(x)))
 			{
-				if (!Source.FileNames.Any(x => Target.FileNames.Contains(x)) && !Source.Identifiers.Any(x => Target.Identifiers.Contains(x)))
-				{
-					return false;
-				}
+				return false;
 			}
+
 			return true;
 		}
 
+		/// <summary>
+		/// Determines if an issue can be merged into another issue that occurred at the same initial job
+		/// </summary>
+		/// <param name="Source">The source issue</param>
+		/// <param name="Target">The target issue</param>
+		/// <returns>True if the two new issues can be merged</returns>
+		public virtual bool CanMergeInitialJob(BuildHealthIssue Source, BuildHealthIssue Target)
+		{
+			return Source.Category == Target.Category;
+		}
+
 		/// <summary>
 		/// Merge one fingerprint with another
 		/// </summary>
diff --git a/Engine/Source/Programs/UnrealGameSync/MetadataTool/CommandHandler_BuildHealth.cs b/Engine/Source/Programs/UnrealGameSync/MetadataTool/CommandHandler_BuildHealth.cs
--- a/Engine/Source/Programs/UnrealGameSync/MetadataTool/CommandHandler_BuildHealth.cs
+++ b/Engine/Source/Programs/UnrealGameSync/MetadataTool/CommandHandler_BuildHealth.cs
@@ -656,6 +656,16 @@ namespace MetadataTool
 				}
 			}
 
+			// Check if it can be merged into an issue that's been created for this job. We only do this after exhausting all other options.
+			foreach (BuildHealthIssue Issue in State.Issues)
+			{
+				if(Issue.InitialJobUrl == InputIssue.InitialJobUrl && Matcher.CanMergeInitialJob(InputIssue, Issue))
+				{
+					Matcher.Merge(InputIssue, Issue);
+					return Issue;
+				}
+			}
+
 			return null;
 		}
 
