UGS: Change the content pattern matcher to just recognize things output by the UE_ASSET_LOG macro, to avoid ambiguities.

#rnx

diff --git a/Engine/Source/Programs/UnrealGameSync/MetadataTool/BuildHealth/Matchers/ContentPatternMatcher.cs b/Engine/Source/Programs/UnrealGameSync/MetadataTool/BuildHealth/Matchers/ContentPatternMatcher.cs
--- a/Engine/Source/Programs/UnrealGameSync/MetadataTool/BuildHealth/Matchers/ContentPatternMatcher.cs
+++ b/Engine/Source/Programs/UnrealGameSync/MetadataTool/BuildHealth/Matchers/ContentPatternMatcher.cs
@@ -16,64 +16,31 @@ namespace MetadataTool
 
 		public override bool TryMatch(InputJob Job, InputJobStep JobStep, InputDiagnostic Diagnostic, List<BuildHealthIssue> Issues)
 		{
-			if(JobStep.Name.Contains("Check Asset References") && Regex.IsMatch(Diagnostic.Message, @"^\s*Log[a-zA-Z0-9]+: (?:Error|Warning):"))
+			HashSet<string> FileNames = new HashSet<string>();
+			foreach(Match Match in Regex.Matches(Diagnostic.Message, @"^\s*Log[a-zA-Z0-9]+:\s+(?:Error:|Warning:)\s+((?:[a-zA-Z]:)?[^:]+(?:.uasset|.umap)):\s*(.*)"))
 			{
-				BuildHealthIssue Issue = new BuildHealthIssue(Job.Project, Category, Job.Url, new BuildHealthDiagnostic(JobStep.Name, JobStep.Url, Diagnostic.Message, Diagnostic.Url));
-
-				// Use the error text to identify this error
-				Issue.Identifiers.Add(Diagnostic.Message);
-
-				// Find any asset references
-				foreach (Match FileMatch in Regex.Matches(Diagnostic.Message, @"[^a-zA-Z0-9_](?:/Game/|/Engine/)([^(). ]+)"))
-				{
-					if (FileMatch.Success)
-					{
-						Issue.References.Add(FileMatch.Groups[1].Value);
-					}
-				}
+				FileNames.Add(GetNormalizedFileName(Match.Groups[1].Value, JobStep.BaseDirectory));
+			}
 
+			if(FileNames.Count > 0)
+			{
+				BuildHealthIssue Issue = new BuildHealthIssue(Job.Project, Category, Job.Url, new BuildHealthDiagnostic(JobStep.Name, JobStep.Url, Diagnostic.Message, Diagnostic.Url));
+				Issue.FileNames.UnionWith(FileNames);
 				Issues.Add(Issue);
 				return true;
 			}
 			return false;
 		}
 
-		public override List<ChangeInfo> FindCausers(PerforceConnection Perforce, BuildHealthIssue Issue, IReadOnlyList<ChangeInfo> Changes)
-		{
-			// Check for any changes that modify those files
-			List<ChangeInfo> Causers = new List<ChangeInfo>();
-			foreach(ChangeInfo Change in Changes)
-			{
-				DescribeRecord DescribeRecord = GetDescribeRecord(Perforce, Change);
-				foreach(DescribeFileRecord File in DescribeRecord.Files)
-				{
-					if(Issue.References.Any(Reference => File.DepotFile.IndexOf(Reference, StringComparison.OrdinalIgnoreCase) != -1))
-					{
-						Causers.Add(Change);
-						break;
-					}
-				}
-			}
-
-			if(Causers.Count > 0)
-			{
-				return Causers;
-			}
-			else
-			{
-				return base.FindCausers(Perforce, Issue, Changes);
-			}
-		}
-
 		public override string GetSummary(BuildHealthIssue Issue)
 		{
-			if (Issue.References.Count == 0)
+			if (Issue.FileNames.Count == 0)
 			{
 				return "Content errors";
 			}
 			else
 			{
-				return String.Format("Content errors in {0}", String.Join(", ", Issue.References));
+				return String.Format("Content errors in {0}", String.Join(", ", GetFileNamesWithoutPath(Issue.FileNames)));
 			}
 		}
 	}
