UGS: Tweak appearance of the build health notification.

diff --git a/Engine/Source/Programs/UnrealGameSync/UnrealGameSync/Controls/AlertButtonControl.cs b/Engine/Source/Programs/UnrealGameSync/UnrealGameSync/Controls/AlertButtonControl.cs
--- a/Engine/Source/Programs/UnrealGameSync/UnrealGameSync/Controls/AlertButtonControl.cs
+++ b/Engine/Source/Programs/UnrealGameSync/UnrealGameSync/Controls/AlertButtonControl.cs
@@ -130,8 +130,8 @@ namespace UnrealGameSync
 				case AlertButtonTheme.Strong:
 					Colors.ForeColor = Color.FromArgb(255, 255, 255);
 					Colors.BorderColor = Color.FromArgb(230, 232, 235);
-					Colors.BackgroundColor1 = Color.FromArgb(214, 69, 64);
-					Colors.BackgroundColor2 = Color.FromArgb(214, 69, 64);
+					Colors.BackgroundColor1 = Color.FromArgb(200, 74, 49);
+					Colors.BackgroundColor2 = Color.FromArgb(200, 74, 49); 
 					Colors.BackgroundColorHover1 = Color.FromArgb(222, 108, 86);
 					Colors.BackgroundColorHover2 = Color.FromArgb(222, 108, 86);
 					Colors.BackgroundColorDown1 = Color.FromArgb(204, 59, 54);
diff --git a/Engine/Source/Programs/UnrealGameSync/UnrealGameSync/Forms/IssueAlertWindow.cs b/Engine/Source/Programs/UnrealGameSync/UnrealGameSync/Forms/IssueAlertWindow.cs
--- a/Engine/Source/Programs/UnrealGameSync/UnrealGameSync/Forms/IssueAlertWindow.cs
+++ b/Engine/Source/Programs/UnrealGameSync/UnrealGameSync/Forms/IssueAlertWindow.cs
@@ -55,11 +55,15 @@ namespace UnrealGameSync
 
 			if(bStrongAlert ?? false)
 			{
-				Color StripeColor = bIsWarning? Color.FromArgb(216, 167, 64) : Color.FromArgb(214, 69, 64);
-				using(Brush StripeBrush = new SolidBrush(StripeColor))
+				Color StripeColor = bIsWarning? Color.FromArgb(216, 167, 64) : Color.FromArgb(200, 74, 49);//214, 69, 64);
+				using (Brush StripeBrush = new SolidBrush(StripeColor))
 				{
 					e.Graphics.FillRectangle(StripeBrush, 0, 0, Bounds.Width, Bounds.Height);
 				}
+				using (Pen Pen = new Pen(Color.FromArgb(255, 255, 255)))
+				{
+					e.Graphics.DrawRectangle(Pen, 0, 0, Bounds.Width - 1, Bounds.Height - 1);
+				}
 			}
 			else
 			{
@@ -95,9 +99,16 @@ namespace UnrealGameSync
 			}
 			else
 			{
-				if(String.Compare(NewIssue.Owner, IssueMonitor.UserName, StringComparison.OrdinalIgnoreCase) == 0 && NewIssue.NominatedBy != null)
+				if(String.Compare(NewIssue.Owner, IssueMonitor.UserName, StringComparison.OrdinalIgnoreCase) == 0)
 				{
-					OwnerTextBuilder.AppendFormat("You have been nominated to fix this issue by {0}.", Utility.FormatUserName(NewIssue.NominatedBy));
+					if(NewIssue.NominatedBy != null)
+					{
+						OwnerTextBuilder.AppendFormat("You have been nominated to fix this issue by {0}.", Utility.FormatUserName(NewIssue.NominatedBy));
+					}
+					else
+					{
+						OwnerTextBuilder.AppendFormat("Assigned to {0}.", Utility.FormatUserName(NewIssue.Owner));
+					}
 					bNewStrongAlert = true;
 				}
 				else
@@ -147,6 +158,7 @@ namespace UnrealGameSync
 						SummaryLabel.ForeColor = Color.FromArgb(255, 255, 255);
 						SummaryLabel.LinkColor = Color.FromArgb(255, 255, 255);
 						OwnerLabel.ForeColor = Color.FromArgb(255, 255, 255);
+						DetailsBtn.Theme = AlertButtonControl.AlertButtonTheme.Strong;
 						AcceptBtn.Theme = AlertButtonControl.AlertButtonTheme.Strong;
 						LatestBuildLinkLabel.LinkColor = Color.FromArgb(255, 255, 255);
 					}
@@ -155,6 +167,7 @@ namespace UnrealGameSync
 						SummaryLabel.ForeColor = Color.FromArgb(32, 32, 64);
 						SummaryLabel.LinkColor = Color.FromArgb(32, 32, 64);
 						OwnerLabel.ForeColor = Color.FromArgb(32, 32, 64);
+						DetailsBtn.Theme = AlertButtonControl.AlertButtonTheme.Normal;
 						AcceptBtn.Theme = AlertButtonControl.AlertButtonTheme.Green;
 						LatestBuildLinkLabel.LinkColor = Color.FromArgb(16, 102, 192);
 					}
@@ -172,10 +185,7 @@ namespace UnrealGameSync
 					Reason = NewReason;
 
 					List<Button> Buttons = new List<Button>();
-					if(!bStrongAlert.Value)
-					{
-						Buttons.Add(DetailsBtn);
-					}
+					Buttons.Add(DetailsBtn);
 					if((NewReason & IssueAlertReason.Owner) != 0)
 					{
 						AcceptBtn.Text = "Acknowledge";
