UGS: Add a "Dismiss" button for notifications that the current user does not own.

#rnx

diff --git a/Engine/Source/Programs/UnrealGameSync/UnrealGameSync/Forms/IssueAlertWindow.cs b/Engine/Source/Programs/UnrealGameSync/UnrealGameSync/Forms/IssueAlertWindow.cs
--- a/Engine/Source/Programs/UnrealGameSync/UnrealGameSync/Forms/IssueAlertWindow.cs
+++ b/Engine/Source/Programs/UnrealGameSync/UnrealGameSync/Forms/IssueAlertWindow.cs
@@ -198,6 +198,11 @@ namespace UnrealGameSync
 						DeclineBtn.Text = "Not Me";
 						Buttons.Add(DeclineBtn);
 					}
+					else
+					{
+						DeclineBtn.Text = "Dismiss";
+						Buttons.Add(DeclineBtn);
+					}
 
 					tableLayoutPanel3.ColumnCount = Buttons.Count;
 					tableLayoutPanel3.Controls.Clear();
diff --git a/Engine/Source/Programs/UnrealGameSync/UnrealGameSync/Forms/IssueFixedWindow.cs b/Engine/Source/Programs/UnrealGameSync/UnrealGameSync/Forms/IssueFixedWindow.cs
--- a/Engine/Source/Programs/UnrealGameSync/UnrealGameSync/Forms/IssueFixedWindow.cs
+++ b/Engine/Source/Programs/UnrealGameSync/UnrealGameSync/Forms/IssueFixedWindow.cs
@@ -84,7 +84,7 @@ namespace UnrealGameSync
 					List<PerforceDescribeRecord> Descriptions = null;
 
 					List<PerforceChangeSummary> Changes;
-					if(Perforce.FindChanges(new string[]{ "//..." }, UserName, 250, out Changes, Log))
+					if(Perforce.FindChanges(new string[]{ "//..." }, UserName, 100, out Changes, Log))
 					{
 						Perforce.DescribeMultiple(Changes.Select(x => x.Number), out Descriptions, Log);
 					}
