UGS: Rewrite workspace syncing to do all filtering in one place, and to convert to client-relative paths where possible.

#rnx

diff --git a/Engine/Source/Programs/UnrealGameSync/UnrealGameSync/Workspace.cs b/Engine/Source/Programs/UnrealGameSync/UnrealGameSync/Workspace.cs
--- a/Engine/Source/Programs/UnrealGameSync/UnrealGameSync/Workspace.cs
+++ b/Engine/Source/Programs/UnrealGameSync/UnrealGameSync/Workspace.cs
@@ -362,10 +362,9 @@ namespace UnrealGameSync
 					List<string> SyncPaths = GetSyncPaths((Context.Options & WorkspaceUpdateOptions.SyncAllProjects) != 0, Context.SyncFilter);
 
 					// Get the user's sync filter
-					FileFilter UserFilter = null;
+					FileFilter UserFilter = new FileFilter(FileFilterType.Include);
 					if(Context.SyncFilter != null)
 					{
-						UserFilter = new FileFilter(FileFilterType.Include);
 						UserFilter.AddRules(Context.SyncFilter.Select(x => x.Trim()).Where(x => x.Length > 0 && !x.StartsWith(";") && !x.StartsWith("#")));
 					}
 
@@ -490,12 +489,34 @@ namespace UnrealGameSync
 						CurrentSyncFilterHash = NextSyncFilterHash;
 					}
 
+					// Create a filter for all the files we don't want
+					FileFilter Filter = new FileFilter(UserFilter);
+					Filter.Exclude("..." + BuildVersionFileName);
+					if (!ProjectConfigFile.GetValue("Options.UseFastModularVersioningV2", false))
+					{
+						Filter.Exclude("..." + VersionHeaderFileName);
+						Filter.Exclude("..." + ObjectVersionFileName);
+					}
+					if (Context.Options.HasFlag(WorkspaceUpdateOptions.ContentOnly))
+					{
+						Filter.Exclude("*.usf");
+						Filter.Exclude("*.ush");
+					}
+
 					// Find all the server changes, and anything that's opened for edit locally. We need to sync files we have open to schedule a resolve.
+					SortedSet<string> SyncCommands = new SortedSet<string>(StringComparer.Ordinal);
 					List<string> SyncDepotPaths = new List<string>();
 					using(RecordCounter Counter = new RecordCounter(Progress, "Filtering files..."))
 					{
 						foreach(string SyncPath in SyncPaths)
 						{
+							if(!SyncPath.StartsWith(ClientRootPath))
+							{
+								Log.WriteLine("Sync path ({0}) does not start with client root path ({1}) - unable to sync.", SyncPath, ClientRootPath);
+								StatusMessage = "Internal error while syncing";
+								return WorkspaceUpdateResult.FailedToSync;
+							}
+
 							List<PerforceFileRecord> SyncRecords = new List<PerforceFileRecord>();
 							if(!Perforce.SyncPreview(SyncPath, PendingChangeNumber, !Context.Options.HasFlag(WorkspaceUpdateOptions.Sync), Record => { SyncRecords.Add(Record); Counter.Increment(); }, Log))
 							{
@@ -503,14 +524,30 @@ namespace UnrealGameSync
 								return WorkspaceUpdateResult.FailedToSync;
 							}
 
-							foreach(PerforceFileRecord SyncRecord in SyncRecords)
+							List<PerforceFileRecord> OpenRecords;
+							if (!Perforce.GetOpenFiles(SyncPath, out OpenRecords, Log))
 							{
+								StatusMessage = String.Format("Couldn't find open files matching {0}.", SyncPath);
+								return WorkspaceUpdateResult.FailedToSync;
+							}
+
+							SyncRecords.AddRange(OpenRecords.Where(x => x.Action != "add" && x.Action != "branch" && x.Action != "move/add"));
+
+							foreach (PerforceFileRecord SyncRecord in SyncRecords)
+							{
+								// If it doesn't exist locally, just add a sync command for it
+								if (String.IsNullOrEmpty(SyncRecord.ClientPath))
+								{
+									SyncCommands.Add(String.Format("{0}@{1}", SyncRecord.DepotPath, PendingChangeNumber));
+									SyncDepotPaths.Add(SyncRecord.DepotPath);
+									continue;
+								}
+
+								// Get the full local path
+								string FullName;
 								try
 								{
-									if(!String.IsNullOrEmpty(SyncRecord.ClientPath))
-									{
-										Path.GetFullPath(SyncRecord.ClientPath);
-									}
+									FullName = Path.GetFullPath(SyncRecord.ClientPath);
 								}
 								catch(PathTooLongException)
 								{
@@ -518,51 +555,31 @@ namespace UnrealGameSync
 									StatusMessage = "File exceeds maximum path length allowed by Windows.";
 									return WorkspaceUpdateResult.FailedToSync;
 								}
-							}
 
-							if(UserFilter != null)
-							{
-								SyncRecords.RemoveAll(x => !String.IsNullOrEmpty(x.ClientPath) && !MatchFilter(Path.GetFullPath(x.ClientPath), UserFilter));
-							}
-
-							SyncDepotPaths.AddRange(SyncRecords.Select(x => x.DepotPath));
+								// Make sure it's under the current directory. Not sure why this would happen, just being safe.
+								if (!FullName.StartsWith(LocalRootPath, StringComparison.OrdinalIgnoreCase))
+								{
+									SyncCommands.Add(String.Format("{0}@{1}", SyncRecord.DepotPath, PendingChangeNumber));
+									SyncDepotPaths.Add(SyncRecord.DepotPath);
+									continue;
+								}
 
-							List<PerforceFileRecord> OpenRecords;
-							if(!Perforce.GetOpenFiles(SyncPath, out OpenRecords, Log))
-							{
-								StatusMessage = String.Format("Couldn't find open files matching {0}.", SyncPath);
-								return WorkspaceUpdateResult.FailedToSync;
+								// Check that it matches the filter
+								string RelativePath = FullName.Substring(LocalRootPath.Length).Replace('\\', '/');
+								if (UserFilter.Matches(RelativePath))
+								{
+									SyncCommands.Add(String.Format("{0}{1}@{2}", ClientRootPath, PerforceUtils.EscapePath(RelativePath), PendingChangeNumber));
+									SyncDepotPaths.Add(SyncRecord.DepotPath);
+								}
 							}
-
-							// don't force a sync on added files
-							SyncDepotPaths.AddRange(OpenRecords.Where(x => x.Action != "add" && x.Action != "branch" && x.Action != "move/add").Select(x => x.DepotPath));
 						}
 					}
 
-					// Filter out all the binaries that we don't want
-					FileFilter Filter = new FileFilter(FileFilterType.Include);
-					Filter.Exclude("..." + BuildVersionFileName);
-					if(!ProjectConfigFile.GetValue("Options.UseFastModularVersioningV2", false))
-					{
-						Filter.Exclude("..." + VersionHeaderFileName);
-						Filter.Exclude("..." + ObjectVersionFileName);
-					}
-					if(Context.Options.HasFlag(WorkspaceUpdateOptions.ContentOnly))
-					{
-						Filter.Exclude("*.usf");
-						Filter.Exclude("*.ush");
-					}
-					SyncDepotPaths.RemoveAll(x => !Filter.Matches(x));
-
 					// Find all the depot paths that will be synced
 					HashSet<string> RemainingDepotPaths = new HashSet<string>(StringComparer.InvariantCultureIgnoreCase);
 					RemainingDepotPaths.UnionWith(SyncDepotPaths);
 
-					// Build the list of revisions to sync
-					List<string> SyncRevisions = new List<string>();
-					SyncRevisions.AddRange(SyncDepotPaths.Select(x => String.Format("{0}@{1}", x, PendingChangeNumber)));
-
-					WorkspaceUpdateResult SyncResult = SyncFileRevisions("Syncing files...", Context, SyncRevisions, RemainingDepotPaths, out StatusMessage);
+					WorkspaceUpdateResult SyncResult = SyncFileRevisions("Syncing files...", Context, SyncCommands.ToList(), RemainingDepotPaths, out StatusMessage);
 					if(SyncResult != WorkspaceUpdateResult.Success)
 					{
 						return SyncResult;
@@ -785,7 +802,7 @@ namespace UnrealGameSync
 					Times.Add(new Tuple<string,TimeSpan>("Sync", Stopwatch.Stop("Success")));
 
 					// Save the number of files synced
-					NumFilesSynced = SyncRevisions.Count;
+					NumFilesSynced = SyncDepotPaths.Count;
 					Log.WriteLine();
 				}
 			}
