UGS: Reduce the number of sync commands sent to the server, by collapsing the list of files to sync into a set of wildcards.

diff --git a/Engine/Source/Programs/UnrealGameSync/UnrealGameSync/Workspace.cs b/Engine/Source/Programs/UnrealGameSync/UnrealGameSync/Workspace.cs
--- a/Engine/Source/Programs/UnrealGameSync/UnrealGameSync/Workspace.cs
+++ b/Engine/Source/Programs/UnrealGameSync/UnrealGameSync/Workspace.cs
@@ -128,6 +128,7 @@ namespace UnrealGameSync
 
 		public readonly PerforceConnection Perforce;
 		public readonly string LocalRootPath;
+		public readonly string LocalRootPrefix;
 		public readonly string SelectedLocalFileName;
 		public readonly string ClientRootPath;
 		public readonly string SelectedClientFileName;
@@ -178,10 +179,79 @@ namespace UnrealGameSync
 			}
 		}
 
+		class SyncTree
+		{
+			public bool bCanUseWildcard;
+			public int TotalExcludedFiles;
+			public HashSet<string> IncludedFiles = new HashSet<string>();
+			public Dictionary<string, SyncTree> NameToSubTree = new Dictionary<string, SyncTree>(StringComparer.OrdinalIgnoreCase);
+
+			public SyncTree(bool bCanUseWildcard)
+			{
+				this.bCanUseWildcard = bCanUseWildcard;
+			}
+
+			public SyncTree FindOrAddSubTree(string Name)
+			{
+				SyncTree Result;
+				if (!NameToSubTree.TryGetValue(Name, out Result))
+				{
+					Result = new SyncTree(bCanUseWildcard);
+					NameToSubTree.Add(Name, Result);
+				}
+				return Result;
+			}
+
+			public void IncludeFile(string Path)
+			{
+				int Idx = Path.IndexOf('/');
+				if (Idx == -1)
+				{
+					IncludedFiles.Add(Path);
+				}
+				else
+				{
+					SyncTree SubTree = FindOrAddSubTree(Path.Substring(0, Idx));
+					SubTree.IncludeFile(Path.Substring(Idx + 1));
+				}
+			}
+
+			public void ExcludeFile(string Path)
+			{
+				int Idx = Path.IndexOf('/');
+				if (Idx != -1)
+				{
+					SyncTree SubTree = FindOrAddSubTree(Path.Substring(0, Idx));
+					SubTree.ExcludeFile(Path.Substring(Idx + 1));
+				}
+				TotalExcludedFiles++;
+			}
+
+			public void GetOptimizedSyncCommands(string Prefix, int ChangeNumber, SortedSet<string> SyncCommands)
+			{
+				if (bCanUseWildcard && TotalExcludedFiles == 0)
+				{
+					SyncCommands.Add(String.Format("{0}/...@{1}", Prefix, ChangeNumber));
+				}
+				else
+				{
+					foreach(string File in IncludedFiles)
+					{
+						SyncCommands.Add(String.Format("{0}/{1}@{2}", Prefix, File, ChangeNumber));
+					}
+					foreach (KeyValuePair<string, SyncTree> Pair in NameToSubTree)
+					{
+						Pair.Value.GetOptimizedSyncCommands(String.Format("{0}/{1}", Prefix, PerforceUtils.EscapePath(Pair.Key)), ChangeNumber, SyncCommands);
+					}
+				}
+			}
+		}
+
 		public Workspace(PerforceConnection InPerforce, string InLocalRootPath, string InSelectedLocalFileName, string InClientRootPath, string InSelectedClientFileName, int InInitialChangeNumber, string InInitialSyncFilterHash, int InLastBuiltChangeNumber, string InTelemetryProjectPath, bool bInIsEnterpriseProject, TextWriter InLog)
 		{
 			Perforce = InPerforce;
 			LocalRootPath = InLocalRootPath;
+			LocalRootPrefix = InLocalRootPath.TrimEnd(Path.DirectorySeparatorChar) + Path.DirectorySeparatorChar;
 			SelectedLocalFileName = InSelectedLocalFileName;
 			ClientRootPath = InClientRootPath;
 			SelectedClientFileName = InSelectedClientFileName;
@@ -359,7 +429,8 @@ namespace UnrealGameSync
 					}
 
 					// Figure out which paths to sync
-					List<string> SyncPaths = GetSyncPaths((Context.Options & WorkspaceUpdateOptions.SyncAllProjects) != 0, Context.SyncFilter);
+					List<string> RelativeSyncPaths = GetRelativeSyncPaths((Context.Options & WorkspaceUpdateOptions.SyncAllProjects) != 0, Context.SyncFilter);
+					List<string> SyncPaths = new List<string>(RelativeSyncPaths.Select(x => ClientRootPath + x));
 
 					// Get the user's sync filter
 					FileFilter UserFilter = new FileFilter(FileFilterType.Include);
@@ -372,7 +443,11 @@ namespace UnrealGameSync
 					string NextSyncFilterHash = null;
 					using (SHA1Managed SHA = new SHA1Managed())
 					{
-						StringBuilder CombinedFilter = new StringBuilder(String.Join("\n", SyncPaths));
+						StringBuilder CombinedFilter = new StringBuilder();
+						foreach(string RelativeSyncPath in RelativeSyncPaths)
+						{
+							CombinedFilter.AppendFormat("{0}\n", RelativeSyncPath);
+						}
 						if(Context.SyncFilter != null)
 						{
 							CombinedFilter.Append("--FROM--\n");
@@ -404,15 +479,9 @@ namespace UnrealGameSync
 
 						// Build a filter for the current sync paths
 						FileFilter SyncPathsFilter = new FileFilter(FileFilterType.Exclude);
-						foreach(string SyncPath in SyncPaths)
+						foreach(string RelativeSyncPath in RelativeSyncPaths)
 						{
-							if(!SyncPath.StartsWith(ClientRootPath))
-							{
-								Log.WriteLine("Invalid sync path; '{0}' does not begin with '{1}'", SyncPath, ClientRootPath);
-								StatusMessage = "Unable to sync files.";
-								return WorkspaceUpdateResult.FailedToSync;
-							}
-							SyncPathsFilter.Include(SyncPath.Substring(ClientRootPath.Length));
+							SyncPathsFilter.Include(RelativeSyncPath);
 						}
 
 						// Remove all the files that are not included by the filter
@@ -437,12 +506,9 @@ namespace UnrealGameSync
 						// Check if there are any paths outside the regular sync paths
 						if(RemoveDepotPaths.Count > 0)
 						{
-							FileFilter ProjectFilter = new FileFilter(FileFilterType.Exclude);
-							ProjectFilter.AddRules(SyncPaths);
-
 							bool bDeleteListMatches = true;
 
-							Dictionary<string, bool> NewDeleteFiles = new Dictionary<string, bool>(StringComparer.InvariantCultureIgnoreCase);
+							Dictionary<string, bool> NewDeleteFiles = new Dictionary<string, bool>(StringComparer.OrdinalIgnoreCase);
 							foreach(string RemoveDepotPath in RemoveDepotPaths)
 							{
 								bool bDelete;
@@ -471,7 +537,7 @@ namespace UnrealGameSync
 							CurrentSyncFilterHash = "INVALID";
 
 							// Find all the depot paths that will be synced
-							HashSet<string> RemainingDepotPathsToRemove = new HashSet<string>(StringComparer.InvariantCultureIgnoreCase);
+							HashSet<string> RemainingDepotPathsToRemove = new HashSet<string>(StringComparer.OrdinalIgnoreCase);
 							RemainingDepotPathsToRemove.UnionWith(RemoveDepotPaths);
 
 							// Build the list of revisions to sync
@@ -491,11 +557,11 @@ namespace UnrealGameSync
 
 					// Create a filter for all the files we don't want
 					FileFilter Filter = new FileFilter(UserFilter);
-					Filter.Exclude("..." + BuildVersionFileName);
+					Filter.Exclude(BuildVersionFileName);
 					if (!ProjectConfigFile.GetValue("Options.UseFastModularVersioningV2", false))
 					{
-						Filter.Exclude("..." + VersionHeaderFileName);
-						Filter.Exclude("..." + ObjectVersionFileName);
+						Filter.Exclude(VersionHeaderFileName);
+						Filter.Exclude(ObjectVersionFileName);
 					}
 					if (Context.Options.HasFlag(WorkspaceUpdateOptions.ContentOnly))
 					{
@@ -503,6 +569,25 @@ namespace UnrealGameSync
 						Filter.Exclude("*.ush");
 					}
 
+					// Create a tree to store the sync path
+					SyncTree SyncTree = new SyncTree(false);
+					foreach (string RelativeSyncPath in RelativeSyncPaths)
+					{
+						const string WildcardSuffix = "/...";
+						if (RelativeSyncPath.EndsWith(WildcardSuffix, StringComparison.Ordinal))
+						{
+							SyncTree Leaf = SyncTree;
+
+							string[] Fragments = RelativeSyncPath.Split('/');
+							for (int Idx = 1; Idx < Fragments.Length - 1; Idx++)
+							{
+								Leaf = Leaf.FindOrAddSubTree(Fragments[Idx]);
+							}
+
+							Leaf.bCanUseWildcard = true;
+						}
+					}
+
 					// Find all the server changes, and anything that's opened for edit locally. We need to sync files we have open to schedule a resolve.
 					SortedSet<string> SyncCommands = new SortedSet<string>(StringComparer.Ordinal);
 					List<string> SyncDepotPaths = new List<string>();
@@ -510,13 +595,6 @@ namespace UnrealGameSync
 					{
 						foreach(string SyncPath in SyncPaths)
 						{
-							if(!SyncPath.StartsWith(ClientRootPath))
-							{
-								Log.WriteLine("Sync path ({0}) does not start with client root path ({1}) - unable to sync.", SyncPath, ClientRootPath);
-								StatusMessage = "Internal error while syncing";
-								return WorkspaceUpdateResult.FailedToSync;
-							}
-
 							List<PerforceFileRecord> SyncRecords = new List<PerforceFileRecord>();
 							if(!Perforce.SyncPreview(SyncPath, PendingChangeNumber, !Context.Options.HasFlag(WorkspaceUpdateOptions.Sync), Record => { SyncRecords.Add(Record); Counter.Increment(); }, Log))
 							{
@@ -557,7 +635,7 @@ namespace UnrealGameSync
 								}
 
 								// Make sure it's under the current directory. Not sure why this would happen, just being safe.
-								if (!FullName.StartsWith(LocalRootPath, StringComparison.OrdinalIgnoreCase))
+								if (!FullName.StartsWith(LocalRootPrefix, StringComparison.OrdinalIgnoreCase))
 								{
 									SyncCommands.Add(String.Format("{0}@{1}", SyncRecord.DepotPath, PendingChangeNumber));
 									SyncDepotPaths.Add(SyncRecord.DepotPath);
@@ -565,18 +643,23 @@ namespace UnrealGameSync
 								}
 
 								// Check that it matches the filter
-								string RelativePath = FullName.Substring(LocalRootPath.Length).Replace('\\', '/');
-								if (UserFilter.Matches(RelativePath))
+								string RelativePath = FullName.Substring(LocalRootPrefix.Length).Replace('\\', '/');
+								if (Filter.Matches(RelativePath))
 								{
-									SyncCommands.Add(String.Format("{0}{1}@{2}", ClientRootPath, PerforceUtils.EscapePath(RelativePath), PendingChangeNumber));
+									SyncTree.IncludeFile(RelativePath);
 									SyncDepotPaths.Add(SyncRecord.DepotPath);
 								}
+								else
+								{
+									SyncTree.ExcludeFile(RelativePath);
+								}
 							}
 						}
 					}
+					SyncTree.GetOptimizedSyncCommands(ClientRootPath, PendingChangeNumber, SyncCommands);
 
 					// Find all the depot paths that will be synced
-					HashSet<string> RemainingDepotPaths = new HashSet<string>(StringComparer.InvariantCultureIgnoreCase);
+					HashSet<string> RemainingDepotPaths = new HashSet<string>(StringComparer.OrdinalIgnoreCase);
 					RemainingDepotPaths.UnionWith(SyncDepotPaths);
 
 					WorkspaceUpdateResult SyncResult = SyncFileRevisions("Syncing files...", Context, SyncCommands.ToList(), RemainingDepotPaths, out StatusMessage);
@@ -723,7 +806,7 @@ namespace UnrealGameSync
 						}
 
 						// Remove all the receipts for build targets in this branch
-						if(SelectedClientFileName.EndsWith(".uproject", StringComparison.InvariantCultureIgnoreCase))
+						if(SelectedClientFileName.EndsWith(".uproject", StringComparison.OrdinalIgnoreCase))
 						{
 							Perforce.Sync(PerforceUtils.GetClientOrDepotDirectoryName(SelectedClientFileName) + "/Build/Receipts/...#0", Log);
 						}
@@ -814,7 +897,7 @@ namespace UnrealGameSync
 				{
 					// Create the directory for extracted archive manifests
 					string ManifestDirectoryName;
-					if(SelectedLocalFileName.EndsWith(".uproject", StringComparison.InvariantCultureIgnoreCase))
+					if(SelectedLocalFileName.EndsWith(".uproject", StringComparison.OrdinalIgnoreCase))
 					{
 						ManifestDirectoryName = Path.Combine(Path.GetDirectoryName(SelectedLocalFileName), "Saved", "UnrealGameSync");
 					}
@@ -891,7 +974,7 @@ namespace UnrealGameSync
 					CommandLine.AppendFormat("/C \"\"{0}\"", Path.Combine(LocalRootPath, "GenerateProjectFiles.bat"));
 					if((Context.Options & WorkspaceUpdateOptions.SyncAllProjects) == 0 && (Context.Options & WorkspaceUpdateOptions.IncludeAllProjectsInSolution) == 0)
 					{
-						if(SelectedLocalFileName.EndsWith(".uproject", StringComparison.InvariantCultureIgnoreCase))
+						if(SelectedLocalFileName.EndsWith(".uproject", StringComparison.OrdinalIgnoreCase))
 						{
 							CommandLine.AppendFormat(" \"{0}\"", SelectedLocalFileName);
 						}
@@ -1094,22 +1177,35 @@ namespace UnrealGameSync
 
 		public List<string> GetSyncPaths(bool bSyncAllProjects, string[] SyncFilter)
 		{
-			// Add the default project paths
+			List<string> SyncPaths = GetRelativeSyncPaths(bSyncAllProjects, SyncFilter);
+			return SyncPaths.Select(x => ClientRootPath + x).ToList();
+		}
+
+		public List<string> GetRelativeSyncPaths(bool bSyncAllProjects, string[] SyncFilter)
+		{
 			List<string> SyncPaths = new List<string>();
-			if(bSyncAllProjects || !SelectedClientFileName.EndsWith(".uproject", StringComparison.InvariantCultureIgnoreCase))
+
+			// Check the client path is formatted correctly
+			if (!SelectedClientFileName.StartsWith(ClientRootPath + "/"))
+			{
+				throw new Exception(String.Format("Expected '{0}' to start with '{1}'", SelectedClientFileName, ClientRootPath));
+			}
+
+			// Add the default project paths
+			int LastSlashIdx = SelectedClientFileName.LastIndexOf('/');
+			if (bSyncAllProjects || !SelectedClientFileName.EndsWith(".uproject", StringComparison.OrdinalIgnoreCase) || LastSlashIdx <= ClientRootPath.Length)
 			{
-				SyncPaths.Add(ClientRootPath + "/...");
+				SyncPaths.Add("/...");
 			}
 			else
 			{
-				SyncPaths.Add(ClientRootPath + "/*");
-				SyncPaths.Add(ClientRootPath + "/Engine/...");
+				SyncPaths.Add("/*");
+				SyncPaths.Add("/Engine/...");
 				if(bIsEnterpriseProject)
 				{
-					SyncPaths.Add(ClientRootPath + "/Enterprise/...");
+					SyncPaths.Add("/Enterprise/...");
 				}
-//				SyncPaths.Add(ClientRootPath + "/Platforms/...");
-				SyncPaths.Add(PerforceUtils.GetClientOrDepotDirectoryName(SelectedClientFileName) + "/...");
+				SyncPaths.Add(SelectedClientFileName.Substring(ClientRootPath.Length, LastSlashIdx - ClientRootPath.Length) + "/...");
 			}
 
 			// Apply the sync filter to that list. We only want inclusive rules in the output list, but we can manually apply exclusions to previous entries.
@@ -1120,11 +1216,11 @@ namespace UnrealGameSync
 					string TrimSyncPath = SyncPath.Trim();
 					if(TrimSyncPath.StartsWith("/"))
 					{
-						SyncPaths.Add(ClientRootPath + TrimSyncPath);
+						SyncPaths.Add(TrimSyncPath);
 					}
 					else if(TrimSyncPath.StartsWith("-/") && TrimSyncPath.EndsWith("..."))
 					{
-						SyncPaths.RemoveAll(x => x.StartsWith(ClientRootPath + TrimSyncPath.Substring(1, TrimSyncPath.Length - 4)));
+						SyncPaths.RemoveAll(x => x.StartsWith(TrimSyncPath.Substring(1, TrimSyncPath.Length - 4)));
 					}
 				}
 			}
@@ -1153,7 +1249,7 @@ namespace UnrealGameSync
 		public bool MatchFilter(string FileName, FileFilter Filter)
 		{
 			bool bMatch = true;
-			if(FileName.StartsWith(LocalRootPath, StringComparison.InvariantCultureIgnoreCase))
+			if(FileName.StartsWith(LocalRootPath, StringComparison.OrdinalIgnoreCase))
 			{
 				if(!Filter.Matches(FileName.Substring(LocalRootPath.Length)))
 				{
@@ -1163,11 +1259,13 @@ namespace UnrealGameSync
 			return bMatch;
 		}
 
-		WorkspaceUpdateResult SyncFileRevisions(string Prefix, WorkspaceUpdateContext Context, List<string> SyncRevisions, HashSet<string> RemainingDepotPaths, out string StatusMessage)
+		WorkspaceUpdateResult SyncFileRevisions(string Prefix, WorkspaceUpdateContext Context, List<string> SyncCommands, HashSet<string> RemainingDepotPaths, out string StatusMessage)
 		{
+			int TotalDepotPaths = RemainingDepotPaths.Count;
+
 			// Sync them all
 			List<string> TamperedFiles = new List<string>();
-			if(!Perforce.Sync(SyncRevisions, Record => UpdateSyncProgress(Prefix, Record, RemainingDepotPaths, SyncRevisions.Count), TamperedFiles, false, Context.PerforceSyncOptions, Log))
+			if(!Perforce.Sync(SyncCommands, Record => UpdateSyncProgress(Prefix, Record, RemainingDepotPaths, TotalDepotPaths), TamperedFiles, false, Context.PerforceSyncOptions, Log))
 			{
 				StatusMessage = "Aborted sync due to errors.";
 				return WorkspaceUpdateResult.FailedToSync;
@@ -1272,7 +1370,7 @@ namespace UnrealGameSync
 			{
 				return true;
 			}
-			if(OpenFiles.Any(x => x.DepotPath.IndexOf("/Source/", StringComparison.InvariantCultureIgnoreCase) != -1))
+			if(OpenFiles.Any(x => x.DepotPath.IndexOf("/Source/", StringComparison.OrdinalIgnoreCase) != -1))
 			{
 				return true;
 			}
