UGS: Less confusing and more consistent way of storing server settings. Settings in the "Application Settings..." dialog now reflect the default server used by all projects, and open projects only store specific connection info if they override the default.

Side-effect of this is that using a P4CONFIG file to override settings for a particular workspace is no longer supported, but a specific server can still be selected from the project settings.

#rnx

diff --git a/Engine/Source/Programs/UnrealGameSync/UnrealGameSync/Controls/WorkspaceControl.cs b/Engine/Source/Programs/UnrealGameSync/UnrealGameSync/Controls/WorkspaceControl.cs
--- a/Engine/Source/Programs/UnrealGameSync/UnrealGameSync/Controls/WorkspaceControl.cs
+++ b/Engine/Source/Programs/UnrealGameSync/UnrealGameSync/Controls/WorkspaceControl.cs
@@ -2872,7 +2872,7 @@ namespace UnrealGameSync
 				ProjectLine.AddText(String.Format("Opened "));
 				ProjectLine.AddLink(SelectedFileName + " \u25BE", FontStyle.Regular, (P, R) => { SelectRecentProject(R); });
 				ProjectLine.AddText("  |  ");
-				ProjectLine.AddLink("Browse...", FontStyle.Regular, (P, R) => { Owner.EditSelectedProject(this); });
+				ProjectLine.AddLink("Settings...", FontStyle.Regular, (P, R) => { Owner.EditSelectedProject(this); });
 				Lines.Add(ProjectLine);
 
 				// Spacer
@@ -3347,7 +3347,7 @@ namespace UnrealGameSync
 		private void SelectOtherStreamDialog()
 		{
 			string NewStreamName;
-			if(SelectStreamWindow.ShowModal(this, SelectedProject.ServerAndPort, SelectedProject.UserName, StreamName, Log, out NewStreamName))
+			if(SelectStreamWindow.ShowModal(this, Workspace.Perforce, StreamName, Log, out NewStreamName))
 			{
 				SelectStream(NewStreamName);
 			}
diff --git a/Engine/Source/Programs/UnrealGameSync/UnrealGameSync/DetectProjectSettingsTask.cs b/Engine/Source/Programs/UnrealGameSync/UnrealGameSync/DetectProjectSettingsTask.cs
--- a/Engine/Source/Programs/UnrealGameSync/UnrealGameSync/DetectProjectSettingsTask.cs
+++ b/Engine/Source/Programs/UnrealGameSync/UnrealGameSync/DetectProjectSettingsTask.cs
@@ -11,7 +11,7 @@ using System.Web.Script.Serialization;
 
 namespace UnrealGameSync
 {
-	class DetectProjectSettingsTask : IModalTask, IPerforceModalTask, IDisposable
+	class DetectProjectSettingsTask : IPerforceModalTask, IDisposable
 	{
 		public UserSelectedProjectSettings SelectedProject;
 		public PerforceConnection PerforceClient;
@@ -47,27 +47,8 @@ namespace UnrealGameSync
 			}
 		}
 
-		public bool Run(out string ErrorMessage)
+		public bool Run(PerforceConnection Perforce, out string ErrorMessage)
 		{
-			// Get the server settings
-			string ServerAndPort = SelectedProject.ServerAndPort;
-			string UserName = SelectedProject.UserName;
-
-			string ProjectFileName = null;
-			if(SelectedProject.Type == UserSelectedProjectType.Local)
-			{
-				ProjectFileName = SelectedProject.LocalPath;
-			}
-
-			if(!PerforceModalTask.TryGetServerSettings(ProjectFileName, ref ServerAndPort, ref UserName, Log))
-			{
-				ErrorMessage = "Unable to get Perforce server settings.";
-				return false;
-			}
-
-			// Create the connection
-			PerforceConnection Perforce = new PerforceConnection(UserName, null, ServerAndPort);
-
 			// Make sure we're logged in
 			bool bLoggedIn;
 			if(!Perforce.GetLoggedInState(out bLoggedIn, Log))
@@ -195,7 +176,7 @@ namespace UnrealGameSync
 			NewSelectedFileName = Utility.GetPathWithCorrectCase(new FileInfo(NewSelectedFileName));
 
 			// Update the selected project with all the data we've found
-			SelectedProject = new UserSelectedProjectSettings(Perforce.ServerAndPort, Perforce.UserName, SelectedProject.Type, NewSelectedClientFileName, NewSelectedFileName);
+			SelectedProject = new UserSelectedProjectSettings(SelectedProject.ServerAndPort, SelectedProject.UserName, SelectedProject.Type, NewSelectedClientFileName, NewSelectedFileName);
 
 			// Figure out where the engine is in relation to it
 			int EndIdx = NewSelectedClientFileName.Length - 1;
@@ -416,11 +397,13 @@ namespace UnrealGameSync
 
 	class DetectMultipleProjectSettingsTask : IModalTask, IDisposable
 	{
+		PerforceConnection DefaultConnection;
 		public DetectProjectSettingsTask[] Tasks;
 		public DetectProjectSettingsResult[] Results;
 
-		public DetectMultipleProjectSettingsTask(IEnumerable<DetectProjectSettingsTask> Tasks)
+		public DetectMultipleProjectSettingsTask(PerforceConnection DefaultConnection, IEnumerable<DetectProjectSettingsTask> Tasks)
 		{
+			this.DefaultConnection = DefaultConnection;
 			this.Tasks = Tasks.ToArray();
 		}
 
@@ -451,10 +434,34 @@ namespace UnrealGameSync
 
 		void RunTask(int Idx)
 		{
+			DetectProjectSettingsTask Task = Tasks[Idx];
+
+			UserSelectedProjectSettings Project = Task.SelectedProject;
+			PerforceConnection Perforce = Utility.OverridePerforceSettings(DefaultConnection, Project.ServerAndPort, Project.UserName);
+
 			string TaskErrorMessage;
-			bool bTaskSucceeded = Tasks[Idx].Run(out TaskErrorMessage);
+			bool bTaskSucceeded = Tasks[Idx].Run(Perforce, out TaskErrorMessage);
 			Results[Idx] = new DetectProjectSettingsResult(Tasks[Idx], bTaskSucceeded, TaskErrorMessage);
 			Tasks[Idx] = null;
 		}
 	}
+
+	class QuietDetectProjectSettingsTask : IModalTask
+	{
+		PerforceConnection Perforce;
+		DetectProjectSettingsTask Inner;
+		TextWriter Log;
+
+		public QuietDetectProjectSettingsTask(PerforceConnection Perforce, DetectProjectSettingsTask Inner, TextWriter Log)
+		{
+			this.Perforce = Perforce;
+			this.Inner = Inner;
+			this.Log = Log;
+		}
+
+		public bool Run(out string ErrorMessage)
+		{
+			return Inner.Run(Perforce, out ErrorMessage);
+		}
+	}
 }
diff --git a/Engine/Source/Programs/UnrealGameSync/UnrealGameSync/Forms/ApplicationSettingsWindow.cs b/Engine/Source/Programs/UnrealGameSync/UnrealGameSync/Forms/ApplicationSettingsWindow.cs
--- a/Engine/Source/Programs/UnrealGameSync/UnrealGameSync/Forms/ApplicationSettingsWindow.cs
+++ b/Engine/Source/Programs/UnrealGameSync/UnrealGameSync/Forms/ApplicationSettingsWindow.cs
@@ -19,33 +19,6 @@ namespace UnrealGameSync
 {
 	partial class ApplicationSettingsWindow : Form
 	{
-		class GetDefaultSettingsTask : IModalTask
-		{
-			TextWriter Log;
-
-			public string ServerAndPort;
-			public string UserName;
-
-			public GetDefaultSettingsTask(TextWriter Log)
-			{
-				this.Log = Log;
-			}
-
-			public bool Run(out string ErrorMessage)
-			{
-				if(PerforceModalTask.TryGetServerSettings(null, ref ServerAndPort, ref UserName, Log))
-				{
-					ErrorMessage = null;
-					return true;
-				}
-				else
-				{
-					ErrorMessage = "Unable to query server settings";
-					return false;
-				}
-			}
-		}
-
 		class PerforceTestConnectionTask : IPerforceModalTask
 		{
 			string DepotPath;
@@ -111,14 +84,9 @@ namespace UnrealGameSync
 			this.UseUnstableBuildCheckBox.Checked = bUnstable;
 		}
 
-		public static bool? ShowModal(IWin32Window Owner, bool bUnstable, string OriginalExecutableFileName, UserSettings Settings, TextWriter Log)
+		public static bool? ShowModal(IWin32Window Owner, PerforceConnection DefaultConnection, bool bUnstable, string OriginalExecutableFileName, UserSettings Settings, TextWriter Log)
 		{
-			GetDefaultSettingsTask DefaultSettings = new GetDefaultSettingsTask(Log);
-
-			string ErrorMessage;
-			ModalTask.Execute(Owner, DefaultSettings, "Checking Settings", "Checking settings, please wait...", out ErrorMessage);
-
-			ApplicationSettingsWindow ApplicationSettings = new ApplicationSettingsWindow(DefaultSettings.ServerAndPort, DefaultSettings.UserName, bUnstable, OriginalExecutableFileName, Settings, Log);
+			ApplicationSettingsWindow ApplicationSettings = new ApplicationSettingsWindow(DefaultConnection.ServerAndPort, DefaultConnection.UserName, bUnstable, OriginalExecutableFileName, Settings, Log);
 			if(ApplicationSettings.ShowDialog() == DialogResult.OK)
 			{
 				return ApplicationSettings.bRestartUnstable;
@@ -164,7 +132,7 @@ namespace UnrealGameSync
 				if(ServerAndPort != InitialServerAndPort || UserName != InitialUserName || DepotPath != InitialDepotPath)
 				{
 					string ErrorMessage;
-					ModalTaskResult Result = PerforceModalTask.Execute(this, null, ServerAndPort, UserName, new PerforceTestConnectionTask(DepotPath), "Connecting", "Checking connection, please wait...", Log, out ErrorMessage);
+					ModalTaskResult Result = PerforceModalTask.Execute(this, new PerforceConnection(UserName, null, ServerAndPort), new PerforceTestConnectionTask(DepotPath), "Connecting", "Checking connection, please wait...", Log, out ErrorMessage);
 					if(Result != ModalTaskResult.Succeeded)
 					{
 						if(Result == ModalTaskResult.Failed)
diff --git a/Engine/Source/Programs/UnrealGameSync/UnrealGameSync/Forms/AutomatedSyncWindow.cs b/Engine/Source/Programs/UnrealGameSync/UnrealGameSync/Forms/AutomatedSyncWindow.cs
--- a/Engine/Source/Programs/UnrealGameSync/UnrealGameSync/Forms/AutomatedSyncWindow.cs
+++ b/Engine/Source/Programs/UnrealGameSync/UnrealGameSync/Forms/AutomatedSyncWindow.cs
@@ -109,17 +109,19 @@ namespace UnrealGameSync
 			public string WorkspaceName;
 			public bool bRequiresStreamSwitch;
 		}
-		
+
 		string StreamName;
 		TextWriter Log;
 
-		string ServerAndPort;
-		string UserName;
+		string ServerAndPortOverride;
+		string UserNameOverride;
+        PerforceConnection DefaultConnection;
 		WorkspaceInfo SelectedWorkspaceInfo;
 
-		private AutomatedSyncWindow(string StreamName, string ProjectPath, string WorkspaceName, TextWriter Log)
+		private AutomatedSyncWindow(string StreamName, string ProjectPath, string WorkspaceName, PerforceConnection DefaultConnection, TextWriter Log)
 		{
 			this.StreamName = StreamName;
+			this.DefaultConnection = DefaultConnection;
 			this.Log = Log;
 
 			InitializeComponent();
@@ -141,14 +143,19 @@ namespace UnrealGameSync
 			UpdateOkButton();
 		}
 
-		public static bool ShowModal(IWin32Window Owner, string StreamName, string ProjectPath, out WorkspaceInfo WorkspaceInfo, TextWriter Log)
+		private PerforceConnection Perforce
+		{
+			get { return Utility.OverridePerforceSettings(DefaultConnection, ServerAndPortOverride, UserNameOverride); }
+		}
+
+		public static bool ShowModal(IWin32Window Owner, PerforceConnection DefaultConnection, string StreamName, string ProjectPath, out WorkspaceInfo WorkspaceInfo, TextWriter Log)
 		{
 			FindDefaultWorkspaceTask FindWorkspace = new FindDefaultWorkspaceTask(StreamName);
 
 			string ErrorMessage;
-			PerforceModalTask.Execute(Owner, null, null, null, FindWorkspace, "Finding workspace", "Finding default workspace, please wait...", Log, out ErrorMessage);
+			PerforceModalTask.Execute(Owner, DefaultConnection, FindWorkspace, "Finding workspace", "Finding default workspace, please wait...", Log, out ErrorMessage);
 
-			AutomatedSyncWindow Window = new AutomatedSyncWindow(StreamName, ProjectPath, FindWorkspace.WorkspaceName, Log);
+			AutomatedSyncWindow Window = new AutomatedSyncWindow(StreamName, ProjectPath, FindWorkspace.WorkspaceName, DefaultConnection, Log);
 			if(Window.ShowDialog() == DialogResult.OK)
 			{
 				WorkspaceInfo = Window.SelectedWorkspaceInfo;
@@ -163,26 +170,22 @@ namespace UnrealGameSync
 
 		private void ChangeLink_LinkClicked(object sender, LinkLabelLinkClickedEventArgs e)
 		{
-			string NewServerAndPort;
-			string NewUserName;
-			if(ConnectWindow.ShowModal(this, ServerAndPort, UserName, Log, out NewServerAndPort, out NewUserName))
+			if(ConnectWindow.ShowModal(this, DefaultConnection, ref ServerAndPortOverride, ref UserNameOverride, Log))
 			{
-				ServerAndPort = NewServerAndPort;
-				UserName = NewUserName;
 				UpdateServerLabel();
 			}
 		}
 
 		private void UpdateServerLabel()
 		{
-			ServerLabel.Text = OpenProjectWindow.GetServerLabelText(ServerAndPort, UserName);
+			ServerLabel.Text = OpenProjectWindow.GetServerLabelText(DefaultConnection, ServerAndPortOverride, UserNameOverride);
 			ChangeLink.Location = new Point(ServerLabel.Right + 5, ChangeLink.Location.Y);
 		}
 
 		private void WorkspaceNameNewBtn_Click(object sender, EventArgs e)
 		{
 			string WorkspaceName;
-			if(NewWorkspaceWindow.ShowModal(this, ServerAndPort, UserName, StreamName, WorkspaceNameTextBox.Text, Log, out WorkspaceName))
+			if(NewWorkspaceWindow.ShowModal(this, Perforce, StreamName, WorkspaceNameTextBox.Text, Log, out WorkspaceName))
 			{
 				WorkspaceNameTextBox.Text = WorkspaceName;
 				UpdateOkButton();
@@ -192,7 +195,7 @@ namespace UnrealGameSync
 		private void WorkspaceNameBrowseBtn_Click(object sender, EventArgs e)
 		{
 			string WorkspaceName = WorkspaceNameTextBox.Text;
-			if(SelectWorkspaceWindow.ShowModal(this, ServerAndPort, UserName, WorkspaceName, Log, out WorkspaceName))
+			if(SelectWorkspaceWindow.ShowModal(this, Perforce, WorkspaceName, Log, out WorkspaceName))
 			{
 				WorkspaceNameTextBox.Text = WorkspaceName;
 				UpdateOkButton();
@@ -209,7 +212,7 @@ namespace UnrealGameSync
 			ValidateWorkspaceTask ValidateWorkspace = new ValidateWorkspaceTask(WorkspaceNameTextBox.Text, StreamName);
 
 			string ErrorMessage;
-			ModalTaskResult Result = PerforceModalTask.Execute(Owner, null, ServerAndPort, UserName, ValidateWorkspace, "Checking workspace", "Checking workspace, please wait...", Log, out ErrorMessage);
+			ModalTaskResult Result = PerforceModalTask.Execute(Owner, Perforce, ValidateWorkspace, "Checking workspace", "Checking workspace, please wait...", Log, out ErrorMessage);
 			if(Result == ModalTaskResult.Failed)
 			{
 				MessageBox.Show(ErrorMessage);
diff --git a/Engine/Source/Programs/UnrealGameSync/UnrealGameSync/Forms/ConnectWindow.cs b/Engine/Source/Programs/UnrealGameSync/UnrealGameSync/Forms/ConnectWindow.cs
--- a/Engine/Source/Programs/UnrealGameSync/UnrealGameSync/Forms/ConnectWindow.cs
+++ b/Engine/Source/Programs/UnrealGameSync/UnrealGameSync/Forms/ConnectWindow.cs
@@ -15,53 +15,46 @@ namespace UnrealGameSync
 {
 	partial class ConnectWindow : Form
 	{
-		string ServerAndPort;
-		string UserName;
+		PerforceConnection DefaultConnection;
+		string ServerAndPortOverride;
+		string UserNameOverride;
 		TextWriter Log;
 
-		public ConnectWindow(string ServerAndPort, string UserName, TextWriter Log)
+		private ConnectWindow(PerforceConnection DefaultConnection, string ServerAndPortOverride, string UserNameOverride, TextWriter Log)
 		{
 			InitializeComponent();
 
+			this.DefaultConnection = DefaultConnection;
 			this.Log = Log;
 
-			if(String.IsNullOrEmpty(ServerAndPort))
+			if(!String.IsNullOrWhiteSpace(ServerAndPortOverride))
 			{
-				ServerAndPort = null;
+				this.ServerAndPortOverride = ServerAndPortOverride.Trim();
 			}
-			if(String.IsNullOrEmpty(UserName))
+			if(!String.IsNullOrEmpty(UserNameOverride))
 			{
-				UserName = null;
+				this.UserNameOverride = UserNameOverride.Trim();
 			}
 
-			if(ServerAndPort != null)
-			{
-				ServerAndPortTextBox.Text = ServerAndPort;
-			}
-			if(UserName != null)
-			{
-				UserNameTextBox.Text = UserName;
-			}
-			UseDefaultConnectionSettings.Checked = ServerAndPort == null && UserName == null;
+			ServerAndPortTextBox.CueBanner = DefaultConnection.ServerAndPort;
+			ServerAndPortTextBox.Text = this.ServerAndPortOverride ?? DefaultConnection.ServerAndPort;
+			UserNameTextBox.CueBanner = DefaultConnection.UserName;
+			UserNameTextBox.Text = this.UserNameOverride ?? DefaultConnection.UserName;
+			UseDefaultConnectionSettings.Checked = this.ServerAndPortOverride == null && this.UserNameOverride == null;
 
 			UpdateEnabledControls();
 		}
 
-		public static bool ShowModal(IWin32Window Owner, string ServerAndPort, string UserName, TextWriter Log, out string NewServerAndPort, out string NewUserName)
+		public static bool ShowModal(IWin32Window Owner, PerforceConnection DefaultConnection, ref string ServerAndPortOverride, ref string UserNameOverride, TextWriter Log)
 		{
-			ConnectWindow Connect = new ConnectWindow(ServerAndPort, UserName, Log);
+			ConnectWindow Connect = new ConnectWindow(DefaultConnection, ServerAndPortOverride, UserNameOverride, Log);
 			if(Connect.ShowDialog() == DialogResult.OK)
 			{
-				NewServerAndPort = Connect.ServerAndPort;
-				NewUserName = Connect.UserName;
+				ServerAndPortOverride = Connect.ServerAndPortOverride;
+				UserNameOverride = Connect.UserNameOverride;
 				return true;
 			}
-			else
-			{
-				NewServerAndPort = null;
-				NewUserName = null;
-				return false;
-			}
+			return false;
 		}
 
 		private void CancelBtn_Click(object sender, EventArgs e)
@@ -74,21 +67,21 @@ namespace UnrealGameSync
 		{
 			if(UseDefaultConnectionSettings.Checked)
 			{
-				ServerAndPort = null;
-				UserName = null;
+				ServerAndPortOverride = null;
+				UserNameOverride = null;
 			}
 			else
 			{
-				ServerAndPort = ServerAndPortTextBox.Text.Trim();
-				if(ServerAndPort.Length == 0)
+				ServerAndPortOverride = ServerAndPortTextBox.Text.Trim();
+				if(ServerAndPortOverride.Length == 0)
 				{
-					ServerAndPort = null;
+					ServerAndPortOverride = DefaultConnection.ServerAndPort;
 				}
 
-				UserName = UserNameTextBox.Text.Trim();
-				if(UserName.Length == 0)
+				UserNameOverride = UserNameTextBox.Text.Trim();
+				if(UserNameOverride.Length == 0)
 				{
-					UserName = null;
+					UserNameOverride = DefaultConnection.UserName;
 				}
 			}
 
@@ -115,7 +108,7 @@ namespace UnrealGameSync
 		private void BrowseUserBtn_Click(object sender, EventArgs e)
 		{
 			string NewUserName;
-			if(SelectUserWindow.ShowModal(this, ServerAndPortTextBox.Text, UserNameTextBox.Text, Log, out NewUserName))
+			if(SelectUserWindow.ShowModal(this, new PerforceConnection(UserNameTextBox.Text, null, ServerAndPortTextBox.Text), Log, out NewUserName))
 			{
 				UserNameTextBox.Text = NewUserName;
 			}
diff --git a/Engine/Source/Programs/UnrealGameSync/UnrealGameSync/Forms/IssueDetailsWindow.cs b/Engine/Source/Programs/UnrealGameSync/UnrealGameSync/Forms/IssueDetailsWindow.cs
--- a/Engine/Source/Programs/UnrealGameSync/UnrealGameSync/Forms/IssueDetailsWindow.cs
+++ b/Engine/Source/Programs/UnrealGameSync/UnrealGameSync/Forms/IssueDetailsWindow.cs
@@ -1307,7 +1307,7 @@ namespace UnrealGameSync
 		private void AssignToOtherBtn_Click(object sender, EventArgs e)
 		{
 			string SelectedUserName;
-			if(SelectUserWindow.ShowModal(this, Perforce.ServerAndPort, Perforce.UserName, Log, out SelectedUserName))
+			if(SelectUserWindow.ShowModal(this, Perforce, Log, out SelectedUserName))
 			{
 				AssignToUser(SelectedUserName);
 			}
diff --git a/Engine/Source/Programs/UnrealGameSync/UnrealGameSync/Forms/IssueFixedWindow.cs b/Engine/Source/Programs/UnrealGameSync/UnrealGameSync/Forms/IssueFixedWindow.cs
--- a/Engine/Source/Programs/UnrealGameSync/UnrealGameSync/Forms/IssueFixedWindow.cs
+++ b/Engine/Source/Programs/UnrealGameSync/UnrealGameSync/Forms/IssueFixedWindow.cs
@@ -350,7 +350,7 @@ namespace UnrealGameSync
 		private void UserBrowseBtn_Click(object sender, EventArgs e)
 		{
 			string SelectedUserName;
-			if(SelectUserWindow.ShowModal(this, Perforce.ServerAndPort, Perforce.UserName, new BufferedTextWriter(), out SelectedUserName))
+			if(SelectUserWindow.ShowModal(this, Perforce, new BufferedTextWriter(), out SelectedUserName))
 			{
 				UserNameTextBox.Text = SelectedUserName;
 			}
diff --git a/Engine/Source/Programs/UnrealGameSync/UnrealGameSync/Forms/MainWindow.cs b/Engine/Source/Programs/UnrealGameSync/UnrealGameSync/Forms/MainWindow.cs
--- a/Engine/Source/Programs/UnrealGameSync/UnrealGameSync/Forms/MainWindow.cs
+++ b/Engine/Source/Programs/UnrealGameSync/UnrealGameSync/Forms/MainWindow.cs
@@ -63,6 +63,7 @@ namespace UnrealGameSync
 		string ApiUrl;
 		string DataFolder;
 		string CacheFolder;
+		PerforceConnection DefaultConnection;
 		LineBasedTextWriter Log;
 		UserSettings Settings;
 		int TabMenu_TabIdx = -1;
@@ -88,7 +89,7 @@ namespace UnrealGameSync
 		Rectangle PrimaryWorkArea;
 		List<IssueAlertWindow> AlertWindows = new List<IssueAlertWindow>();
 
-		public MainWindow(UpdateMonitor InUpdateMonitor, string InApiUrl, string InDataFolder, string InCacheFolder, bool bInRestoreStateOnLoad, string InOriginalExecutableFileName, bool bInUnstable, DetectProjectSettingsResult[] StartupProjects, LineBasedTextWriter InLog, UserSettings InSettings)
+		public MainWindow(UpdateMonitor InUpdateMonitor, string InApiUrl, string InDataFolder, string InCacheFolder, bool bInRestoreStateOnLoad, string InOriginalExecutableFileName, bool bInUnstable, DetectProjectSettingsResult[] StartupProjects, PerforceConnection InDefaultConnection, LineBasedTextWriter InLog, UserSettings InSettings)
 		{
 			InitializeComponent();
 
@@ -100,6 +101,7 @@ namespace UnrealGameSync
 			bRestoreStateOnLoad = bInRestoreStateOnLoad;
 			OriginalExecutableFileName = InOriginalExecutableFileName;
 			bUnstable = bInUnstable;
+			DefaultConnection = InDefaultConnection;
 			Log = InLog;
 			Settings = InSettings;
 
@@ -217,7 +219,7 @@ namespace UnrealGameSync
 			string ProjectPath = Reader.ReadString();
 
 			AutomatedSyncWindow.WorkspaceInfo WorkspaceInfo;
-			if(!AutomatedSyncWindow.ShowModal(this, StreamName, ProjectPath, out WorkspaceInfo, Log))
+			if(!AutomatedSyncWindow.ShowModal(this, DefaultConnection, StreamName, ProjectPath, out WorkspaceInfo, Log))
 			{
 				return new AutomationRequestOutput(AutomationRequestResult.Canceled);
 			}
@@ -861,7 +863,7 @@ namespace UnrealGameSync
 		public void OpenNewProject()
 		{
 			DetectProjectSettingsTask DetectedProjectSettings;
-			if(OpenProjectWindow.ShowModal(this, null, out DetectedProjectSettings, Settings, DataFolder, CacheFolder, Log))
+			if(OpenProjectWindow.ShowModal(this, null, out DetectedProjectSettings, Settings, DataFolder, CacheFolder, DefaultConnection, Log))
 			{
 				int NewTabIdx = TryOpenProject(DetectedProjectSettings, -1, OpenProjectOptions.None);
 				if(NewTabIdx != -1)
@@ -913,7 +915,7 @@ namespace UnrealGameSync
 		public void EditSelectedProject(int TabIdx, UserSelectedProjectSettings SelectedProject)
 		{
 			DetectProjectSettingsTask DetectedProjectSettings;
-			if(OpenProjectWindow.ShowModal(this, SelectedProject, out DetectedProjectSettings, Settings, DataFolder, CacheFolder, Log))
+			if(OpenProjectWindow.ShowModal(this, SelectedProject, out DetectedProjectSettings, Settings, DataFolder, CacheFolder, DefaultConnection, Log))
 			{
 				int NewTabIdx = TryOpenProject(DetectedProjectSettings, TabIdx, OpenProjectOptions.None);
 				if(NewTabIdx != -1)
@@ -933,16 +935,18 @@ namespace UnrealGameSync
 			Log.WriteLine("Detecting settings for {0}", Project);
 			using(DetectProjectSettingsTask DetectProjectSettings = new DetectProjectSettingsTask(Project, DataFolder, CacheFolder, new PrefixedTextWriter("  ", Log)))
 			{
+				PerforceConnection Perforce = Utility.OverridePerforceSettings(DefaultConnection, Project.ServerAndPort, Project.UserName);
+
 				string ErrorMessage;
 
 				ModalTaskResult Result;
 				if((Options & OpenProjectOptions.Quiet) != 0)
 				{
-					Result = ModalTask.Execute(this, DetectProjectSettings, "Opening Project", "Opening project, please wait...", out ErrorMessage);
+					Result = ModalTask.Execute(this, new QuietDetectProjectSettingsTask(Perforce, DetectProjectSettings, Log), "Opening Project", "Opening project, please wait...", out ErrorMessage);
 				}
 				else
 				{
-					Result = PerforceModalTask.Execute(this, Project.LocalPath, Project.ServerAndPort, Project.UserName, DetectProjectSettings, "Opening Project", "Opening project, please wait...", Log, out ErrorMessage);
+					Result = PerforceModalTask.Execute(this, Perforce, DetectProjectSettings, "Opening Project", "Opening project, please wait...", Log, out ErrorMessage);
 				}
 
 				if(Result != ModalTaskResult.Succeeded)
@@ -1211,7 +1215,7 @@ namespace UnrealGameSync
 
 		public void ModifyApplicationSettings()
 		{
-			bool? bRelaunchUnstable = ApplicationSettingsWindow.ShowModal(this, bUnstable, OriginalExecutableFileName, Settings, Log);
+			bool? bRelaunchUnstable = ApplicationSettingsWindow.ShowModal(this, DefaultConnection, bUnstable, OriginalExecutableFileName, Settings, Log);
 			if(bRelaunchUnstable.HasValue)
 			{
 				UpdateMonitor.TriggerUpdate(UpdateType.UserInitiated, bRelaunchUnstable);
diff --git a/Engine/Source/Programs/UnrealGameSync/UnrealGameSync/Forms/NewWorkspaceWindow.cs b/Engine/Source/Programs/UnrealGameSync/UnrealGameSync/Forms/NewWorkspaceWindow.cs
--- a/Engine/Source/Programs/UnrealGameSync/UnrealGameSync/Forms/NewWorkspaceWindow.cs
+++ b/Engine/Source/Programs/UnrealGameSync/UnrealGameSync/Forms/NewWorkspaceWindow.cs
@@ -104,20 +104,18 @@ namespace UnrealGameSync
 			}
 		}
 
-		string ServerAndPort;
-		string UserName;
+		PerforceConnection Perforce;
 		PerforceInfoRecord Info;
 		List<PerforceClientRecord> Clients;
 		TextWriter Log;
 		NewWorkspaceSettings Settings;
 		string DefaultRootPath;
 
-		private NewWorkspaceWindow(string ServerAndPort, string UserName, string ForceStream, string DefaultStream, PerforceInfoRecord Info, List<PerforceClientRecord> Clients, TextWriter Log)
+		private NewWorkspaceWindow(PerforceConnection Perforce, string ForceStream, string DefaultStream, PerforceInfoRecord Info, List<PerforceClientRecord> Clients, TextWriter Log)
 		{
 			InitializeComponent();
 
-			this.ServerAndPort = ServerAndPort;
-			this.UserName = UserName;
+			this.Perforce = Perforce;
 			this.Info = Info;
 			this.Clients = Clients;
 			this.Log = Log;
@@ -180,12 +178,12 @@ namespace UnrealGameSync
 			UpdateRootDirCueBanner();
 		}
 
-		public static bool ShowModal(IWin32Window Owner, string ServerAndPort, string UserName, string ForceStreamName, string CurrentWorkspaceName, TextWriter Log, out string WorkspaceName)
+		public static bool ShowModal(IWin32Window Owner, PerforceConnection Perforce, string ForceStreamName, string CurrentWorkspaceName, TextWriter Log, out string WorkspaceName)
 		{
 			FindWorkspaceSettingsTask FindSettings = new FindWorkspaceSettingsTask(CurrentWorkspaceName);
 
 			string ErrorMessage;
-			if(PerforceModalTask.Execute(Owner, null, ServerAndPort, UserName, FindSettings, "Checking settings", "Checking settings, please wait...", Log, out ErrorMessage) != ModalTaskResult.Succeeded)
+			if(PerforceModalTask.Execute(Owner, Perforce, FindSettings, "Checking settings", "Checking settings, please wait...", Log, out ErrorMessage) != ModalTaskResult.Succeeded)
 			{
 				if(!String.IsNullOrEmpty(ErrorMessage))
 				{
@@ -196,7 +194,7 @@ namespace UnrealGameSync
 				return false;
 			}
 
-			NewWorkspaceWindow Window = new NewWorkspaceWindow(ServerAndPort, UserName, ForceStreamName, FindSettings.CurrentStream, FindSettings.Info, FindSettings.Clients, Log);
+			NewWorkspaceWindow Window = new NewWorkspaceWindow(Perforce, ForceStreamName, FindSettings.CurrentStream, FindSettings.Info, FindSettings.Clients, Log);
 			if(Window.ShowDialog(Owner) == DialogResult.OK)
 			{
 				WorkspaceName = Window.Settings.Name;
@@ -354,7 +352,7 @@ namespace UnrealGameSync
 				NewWorkspaceTask NewWorkspace = new NewWorkspaceTask(Settings, Info.UserName, Info.HostName);
 
 				string ErrorMessage;
-				if(PerforceModalTask.Execute(Owner, null, ServerAndPort, UserName, NewWorkspace, "Creating workspace", "Creating workspace, please wait...", Log, out ErrorMessage) != ModalTaskResult.Succeeded)
+				if(PerforceModalTask.Execute(Owner, Perforce, NewWorkspace, "Creating workspace", "Creating workspace, please wait...", Log, out ErrorMessage) != ModalTaskResult.Succeeded)
 				{
 					MessageBox.Show(ErrorMessage);
 					return;
@@ -374,7 +372,7 @@ namespace UnrealGameSync
 		private void StreamBrowseBtn_Click(object sender, EventArgs e)
 		{
 			string StreamName = StreamTextBox.Text.Trim();
-			if(SelectStreamWindow.ShowModal(this, ServerAndPort, UserName, StreamName, Log, out StreamName))
+			if(SelectStreamWindow.ShowModal(this, Perforce, StreamName, Log, out StreamName))
 			{
 				StreamTextBox.Text = StreamName;
 			}
diff --git a/Engine/Source/Programs/UnrealGameSync/UnrealGameSync/Forms/OpenProjectWindow.cs b/Engine/Source/Programs/UnrealGameSync/UnrealGameSync/Forms/OpenProjectWindow.cs
--- a/Engine/Source/Programs/UnrealGameSync/UnrealGameSync/Forms/OpenProjectWindow.cs
+++ b/Engine/Source/Programs/UnrealGameSync/UnrealGameSync/Forms/OpenProjectWindow.cs
@@ -15,15 +15,16 @@ namespace UnrealGameSync
 {
 	partial class OpenProjectWindow : Form
 	{
-		string ServerAndPort;
-		string UserName;
+		string ServerAndPortOverride;
+		string UserNameOverride;
 		DetectProjectSettingsTask DetectedProjectSettings;
 		string DataFolder;
 		string CacheFolder;
+		PerforceConnection DefaultConnection;
 		TextWriter Log;
 		UserSettings Settings;
 
-		private OpenProjectWindow(UserSelectedProjectSettings Project, UserSettings Settings, string DataFolder, string CacheFolder, TextWriter Log)
+		private OpenProjectWindow(UserSelectedProjectSettings Project, UserSettings Settings, string DataFolder, string CacheFolder, PerforceConnection DefaultConnection, TextWriter Log)
 		{
 			InitializeComponent();
 
@@ -31,6 +32,7 @@ namespace UnrealGameSync
 			this.DetectedProjectSettings = null;
 			this.DataFolder = DataFolder;
 			this.CacheFolder = CacheFolder;
+			this.DefaultConnection = DefaultConnection;
 			this.Log = Log;
 
 			if(Project == null)
@@ -41,11 +43,11 @@ namespace UnrealGameSync
 			{
 				if(!String.IsNullOrWhiteSpace(Project.ServerAndPort))
 				{
-					ServerAndPort = Project.ServerAndPort;
+					ServerAndPortOverride = Project.ServerAndPort;
 				}
 				if(!String.IsNullOrWhiteSpace(Project.UserName))
 				{
-					UserName = Project.UserName;
+					UserNameOverride = Project.UserName;
 				}
 
 				if(Project.ClientPath != null && Project.ClientPath.StartsWith("//"))
@@ -79,9 +81,14 @@ namespace UnrealGameSync
 			UpdateOkButton();
 		}
 
-		public static bool ShowModal(IWin32Window Owner, UserSelectedProjectSettings Project, out DetectProjectSettingsTask NewDetectedProjectSettings, UserSettings Settings, string DataFolder, string CacheFolder, TextWriter Log)
+		private PerforceConnection Perforce
 		{
-			OpenProjectWindow Window = new OpenProjectWindow(Project, Settings, DataFolder, CacheFolder, Log);
+			get { return Utility.OverridePerforceSettings(DefaultConnection, ServerAndPortOverride, UserNameOverride); }
+		}
+
+		public static bool ShowModal(IWin32Window Owner, UserSelectedProjectSettings Project, out DetectProjectSettingsTask NewDetectedProjectSettings, UserSettings Settings, string DataFolder, string CacheFolder, PerforceConnection DefaultConnection, TextWriter Log)
+		{
+			OpenProjectWindow Window = new OpenProjectWindow(Project, Settings, DataFolder, CacheFolder, DefaultConnection, Log);
 			if(Window.ShowDialog(Owner) == DialogResult.OK)
 			{
 				NewDetectedProjectSettings = Window.DetectedProjectSettings;
@@ -113,11 +120,11 @@ namespace UnrealGameSync
 			UpdateWorkspacePathBrowseButton();
 		}
 
-		public static string GetServerLabelText(string ServerAndPort, string UserName)
+		public static string GetServerLabelText(PerforceConnection DefaultConnection, string ServerAndPort, string UserName)
 		{
 			if(ServerAndPort == null && UserName == null)
 			{
-				return "Using default Perforce server settings.";
+				return String.Format("Using default connection settings (user '{0}' on server '{1}').", DefaultConnection.UserName, DefaultConnection.ServerAndPort);
 			}
 			else
 			{
@@ -145,8 +152,7 @@ namespace UnrealGameSync
 
 		private void UpdateServerLabel()
 		{
-			ServerLabel.Text = GetServerLabelText(ServerAndPort, UserName);
-			ChangeLink.Location = new Point(ServerLabel.Right + 5, ChangeLink.Location.Y);
+			ServerLabel.Text = GetServerLabelText(DefaultConnection, ServerAndPortOverride, UserNameOverride);
 		}
 
 		private void UpdateWorkspacePathBrowseButton()
@@ -166,7 +172,7 @@ namespace UnrealGameSync
 			WorkspaceRadioBtn.Checked = true;
 			
 			string WorkspaceName;
-			if(NewWorkspaceWindow.ShowModal(this, ServerAndPort, UserName, null, WorkspaceNameTextBox.Text, Log, out WorkspaceName))
+			if(NewWorkspaceWindow.ShowModal(this, Perforce, null, WorkspaceNameTextBox.Text, Log, out WorkspaceName))
 			{
 				WorkspaceNameTextBox.Text = WorkspaceName;
 				UpdateOkButton();
@@ -178,7 +184,7 @@ namespace UnrealGameSync
 			WorkspaceRadioBtn.Checked = true;
 
 			string WorkspaceName = WorkspaceNameTextBox.Text;
-			if(SelectWorkspaceWindow.ShowModal(this, ServerAndPort, UserName, WorkspaceName, Log, out WorkspaceName))
+			if(SelectWorkspaceWindow.ShowModal(this, Perforce, WorkspaceName, Log, out WorkspaceName))
 			{
 				WorkspaceNameTextBox.Text = WorkspaceName;
 			}
@@ -192,7 +198,7 @@ namespace UnrealGameSync
 			if(TryGetWorkspaceName(out WorkspaceName))
 			{
 				string WorkspacePath = WorkspacePathTextBox.Text.Trim();
-				if(SelectProjectFromWorkspaceWindow.ShowModal(this, ServerAndPort, UserName, WorkspaceName, WorkspacePath, Log, out WorkspacePath))
+				if(SelectProjectFromWorkspaceWindow.ShowModal(this, Perforce, WorkspaceName, WorkspacePath, Log, out WorkspacePath))
 				{
 					WorkspacePathTextBox.Text = WorkspacePath;
 					UpdateOkButton();
@@ -253,7 +259,7 @@ namespace UnrealGameSync
 				string ClientPath;
 				if(TryGetClientPath(out ClientPath))
 				{
-					Project = new UserSelectedProjectSettings(ServerAndPort, UserName, UserSelectedProjectType.Client, ClientPath, null);
+					Project = new UserSelectedProjectSettings(ServerAndPortOverride, UserNameOverride, UserSelectedProjectType.Client, ClientPath, null);
 					return true;
 				}
 			}
@@ -262,7 +268,7 @@ namespace UnrealGameSync
 				string LocalPath;
 				if(TryGetLocalPath(out LocalPath))
 				{
-					Project = new UserSelectedProjectSettings(ServerAndPort, UserName, UserSelectedProjectType.Local, null, LocalPath);
+					Project = new UserSelectedProjectSettings(ServerAndPortOverride, UserNameOverride, UserSelectedProjectType.Local, null, LocalPath);
 					return true;
 				}
 			}
@@ -279,14 +285,8 @@ namespace UnrealGameSync
 				DetectProjectSettingsTask NewDetectedProjectSettings = new DetectProjectSettingsTask(SelectedProject, DataFolder, CacheFolder, Log);
 				try
 				{
-					string ProjectFileName = null;
-					if(SelectedProject.Type == UserSelectedProjectType.Local)
-					{
-						ProjectFileName = SelectedProject.LocalPath;
-					}
-
 					string ErrorMessage;
-					if(PerforceModalTask.Execute(this, ProjectFileName, SelectedProject.ServerAndPort, SelectedProject.UserName, NewDetectedProjectSettings, "Opening Project", "Opening project, please wait...", Log, out ErrorMessage) != ModalTaskResult.Succeeded)
+					if(PerforceModalTask.Execute(this, Perforce, NewDetectedProjectSettings, "Opening Project", "Opening project, please wait...", Log, out ErrorMessage) != ModalTaskResult.Succeeded)
 					{
 						if(!String.IsNullOrEmpty(ErrorMessage))
 						{
@@ -313,12 +313,8 @@ namespace UnrealGameSync
 
 		private void ChangeLink_LinkClicked(object sender, LinkLabelLinkClickedEventArgs e)
 		{
-			string NewServerAndPort;
-			string NewUserName;
-			if(ConnectWindow.ShowModal(this, ServerAndPort, UserName, Log, out NewServerAndPort, out NewUserName))
+			if(ConnectWindow.ShowModal(this, DefaultConnection, ref ServerAndPortOverride, ref UserNameOverride, Log))
 			{
-				ServerAndPort = NewServerAndPort;
-				UserName = NewUserName;
 				UpdateServerLabel();
 			}
 		}
diff --git a/Engine/Source/Programs/UnrealGameSync/UnrealGameSync/Forms/SelectProjectFromWorkspaceWindow.cs b/Engine/Source/Programs/UnrealGameSync/UnrealGameSync/Forms/SelectProjectFromWorkspaceWindow.cs
--- a/Engine/Source/Programs/UnrealGameSync/UnrealGameSync/Forms/SelectProjectFromWorkspaceWindow.cs
+++ b/Engine/Source/Programs/UnrealGameSync/UnrealGameSync/Forms/SelectProjectFromWorkspaceWindow.cs
@@ -207,12 +207,12 @@ namespace UnrealGameSync
 			return NextNode;
 		}
 
-		public static bool ShowModal(IWin32Window Owner, string ServerAndPort, string UserName, string WorkspaceName, string WorkspacePath, TextWriter Log, out string NewWorkspacePath)
+		public static bool ShowModal(IWin32Window Owner, PerforceConnection Perforce, string WorkspaceName, string WorkspacePath, TextWriter Log, out string NewWorkspacePath)
 		{
 			EnumerateWorkspaceProjectsTask EnumerateProjectsTask = new EnumerateWorkspaceProjectsTask(WorkspaceName);
 
 			string ErrorMessage;
-			if(PerforceModalTask.Execute(Owner, null, ServerAndPort, UserName, EnumerateProjectsTask, "Finding Projects", "Finding projects, please wait...", Log, out ErrorMessage) != ModalTaskResult.Succeeded)
+			if(PerforceModalTask.Execute(Owner, Perforce, EnumerateProjectsTask, "Finding Projects", "Finding projects, please wait...", Log, out ErrorMessage) != ModalTaskResult.Succeeded)
 			{
 				if(!String.IsNullOrEmpty(ErrorMessage))
 				{
diff --git a/Engine/Source/Programs/UnrealGameSync/UnrealGameSync/Forms/SelectStreamWindow.cs b/Engine/Source/Programs/UnrealGameSync/UnrealGameSync/Forms/SelectStreamWindow.cs
--- a/Engine/Source/Programs/UnrealGameSync/UnrealGameSync/Forms/SelectStreamWindow.cs
+++ b/Engine/Source/Programs/UnrealGameSync/UnrealGameSync/Forms/SelectStreamWindow.cs
@@ -321,12 +321,12 @@ namespace UnrealGameSync
 			return bExpand || (SelectedStream == Stream.Record.Identifier);
 		}
 
-		public static bool ShowModal(IWin32Window Owner, string ServerAndPort, string UserName, string StreamName, TextWriter Log, out string NewStreamName)
+		public static bool ShowModal(IWin32Window Owner, PerforceConnection Perforce, string StreamName, TextWriter Log, out string NewStreamName)
 		{
 			EnumerateStreamsTask Task = new EnumerateStreamsTask();
 
 			string ErrorMessage;
-			ModalTaskResult Result = PerforceModalTask.Execute(Owner, null, ServerAndPort, UserName, Task, "Finding streams", "Finding streams, please wait...", Log, out ErrorMessage);
+			ModalTaskResult Result = PerforceModalTask.Execute(Owner, Perforce, Task, "Finding streams", "Finding streams, please wait...", Log, out ErrorMessage);
 			if(Result != ModalTaskResult.Succeeded)
 			{
 				if(!String.IsNullOrEmpty(ErrorMessage))
diff --git a/Engine/Source/Programs/UnrealGameSync/UnrealGameSync/Forms/SelectUserWindow.cs b/Engine/Source/Programs/UnrealGameSync/UnrealGameSync/Forms/SelectUserWindow.cs
--- a/Engine/Source/Programs/UnrealGameSync/UnrealGameSync/Forms/SelectUserWindow.cs
+++ b/Engine/Source/Programs/UnrealGameSync/UnrealGameSync/Forms/SelectUserWindow.cs
@@ -151,12 +151,12 @@ namespace UnrealGameSync
 			UpdateOkButton();
 		}
 
-		public static bool ShowModal(IWin32Window Owner, string ServerAndPort, string UserName, TextWriter Log, out string SelectedUserName)
+		public static bool ShowModal(IWin32Window Owner, PerforceConnection Perforce, TextWriter Log, out string SelectedUserName)
 		{
 			EnumerateUsersTask Task = new EnumerateUsersTask();
 
 			string ErrorMessage;
-			ModalTaskResult Result = PerforceModalTask.Execute(Owner, null, ServerAndPort, UserName, Task, "Finding users", "Finding users, please wait...", Log, out ErrorMessage);
+			ModalTaskResult Result = PerforceModalTask.Execute(Owner, Perforce, Task, "Finding users", "Finding users, please wait...", Log, out ErrorMessage);
 			if(Result != ModalTaskResult.Succeeded)
 			{
 				if(!String.IsNullOrEmpty(ErrorMessage))
diff --git a/Engine/Source/Programs/UnrealGameSync/UnrealGameSync/Forms/SelectWorkspaceWindow.cs b/Engine/Source/Programs/UnrealGameSync/UnrealGameSync/Forms/SelectWorkspaceWindow.cs
--- a/Engine/Source/Programs/UnrealGameSync/UnrealGameSync/Forms/SelectWorkspaceWindow.cs
+++ b/Engine/Source/Programs/UnrealGameSync/UnrealGameSync/Forms/SelectWorkspaceWindow.cs
@@ -86,12 +86,12 @@ namespace UnrealGameSync
 			OkBtn.Enabled = (WorkspaceListView.SelectedItems.Count == 1);
 		}
 
-		public static bool ShowModal(IWin32Window Owner, string ServerAndPort, string UserName, string WorkspaceName, TextWriter Log, out string NewWorkspaceName)
+		public static bool ShowModal(IWin32Window Owner, PerforceConnection Perforce, string WorkspaceName, TextWriter Log, out string NewWorkspaceName)
 		{
 			EnumerateWorkspacesTask EnumerateWorkspaces = new EnumerateWorkspacesTask();
 
 			string ErrorMessage;
-			ModalTaskResult Result = PerforceModalTask.Execute(Owner, null, ServerAndPort, UserName, EnumerateWorkspaces, "Finding workspaces", "Finding workspaces, please wait...", Log, out ErrorMessage);
+			ModalTaskResult Result = PerforceModalTask.Execute(Owner, Perforce, EnumerateWorkspaces, "Finding workspaces", "Finding workspaces, please wait...", Log, out ErrorMessage);
 			if(Result != ModalTaskResult.Succeeded)
 			{
 				if(!String.IsNullOrEmpty(ErrorMessage))
diff --git a/Engine/Source/Programs/UnrealGameSync/UnrealGameSync/PerforceModalTask.cs b/Engine/Source/Programs/UnrealGameSync/UnrealGameSync/PerforceModalTask.cs
--- a/Engine/Source/Programs/UnrealGameSync/UnrealGameSync/PerforceModalTask.cs
+++ b/Engine/Source/Programs/UnrealGameSync/UnrealGameSync/PerforceModalTask.cs
@@ -17,95 +17,24 @@ namespace UnrealGameSync
 
 	class PerforceModalTask : IModalTask
 	{
-		public string ProjectFileName;
-		public string ServerAndPort;
-		public string UserName;
+		public PerforceConnection Perforce;
 		public string Password;
 		public LoginResult LoginResult;
 		public IPerforceModalTask InnerTask;
 		public TextWriter Log;
 
-		public PerforceModalTask(string ProjectFileName, string ServerAndPort, string UserName, IPerforceModalTask InnerTask, TextWriter Log)
+		public PerforceModalTask(PerforceConnection Perforce, IPerforceModalTask InnerTask, TextWriter Log)
 		{
-			this.ProjectFileName = ProjectFileName;
-			this.ServerAndPort = ServerAndPort;
-			this.UserName = UserName;
+			this.Perforce = Perforce;
 			this.InnerTask = InnerTask;
 			this.Log = Log;
 		}
 
-		public static bool TryGetServerSettings(string ProjectFileName, ref string ServerAndPort, ref string UserName, TextWriter Log)
-		{
-			// Read the P4PORT setting for the server, if necessary. Change to the project folder if set, so we can respect the contents of any P4CONFIG file.
-			if(ServerAndPort == null)
-			{
-				string PrevDirectory = Directory.GetCurrentDirectory();
-				try
-				{
-					PerforceConnection Perforce = new PerforceConnection(UserName, null, null);
-					if(ProjectFileName != null)
-					{
-						try
-						{
-							Directory.SetCurrentDirectory(Path.GetDirectoryName(ProjectFileName));
-						}
-						catch
-						{
-							// Just ignore
-						}
-					}
-
-					string NewServerAndPort;
-					if (Perforce.GetSetting("P4PORT", out NewServerAndPort, Log))
-					{
-						ServerAndPort = NewServerAndPort;
-					}
-				}
-				finally
-				{
-					Directory.SetCurrentDirectory(PrevDirectory);
-				}
-			}
-
-			// If the server setting is still null, user the default
-			if(ServerAndPort == null)
-			{
-				ServerAndPort = PerforceConnection.DefaultServerAndPort;
-			}
-
-			// Update the server and username from the reported server info if it's not set
-			if(UserName == null)
-			{
-				PerforceConnection Perforce = new PerforceConnection(UserName, null, ServerAndPort);
-
-				PerforceInfoRecord PerforceInfo;
-				if(!Perforce.Info(out PerforceInfo, Log) || String.IsNullOrEmpty(PerforceInfo.UserName))
-				{
-					return false;
-				}
-
-				UserName = PerforceInfo.UserName;
-			}
-
-			// Otherwise succeed
-			return true;
-		}
-
 		public bool Run(out string ErrorMessage)
 		{
 			// Set the default login state to failed
 			LoginResult = LoginResult.Failed;
 
-			// Get the server settings
-			if(!TryGetServerSettings(ProjectFileName, ref ServerAndPort, ref UserName, Log))
-			{
-				ErrorMessage = "Unable to get Perforce server settings.";
-				return false;
-			}
-
-			// Create the connection
-			PerforceConnection Perforce = new PerforceConnection(UserName, null, ServerAndPort);
-
 			// If we've got a password, execute the login command
 			if(Password != null)
 			{
@@ -139,9 +68,9 @@ namespace UnrealGameSync
 			return InnerTask.Run(Perforce, Log, out ErrorMessage);
 		}
 
-		public static ModalTaskResult Execute(IWin32Window Owner, string ProjectFileName, string ServerAndPort, string UserName, IPerforceModalTask PerforceTask, string Title, string Message, TextWriter Log, out string ErrorMessage)
+		public static ModalTaskResult Execute(IWin32Window Owner, PerforceConnection Perforce, IPerforceModalTask PerforceTask, string Title, string Message, TextWriter Log, out string ErrorMessage)
 		{
-			PerforceModalTask Task = new PerforceModalTask(ProjectFileName, ServerAndPort, UserName, PerforceTask, Log);
+			PerforceModalTask Task = new PerforceModalTask(Perforce, PerforceTask, Log);
 			for(;;)
 			{
 				string TaskErrorMessage;
@@ -154,7 +83,7 @@ namespace UnrealGameSync
 				}
 				else if(Task.LoginResult == LoginResult.MissingPassword)
 				{
-					PasswordWindow PasswordWindow = new PasswordWindow(String.Format("Enter the password for user '{0}' on server '{1}'.", Task.UserName, Task.ServerAndPort), Task.Password);
+					PasswordWindow PasswordWindow = new PasswordWindow(String.Format("Enter the password for user '{0}' on server '{1}'.", Perforce.UserName, Perforce.ServerAndPort), Task.Password);
 					if(Owner == null)
 					{
 						PasswordWindow.StartPosition = FormStartPosition.CenterScreen;
@@ -168,7 +97,7 @@ namespace UnrealGameSync
 				}
 				else if(Task.LoginResult == LoginResult.IncorrectPassword)
 				{
-					PasswordWindow PasswordWindow = new PasswordWindow(String.Format("Authentication failed. Enter the password for user '{0}' on server '{1}'.", Task.UserName, Task.ServerAndPort), Task.Password);
+					PasswordWindow PasswordWindow = new PasswordWindow(String.Format("Authentication failed. Enter the password for user '{0}' on server '{1}'.", Perforce.UserName, Perforce.ServerAndPort), Task.Password);
 					if (Owner == null)
 					{
 						PasswordWindow.StartPosition = FormStartPosition.CenterScreen;
@@ -187,5 +116,20 @@ namespace UnrealGameSync
 				}
 			}
 		}
+
+		public static bool ExecuteAndShowError(IWin32Window Owner, PerforceConnection Perforce, IPerforceModalTask Task, string Title, string Message, TextWriter Log)
+		{
+			string ErrorMessage;
+			ModalTaskResult Result = Execute(Owner, Perforce, Task, Title, Message, Log, out ErrorMessage);
+			if (Result != ModalTaskResult.Succeeded)
+			{
+				if (!String.IsNullOrEmpty(ErrorMessage))
+				{
+					MessageBox.Show(ErrorMessage);
+				}
+				return false;
+			}
+			return true;
+		}
 	}
 }
diff --git a/Engine/Source/Programs/UnrealGameSync/UnrealGameSync/Program.cs b/Engine/Source/Programs/UnrealGameSync/UnrealGameSync/Program.cs
--- a/Engine/Source/Programs/UnrealGameSync/UnrealGameSync/Program.cs
+++ b/Engine/Source/Programs/UnrealGameSync/UnrealGameSync/Program.cs
@@ -85,21 +85,72 @@ namespace UnrealGameSync
 			using(TelemetryWriter Telemetry = new TelemetryWriter(DeploymentSettings.ApiUrl, Path.Combine(DataFolder, "Telemetry.log")))
 			{
 				AppDomain.CurrentDomain.UnhandledException += CurrentDomain_UnhandledException;
-				using(UpdateMonitor UpdateMonitor = new UpdateMonitor(new PerforceConnection(UserName, null, ServerAndPort), UpdatePath))
+
+				// Create the log file
+				using (TimestampLogWriter Log = new TimestampLogWriter(new BoundedLogWriter(Path.Combine(DataFolder, "UnrealGameSync.log"))))
 				{
-					ProgramApplicationContext Context = new ProgramApplicationContext(UpdateMonitor, DeploymentSettings.ApiUrl, DataFolder, ActivateEvent, bRestoreState, UpdateSpawn, ProjectFileName, bUnstable);
-					Application.Run(Context);
+					Log.WriteLine("Application version: {0}", Assembly.GetExecutingAssembly().GetName().Version);
+					Log.WriteLine("Started at {0}", DateTime.Now.ToString());
+
+					if (ServerAndPort == null || UserName == null)
+					{
+						Log.WriteLine("Missing server settings; finding defaults.");
+						GetDefaultServerSettings(ref ServerAndPort, ref UserName, Log);
+						Utility.SaveGlobalPerforceSettings(ServerAndPort, UserName, UpdatePath);
+					}
 
-					if(UpdateMonitor.IsUpdateAvailable && UpdateSpawn != null)
+					PerforceConnection DefaultConnection = new PerforceConnection(UserName, null, ServerAndPort);
+					using (UpdateMonitor UpdateMonitor = new UpdateMonitor(DefaultConnection, UpdatePath))
 					{
-						InstanceMutex.Close();
-						bool bLaunchUnstable = UpdateMonitor.RelaunchUnstable ?? bUnstable;
-						Utility.SpawnProcess(UpdateSpawn, "-restorestate" + (bLaunchUnstable? " -unstable" : ""));
+						ProgramApplicationContext Context = new ProgramApplicationContext(DefaultConnection, UpdateMonitor, DeploymentSettings.ApiUrl, DataFolder, ActivateEvent, bRestoreState, UpdateSpawn, ProjectFileName, bUnstable, Log);
+						Application.Run(Context);
+
+						if(UpdateMonitor.IsUpdateAvailable && UpdateSpawn != null)
+						{
+							InstanceMutex.Close();
+							bool bLaunchUnstable = UpdateMonitor.RelaunchUnstable ?? bUnstable;
+							Utility.SpawnProcess(UpdateSpawn, "-restorestate" + (bLaunchUnstable? " -unstable" : ""));
+						}
 					}
 				}
 			}
 		}
 
+		public static void GetDefaultServerSettings(ref string ServerAndPort, ref string UserName, TextWriter Log)
+		{
+			// Read the P4PORT setting for the server, if necessary. Change to the project folder if set, so we can respect the contents of any P4CONFIG file.
+			if(ServerAndPort == null)
+			{
+				PerforceConnection Perforce = new PerforceConnection(UserName, null, null);
+
+				string NewServerAndPort;
+				if (Perforce.GetSetting("P4PORT", out NewServerAndPort, Log))
+				{
+					ServerAndPort = NewServerAndPort;
+				}
+				else
+				{
+					ServerAndPort = PerforceConnection.DefaultServerAndPort;
+				}
+			}
+
+			// Update the server and username from the reported server info if it's not set
+			if(UserName == null)
+			{
+				PerforceConnection Perforce = new PerforceConnection(UserName, null, ServerAndPort);
+
+				PerforceInfoRecord PerforceInfo;
+				if(Perforce.Info(out PerforceInfo, Log) && !String.IsNullOrEmpty(PerforceInfo.UserName))
+				{
+					UserName = PerforceInfo.UserName;
+				}
+				else
+				{
+					UserName = Environment.UserName;
+				}
+			}
+		}
+
 		private static void CurrentDomain_UnhandledException(object Sender, UnhandledExceptionEventArgs Args)
 		{
 			Exception Ex = Args.ExceptionObject as Exception;
diff --git a/Engine/Source/Programs/UnrealGameSync/UnrealGameSync/ProgramApplicationContext.cs b/Engine/Source/Programs/UnrealGameSync/UnrealGameSync/ProgramApplicationContext.cs
--- a/Engine/Source/Programs/UnrealGameSync/UnrealGameSync/ProgramApplicationContext.cs
+++ b/Engine/Source/Programs/UnrealGameSync/UnrealGameSync/ProgramApplicationContext.cs
@@ -18,6 +18,7 @@ namespace UnrealGameSync
 	{
 		SynchronizationContext MainThreadSynchronizationContext;
 
+		PerforceConnection DefaultConnection;
 		UpdateMonitor UpdateMonitor;
 		string ApiUrl;
 		string DataFolder;
@@ -46,8 +47,9 @@ namespace UnrealGameSync
 		ModalTaskWindow DetectStartupProjectSettingsWindow;
 		MainWindow MainWindowInstance;
 
-		public ProgramApplicationContext(UpdateMonitor UpdateMonitor, string ApiUrl, string DataFolder, EventWaitHandle ActivateEvent, bool bRestoreState, string UpdateSpawn, string ProjectFileName, bool bUnstable)
+		public ProgramApplicationContext(PerforceConnection DefaultConnection, UpdateMonitor UpdateMonitor, string ApiUrl, string DataFolder, EventWaitHandle ActivateEvent, bool bRestoreState, string UpdateSpawn, string ProjectFileName, bool bUnstable, TimestampLogWriter Log)
 		{
+			this.DefaultConnection = DefaultConnection;
 			this.UpdateMonitor = UpdateMonitor;
 			this.ApiUrl = ApiUrl;
 			this.DataFolder = DataFolder;
@@ -55,6 +57,7 @@ namespace UnrealGameSync
 			this.bRestoreState = bRestoreState;
 			this.UpdateSpawn = UpdateSpawn;
 			this.bUnstable = bUnstable;
+			this.Log = Log;
 
 			// Create the directories
 			Directory.CreateDirectory(DataFolder);
@@ -70,11 +73,6 @@ namespace UnrealGameSync
 			// Capture the main thread's synchronization context for callbacks
 			MainThreadSynchronizationContext = WindowsFormsSynchronizationContext.Current;
 
-			// Create the log file
-			Log = new TimestampLogWriter(new BoundedLogWriter(Path.Combine(DataFolder, "UnrealGameSync.log")));
-			Log.WriteLine("Application version: {0}", Assembly.GetExecutingAssembly().GetName().Version);
-			Log.WriteLine("Started at {0}", DateTime.Now.ToString());
-
 			// Read the user's settings
 			Settings = new UserSettings(Path.Combine(DataFolder, "UnrealGameSync.ini"));
 			if(!String.IsNullOrEmpty(ProjectFileName))
@@ -86,6 +84,27 @@ namespace UnrealGameSync
 				}
 			}
 
+			// Update the settings to the latest version
+			if(Settings.Version < UserSettingsVersion.Latest)
+			{
+				// Clear out the server settings for anything using the default server
+				if(Settings.Version < UserSettingsVersion.DefaultServerSettings)
+				{
+					for(int Idx = 0; Idx < Settings.OpenProjects.Count; Idx++)
+					{
+						Settings.OpenProjects[Idx] = UpgradeSelectedProjectSettings(Settings.OpenProjects[Idx]);
+					}
+					for (int Idx = 0; Idx < Settings.RecentProjects.Count; Idx++)
+					{
+						Settings.RecentProjects[Idx] = UpgradeSelectedProjectSettings(Settings.RecentProjects[Idx]);
+					}
+				}
+
+				// Save the new settings
+				Settings.Version = UserSettingsVersion.Latest;
+				Settings.Save();
+			}
+
 			// Register the update listener
 			UpdateMonitor.OnUpdateAvailable += OnUpdateAvailableCallback;
 
@@ -162,7 +181,7 @@ namespace UnrealGameSync
 			}
 
 			// Detect settings for the project we want to open
-			DetectStartupProjectSettingsTask = new DetectMultipleProjectSettingsTask(Tasks);
+			DetectStartupProjectSettingsTask = new DetectMultipleProjectSettingsTask(DefaultConnection, Tasks);
 
 			DetectStartupProjectSettingsWindow = new ModalTaskWindow(DetectStartupProjectSettingsTask, "Opening Projects", "Opening projects, please wait...", FormStartPosition.CenterScreen);
 			if(bRestoreState)
@@ -180,6 +199,18 @@ namespace UnrealGameSync
 			DetectStartupProjectSettingsWindow.Complete += OnDetectStartupProjectsComplete;
 		}
 
+		private UserSelectedProjectSettings UpgradeSelectedProjectSettings(UserSelectedProjectSettings Project)
+		{
+			if (Project.ServerAndPort == null || String.Compare(Project.ServerAndPort, DefaultConnection.ServerAndPort, StringComparison.OrdinalIgnoreCase) == 0)
+			{
+				if (Project.UserName == null || String.Compare(Project.UserName, DefaultConnection.UserName, StringComparison.OrdinalIgnoreCase) == 0)
+				{
+					Project = new UserSelectedProjectSettings(null, null, Project.Type, Project.ClientPath, Project.LocalPath);
+				}
+			}
+			return Project;
+		}
+
 		private void OnDetectStartupProjectsComplete()
 		{
 			// Close the startup window
@@ -203,7 +234,7 @@ namespace UnrealGameSync
 			DetectProjectSettingsResult[] StartupProjects = DetectStartupProjectSettingsTask.Results.Where(x => x != null).ToArray();
 
 			// Create the main window
-			MainWindowInstance = new MainWindow(UpdateMonitor, ApiUrl, DataFolder, CacheFolder, bRestoreState, UpdateSpawn ?? Assembly.GetExecutingAssembly().Location, bUnstable, StartupProjects, Log, Settings);
+			MainWindowInstance = new MainWindow(UpdateMonitor, ApiUrl, DataFolder, CacheFolder, bRestoreState, UpdateSpawn ?? Assembly.GetExecutingAssembly().Location, bUnstable, StartupProjects, DefaultConnection, Log, Settings);
 			if(bVisible)
 			{
 				MainWindowInstance.Show();
diff --git a/Engine/Source/Programs/UnrealGameSync/UnrealGameSync/UserSettings.cs b/Engine/Source/Programs/UnrealGameSync/UnrealGameSync/UserSettings.cs
--- a/Engine/Source/Programs/UnrealGameSync/UnrealGameSync/UserSettings.cs
+++ b/Engine/Source/Programs/UnrealGameSync/UnrealGameSync/UserSettings.cs
@@ -47,6 +47,13 @@ namespace UnrealGameSync
 		Content
 	}
 
+	enum UserSettingsVersion
+	{
+		Initial = 0,
+		DefaultServerSettings = 1,
+		Latest = DefaultServerSettings
+	}
+
 	class UserSelectedProjectSettings
 	{
 		public readonly string ServerAndPort;
@@ -190,6 +197,7 @@ namespace UnrealGameSync
 		ConfigFile ConfigFile = new ConfigFile();
 
 		// General settings
+		public UserSettingsVersion Version = UserSettingsVersion.Latest;
 		public bool bBuildAfterSync;
 		public bool bRunAfterSync;
 		public bool bSyncPrecompiledEditor;
@@ -284,6 +292,7 @@ namespace UnrealGameSync
 			}
 
 			// General settings
+			Version = (UserSettingsVersion)ConfigFile.GetValue("General.Version", (int)UserSettingsVersion.Initial);
 			bBuildAfterSync = (ConfigFile.GetValue("General.BuildAfterSync", "1") != "0");
 			bRunAfterSync = (ConfigFile.GetValue("General.RunAfterSync", "1") != "0");
 			bSyncPrecompiledEditor = (ConfigFile.GetValue("General.SyncPrecompiledEditor", "0") != "0");
@@ -574,6 +583,7 @@ namespace UnrealGameSync
 			// General settings
 			ConfigSection GeneralSection = ConfigFile.FindOrAddSection("General");
 			GeneralSection.Clear();
+			GeneralSection.SetValue("Version", (int)Version);
 			GeneralSection.SetValue("BuildAfterSync", bBuildAfterSync);
 			GeneralSection.SetValue("RunAfterSync", bRunAfterSync);
 			GeneralSection.SetValue("SyncPrecompiledEditor", bSyncPrecompiledEditor);
diff --git a/Engine/Source/Programs/UnrealGameSync/UnrealGameSync/Utility.cs b/Engine/Source/Programs/UnrealGameSync/UnrealGameSync/Utility.cs
--- a/Engine/Source/Programs/UnrealGameSync/UnrealGameSync/Utility.cs
+++ b/Engine/Source/Programs/UnrealGameSync/UnrealGameSync/Utility.cs
@@ -306,6 +306,21 @@ namespace UnrealGameSync
 			}
 		}
 
+		private static void DeleteKey(RegistryKey Key, string Name)
+		{
+			string[] ValueNames = Key.GetValueNames();
+			if (ValueNames.Any(x => String.Compare(x, Name, StringComparison.OrdinalIgnoreCase) == 0))
+			{
+				try
+				{ 
+					Key.DeleteValue(Name);
+				}
+				catch
+				{ 
+				}
+			}
+		}
+
 		public static void SaveGlobalPerforceSettings(string ServerAndPort, string UserName, string DepotPath)
 		{
 			try
@@ -313,7 +328,7 @@ namespace UnrealGameSync
 				using (RegistryKey Key = Registry.CurrentUser.CreateSubKey("SOFTWARE\\Epic Games\\UnrealGameSync"))
 				{
 					// Delete this legacy setting
-					try { Key.DeleteValue("Server"); } catch(Exception) { }
+					DeleteKey(Key, "Server");
 
 					if(String.IsNullOrEmpty(ServerAndPort))
 					{
@@ -335,7 +350,7 @@ namespace UnrealGameSync
 
 					if(String.IsNullOrEmpty(DepotPath) || (DeploymentSettings.DefaultDepotPath != null && String.Equals(DepotPath, DeploymentSettings.DefaultDepotPath, StringComparison.InvariantCultureIgnoreCase)))
 					{
-						try { Key.DeleteValue("DepotPath"); } catch(Exception) { }
+						DeleteKey(Key, "DepotPath");
 					}
 					else
 					{
@@ -431,6 +446,13 @@ namespace UnrealGameSync
 			return Color.FromArgb((int)(First.R + (Second.R - First.R) * T), (int)(First.G + (Second.G - First.G) * T), (int)(First.B + (Second.B - First.B) * T));
 		}
 
+		public static PerforceConnection OverridePerforceSettings(PerforceConnection DefaultConnection, string ServerAndPort, string UserName)
+		{
+			string ResolvedServerAndPort = String.IsNullOrWhiteSpace(ServerAndPort) ? DefaultConnection.ServerAndPort : ServerAndPort;
+			string ResolvedUserName = String.IsNullOrWhiteSpace(UserName) ? DefaultConnection.UserName : UserName;
+			return new PerforceConnection(ResolvedUserName, null, ResolvedServerAndPort);
+		}
+
 		public static string FormatRecentDateTime(DateTime Date)
 		{
 			DateTime Now = DateTime.Now;
