UGS: Allow enabling and disabling automation via the Application Settings dialog.

#rnx

diff --git a/Engine/Source/Programs/UnrealGameSync/UnrealGameSync/AutomationServer.cs b/Engine/Source/Programs/UnrealGameSync/UnrealGameSync/AutomationServer.cs
--- a/Engine/Source/Programs/UnrealGameSync/UnrealGameSync/AutomationServer.cs
+++ b/Engine/Source/Programs/UnrealGameSync/UnrealGameSync/AutomationServer.cs
@@ -131,6 +131,8 @@ namespace UnrealGameSync
 
 	class AutomationServer : IDisposable
 	{
+		public const int DefaultPortNumber = 30422;
+
 		TcpListener Listener;
 		TcpClient CurrentClient;
 		Thread BackgroundThread;
@@ -143,13 +145,11 @@ namespace UnrealGameSync
 			this.PostRequest = PostRequest;
 			this.Log = Log;
 
-			object PortValue = Registry.GetValue("HKEY_CURRENT_USER\\Software\\Epic Games\\UnrealGameSync", "AutomationPort", null);
-			if(PortValue != null)
+			int PortNumber = GetPortNumber();
+			if(PortNumber > 0)
 			{
 				try
 				{
-					int PortNumber = (int)PortValue;
-
 					Listener = new TcpListener(IPAddress.Loopback, PortNumber);
 					Listener.Start();
 
@@ -163,6 +163,31 @@ namespace UnrealGameSync
 			}
 		}
 
+		public static void SetPortNumber(int PortNumber)
+		{
+			if(PortNumber <= 0)
+			{
+				Utility.DeleteRegistryKey(Registry.CurrentUser, "Software\\Epic Games\\UnrealGameSync", "AutomationPort");
+			}
+			else
+			{
+				Registry.SetValue("HKEY_CURRENT_USER\\Software\\Epic Games\\UnrealGameSync", "AutomationPort", PortNumber);
+			}
+		}
+
+		public static int GetPortNumber()
+		{
+			object PortValue = Registry.GetValue("HKEY_CURRENT_USER\\Software\\Epic Games\\UnrealGameSync", "AutomationPort", null);
+			if (PortValue != null && PortValue is int)
+			{
+				return (int)PortValue;
+			}
+			else
+			{
+				return -1;
+			}
+		}
+
 		public void Run()
 		{
 			try
diff --git a/Engine/Source/Programs/UnrealGameSync/UnrealGameSync/Forms/ApplicationSettingsWindow.cs b/Engine/Source/Programs/UnrealGameSync/UnrealGameSync/Forms/ApplicationSettingsWindow.cs
--- a/Engine/Source/Programs/UnrealGameSync/UnrealGameSync/Forms/ApplicationSettingsWindow.cs
+++ b/Engine/Source/Programs/UnrealGameSync/UnrealGameSync/Forms/ApplicationSettingsWindow.cs
@@ -52,6 +52,7 @@ namespace UnrealGameSync
 		string InitialUserName;
 		string InitialDepotPath;
 		bool bInitialUnstable;
+		int InitialAutomationPortNumber;
 
 		bool? bRestartUnstable;
 
@@ -66,6 +67,8 @@ namespace UnrealGameSync
 			Utility.ReadGlobalPerforceSettings(ref InitialServerAndPort, ref InitialUserName, ref InitialDepotPath);
 			bInitialUnstable = bUnstable;
 
+			InitialAutomationPortNumber = AutomationServer.GetPortNumber();
+
 			this.AutomaticallyRunAtStartupCheckBox.Checked = IsAutomaticallyRunAtStartup();
 			this.KeepInTrayCheckBox.Checked = Settings.bKeepInTray;
 						
@@ -82,6 +85,19 @@ namespace UnrealGameSync
 			this.DepotPathTextBox.CueBanner = DeploymentSettings.DefaultDepotPath;
 
 			this.UseUnstableBuildCheckBox.Checked = bUnstable;
+
+			if(InitialAutomationPortNumber > 0)
+			{
+				this.EnableAutomationCheckBox.Checked = true;
+				this.AutomationPortTextBox.Enabled = true;
+				this.AutomationPortTextBox.Text = InitialAutomationPortNumber.ToString();
+			}
+			else
+			{
+				this.EnableAutomationCheckBox.Checked = false;
+				this.AutomationPortTextBox.Enabled = false;
+				this.AutomationPortTextBox.Text = AutomationServer.DefaultPortNumber.ToString();
+			}
 		}
 
 		public static bool? ShowModal(IWin32Window Owner, PerforceConnection DefaultConnection, bool bUnstable, string OriginalExecutableFileName, UserSettings Settings, TextWriter Log)
@@ -126,7 +142,14 @@ namespace UnrealGameSync
 
 			bool bUnstable = UseUnstableBuildCheckBox.Checked;
 
-			if(ServerAndPort != InitialServerAndPort || UserName != InitialUserName || DepotPath != InitialDepotPath || bUnstable != bInitialUnstable)
+
+			int AutomationPortNumber;
+			if(!EnableAutomationCheckBox.Checked || !int.TryParse(AutomationPortTextBox.Text, out AutomationPortNumber))
+			{
+				AutomationPortNumber = -1;
+			}
+			
+			if(ServerAndPort != InitialServerAndPort || UserName != InitialUserName || DepotPath != InitialDepotPath || bUnstable != bInitialUnstable || AutomationPortNumber != InitialAutomationPortNumber)
 			{
 				// Try to log in to the new server, and check the application is there
 				if(ServerAndPort != InitialServerAndPort || UserName != InitialUserName || DepotPath != InitialDepotPath)
@@ -150,6 +173,7 @@ namespace UnrealGameSync
 
 				bRestartUnstable = UseUnstableBuildCheckBox.Checked;
 				Utility.SaveGlobalPerforceSettings(ServerAndPort, UserName, DepotPath);
+				AutomationServer.SetPortNumber(AutomationPortNumber);
 			}
 
 			RegistryKey Key = Registry.CurrentUser.CreateSubKey("SOFTWARE\\Microsoft\\Windows\\CurrentVersion\\Run");
@@ -177,5 +201,10 @@ namespace UnrealGameSync
 			DialogResult = DialogResult.Cancel;
 			Close();
 		}
+
+		private void EnableAutomationCheckBox_CheckedChanged(object sender, EventArgs e)
+		{
+			AutomationPortTextBox.Enabled = EnableAutomationCheckBox.Checked;
+		}
 	}
 }
diff --git a/Engine/Source/Programs/UnrealGameSync/UnrealGameSync/Utility.cs b/Engine/Source/Programs/UnrealGameSync/UnrealGameSync/Utility.cs
--- a/Engine/Source/Programs/UnrealGameSync/UnrealGameSync/Utility.cs
+++ b/Engine/Source/Programs/UnrealGameSync/UnrealGameSync/Utility.cs
@@ -306,7 +306,18 @@ namespace UnrealGameSync
 			}
 		}
 
-		private static void DeleteKey(RegistryKey Key, string Name)
+		public static void DeleteRegistryKey(RegistryKey RootKey, string KeyName, string ValueName)
+		{
+			using (RegistryKey Key = RootKey.OpenSubKey(KeyName, true))
+			{
+				if (Key != null)
+				{
+					DeleteRegistryKey(Key, ValueName);
+				}
+			}
+		}
+
+		public static void DeleteRegistryKey(RegistryKey Key, string Name)
 		{
 			string[] ValueNames = Key.GetValueNames();
 			if (ValueNames.Any(x => String.Compare(x, Name, StringComparison.OrdinalIgnoreCase) == 0))
@@ -328,7 +339,7 @@ namespace UnrealGameSync
 				using (RegistryKey Key = Registry.CurrentUser.CreateSubKey("SOFTWARE\\Epic Games\\UnrealGameSync"))
 				{
 					// Delete this legacy setting
-					DeleteKey(Key, "Server");
+					DeleteRegistryKey(Key, "Server");
 
 					if(String.IsNullOrEmpty(ServerAndPort))
 					{
@@ -350,7 +361,7 @@ namespace UnrealGameSync
 
 					if(String.IsNullOrEmpty(DepotPath) || (DeploymentSettings.DefaultDepotPath != null && String.Equals(DepotPath, DeploymentSettings.DefaultDepotPath, StringComparison.InvariantCultureIgnoreCase)))
 					{
-						DeleteKey(Key, "DepotPath");
+						DeleteRegistryKey(Key, "DepotPath");
 					}
 					else
 					{
