UGS: Fix typo.

#rnx

diff --git a/Engine/Source/Programs/UnrealGameSync/UnrealGameSync/Forms/IssueDetailsWindow.cs b/Engine/Source/Programs/UnrealGameSync/UnrealGameSync/Forms/IssueDetailsWindow.cs
--- a/Engine/Source/Programs/UnrealGameSync/UnrealGameSync/Forms/IssueDetailsWindow.cs
+++ b/Engine/Source/Programs/UnrealGameSync/UnrealGameSync/Forms/IssueDetailsWindow.cs
@@ -1038,7 +1038,7 @@ namespace UnrealGameSync
 				List<IssueDiagnosticData> Diagnostics = new List<IssueDiagnosticData>();
 
 				UpdateIssueDetailsTask Task = new UpdateIssueDetailsTask(IssueMonitor.ApiUrl, Issue.Id, Diagnostics);
-				if(!ModalTask.ExecuteAndShowError(Owner, Task, "Fetching data", "Fetchind data, please wait..."))
+				if(!ModalTask.ExecuteAndShowError(Owner, Task, "Fetching data", "Fetching data, please wait..."))
 				{
 					return;
 				}
