UGS: Fix incorrect extraction of symbol names from undefined symbol errors starting with __declspec(dllexport).

#rnx

diff --git a/Engine/Source/Programs/UnrealGameSync/MetadataTool/BuildHealth/Matchers/UndefinedSymbolPatternMatcher.cs b/Engine/Source/Programs/UnrealGameSync/MetadataTool/BuildHealth/Matchers/UndefinedSymbolPatternMatcher.cs
--- a/Engine/Source/Programs/UnrealGameSync/MetadataTool/BuildHealth/Matchers/UndefinedSymbolPatternMatcher.cs
+++ b/Engine/Source/Programs/UnrealGameSync/MetadataTool/BuildHealth/Matchers/UndefinedSymbolPatternMatcher.cs
@@ -70,6 +70,9 @@ namespace MetadataTool
 			{
 				string SymbolName = SymbolMatch;
 
+				// Remove any __declspec qualifiers
+				SymbolName = Regex.Replace(SymbolName, "(?<![^a-zA-Z_])__declspec\\([^\\)]+\\)", "");
+
 				// Remove any argument lists for functions (anything after the first paren)
 				SymbolName = Regex.Replace(SymbolName, "\\(.*$", "");
 
