UGS: Parse fatal errors from Clang as filenames to use for blaming people.

#rnx

diff --git a/Engine/Source/Programs/UnrealGameSync/MetadataTool/BuildHealth/Matchers/CompilePatternMatcher.cs b/Engine/Source/Programs/UnrealGameSync/MetadataTool/BuildHealth/Matchers/CompilePatternMatcher.cs
--- a/Engine/Source/Programs/UnrealGameSync/MetadataTool/BuildHealth/Matchers/CompilePatternMatcher.cs
+++ b/Engine/Source/Programs/UnrealGameSync/MetadataTool/BuildHealth/Matchers/CompilePatternMatcher.cs
@@ -30,7 +30,7 @@ namespace MetadataTool
 		{
 			// Find a list of source files with errors
 			HashSet<string> ErrorFileNames = new HashSet<string>(StringComparer.OrdinalIgnoreCase);
-			foreach (Match FileMatch in Regex.Matches(Diagnostic.Message, @"^\s*((?:[A-Za-z]:)?[^\s(:]+)[\(:]\d[\s\d:\)]+(?:warning|error)", RegexOptions.Multiline))
+			foreach (Match FileMatch in Regex.Matches(Diagnostic.Message, @"^\s*((?:[A-Za-z]:)?[^\s(:]+)[\(:]\d[\s\d:\)]+(?:warning|error|fatal error)", RegexOptions.Multiline))
 			{
 				if (FileMatch.Success)
 				{
