UGS: Attempt to identify a list of causers for undefined symbol errors based on the name of the class, and an assumption that classes are defined in headers and cpp files with the same name.

#rnx

diff --git a/Engine/Source/Programs/UnrealGameSync/MetadataTool/BuildHealth/Matchers/UndefinedSymbolPatternMatcher.cs b/Engine/Source/Programs/UnrealGameSync/MetadataTool/BuildHealth/Matchers/UndefinedSymbolPatternMatcher.cs
--- a/Engine/Source/Programs/UnrealGameSync/MetadataTool/BuildHealth/Matchers/UndefinedSymbolPatternMatcher.cs
+++ b/Engine/Source/Programs/UnrealGameSync/MetadataTool/BuildHealth/Matchers/UndefinedSymbolPatternMatcher.cs
@@ -6,6 +6,7 @@ using System.Linq;
 using System.Text;
 using System.Text.RegularExpressions;
 using System.Threading.Tasks;
+using Tools.DotNETCommon.Perforce;
 
 namespace MetadataTool
 {
@@ -96,6 +97,40 @@ namespace MetadataTool
 			return false;
 		}
 
+		public override List<ChangeInfo> FindCausers(PerforceConnection Perforce, BuildHealthIssue Issue, IReadOnlyList<ChangeInfo> Changes)
+		{
+			SortedSet<string> FileNamesWithoutPath = new SortedSet<string>(StringComparer.Ordinal);
+			foreach(string Identifier in Issue.Identifiers)
+			{
+				Match Match = Regex.Match(Identifier, "^[AUFS]([A-Z][A-Za-z_]*)::");
+				if(Match.Success)
+				{
+					FileNamesWithoutPath.Add(Match.Groups[1].Value + ".h");
+					FileNamesWithoutPath.Add(Match.Groups[1].Value + ".cpp");
+				}
+			}
+
+			if(FileNamesWithoutPath.Count > 0)
+			{
+				List<ChangeInfo> Causers = new List<ChangeInfo>();
+				foreach (ChangeInfo Change in Changes)
+				{
+					DescribeRecord DescribeRecord = GetDescribeRecord(Perforce, Change);
+					if (ContainsFileNames(DescribeRecord, FileNamesWithoutPath))
+					{
+						Causers.Add(Change);
+					}
+				}
+
+				if (Causers.Count > 0)
+				{
+					return Causers;
+				}
+			}
+
+			return base.FindCausers(Perforce, Issue, Changes);
+		}
+
 		public override string GetSummary(BuildHealthIssue Issue)
 		{
 			if(Issue.Identifiers.Count == 1)
