update metadata server to use mysql backend instead of sqlite

diff --git a/Engine/Source/Programs/UnrealGameSync/MetadataServer/Connectors/SqlConnector.cs b/Engine/Source/Programs/UnrealGameSync/MetadataServer/Connectors/SqlConnector.cs
--- a/Engine/Source/Programs/UnrealGameSync/MetadataServer/Connectors/SqlConnector.cs
+++ b/Engine/Source/Programs/UnrealGameSync/MetadataServer/Connectors/SqlConnector.cs
@@ -5,12 +5,12 @@ using Newtonsoft.Json.Linq;
 using System;
 using System.Collections.Generic;
 using System.Data.SqlClient;
-using System.Data.SQLite;
 using System.IO;
 using System.Linq;
 using System.Text;
 using System.Text.RegularExpressions;
 using MetadataServer.Models;
+using MySql.Data.MySqlClient;
 
 namespace MetadataServer.Connectors
 {
@@ -25,30 +25,32 @@ namespace MetadataServer.Connectors
 			long LastCommentId = 0;
 			long LastBuildId = 0;
 			string ProjectLikeString = "%" + (Project == null ? String.Empty : GetProjectStream(Project)) + "%";
-			using (SQLiteConnection Connection = new SQLiteConnection(ConnectionString))
+			using (MySqlConnection Connection = new MySqlConnection(ConnectionString))
 			{
 				Connection.Open();
-				using (SQLiteCommand Command = new SQLiteCommand("SELECT UserVotes.Id FROM [UserVotes] " +
-																 "INNER JOIN [Projects] ON Projects.Id = UserVotes.ProjectId " +
-																 "WHERE Projects.Name LIKE @param1 GROUP BY Changelist ORDER BY Changelist DESC LIMIT 1 OFFSET 432", Connection))
+				using (MySqlCommand Command = new MySqlCommand("WITH user_votes AS (SELECT UserVotes.Id, UserVotes.Changelist FROM ugs_db.UserVotes " +
+																 "INNER JOIN ugs_db.Projects ON Projects.Id = UserVotes.ProjectId " +
+																 "WHERE Projects.Name LIKE @param1 GROUP BY Changelist ORDER BY Changelist DESC LIMIT 432) " +
+																 "SELECT * FROM user_votes ORDER BY user_votes.Changelist ASC LIMIT 1", Connection))
 				{
 					Command.Parameters.AddWithValue("@param1", ProjectLikeString);
-					using (SQLiteDataReader Reader = Command.ExecuteReader())
+					using (MySqlDataReader Reader = Command.ExecuteReader())
 					{
 						while (Reader.Read())
 						{
 							LastEventId = Reader.GetInt64(0);
 							break;
 						}
-					}
+					}	
 				}
 				
-				using (SQLiteCommand Command = new SQLiteCommand("SELECT Comments.Id FROM [Comments] " +
-																 "INNER JOIN [Projects] ON Projects.Id = Comments.ProjectId " +
-																 "WHERE Projects.Name LIKE @param1 GROUP BY ChangeNumber ORDER BY ChangeNumber DESC LIMIT 1 OFFSET 432", Connection))
+				using (MySqlCommand Command = new MySqlCommand("WITH comments AS (SELECT Comments.Id, Comments.ChangeNumber FROM ugs_db.Comments " +
+																 "INNER JOIN ugs_db.Projects ON Projects.Id = Comments.ProjectId " +
+																 "WHERE Projects.Name LIKE @param1 GROUP BY ChangeNumber ORDER BY ChangeNumber DESC LIMIT 432) " +
+																 "SELECT * FROM comments ORDER BY comments.ChangeNumber ASC LIMIT 1", Connection))
 				{
 					Command.Parameters.AddWithValue("@param1", ProjectLikeString);
-					using (SQLiteDataReader Reader = Command.ExecuteReader())
+					using (MySqlDataReader Reader = Command.ExecuteReader())
 					{
 						while (Reader.Read())
 						{
@@ -57,13 +59,14 @@ namespace MetadataServer.Connectors
 						}
 					}
 				}
-				using (SQLiteCommand Command = new SQLiteCommand("SELECT Badges.Id FROM [Badges] " +
-																 "INNER JOIN [Projects] ON Projects.Id = Badges.ProjectId " +
-																 "WHERE Projects.Name LIKE @param1 GROUP BY ChangeNumber ORDER BY ChangeNumber DESC LIMIT 1 OFFSET 432", Connection))
+				using (MySqlCommand Command = new MySqlCommand("WITH badges AS (SELECT Badges.Id, Badges.ChangeNumber FROM ugs_db.Badges " +
+																 "INNER JOIN ugs_db.Projects ON Projects.Id = Badges.ProjectId " +
+																 "WHERE Projects.Name LIKE @param1 GROUP BY ChangeNumber ORDER BY ChangeNumber DESC LIMIT 432) " + 
+																 "SELECT * FROM badges ORDER BY badges.ChangeNumber ASC LIMIT 1", Connection))
 				{
 					//Command.Parameters.AddWithValue("@param1", ProjectId);
 					Command.Parameters.AddWithValue("@param1", ProjectLikeString);
-					using (SQLiteDataReader Reader = Command.ExecuteReader())
+					using (MySqlDataReader Reader = Command.ExecuteReader())
 					{
 						while (Reader.Read())
 						{
@@ -79,15 +82,15 @@ namespace MetadataServer.Connectors
 		{
 			List<EventData> ReturnedEvents = new List<EventData>();
 			string ProjectLikeString = "%" + (Project == null ? String.Empty : GetProjectStream(Project)) + "%";
-			using (SQLiteConnection Connection = new SQLiteConnection(ConnectionString))
+			using (MySqlConnection Connection = new MySqlConnection(ConnectionString))
 			{
 				Connection.Open();
-				using (SQLiteCommand Command = new SQLiteCommand("SELECT UserVotes.Id, UserVotes.Changelist, UserVotes.UserName, UserVotes.Verdict, UserVotes.Project FROM [UserVotes] " +
-																 "INNER JOIN [Projects] ON Projects.Id = UserVotes.ProjectId WHERE UserVotes.Id > @param1 AND Projects.Name LIKE @param2 ORDER BY UserVotes.Id", Connection))
+				using (MySqlCommand Command = new MySqlCommand("SELECT UserVotes.Id, UserVotes.Changelist, UserVotes.UserName, UserVotes.Verdict, UserVotes.Project FROM ugs_db.UserVotes " +
+															   "INNER JOIN ugs_db.Projects ON Projects.Id = UserVotes.ProjectId WHERE UserVotes.Id > @param1 AND Projects.Name LIKE @param2 ORDER BY UserVotes.Id", Connection))
 				{
 					Command.Parameters.AddWithValue("@param1", LastEventId);
 					Command.Parameters.AddWithValue("@param2", ProjectLikeString);
-					using (SQLiteDataReader Reader = Command.ExecuteReader())
+					using (MySqlDataReader Reader = Command.ExecuteReader())
 					{
 						while (Reader.Read())
 						{
@@ -113,15 +116,15 @@ namespace MetadataServer.Connectors
 		{
 			List<CommentData> ReturnedComments = new List<CommentData>();
 			string ProjectLikeString = "%" + (Project == null ? String.Empty : GetProjectStream(Project)) + "%";
-			using (SQLiteConnection Connection = new SQLiteConnection(ConnectionString))
+			using (MySqlConnection Connection = new MySqlConnection(ConnectionString))
 			{
 				Connection.Open();
-				using (SQLiteCommand Command = new SQLiteCommand("SELECT Comments.Id, Comments.ChangeNumber, Comments.UserName, Comments.Text, Comments.Project FROM [Comments] " +
-															     "INNER JOIN [Projects] ON Projects.Id = Comments.ProjectId WHERE Comments.Id > @param1 AND Projects.Name LIKE @param2 ORDER BY Comments.Id", Connection))
+				using (MySqlCommand Command = new MySqlCommand("SELECT Comments.Id, Comments.ChangeNumber, Comments.UserName, Comments.Text, Comments.Project FROM ugs_db.Comments " +
+															     "INNER JOIN ugs_db.Projects ON Projects.Id = Comments.ProjectId WHERE Comments.Id > @param1 AND Projects.Name LIKE @param2 ORDER BY Comments.Id", Connection))
 				{
 					Command.Parameters.AddWithValue("@param1", LastCommentId);
 					Command.Parameters.AddWithValue("@param2", ProjectLikeString);
-					using (SQLiteDataReader Reader = Command.ExecuteReader())
+					using (MySqlDataReader Reader = Command.ExecuteReader())
 					{
 						while (Reader.Read())
 						{
@@ -146,15 +149,15 @@ namespace MetadataServer.Connectors
 		{
 			List<BuildData> ReturnedBuilds = new List<BuildData>();
 			string ProjectLikeString = "%" + (Project == null ? String.Empty : GetProjectStream(Project)) + "%";
-			using (SQLiteConnection Connection = new SQLiteConnection(ConnectionString))
+			using (MySqlConnection Connection = new MySqlConnection(ConnectionString))
 			{
 				Connection.Open();
-				using (SQLiteCommand Command = new SQLiteCommand("SELECT Badges.Id, Badges.ChangeNumber, Badges.BuildType, Badges.Result, Badges.Url, Projects.Name, Badges.ArchivePath FROM [Badges] " +
-																 "INNER JOIN [Projects] ON Projects.Id = Badges.ProjectId WHERE Badges.Id > @param1 AND Projects.Name LIKE @param2 ORDER BY Badges.Id", Connection))
+				using (MySqlCommand Command = new MySqlCommand("SELECT Badges.Id, Badges.ChangeNumber, Badges.BuildType, Badges.Result, Badges.Url, Projects.Name, Badges.ArchivePath FROM ugs_db.Badges " +
+																 "INNER JOIN ugs_db.Projects ON Projects.Id = Badges.ProjectId WHERE Badges.Id > @param1 AND Projects.Name LIKE @param2 ORDER BY Badges.Id", Connection))
 				{ 
 					Command.Parameters.AddWithValue("@param1", LastBuildId);
 					Command.Parameters.AddWithValue("@param2", ProjectLikeString);
-					using (SQLiteDataReader Reader = Command.ExecuteReader())
+					using (MySqlDataReader Reader = Command.ExecuteReader())
 					{
 						while (Reader.Read())
 						{
@@ -182,13 +185,13 @@ namespace MetadataServer.Connectors
 		public static List<TelemetryErrorData> GetErrorData(int Records)
 		{
 			List<TelemetryErrorData> ReturnedErrors = new List<TelemetryErrorData>();
-			using (SQLiteConnection Connection = new SQLiteConnection(ConnectionString))
+			using (MySqlConnection Connection = new MySqlConnection(ConnectionString))
 			{
 				Connection.Open();
-				using (SQLiteCommand Command = new SQLiteCommand("SELECT Id, Type, Text, UserName, Project, Timestamp, Version, IpAddress FROM [Errors] ORDER BY Id DESC LIMIT @param1", Connection))
+				using (MySqlCommand Command = new MySqlCommand("SELECT Id, Type, Text, UserName, Project, Timestamp, Version, IpAddress FROM ugs_db.Errors ORDER BY Id DESC LIMIT @param1", Connection))
 				{
 					Command.Parameters.AddWithValue("@param1", Records);
-					using (SQLiteDataReader Reader = Command.ExecuteReader())
+					using (MySqlDataReader Reader = Command.ExecuteReader())
 					{
 						while (Reader.Read())
 						{
@@ -209,22 +212,22 @@ namespace MetadataServer.Connectors
 			return ReturnedErrors;
 		}
 
-		private static long TryInsertAndGetProject(SQLiteConnection Connection, string Project)
+		private static long TryInsertAndGetProject(MySqlConnection Connection, string Project)
 		{
-			using (SQLiteCommand Command = new SQLiteCommand("INSERT OR IGNORE INTO [Projects] (Name) VALUES (@Project); SELECT [Id] FROM [Projects] WHERE Name = @Project", Connection))
+			using (MySqlCommand Command = new MySqlCommand("INSERT IGNORE INTO ugs_db.Projects (Name) VALUES (@Project); SELECT Id FROM ugs_db.Projects WHERE Name = @Project", Connection))
 			{
 				Command.Parameters.AddWithValue("@Project", Project);
 				object ProjectId = Command.ExecuteScalar();
-				return (long)ProjectId;
+				return Convert.ToInt64(ProjectId);
 			}
 		}
 		public static void PostBuild(BuildData Build)
 		{
-			using (SQLiteConnection Connection = new SQLiteConnection(ConnectionString))
+			using (MySqlConnection Connection = new MySqlConnection(ConnectionString))
 			{
 				Connection.Open();
 				long ProjectId = TryInsertAndGetProject(Connection, Build.Project);
-				using (SQLiteCommand Command = new SQLiteCommand("INSERT INTO [Badges] (ChangeNumber, BuildType, Result, URL, ArchivePath, ProjectId) VALUES (@ChangeNumber, @BuildType, @Result, @URL, @ArchivePath, @ProjectId)", Connection))
+				using (MySqlCommand Command = new MySqlCommand("INSERT INTO ugs_db.Badges (ChangeNumber, BuildType, Result, URL, ArchivePath, ProjectId) VALUES (@ChangeNumber, @BuildType, @Result, @URL, @ArchivePath, @ProjectId)", Connection))
 				{
 					Command.Parameters.AddWithValue("@ChangeNumber", Build.ChangeNumber);
 					Command.Parameters.AddWithValue("@BuildType", Build.BuildType);
@@ -239,11 +242,11 @@ namespace MetadataServer.Connectors
 
 		public static void PostEvent(EventData Event)
 		{
-			using (SQLiteConnection Connection = new SQLiteConnection(ConnectionString))
+			using (MySqlConnection Connection = new MySqlConnection(ConnectionString))
 			{
 				Connection.Open();
 				long ProjectId = TryInsertAndGetProject(Connection, Event.Project);
-				using (SQLiteCommand Command = new SQLiteCommand("INSERT INTO [UserVotes] (Changelist, UserName, Verdict, Project, ProjectId) VALUES (@Changelist, @UserName, @Verdict, @Project, @ProjectId)", Connection))
+				using (MySqlCommand Command = new MySqlCommand("INSERT INTO ugs_db.UserVotes (Changelist, UserName, Verdict, Project, ProjectId) VALUES (@Changelist, @UserName, @Verdict, @Project, @ProjectId)", Connection))
 				{
 					Command.Parameters.AddWithValue("@Changelist", Event.Change);
 					Command.Parameters.AddWithValue("@UserName", Event.UserName.ToString());
@@ -257,11 +260,11 @@ namespace MetadataServer.Connectors
 
 		public static void PostComment(CommentData Comment)
 		{
-			using (SQLiteConnection Connection = new SQLiteConnection(ConnectionString))
+			using (MySqlConnection Connection = new MySqlConnection(ConnectionString))
 			{
 				Connection.Open();
 				long ProjectId = TryInsertAndGetProject(Connection, Comment.Project);
-				using (SQLiteCommand Command = new SQLiteCommand("INSERT INTO [Comments] (ChangeNumber, UserName, Text, Project, ProjectId) VALUES (@ChangeNumber, @UserName, @Text, @Project, @ProjectId)", Connection))
+				using (MySqlCommand Command = new MySqlCommand("INSERT INTO ugs_db.Comments (ChangeNumber, UserName, Text, Project, ProjectId) VALUES (@ChangeNumber, @UserName, @Text, @Project, @ProjectId)", Connection))
 				{
 					Command.Parameters.AddWithValue("@ChangeNumber", Comment.ChangeNumber);
 					Command.Parameters.AddWithValue("@UserName", Comment.UserName);
@@ -275,11 +278,11 @@ namespace MetadataServer.Connectors
 
 		public static void PostTelemetryData(TelemetryTimingData Data, string Version, string IpAddress)
 		{
-			using (SQLiteConnection Connection = new SQLiteConnection(ConnectionString))
+			using (MySqlConnection Connection = new MySqlConnection(ConnectionString))
 			{
 				Connection.Open();
 				long ProjectId = TryInsertAndGetProject(Connection, Data.Project);
-				using (SQLiteCommand Command = new SQLiteCommand("INSERT INTO [Telemetry.v2] (Action, Result, UserName, Project, Timestamp, Duration, Version, IpAddress, ProjectId) VALUES (@Action, @Result, @UserName, @Project, @Timestamp, @Duration, @Version, @IpAddress, @ProjectId)", Connection))
+				using (MySqlCommand Command = new MySqlCommand("INSERT INTO ugs_db.Telemetry_v2 (Action, Result, UserName, Project, Timestamp, Duration, Version, IpAddress, ProjectId) VALUES (@Action, @Result, @UserName, @Project, @Timestamp, @Duration, @Version, @IpAddress, @ProjectId)", Connection))
 				{
 					Command.Parameters.AddWithValue("@Action", Data.Action);
 					Command.Parameters.AddWithValue("@Result", Data.Result);
@@ -297,7 +300,7 @@ namespace MetadataServer.Connectors
 
 		public static void PostErrorData(TelemetryErrorData Data, string Version, string IpAddress)
 		{
-			using (SQLiteConnection Connection = new SQLiteConnection(ConnectionString))
+			using (MySqlConnection Connection = new MySqlConnection(ConnectionString))
 			{
 				Connection.Open();
 				long? ProjectId = null;
@@ -305,7 +308,7 @@ namespace MetadataServer.Connectors
 				{
 					ProjectId = TryInsertAndGetProject(Connection, Data.Project);
 				}
-				using (SQLiteCommand Command = new SQLiteCommand("INSERT INTO [Errors] (Type, Text, UserName, Project, Timestamp, Version, IpAddress, ProjectId) VALUES (@Type, @Text, @UserName, @Project, @Timestamp, @Version, @IpAddress, @ProjectId)", Connection))
+				using (MySqlCommand Command = new MySqlCommand("INSERT INTO ugs_db.Errors (Type, Text, UserName, Project, Timestamp, Version, IpAddress, ProjectId) VALUES (@Type, @Text, @UserName, @Project, @Timestamp, @Version, @IpAddress, @ProjectId)", Connection))
 				{
 					Command.Parameters.AddWithValue("@Type", Data.Type.ToString());
 					Command.Parameters.AddWithValue("@Text", Data.Text);
@@ -350,14 +353,14 @@ namespace MetadataServer.Connectors
 
 		public static long FindOrAddUserId(string Name)
 		{
-			using (SQLiteConnection Connection = new SQLiteConnection(SqlConnector.ConnectionString))
+			using (MySqlConnection Connection = new MySqlConnection(SqlConnector.ConnectionString))
 			{
 				Connection.Open();
 				return FindOrAddUserId(Name, Connection);
 			}
 		}
 
-		private static long FindOrAddUserId(string Name, SQLiteConnection Connection)
+		private static long FindOrAddUserId(string Name, MySqlConnection Connection)
 		{
 			if(Name.Length == 0)
 			{
@@ -366,21 +369,21 @@ namespace MetadataServer.Connectors
 
 			string NormalizedName = NormalizeUserName(Name);
 
-			using (SQLiteCommand Command = new SQLiteCommand("SELECT [Id] FROM [Users] WHERE Name = @Name", Connection))
+			using (MySqlCommand Command = new MySqlCommand("SELECT Id FROM ugs_db.Users WHERE Name = @Name", Connection))
 			{
 				Command.Parameters.AddWithValue("@Name", NormalizedName);
 				object UserId = Command.ExecuteScalar();
 				if(UserId != null)
 				{
-					return (long)UserId;
+					return Convert.ToInt64(UserId);
 				}
 			}
 
-			using (SQLiteCommand Command = new SQLiteCommand("INSERT OR IGNORE INTO [Users] (Name) VALUES (@Name); SELECT [Id] FROM [Users] WHERE Name = @Name", Connection))
+			using (MySqlCommand Command = new MySqlCommand("INSERT IGNORE INTO ugs_db.Users (Name) VALUES (@Name); SELECT Id FROM ugs_db.Users WHERE Name = @Name", Connection))
 			{
 				Command.Parameters.AddWithValue("@Name", NormalizedName);
 				object UserId = Command.ExecuteScalar();
-				return (long)UserId;
+				return Convert.ToInt64(UserId);
 			}
 		}
 
@@ -389,11 +392,11 @@ namespace MetadataServer.Connectors
 		public static long AddIssue(IssueData Issue)
 		{
 			long IssueId;
-			using (SQLiteConnection Connection = new SQLiteConnection(SqlConnector.ConnectionString))
+			using (MySqlConnection Connection = new MySqlConnection(SqlConnector.ConnectionString))
 			{
 				Connection.Open();
 
-				using (SQLiteCommand Command = new SQLiteCommand("INSERT INTO [Issues] (Project, Summary, OwnerId, CreatedAt, FixChange) VALUES (@Project, @Summary, @OwnerId, DATETIME('now'), 0)", Connection))
+				using (MySqlCommand Command = new MySqlCommand("INSERT INTO ugs_db.Issues (Project, Summary, OwnerId, CreatedAt, FixChange) VALUES (@Project, @Summary, @OwnerId, UTC_TIMESTAMP(), 0)", Connection))
 				{
 					Command.Parameters.AddWithValue("@Project", Issue.Project);
 					Command.Parameters.AddWithValue("@Summary", SanitizeText(Issue.Summary, IssueSummaryMaxLength));
@@ -407,7 +410,7 @@ namespace MetadataServer.Connectors
 					}
 					Command.ExecuteNonQuery();
 
-					IssueId = Connection.LastInsertRowId;
+					IssueId = Command.LastInsertedId;
 				}
 			}
 			return IssueId;
@@ -439,7 +442,7 @@ namespace MetadataServer.Connectors
 		private static List<IssueData> GetIssuesInternal(long IssueId, string UserName, bool IncludeResolved, int NumResults)
 		{
 			List<IssueData> Issues = new List<IssueData>();
-			using (SQLiteConnection Connection = new SQLiteConnection(SqlConnector.ConnectionString))
+			using (MySqlConnection Connection = new MySqlConnection(SqlConnector.ConnectionString))
 			{
 				Connection.Open();
 
@@ -451,17 +454,17 @@ namespace MetadataServer.Connectors
 
 				StringBuilder CommandBuilder = new StringBuilder();
 				CommandBuilder.Append("SELECT");
-				CommandBuilder.Append(" Issues.Id, Issues.CreatedAt, DATETIME('now'), Issues.Project, Issues.Summary, OwnerUsers.Name, NominatedByUsers.Name, Issues.AcknowledgedAt, Issues.FixChange, Issues.ResolvedAt");
+				CommandBuilder.Append(" Issues.Id, Issues.CreatedAt, UTC_TIMESTAMP(), Issues.Project, Issues.Summary, OwnerUsers.Name, NominatedByUsers.Name, Issues.AcknowledgedAt, Issues.FixChange, Issues.ResolvedAt");
 				if(UserName != null)
 				{
 					CommandBuilder.Append(", IssueWatchers.UserId");
 				}
-				CommandBuilder.Append(" FROM [Issues]");
-				CommandBuilder.Append(" LEFT JOIN [Users] AS [OwnerUsers] ON OwnerUsers.Id = Issues.OwnerId");
-				CommandBuilder.Append(" LEFT JOIN [Users] AS [NominatedByUsers] ON NominatedByUsers.Id = Issues.NominatedById");
+				CommandBuilder.Append(" FROM ugs_db.Issues");
+				CommandBuilder.Append(" LEFT JOIN ugs_db.Users AS OwnerUsers ON OwnerUsers.Id = Issues.OwnerId");
+				CommandBuilder.Append(" LEFT JOIN ugs_db.Users AS NominatedByUsers ON NominatedByUsers.Id = Issues.NominatedById");
 				if(UserName != null)
 				{
-					CommandBuilder.Append(" LEFT JOIN [IssueWatchers] ON IssueWatchers.IssueId = Issues.Id AND IssueWatchers.UserId = @UserId");
+					CommandBuilder.Append(" LEFT JOIN ugs_db.IssueWatchers ON IssueWatchers.IssueId = Issues.Id AND IssueWatchers.UserId = @UserId");
 				}
 				if(IssueId != -1)
 				{
@@ -476,7 +479,7 @@ namespace MetadataServer.Connectors
 					CommandBuilder.AppendFormat(" ORDER BY Issues.Id DESC LIMIT {0}", NumResults);
 				}
 
-				using (SQLiteCommand Command = new SQLiteCommand(CommandBuilder.ToString(), Connection))
+				using (MySqlCommand Command = new MySqlCommand(CommandBuilder.ToString(), Connection))
 				{
 					if(IssueId != -1)
 					{
@@ -487,7 +490,7 @@ namespace MetadataServer.Connectors
 						Command.Parameters.AddWithValue("@UserId", UserId);
 					}
 
-					using(SQLiteDataReader Reader = Command.ExecuteReader())
+					using(MySqlDataReader Reader = Command.ExecuteReader())
 					{
 						while(Reader.Read())
 						{
@@ -516,11 +519,11 @@ namespace MetadataServer.Connectors
 
 		public static void UpdateIssue(long IssueId, IssueUpdateData Issue)
 		{
-			using (SQLiteConnection Connection = new SQLiteConnection(SqlConnector.ConnectionString))
+			using (MySqlConnection Connection = new MySqlConnection(SqlConnector.ConnectionString))
 			{
 				Connection.Open();
 
-				using (SQLiteCommand Command = new SQLiteCommand(Connection))
+				using (MySqlCommand Command = Connection.CreateCommand())
 				{
 					List<string> Columns = new List<string>();
 					List<string> Values = new List<string>();
@@ -545,7 +548,7 @@ namespace MetadataServer.Connectors
 					if(Issue.Acknowledged.HasValue)
 					{
 						Columns.Add("AcknowledgedAt");
-						Values.Add(Issue.Acknowledged.Value? "DATETIME('now')" : "NULL");
+						Values.Add(Issue.Acknowledged.Value? "UTC_TIMESTAMP()" : "NULL");
 					}
 					if(Issue.FixChange.HasValue)
 					{
@@ -556,10 +559,20 @@ namespace MetadataServer.Connectors
 					if(Issue.Resolved.HasValue)
 					{
 						Columns.Add("ResolvedAt");
-						Values.Add(Issue.Resolved.Value? "DATETIME('now')" : "NULL");
+						Values.Add(Issue.Resolved.Value? "UTC_TIMESTAMP()" : "NULL");
 					}
 
-					Command.CommandText = String.Format("UPDATE [Issues] SET ({0}) = ({1}) WHERE Id = @IssueId", String.Join(", ", Columns), String.Join(", ", Values));
+					StringBuilder CommandText = new StringBuilder("UPDATE ugs_db.Issues SET ");
+					for(int idx = 0; idx < Columns.Count; idx++)
+					{
+						CommandText.Append(String.Format("{0}={1}", Columns[idx], Values[idx]));
+						if(idx != Columns.Count - 1)
+						{
+							CommandText.Append(",");
+						}
+					}
+					CommandText.Append(" WHERE Id = @IssueId");
+					Command.CommandText = CommandText.ToString();
 					Command.Parameters.AddWithValue("@IssueId", IssueId);
 					Command.ExecuteNonQuery();
 				}
@@ -585,25 +598,25 @@ namespace MetadataServer.Connectors
 
 		public static void DeleteIssue(long IssueId)
 		{
-			using (SQLiteConnection Connection = new SQLiteConnection(SqlConnector.ConnectionString))
+			using (MySqlConnection Connection = new MySqlConnection(SqlConnector.ConnectionString))
 			{
 				Connection.Open();
 
-				using (SQLiteTransaction Transaction = Connection.BeginTransaction())
+				using (MySqlTransaction Transaction = Connection.BeginTransaction())
 				{
-					using (SQLiteCommand Command = new SQLiteCommand(Connection))
+					using (MySqlCommand Command = Connection.CreateCommand())
 					{
 						Command.Transaction = Transaction;
 
-						Command.CommandText = "DELETE FROM [IssueWatchers] WHERE IssueId = @IssueId";
+						Command.CommandText = "DELETE FROM ugs_db.IssueWatchers WHERE IssueId = @IssueId";
 						Command.Parameters.AddWithValue("@IssueId", IssueId);
 						Command.ExecuteNonQuery();
 
-						Command.CommandText = "DELETE FROM [IssueBuilds] WHERE IssueId = @IssueId";
+						Command.CommandText = "DELETE FROM ugs_db.IssueBuilds WHERE IssueId = @IssueId";
 						Command.Parameters.AddWithValue("@IssueId", IssueId);
 						Command.ExecuteNonQuery();
 
-						Command.CommandText = "DELETE FROM [Issues] WHERE Id = @IssueId";
+						Command.CommandText = "DELETE FROM ugs_db.Issues WHERE Id = @IssueId";
 						Command.Parameters.AddWithValue("@IssueId", IssueId);
 						Command.ExecuteNonQuery();
 
@@ -615,11 +628,11 @@ namespace MetadataServer.Connectors
 
 		public static void AddDiagnostic(long IssueId, IssueDiagnosticData Diagnostic)
 		{
-			using (SQLiteConnection Connection = new SQLiteConnection(SqlConnector.ConnectionString))
+			using (MySqlConnection Connection = new MySqlConnection(SqlConnector.ConnectionString))
 			{
 				Connection.Open();
 
-				using (SQLiteCommand Command = new SQLiteCommand("INSERT INTO [IssueDiagnostics] (IssueId, BuildId, Message, Url) VALUES (@IssueId, @BuildId, @Message, @Url)", Connection))
+				using (MySqlCommand Command = new MySqlCommand("INSERT INTO ugs_db.IssueDiagnostics (IssueId, BuildId, Message, Url) VALUES (@IssueId, @BuildId, @Message, @Url)", Connection))
 				{
 					Command.Parameters.AddWithValue("@IssueId", IssueId);
 					Command.Parameters.AddWithValue("@BuildId", Diagnostic.BuildId);
@@ -633,18 +646,18 @@ namespace MetadataServer.Connectors
 		public static List<IssueDiagnosticData> GetDiagnostics(long IssueId)
 		{
 			List<IssueDiagnosticData> Diagnostics = new List<IssueDiagnosticData>();
-			using (SQLiteConnection Connection = new SQLiteConnection(SqlConnector.ConnectionString))
+			using (MySqlConnection Connection = new MySqlConnection(SqlConnector.ConnectionString))
 			{
 				Connection.Open();
 
 				StringBuilder CommandBuilder = new StringBuilder();
-				CommandBuilder.Append("SELECT BuildId, Message, Url FROM [IssueDiagnostics]");
+				CommandBuilder.Append("SELECT BuildId, Message, Url FROM ugs_db.IssueDiagnostics");
 				CommandBuilder.Append(" WHERE IssueDiagnostics.IssueId = @IssueId");
 
-				using (SQLiteCommand Command = new SQLiteCommand(CommandBuilder.ToString(), Connection))
+				using (MySqlCommand Command = new MySqlCommand(CommandBuilder.ToString(), Connection))
 				{
 					Command.Parameters.AddWithValue("@IssueId", IssueId);
-					using (SQLiteDataReader Reader = Command.ExecuteReader())
+					using (MySqlDataReader Reader = Command.ExecuteReader())
 					{
 						while (Reader.Read())
 						{
@@ -662,13 +675,13 @@ namespace MetadataServer.Connectors
 
 		public static void AddWatcher(long IssueId, string UserName)
 		{
-			using (SQLiteConnection Connection = new SQLiteConnection(SqlConnector.ConnectionString))
+			using (MySqlConnection Connection = new MySqlConnection(SqlConnector.ConnectionString))
 			{
 				Connection.Open();
 
 				long UserId = FindOrAddUserId(UserName, Connection);
 
-				using(SQLiteCommand Command = new SQLiteCommand("INSERT OR IGNORE INTO [IssueWatchers] (IssueId, UserId) VALUES (@IssueId, @UserId)", Connection))
+				using(MySqlCommand Command = new MySqlCommand("INSERT IGNORE INTO ugs_db.IssueWatchers (IssueId, UserId) VALUES (@IssueId, @UserId)", Connection))
 				{
 					Command.Parameters.AddWithValue("@IssueId", IssueId);
 					Command.Parameters.AddWithValue("@UserId", UserId);
@@ -680,19 +693,19 @@ namespace MetadataServer.Connectors
 		public static List<string> GetWatchers(long IssueId)
 		{
 			List<string> Watchers = new List<string>();
-			using (SQLiteConnection Connection = new SQLiteConnection(SqlConnector.ConnectionString))
+			using (MySqlConnection Connection = new MySqlConnection(SqlConnector.ConnectionString))
 			{
 				Connection.Open();
 
 				StringBuilder CommandBuilder = new StringBuilder();
-				CommandBuilder.Append("SELECT Users.Name FROM [IssueWatchers]");
-				CommandBuilder.Append(" LEFT JOIN [Users] ON IssueWatchers.UserId = Users.Id");
+				CommandBuilder.Append("SELECT Users.Name FROM ugs_db.IssueWatchers");
+				CommandBuilder.Append(" LEFT JOIN ugs_db.Users ON IssueWatchers.UserId = Users.Id");
 				CommandBuilder.Append(" WHERE IssueWatchers.IssueId = @IssueId");
 
-				using(SQLiteCommand Command = new SQLiteCommand(CommandBuilder.ToString(), Connection))
+				using(MySqlCommand Command = new MySqlCommand(CommandBuilder.ToString(), Connection))
 				{
 					Command.Parameters.AddWithValue("@IssueId", IssueId);
-					using(SQLiteDataReader Reader = Command.ExecuteReader())
+					using(MySqlDataReader Reader = Command.ExecuteReader())
 					{
 						while(Reader.Read())
 						{
@@ -706,13 +719,13 @@ namespace MetadataServer.Connectors
 
 		public static void RemoveWatcher(long IssueId, string UserName)
 		{
-			using (SQLiteConnection Connection = new SQLiteConnection(SqlConnector.ConnectionString))
+			using (MySqlConnection Connection = new MySqlConnection(SqlConnector.ConnectionString))
 			{
 				Connection.Open();
 
 				long UserId = FindOrAddUserId(UserName, Connection);
 
-				using(SQLiteCommand Command = new SQLiteCommand("DELETE FROM [IssueWatchers] WHERE IssueId = @IssueId AND UserId = @UserId", Connection))
+				using(MySqlCommand Command = new MySqlCommand("DELETE FROM ugs_db.IssueWatchers WHERE IssueId = @IssueId AND UserId = @UserId", Connection))
 				{
 					Command.Parameters.AddWithValue("@IssueId", IssueId);
 					Command.Parameters.AddWithValue("@UserId", UserId);
@@ -723,11 +736,11 @@ namespace MetadataServer.Connectors
 
 		public static long AddBuild(long IssueId, IssueBuildData Build)
 		{
-			using (SQLiteConnection Connection = new SQLiteConnection(SqlConnector.ConnectionString))
+			using (MySqlConnection Connection = new MySqlConnection(SqlConnector.ConnectionString))
 			{
 				Connection.Open();
 
-				using (SQLiteCommand Command = new SQLiteCommand("INSERT INTO [IssueBuilds] (IssueId, Stream, Change, JobName, JobUrl, JobStepName, JobStepUrl, ErrorUrl, Outcome) VALUES (@IssueId, @Stream, @Change, @JobName, @JobUrl, @JobStepName, @JobStepUrl, @ErrorUrl, @Outcome)", Connection))
+				using (MySqlCommand Command = new MySqlCommand("INSERT INTO ugs_db.IssueBuilds (IssueId, Stream, `Change`, JobName, JobUrl, JobStepName, JobStepUrl, ErrorUrl, Outcome) VALUES (@IssueId, @Stream, @Change, @JobName, @JobUrl, @JobStepName, @JobStepUrl, @ErrorUrl, @Outcome)", Connection))
 				{
 					Command.Parameters.AddWithValue("@IssueId", IssueId);
 					Command.Parameters.AddWithValue("@Stream", Build.Stream);
@@ -740,7 +753,7 @@ namespace MetadataServer.Connectors
 					Command.Parameters.AddWithValue("@Outcome", Build.Outcome);
 					Command.ExecuteNonQuery();
 
-					return Connection.LastInsertRowId;
+					return Command.LastInsertedId;
 				}
 			}
 		}
@@ -748,14 +761,14 @@ namespace MetadataServer.Connectors
 		public static List<IssueBuildData> GetBuilds(long IssueId)
 		{
 			List<IssueBuildData> Builds = new List<IssueBuildData>();
-			using (SQLiteConnection Connection = new SQLiteConnection(SqlConnector.ConnectionString))
+			using (MySqlConnection Connection = new MySqlConnection(SqlConnector.ConnectionString))
 			{
 				Connection.Open();
 
-				using(SQLiteCommand Command = new SQLiteCommand("SELECT IssueBuilds.Id, IssueBuilds.Stream, IssueBuilds.Change, IssueBuilds.JobName, IssueBuilds.JobUrl, IssueBuilds.JobStepName, IssueBuilds.JobStepUrl, IssueBuilds.ErrorUrl, IssueBuilds.Outcome FROM [IssueBuilds] WHERE IssueBuilds.IssueId = @IssueId", Connection))
+				using(MySqlCommand Command = new MySqlCommand("SELECT IssueBuilds.Id, IssueBuilds.Stream, IssueBuilds.Change, IssueBuilds.JobName, IssueBuilds.JobUrl, IssueBuilds.JobStepName, IssueBuilds.JobStepUrl, IssueBuilds.ErrorUrl, IssueBuilds.Outcome FROM ugs_db.IssueBuilds WHERE IssueBuilds.IssueId = @IssueId", Connection))
 				{
 					Command.Parameters.AddWithValue("@IssueId", IssueId);
-					using(SQLiteDataReader Reader = Command.ExecuteReader())
+					using(MySqlDataReader Reader = Command.ExecuteReader())
 					{
 						while(Reader.Read())
 						{
@@ -779,14 +792,14 @@ namespace MetadataServer.Connectors
 		public static IssueBuildData GetBuild(long BuildId)
 		{
 			IssueBuildData Build = null;
-			using (SQLiteConnection Connection = new SQLiteConnection(SqlConnector.ConnectionString))
+			using (MySqlConnection Connection = new MySqlConnection(SqlConnector.ConnectionString))
 			{
 				Connection.Open();
 
-				using(SQLiteCommand Command = new SQLiteCommand("SELECT IssueBuilds.Id, IssueBuilds.Stream, IssueBuilds.Change, IssueBuilds.JobName, IssueBuilds.JobUrl, IssueBuilds.JobStepName, IssueBuilds.JobStepUrl, IssueBuilds.ErrorUrl, IssueBuilds.Outcome FROM [IssueBuilds] WHERE IssueBuilds.Id = @BuildId", Connection))
+				using(MySqlCommand Command = new MySqlCommand("SELECT IssueBuilds.Id, IssueBuilds.Stream, IssueBuilds.Change, IssueBuilds.JobName, IssueBuilds.JobUrl, IssueBuilds.JobStepName, IssueBuilds.JobStepUrl, IssueBuilds.ErrorUrl, IssueBuilds.Outcome FROM ugs_db.IssueBuilds WHERE IssueBuilds.Id = @BuildId", Connection))
 				{
 					Command.Parameters.AddWithValue("@BuildId", BuildId);
-					using(SQLiteDataReader Reader = Command.ExecuteReader())
+					using(MySqlDataReader Reader = Command.ExecuteReader())
 					{
 						while(Reader.Read())
 						{
@@ -810,11 +823,11 @@ namespace MetadataServer.Connectors
 
 		public static void UpdateBuild(long BuildId, int Outcome)
 		{
-			using (SQLiteConnection Connection = new SQLiteConnection(SqlConnector.ConnectionString))
+			using (MySqlConnection Connection = new MySqlConnection(SqlConnector.ConnectionString))
 			{
 				Connection.Open();
 
-				using (SQLiteCommand Command = new SQLiteCommand("UPDATE [IssueBuilds] SET (Outcome) = (@Outcome) WHERE Id = @BuildId", Connection))
+				using (MySqlCommand Command = new MySqlCommand("UPDATE ugs_db.IssueBuilds SET (Outcome) = (@Outcome) WHERE Id = @BuildId", Connection))
 				{
 					Command.Parameters.AddWithValue("@BuildId", BuildId);
 					Command.Parameters.AddWithValue("@Outcome", Outcome);
diff --git a/Engine/Source/Programs/UnrealGameSync/MetadataServer/Controllers/UserController.cs b/Engine/Source/Programs/UnrealGameSync/MetadataServer/Controllers/UserController.cs
--- a/Engine/Source/Programs/UnrealGameSync/MetadataServer/Controllers/UserController.cs
+++ b/Engine/Source/Programs/UnrealGameSync/MetadataServer/Controllers/UserController.cs
@@ -2,7 +2,6 @@
 
 using System;
 using System.Collections.Generic;
-using System.Data.SQLite;
 using System.Web.Http;
 using MetadataServer.Connectors;
 using MetadataServer.Models;
diff --git a/Engine/Source/Programs/UnrealGameSync/MetadataServer/Global.asax.cs b/Engine/Source/Programs/UnrealGameSync/MetadataServer/Global.asax.cs
--- a/Engine/Source/Programs/UnrealGameSync/MetadataServer/Global.asax.cs
+++ b/Engine/Source/Programs/UnrealGameSync/MetadataServer/Global.asax.cs
@@ -1,7 +1,7 @@
 ﻿// Copyright 1998-2019 Epic Games, Inc. All Rights Reserved.
+using MySql.Data.MySqlClient;
 using System;
 using System.Collections.Generic;
-using System.Data.SQLite;
 using System.IO;
 using System.Linq;
 using System.Text.RegularExpressions;
@@ -17,134 +17,15 @@ namespace MetadataServer
             try
             {
                 string ConnectionString = System.Configuration.ConfigurationManager.ConnectionStrings["ConnectionString"].ConnectionString;
-                string FileName = Regex.Match(ConnectionString, "Data Source=(.+);.+").Groups[1].Value;
-                Directory.CreateDirectory(new FileInfo(FileName).Directory.FullName);
-                if (!File.Exists(FileName))
-                {
-                    SQLiteConnection.CreateFile(FileName);
-                }
-                // initialize the db if it doesn't exist
-                using (SQLiteConnection Connection = new SQLiteConnection(ConnectionString, true))
-                {
-                    Connection.Open();
-                    using (SQLiteCommand Command = new SQLiteCommand(Properties.Resources.Setup, Connection))
-                    {
-                        Command.ExecuteNonQuery();
-                    }
-
-					////////////////////////////////////////////////////////
-					// DB change for indexing of project IDs by foreign key
-					////////////////////////////////////////////////////////
-					
-					// check if we've seeded the projects table already by looking at the raw number of rows
-					// if we're a brand new database this will techincally trigger with a no-op every time the server is started
-					// but will stop as soon as a badge is added.
-					bool ProjectsTableSeeded = false;
-					using (SQLiteCommand Command = new SQLiteCommand("SELECT COALESCE(MAX(id)+1, 0) FROM [Projects]", Connection))
+				using (MySqlConnection Connection = new MySqlConnection(ConnectionString))
+				{
+					Connection.Open();
+					using (MySqlCommand Command = new MySqlCommand(Properties.Resources.Setup, Connection))
 					{
-						using (SQLiteDataReader Reader = Command.ExecuteReader())
-						{
-							while (Reader.Read())
-							{
-								int ProjectRows = Reader.GetInt32(0);
-								ProjectsTableSeeded = ProjectRows == 0 ? false : true;
-							}
-						}
+						Command.ExecuteNonQuery();
 					}
-					// if not found, get all the existing projects fromm all the other tables.
-					if (!ProjectsTableSeeded)
-					{
-						bool AlterTableRun = false;
-						// check and see if we've run the alter already.
-						using (SQLiteCommand Command = new SQLiteCommand("PRAGMA table_info(CIS)", Connection))
-						{
-							using (SQLiteDataReader Reader = Command.ExecuteReader())
-							{
-								while (Reader.Read())
-								{
-									string ColumnName = Reader.GetString(1);
-									if(ColumnName == "ProjectID")
-									{
-										AlterTableRun = true;
-									}
-								}
-							}
-						}
-						if (!AlterTableRun)
-						{
-							using (SQLiteCommand Command = new SQLiteCommand("ALTER TABLE [CIS] ADD COLUMN ProjectID INTEGER REFERENCES Projects(Id);" +
-																			 "ALTER TABLE [Comments] ADD COLUMN ProjectID INTEGER REFERENCES Projects(Id);" +
-																			 "ALTER TABLE [Errors] ADD COLUMN ProjectID INTEGER REFERENCES Projects(Id);" +
-																			 "ALTER TABLE [Telemetry.v2] ADD COLUMN ProjectID INTEGER REFERENCES Projects(Id);" +
-																			 "ALTER TABLE [UserVotes] ADD COLUMN ProjectID INTEGER REFERENCES Projects(Id);", Connection))
-							{
-								Command.ExecuteNonQuery();
-							}
-
-							// get all the distinct project names.
-							Dictionary<string, long> Projects = new Dictionary<string, long>();
-							using (SQLiteCommand Command = new SQLiteCommand("SELECT Project FROM [CIS] UNION Select Project FROM [Comments] WHERE Project IS NOT NULL " +
-																			 "UNION " +
-																			 "SELECT Project FROM [Comments] WHERE Project IS NOT NULL " +
-																			 "UNION " +
-																			 "SELECT Project FROM [Errors] WHERE Project IS NOT NULL " +
-																			 "UNION " +
-																			 "SELECT Project FROM [Telemetry.v2] WHERE Project is not NULL " +
-																			 "UNION " +
-																			 "SELECT Project FROM [UserVotes] WHERE Project is not NULL", Connection))
-							{
-								using (SQLiteDataReader Reader = Command.ExecuteReader())
-								{
-									while (Reader.Read())
-									{
-										object ProjectName = Reader.GetValue(0);
-										if (ProjectName != null)
-										{
-											Projects.Add((string)ProjectName, 0);
-										}
-
-									}
-								}
-							}
-							// insert all the projects into the new table
-							using (SQLiteCommand Command = new SQLiteCommand(Connection))
-							{
-								using (SQLiteTransaction Transaction = Connection.BeginTransaction())
-								{
-									foreach (string ProjectName in Projects.Keys.ToList())
-									{
-										Command.CommandText = "INSERT INTO [Projects] (Name) VALUES (@ProjectName)";
-										Command.Parameters.AddWithValue("@ProjectName", ProjectName);
-										Command.ExecuteNonQuery();
-										Projects[ProjectName] = Connection.LastInsertRowId;
-									}
-									Transaction.Commit();
-								}
-							}
-							// update all the foreign keys
-							using (SQLiteCommand Command = new SQLiteCommand(Connection))
-							{
-								using (SQLiteTransaction Transaction = Connection.BeginTransaction())
-								{
-									foreach (KeyValuePair<string, long> Project in Projects)
-									{
-										Command.CommandText = "UPDATE [CIS] SET ProjectID = @ProjectID WHERE Project = @ProjectName;" +
-															  "UPDATE [Comments] SET ProjectID = @ProjectID WHERE Project = @ProjectName;" +
-															  "UPDATE [Errors] SET ProjectID = @ProjectID WHERE Project = @ProjectName;" +
-															  "UPDATE [Telemetry.v2] SET ProjectID = @ProjectID WHERE Project = @ProjectName;" +
-															  "UPDATE [UserVotes] SET ProjectID = @ProjectID WHERE Project = @ProjectName;";
-										Command.Parameters.AddWithValue("@ProjectName", Project.Key);
-										Command.Parameters.AddWithValue("@ProjectID", Project.Value);
-										Command.ExecuteNonQuery();
-									}
-									Transaction.Commit();
-								}
-							}
-						}
-					}
-                }
-
-            }
+				}
+			}
             catch (Exception)
             {
                 HttpRuntime.UnloadAppDomain();
diff --git a/Engine/Source/Programs/UnrealGameSync/MetadataServer/MetadataServer.csproj b/Engine/Source/Programs/UnrealGameSync/MetadataServer/MetadataServer.csproj
--- a/Engine/Source/Programs/UnrealGameSync/MetadataServer/MetadataServer.csproj
+++ b/Engine/Source/Programs/UnrealGameSync/MetadataServer/MetadataServer.csproj
@@ -13,7 +13,7 @@
     <AppDesignerFolder>Properties</AppDesignerFolder>
     <RootNamespace>MetadataServer</RootNamespace>
     <AssemblyName>MetadataServer</AssemblyName>
-    <TargetFrameworkVersion>v4.5</TargetFrameworkVersion>
+    <TargetFrameworkVersion>v4.5.2</TargetFrameworkVersion>
     <MvcBuildViews>false</MvcBuildViews>
     <UseIISExpress>true</UseIISExpress>
     <IISExpressSSLPort />
@@ -46,16 +46,22 @@
     <WarningLevel>4</WarningLevel>
   </PropertyGroup>
   <ItemGroup>
+    <Reference Include="BouncyCastle.Crypto, Version=1.8.3.0, Culture=neutral, PublicKeyToken=0e99375e54769942">
+      <HintPath>..\packages\BouncyCastle.1.8.3.1\lib\BouncyCastle.Crypto.dll</HintPath>
+    </Reference>
+    <Reference Include="Google.Protobuf, Version=3.6.1.0, Culture=neutral, PublicKeyToken=a7d26565bac4d604, processorArchitecture=MSIL">
+      <HintPath>..\packages\Google.Protobuf.3.6.1\lib\net45\Google.Protobuf.dll</HintPath>
+    </Reference>
     <Reference Include="Microsoft.CSharp" />
+    <Reference Include="MySql.Data, Version=8.0.16.0, Culture=neutral, PublicKeyToken=c5687fc88969c44d, processorArchitecture=MSIL">
+      <HintPath>..\packages\MySql.Data.8.0.16\lib\net452\MySql.Data.dll</HintPath>
+    </Reference>
     <Reference Include="Newtonsoft.Json, Version=12.0.0.0, Culture=neutral, PublicKeyToken=30ad4fe6b2a6aeed, processorArchitecture=MSIL">
       <HintPath>..\packages\Newtonsoft.Json.12.0.2\lib\net45\Newtonsoft.Json.dll</HintPath>
     </Reference>
     <Reference Include="System" />
     <Reference Include="System.Data" />
     <Reference Include="System.Data.DataSetExtensions" />
-    <Reference Include="System.Data.SQLite, Version=1.0.110.0, Culture=neutral, PublicKeyToken=db937bc2d44ff139, processorArchitecture=MSIL">
-      <HintPath>..\packages\System.Data.SQLite.Core.1.0.110.0\lib\net45\System.Data.SQLite.dll</HintPath>
-    </Reference>
     <Reference Include="System.Drawing" />
     <Reference Include="System.Net.Http.Formatting, Version=5.2.7.0, Culture=neutral, PublicKeyToken=31bf3856ad364e35, processorArchitecture=MSIL">
       <HintPath>..\packages\Microsoft.AspNet.WebApi.Client.5.2.7\lib\net45\System.Net.Http.Formatting.dll</HintPath>
@@ -64,9 +70,9 @@
     <Reference Include="System.Web.Entity" />
     <Reference Include="System.Web.ApplicationServices" />
     <Reference Include="System.ComponentModel.DataAnnotations" />
-    <Reference Include="System.Web.Extensions" />
     <Reference Include="System.Web" />
     <Reference Include="System.Web.Abstractions" />
+    <Reference Include="System.Web.Extensions" />
     <Reference Include="System.Web.Http, Version=5.2.7.0, Culture=neutral, PublicKeyToken=31bf3856ad364e35, processorArchitecture=MSIL">
       <HintPath>..\packages\Microsoft.AspNet.WebApi.Core.5.2.7\lib\net45\System.Web.Http.dll</HintPath>
     </Reference>
@@ -185,11 +191,4 @@
     <Message Importance="High" Text="Found NotForLicensees\Web.config, applying transform" Condition="Exists('NotForLicensees\Web.config')" />
     <TransformXml Source="Web.config" Transform="NotForLicensees\Web.config" Destination="Web.config" Condition="Exists('NotForLicensees\Web.config')" />
   </Target>
-  <Import Project="..\packages\System.Data.SQLite.Core.1.0.110.0\build\net45\System.Data.SQLite.Core.targets" Condition="Exists('..\packages\System.Data.SQLite.Core.1.0.110.0\build\net45\System.Data.SQLite.Core.targets')" />
-  <Target Name="EnsureNuGetPackageBuildImports" BeforeTargets="PrepareForBuild">
-    <PropertyGroup>
-      <ErrorText>This project references NuGet package(s) that are missing on this computer. Use NuGet Package Restore to download them.  For more information, see http://go.microsoft.com/fwlink/?LinkID=322105. The missing file is {0}.</ErrorText>
-    </PropertyGroup>
-    <Error Condition="!Exists('..\packages\System.Data.SQLite.Core.1.0.110.0\build\net45\System.Data.SQLite.Core.targets')" Text="$([System.String]::Format('$(ErrorText)', '..\packages\System.Data.SQLite.Core.1.0.110.0\build\net45\System.Data.SQLite.Core.targets'))" />
-  </Target>
 </Project>
\ No newline at end of file
diff --git a/Engine/Source/Programs/UnrealGameSync/MetadataServer/packages.config b/Engine/Source/Programs/UnrealGameSync/MetadataServer/packages.config
--- a/Engine/Source/Programs/UnrealGameSync/MetadataServer/packages.config
+++ b/Engine/Source/Programs/UnrealGameSync/MetadataServer/packages.config
@@ -1,9 +1,11 @@
 ﻿<?xml version="1.0" encoding="utf-8"?>
 <packages>
+  <package id="BouncyCastle" version="1.8.3.1" targetFramework="net452" />
+  <package id="Google.Protobuf" version="3.6.1" targetFramework="net452" />
   <package id="Microsoft.AspNet.WebApi" version="5.2.7" targetFramework="net45" />
   <package id="Microsoft.AspNet.WebApi.Client" version="5.2.7" targetFramework="net45" />
   <package id="Microsoft.AspNet.WebApi.Core" version="5.2.7" targetFramework="net45" />
   <package id="Microsoft.AspNet.WebApi.WebHost" version="5.2.7" targetFramework="net45" />
+  <package id="MySql.Data" version="8.0.16" targetFramework="net452" />
   <package id="Newtonsoft.Json" version="12.0.2" targetFramework="net45" />
-  <package id="System.Data.SQLite.Core" version="1.0.110.0" targetFramework="net45" />
 </packages>
\ No newline at end of file
