change target framework to 4.6.2

diff --git a/Engine/Source/Programs/UnrealGameSync/MetadataServer/MetadataServer.csproj b/Engine/Source/Programs/UnrealGameSync/MetadataServer/MetadataServer.csproj
--- a/Engine/Source/Programs/UnrealGameSync/MetadataServer/MetadataServer.csproj
+++ b/Engine/Source/Programs/UnrealGameSync/MetadataServer/MetadataServer.csproj
@@ -13,7 +13,7 @@
     <AppDesignerFolder>Properties</AppDesignerFolder>
     <RootNamespace>MetadataServer</RootNamespace>
     <AssemblyName>MetadataServer</AssemblyName>
-    <TargetFrameworkVersion>v4.5.2</TargetFrameworkVersion>
+    <TargetFrameworkVersion>v4.6.2</TargetFrameworkVersion>
     <MvcBuildViews>false</MvcBuildViews>
     <UseIISExpress>true</UseIISExpress>
     <IISExpressSSLPort />
@@ -61,7 +61,6 @@
     </Reference>
     <Reference Include="System" />
     <Reference Include="System.Data" />
-    <Reference Include="System.Data.DataSetExtensions" />
     <Reference Include="System.Drawing" />
     <Reference Include="System.Net.Http.Formatting, Version=5.2.7.0, Culture=neutral, PublicKeyToken=31bf3856ad364e35, processorArchitecture=MSIL">
       <HintPath>..\packages\Microsoft.AspNet.WebApi.Client.5.2.7\lib\net45\System.Net.Http.Formatting.dll</HintPath>
@@ -72,7 +71,6 @@
     <Reference Include="System.ComponentModel.DataAnnotations" />
     <Reference Include="System.Web" />
     <Reference Include="System.Web.Abstractions" />
-    <Reference Include="System.Web.Extensions" />
     <Reference Include="System.Web.Http, Version=5.2.7.0, Culture=neutral, PublicKeyToken=31bf3856ad364e35, processorArchitecture=MSIL">
       <HintPath>..\packages\Microsoft.AspNet.WebApi.Core.5.2.7\lib\net45\System.Web.Http.dll</HintPath>
     </Reference>
@@ -87,7 +85,6 @@
     </Reference>
     <Reference Include="System.Net.Http.WebRequest">
     </Reference>
-    <Reference Include="System.Xml.Linq" />
   </ItemGroup>
   <ItemGroup>
     <Folder Include="App_Data\" />
