UGS: Fix crash when selecting "Open in Explorer..." from the context menu of a file in the Clean Workspace dialog.

#rnx

diff --git a/Engine/Source/Programs/UnrealGameSync/UnrealGameSync/Forms/CleanWorkspaceWindow.cs b/Engine/Source/Programs/UnrealGameSync/UnrealGameSync/Forms/CleanWorkspaceWindow.cs
--- a/Engine/Source/Programs/UnrealGameSync/UnrealGameSync/Forms/CleanWorkspaceWindow.cs
+++ b/Engine/Source/Programs/UnrealGameSync/UnrealGameSync/Forms/CleanWorkspaceWindow.cs
@@ -489,7 +489,15 @@ namespace UnrealGameSync
 		{
 			TreeNode Node = (TreeNode)FolderContextMenu.Tag;
 			TreeNodeData NodeData = (TreeNodeData)Node.Tag;
-			Process.Start("explorer.exe", String.Format("\"{0}\"", NodeData.Folder.Directory.FullName));
+
+			if (NodeData.Folder != null)
+			{
+				Process.Start("explorer.exe", String.Format("\"{0}\"", NodeData.Folder.Directory.FullName));
+			}
+			else if (NodeData.File != null)
+			{
+				Process.Start("explorer.exe", String.Format("\"{0}\"", NodeData.File.Directory.FullName));
+			}
 		}
 	}
 }
