UGS: Fix an issue which can occur when switching between streams, and a fstat call for version files will return the existing file from the previous stream.

diff --git a/Engine/Source/Programs/UnrealGameSync/UnrealGameSync/Workspace.cs b/Engine/Source/Programs/UnrealGameSync/UnrealGameSync/Workspace.cs
--- a/Engine/Source/Programs/UnrealGameSync/UnrealGameSync/Workspace.cs
+++ b/Engine/Source/Programs/UnrealGameSync/UnrealGameSync/Workspace.cs
@@ -1412,7 +1412,19 @@ namespace UnrealGameSync
 				Log.WriteLine("Failed to query records for {0}", ClientPath);
 				return false;
 			}
-			if(Records.Count == 0)
+			if (Records.Count > 1)
+			{
+				// Attempt to remove any existing file which is synced
+				Perforce.ForceSync(String.Format("{0}#0", ClientPath), Log);
+
+				// Try to get the mapped files again
+				if (!Perforce.Stat(ClientPath, out Records, Log))
+				{
+					Log.WriteLine("Failed to query records for {0}", ClientPath);
+					return false;
+				}
+			}
+			if (Records.Count == 0)
 			{
 				Log.WriteLine("Ignoring {0}; not found on server.", ClientPath);
 				return true;
