Copying latest UGS from Dev-Build (pt 1)

#rnx

diff --git a/Engine/Source/Programs/UnrealGameSync/InstallerCustomAction/CustomAction.cs b/Engine/Source/Programs/UnrealGameSync/InstallerCustomAction/CustomAction.cs
--- a/Engine/Source/Programs/UnrealGameSync/InstallerCustomAction/CustomAction.cs
+++ b/Engine/Source/Programs/UnrealGameSync/InstallerCustomAction/CustomAction.cs
@@ -1,4 +1,4 @@
-﻿// Copyright 1998-2019 Epic Games, Inc. All Rights Reserved.
+﻿// Copyright 1998-2018 Epic Games, Inc. All Rights Reserved.
 
 using System;
 using System.Collections.Generic;
diff --git a/Engine/Source/Programs/UnrealGameSync/InstallerCustomAction/Properties/AssemblyInfo.cs b/Engine/Source/Programs/UnrealGameSync/InstallerCustomAction/Properties/AssemblyInfo.cs
--- a/Engine/Source/Programs/UnrealGameSync/InstallerCustomAction/Properties/AssemblyInfo.cs
+++ b/Engine/Source/Programs/UnrealGameSync/InstallerCustomAction/Properties/AssemblyInfo.cs
@@ -1,4 +1,4 @@
-﻿// Copyright 1998-2019 Epic Games, Inc. All Rights Reserved.
+﻿// Copyright 1998-2018 Epic Games, Inc. All Rights Reserved.
 
 using System.Reflection;
 using System.Runtime.CompilerServices;
@@ -11,7 +11,7 @@ using System.Runtime.InteropServices;
 [assembly: AssemblyDescription("")]
 [assembly: AssemblyCompany("")]
 [assembly: AssemblyProduct("InstallerCustomAction")]
-[assembly: AssemblyCopyright("Copyright 1998-2019 Epic Games, Inc. All rights reserved.")]
+[assembly: AssemblyCopyright("Copyright ©  2015")]
 [assembly: AssemblyTrademark("")]
 [assembly: AssemblyCulture("")]
 
diff --git a/Engine/Source/Programs/UnrealGameSync/PostBadgeStatus/Program.cs b/Engine/Source/Programs/UnrealGameSync/PostBadgeStatus/Program.cs
--- a/Engine/Source/Programs/UnrealGameSync/PostBadgeStatus/Program.cs
+++ b/Engine/Source/Programs/UnrealGameSync/PostBadgeStatus/Program.cs
@@ -1,4 +1,4 @@
-﻿// Copyright 1998-2019 Epic Games, Inc. All Rights Reserved.
+﻿// Copyright 1998-2018 Epic Games, Inc. All Rights Reserved.
 
 using System;
 using System.Collections.Generic;
diff --git a/Engine/Source/Programs/UnrealGameSync/PostBadgeStatus/Properties/AssemblyInfo.cs b/Engine/Source/Programs/UnrealGameSync/PostBadgeStatus/Properties/AssemblyInfo.cs
--- a/Engine/Source/Programs/UnrealGameSync/PostBadgeStatus/Properties/AssemblyInfo.cs
+++ b/Engine/Source/Programs/UnrealGameSync/PostBadgeStatus/Properties/AssemblyInfo.cs
@@ -1,4 +1,4 @@
-﻿// Copyright 1998-2019 Epic Games, Inc. All Rights Reserved.
+﻿// Copyright 1998-2018 Epic Games, Inc. All Rights Reserved.
 
 using System.Reflection;
 using System.Runtime.CompilerServices;
@@ -12,7 +12,7 @@ using System.Runtime.InteropServices;
 [assembly: AssemblyConfiguration("")]
 [assembly: AssemblyCompany("")]
 [assembly: AssemblyProduct("WriteBadgeStatus")]
-[assembly: AssemblyCopyright("Copyright 1998-2019 Epic Games, Inc. All Rights Reserved.")]
+[assembly: AssemblyCopyright("Copyright 1998-2018 Epic Games, Inc. All Rights Reserved.")]
 [assembly: AssemblyTrademark("")]
 [assembly: AssemblyCulture("")]
 
diff --git a/Engine/Source/Programs/UnrealGameSync/UnrealGameSync/ActivationListener.cs b/Engine/Source/Programs/UnrealGameSync/UnrealGameSync/ActivationListener.cs
--- a/Engine/Source/Programs/UnrealGameSync/UnrealGameSync/ActivationListener.cs
+++ b/Engine/Source/Programs/UnrealGameSync/UnrealGameSync/ActivationListener.cs
@@ -1,4 +1,4 @@
-﻿// Copyright 1998-2019 Epic Games, Inc. All Rights Reserved.
+﻿// Copyright 1998-2018 Epic Games, Inc. All Rights Reserved.
 
 using System;
 using System.Collections.Generic;
diff --git a/Engine/Source/Programs/UnrealGameSync/UnrealGameSync/ArchiveManifest.cs b/Engine/Source/Programs/UnrealGameSync/UnrealGameSync/ArchiveManifest.cs
--- a/Engine/Source/Programs/UnrealGameSync/UnrealGameSync/ArchiveManifest.cs
+++ b/Engine/Source/Programs/UnrealGameSync/UnrealGameSync/ArchiveManifest.cs
@@ -1,4 +1,4 @@
-﻿// Copyright 1998-2019 Epic Games, Inc. All Rights Reserved.
+﻿// Copyright 1998-2018 Epic Games, Inc. All Rights Reserved.
 
 using System;
 using System.Collections.Generic;
diff --git a/Engine/Source/Programs/UnrealGameSync/UnrealGameSync/BuildStep.cs b/Engine/Source/Programs/UnrealGameSync/UnrealGameSync/BuildStep.cs
--- a/Engine/Source/Programs/UnrealGameSync/UnrealGameSync/BuildStep.cs
+++ b/Engine/Source/Programs/UnrealGameSync/UnrealGameSync/BuildStep.cs
@@ -1,4 +1,4 @@
-﻿// Copyright 1998-2019 Epic Games, Inc. All Rights Reserved.
+// Copyright 1998-2018 Epic Games, Inc. All Rights Reserved.
 
 using System;
 using System.Collections.Generic;
@@ -104,6 +104,21 @@ namespace UnrealGameSync
 			}
 		}
 
+		public bool IsValid()
+		{
+			switch(Type)
+			{
+				case BuildStepType.Compile:
+					return Target != null && Platform != null && Configuration != null;
+				case BuildStepType.Cook:
+					return FileName != null;
+				case BuildStepType.Other:
+					return FileName != null;
+				default:
+					return false;
+			}
+		}
+
 		public static void MergeBuildStepObjects(Dictionary<Guid, ConfigObject> BuildStepObjects, IEnumerable<ConfigObject> ModifyObjects)
 		{
 			foreach(ConfigObject ModifyObject in ModifyObjects)
diff --git a/Engine/Source/Programs/UnrealGameSync/UnrealGameSync/ChildProcess.cs b/Engine/Source/Programs/UnrealGameSync/UnrealGameSync/ChildProcess.cs
--- a/Engine/Source/Programs/UnrealGameSync/UnrealGameSync/ChildProcess.cs
+++ b/Engine/Source/Programs/UnrealGameSync/UnrealGameSync/ChildProcess.cs
@@ -1,4 +1,4 @@
-// Copyright 1998-2019 Epic Games, Inc. All Rights Reserved.
+// Copyright 1998-2018 Epic Games, Inc. All Rights Reserved.
 
 using Microsoft.Win32.SafeHandles;
 using System;
diff --git a/Engine/Source/Programs/UnrealGameSync/UnrealGameSync/ConfigFile.cs b/Engine/Source/Programs/UnrealGameSync/UnrealGameSync/ConfigFile.cs
--- a/Engine/Source/Programs/UnrealGameSync/UnrealGameSync/ConfigFile.cs
+++ b/Engine/Source/Programs/UnrealGameSync/UnrealGameSync/ConfigFile.cs
@@ -1,4 +1,4 @@
-﻿// Copyright 1998-2019 Epic Games, Inc. All Rights Reserved.
+// Copyright 1998-2018 Epic Games, Inc. All Rights Reserved.
 
 using System;
 using System.Collections.Generic;
@@ -300,7 +300,7 @@ namespace UnrealGameSync
 
 		public void SetValues(string Key, string[] Values)
 		{
-			if(Values == null)
+			if(Values == null || Values.Length == 0)
 			{
 				RemoveValue(Key);
 			}
diff --git a/Engine/Source/Programs/UnrealGameSync/UnrealGameSync/Controls/BuildListControl.cs b/Engine/Source/Programs/UnrealGameSync/UnrealGameSync/Controls/BuildListControl.cs
--- a/Engine/Source/Programs/UnrealGameSync/UnrealGameSync/Controls/BuildListControl.cs
+++ b/Engine/Source/Programs/UnrealGameSync/UnrealGameSync/Controls/BuildListControl.cs
@@ -1,4 +1,4 @@
-// Copyright 1998-2019 Epic Games, Inc. All Rights Reserved.
+// Copyright 1998-2018 Epic Games, Inc. All Rights Reserved.
 
 using System;
 using System.Collections.Generic;
diff --git a/Engine/Source/Programs/UnrealGameSync/UnrealGameSync/Controls/CustomListViewControl.cs b/Engine/Source/Programs/UnrealGameSync/UnrealGameSync/Controls/CustomListViewControl.cs
--- a/Engine/Source/Programs/UnrealGameSync/UnrealGameSync/Controls/CustomListViewControl.cs
+++ b/Engine/Source/Programs/UnrealGameSync/UnrealGameSync/Controls/CustomListViewControl.cs
@@ -1,5 +1,3 @@
-// Copyright 1998-2019 Epic Games, Inc. All Rights Reserved.
-
 using System;
 using System.Collections.Generic;
 using System.ComponentModel;
diff --git a/Engine/Source/Programs/UnrealGameSync/UnrealGameSync/Controls/ErrorPanel.cs b/Engine/Source/Programs/UnrealGameSync/UnrealGameSync/Controls/ErrorPanel.cs
--- a/Engine/Source/Programs/UnrealGameSync/UnrealGameSync/Controls/ErrorPanel.cs
+++ b/Engine/Source/Programs/UnrealGameSync/UnrealGameSync/Controls/ErrorPanel.cs
@@ -1,4 +1,4 @@
-﻿// Copyright 1998-2019 Epic Games, Inc. All Rights Reserved.
+﻿// Copyright 1998-2018 Epic Games, Inc. All Rights Reserved.
 
 using System;
 using System.Collections.Generic;
diff --git a/Engine/Source/Programs/UnrealGameSync/UnrealGameSync/Controls/LogControl.cs b/Engine/Source/Programs/UnrealGameSync/UnrealGameSync/Controls/LogControl.cs
--- a/Engine/Source/Programs/UnrealGameSync/UnrealGameSync/Controls/LogControl.cs
+++ b/Engine/Source/Programs/UnrealGameSync/UnrealGameSync/Controls/LogControl.cs
@@ -1,4 +1,4 @@
-﻿// Copyright 1998-2019 Epic Games, Inc. All Rights Reserved.
+﻿// Copyright 1998-2018 Epic Games, Inc. All Rights Reserved.
 
 using System;
 using System.Collections.Generic;
diff --git a/Engine/Source/Programs/UnrealGameSync/UnrealGameSync/Controls/LogSplitContainer.cs b/Engine/Source/Programs/UnrealGameSync/UnrealGameSync/Controls/LogSplitContainer.cs
--- a/Engine/Source/Programs/UnrealGameSync/UnrealGameSync/Controls/LogSplitContainer.cs
+++ b/Engine/Source/Programs/UnrealGameSync/UnrealGameSync/Controls/LogSplitContainer.cs
@@ -1,4 +1,4 @@
-﻿// Copyright 1998-2019 Epic Games, Inc. All Rights Reserved.
+﻿// Copyright 1998-2018 Epic Games, Inc. All Rights Reserved.
 
 using System;
 using System.Collections.Generic;
diff --git a/Engine/Source/Programs/UnrealGameSync/UnrealGameSync/Controls/StatusPanel.cs b/Engine/Source/Programs/UnrealGameSync/UnrealGameSync/Controls/StatusPanel.cs
--- a/Engine/Source/Programs/UnrealGameSync/UnrealGameSync/Controls/StatusPanel.cs
+++ b/Engine/Source/Programs/UnrealGameSync/UnrealGameSync/Controls/StatusPanel.cs
@@ -1,4 +1,4 @@
-// Copyright 1998-2019 Epic Games, Inc. All Rights Reserved.
+// Copyright 1998-2018 Epic Games, Inc. All Rights Reserved.
 
 using System;
 using System.Collections.Generic;
diff --git a/Engine/Source/Programs/UnrealGameSync/UnrealGameSync/Controls/SyncFilterControl.cs b/Engine/Source/Programs/UnrealGameSync/UnrealGameSync/Controls/SyncFilterControl.cs
--- a/Engine/Source/Programs/UnrealGameSync/UnrealGameSync/Controls/SyncFilterControl.cs
+++ b/Engine/Source/Programs/UnrealGameSync/UnrealGameSync/Controls/SyncFilterControl.cs
@@ -1,4 +1,4 @@
-﻿// Copyright 1998-2019 Epic Games, Inc. All Rights Reserved.
+﻿// Copyright 1998-2018 Epic Games, Inc. All Rights Reserved.
 
 using System;
 using System.Collections.Generic;
diff --git a/Engine/Source/Programs/UnrealGameSync/UnrealGameSync/Controls/TabControl.cs b/Engine/Source/Programs/UnrealGameSync/UnrealGameSync/Controls/TabControl.cs
--- a/Engine/Source/Programs/UnrealGameSync/UnrealGameSync/Controls/TabControl.cs
+++ b/Engine/Source/Programs/UnrealGameSync/UnrealGameSync/Controls/TabControl.cs
@@ -1,4 +1,4 @@
-// Copyright 1998-2019 Epic Games, Inc. All Rights Reserved.
+// Copyright 1998-2018 Epic Games, Inc. All Rights Reserved.
 
 using System;
 using System.Collections.Generic;
diff --git a/Engine/Source/Programs/UnrealGameSync/UnrealGameSync/Controls/TextBoxWithCueBanner.cs b/Engine/Source/Programs/UnrealGameSync/UnrealGameSync/Controls/TextBoxWithCueBanner.cs
--- a/Engine/Source/Programs/UnrealGameSync/UnrealGameSync/Controls/TextBoxWithCueBanner.cs
+++ b/Engine/Source/Programs/UnrealGameSync/UnrealGameSync/Controls/TextBoxWithCueBanner.cs
@@ -1,4 +1,4 @@
-// Copyright 1998-2019 Epic Games, Inc. All Rights Reserved.
+// Copyright 1998-2018 Epic Games, Inc. All Rights Reserved.
 
 using System;
 using System.Collections.Generic;
diff --git a/Engine/Source/Programs/UnrealGameSync/UnrealGameSync/Controls/WorkspaceControl.cs b/Engine/Source/Programs/UnrealGameSync/UnrealGameSync/Controls/WorkspaceControl.cs
--- a/Engine/Source/Programs/UnrealGameSync/UnrealGameSync/Controls/WorkspaceControl.cs
+++ b/Engine/Source/Programs/UnrealGameSync/UnrealGameSync/Controls/WorkspaceControl.cs
@@ -1,4 +1,4 @@
-// Copyright 1998-2019 Epic Games, Inc. All Rights Reserved.
+// Copyright 1998-2018 Epic Games, Inc. All Rights Reserved.
 
 using System;
 using System.Collections.Generic;
@@ -69,26 +69,26 @@ namespace UnrealGameSync
 			public int Height;
 			public Color BackgroundColor;
 			public Color HoverBackgroundColor;
-			public object UserData;
+			public Action ClickHandler;
 			public string ToolTip;
 
 			public BadgeInfo(string Label, string Group, Color BadgeColor)
-				: this(Label, Group, null, BadgeColor, BadgeColor)
+				: this(Label, Group, null, BadgeColor, BadgeColor, null)
 			{
 			}
 
-			public BadgeInfo(string Label, string Group, string UniqueId, Color BackgroundColor, Color HoverBackgroundColor, object UserData = null)
+			public BadgeInfo(string Label, string Group, string UniqueId, Color BackgroundColor, Color HoverBackgroundColor, Action ClickHandler)
 			{
 				this.Label = Label;
 				this.Group = Group;
 				this.UniqueId = UniqueId;
 				this.BackgroundColor = BackgroundColor;
 				this.HoverBackgroundColor = HoverBackgroundColor;
-				this.UserData = UserData;
+				this.ClickHandler = ClickHandler;
 			}
 
 			public BadgeInfo(BadgeInfo Other)
-				: this(Other.Label, Other.Group, Other.UniqueId, Other.BackgroundColor, Other.HoverBackgroundColor, Other.UserData)
+				: this(Other.Label, Other.Group, Other.UniqueId, Other.BackgroundColor, Other.HoverBackgroundColor, Other.ClickHandler)
 			{
 				this.Offset = Other.Offset;
 				this.Width = Other.Width;
@@ -206,7 +206,7 @@ namespace UnrealGameSync
 		Font BadgeFont;
 		List<KeyValuePair<string, string>> BadgeNameAndGroupPairs = new List<KeyValuePair<string, string>>();
 		Dictionary<string, Size> BadgeLabelToSize = new Dictionary<string, Size>();
-		List<KeyValuePair<string, BuildData>> ServiceBadges = new List<KeyValuePair<string, BuildData>>();
+		List<KeyValuePair<string, BadgeData>> ServiceBadges = new List<KeyValuePair<string, BadgeData>>();
 
 		string OriginalExecutableFileName;
 
@@ -472,9 +472,9 @@ namespace UnrealGameSync
 			ServiceBadges.Clear();
 			foreach(string ServiceBadgeName in ServiceBadgeNames)
 			{
-				BuildData LatestBuild;
-				EventMonitor.TryGetLatestBuild(ServiceBadgeName, out LatestBuild);
-				ServiceBadges.Add(new KeyValuePair<string, BuildData>(ServiceBadgeName, LatestBuild));
+				BadgeData LatestBuild;
+				EventMonitor.TryGetLatestBadge(ServiceBadgeName, out LatestBuild);
+				ServiceBadges.Add(new KeyValuePair<string, BadgeData>(ServiceBadgeName, LatestBuild));
 			}
 		}
 
@@ -1136,9 +1136,9 @@ namespace UnrealGameSync
 				EventSummary Summary = EventMonitor.GetSummaryForChange(((PerforceChangeSummary)Item.Tag).Number);
 				if(Summary != null)
 				{
-					foreach(BuildData Build in Summary.Builds)
+					foreach(BadgeData Badge in Summary.Badges)
 					{
-						string BadgeSlot = Build.BadgeName;
+						string BadgeSlot = Badge.BadgeName;
 						if(!BadgeNameToGroup.ContainsKey(BadgeSlot))
 						{
 							BadgeNameToGroup.Add(BadgeSlot, "XXXX");
@@ -1148,7 +1148,7 @@ namespace UnrealGameSync
 			}
 
 			// Remove anything that's a service badge
-			foreach(KeyValuePair<string, BuildData> ServiceBadge in ServiceBadges)
+			foreach(KeyValuePair<string, BadgeData> ServiceBadge in ServiceBadges)
 			{
 				BadgeNameToGroup.Remove(ServiceBadge.Key);
 			}
@@ -1282,22 +1282,22 @@ namespace UnrealGameSync
 			ContentBadges.UnionWith(PerforceMonitor.LatestProjectConfigFile.GetValues("Notifications.ContentBadges", new string[0]));
 
 			// Find the most recent build of each type, and the last time it succeeded
-			Dictionary<string, BuildData> TypeToLastBuild = new Dictionary<string,BuildData>();
-			Dictionary<string, BuildData> TypeToLastSucceededBuild = new Dictionary<string,BuildData>();
+			Dictionary<string, BadgeData> TypeToLastBuild = new Dictionary<string,BadgeData>();
+			Dictionary<string, BadgeData> TypeToLastSucceededBuild = new Dictionary<string,BadgeData>();
 			for(int Idx = SortedChangeNumbers.Count - 1; Idx >= 0; Idx--)
 			{
 				EventSummary Summary = EventMonitor.GetSummaryForChange(SortedChangeNumbers[Idx]);
 				if(Summary != null)
 				{
-					foreach(BuildData Build in Summary.Builds)
+					foreach(BadgeData Badge in Summary.Badges)
 					{
-						if(!TypeToLastBuild.ContainsKey(Build.BuildType) && (Build.Result == BuildDataResult.Success || Build.Result == BuildDataResult.Warning || Build.Result == BuildDataResult.Failure))
+						if(!TypeToLastBuild.ContainsKey(Badge.BuildType) && (Badge.Result == BadgeResult.Success || Badge.Result == BadgeResult.Warning || Badge.Result == BadgeResult.Failure))
 						{
-							TypeToLastBuild.Add(Build.BuildType, Build);
+							TypeToLastBuild.Add(Badge.BuildType, Badge);
 						}
-						if(!TypeToLastSucceededBuild.ContainsKey(Build.BuildType) && Build.Result == BuildDataResult.Success)
+						if(!TypeToLastSucceededBuild.ContainsKey(Badge.BuildType) && Badge.Result == BadgeResult.Success)
 						{
-							TypeToLastSucceededBuild.Add(Build.BuildType, Build);
+							TypeToLastSucceededBuild.Add(Badge.BuildType, Badge);
 						}
 					}
 				}
@@ -1305,10 +1305,10 @@ namespace UnrealGameSync
 
 			// Find all the build types that the user needs to be notified about.
 			int RequireNotificationForChange = -1;
-			List<BuildData> NotifyBuilds = new List<BuildData>();
-			foreach(BuildData LastBuild in TypeToLastBuild.Values.OrderBy(x => x.BuildType))
+			List<BadgeData> NotifyBuilds = new List<BadgeData>();
+			foreach(BadgeData LastBuild in TypeToLastBuild.Values.OrderBy(x => x.BuildType))
 			{
-				if(LastBuild.Result == BuildDataResult.Failure || LastBuild.Result == BuildDataResult.Warning)
+				if(LastBuild.Result == BadgeResult.Failure || LastBuild.Result == BadgeResult.Warning)
 				{
 					// Get the last submitted changelist by this user of the correct type
 					int LastChangeByCurrentUserOfType;
@@ -1325,7 +1325,7 @@ namespace UnrealGameSync
 					if(LastChangeByCurrentUserOfType > 0 && LastBuild.ChangeNumber >= LastChangeByCurrentUserOfType)
 					{
 						// And check that there wasn't a successful build after we submitted (if there was, we're in the clear)
-						BuildData LastSuccessfulBuild;
+						BadgeData LastSuccessfulBuild;
 						if(!TypeToLastSucceededBuild.TryGetValue(LastBuild.BuildType, out LastSuccessfulBuild) || LastSuccessfulBuild.ChangeNumber < LastChangeByCurrentUserOfType)
 						{
 							// Add it to the list of notifications
@@ -1357,7 +1357,7 @@ namespace UnrealGameSync
 				}
 
 				// Show the balloon tooltip
-				if(NotifyBuilds.Any(x => x.Result == BuildDataResult.Failure))
+				if(NotifyBuilds.Any(x => x.Result == BadgeResult.Failure))
 				{
 					string Title = String.Format("{0} Errors", PlatformList.ToString());
 					string Message = String.Format("CIS failed after your last submitted changelist ({0}).", RequireNotificationForChange);
@@ -1375,7 +1375,7 @@ namespace UnrealGameSync
 				NotificationWindow.OnMoreInformation = () => { Owner.ShowAndActivate(); SelectChange(HighlightChange); };
 				
 				// Don't show messages for this change again
-				foreach(BuildData NotifyBuild in NotifyBuilds)
+				foreach(BadgeData NotifyBuild in NotifyBuilds)
 				{
 					NotifiedBuildTypeToChangeNumber[NotifyBuild.BuildType] = RequireNotificationForChange;
 				}
@@ -1833,6 +1833,7 @@ namespace UnrealGameSync
 						string Color = BadgeDefinitionObject.GetValue("Color", "#909090");
 						string HoverColor = BadgeDefinitionObject.GetValue("HoverColor", "#b0b0b0");
 						string Url = BadgeDefinitionObject.GetValue("Url", null);
+						string Arguments = BadgeDefinitionObject.GetValue("Arguments", null);
 						if(!String.IsNullOrEmpty(Name) && !String.IsNullOrEmpty(Pattern))
 						{
 							foreach(Match MatchResult in Regex.Matches(Description, Pattern, RegexOptions.Multiline))
@@ -1841,7 +1842,25 @@ namespace UnrealGameSync
 								Color HoverBadgeColor = System.Drawing.ColorTranslator.FromHtml(HoverColor);
 
 								string UniqueId = String.IsNullOrEmpty(Url)? null : String.Format("Description:{0}:{1}", Change.Number, Badges.Count);
-								Badges.Add(new BadgeInfo(ReplaceRegexMatches(Name, MatchResult), Group, UniqueId, BadgeColor, HoverBadgeColor, ReplaceRegexMatches(Url, MatchResult)));
+
+								string ExpandedUrl = ReplaceRegexMatches(Url, MatchResult);
+								string ExpandedArguments = ReplaceRegexMatches(Arguments, MatchResult);
+
+								Action ClickHandler;
+								if(String.IsNullOrEmpty(ExpandedUrl))
+								{
+									ClickHandler = null;
+								}
+								else if(String.IsNullOrEmpty(ExpandedArguments))
+								{
+									ClickHandler = () => Process.Start(ExpandedUrl);
+								}
+								else
+								{
+									ClickHandler = () => Process.Start(ExpandedUrl, ExpandedArguments);
+								}
+
+								Badges.Add(new BadgeInfo(ReplaceRegexMatches(Name, MatchResult), Group, UniqueId, BadgeColor, HoverBadgeColor, ClickHandler));
 							}
 						}
 					}
@@ -1961,28 +1980,28 @@ namespace UnrealGameSync
 		{
 			List<BadgeInfo> Badges = new List<BadgeInfo>();
 
-			if(Summary != null && Summary.Builds.Count > 0)
+			if(Summary != null && Summary.Badges.Count > 0)
 			{
 				// Create a lookup for build data for each badge name
-				Dictionary<string, BuildData> BadgeNameToBuildData = new Dictionary<string, BuildData>();
-				foreach(BuildData Build in Summary.Builds)
+				Dictionary<string, BadgeData> BadgeNameToBuildData = new Dictionary<string, BadgeData>();
+				foreach(BadgeData Badge in Summary.Badges)
 				{
-					BadgeNameToBuildData[Build.BadgeName] = Build;
+					BadgeNameToBuildData[Badge.BadgeName] = Badge;
 				}
 
 				// Add all the badges, sorted by group
 				foreach(KeyValuePair<string, string> BadgeNameAndGroup in BadgeNameAndGroupPairs)
 				{
-					BuildData Build;
-					BadgeNameToBuildData.TryGetValue(BadgeNameAndGroup.Key, out Build);
+					BadgeData BadgeData;
+					BadgeNameToBuildData.TryGetValue(BadgeNameAndGroup.Key, out BadgeData);
 
-					BadgeInfo Badge = CreateBadge(ChangeNumber, BadgeNameAndGroup.Key, BadgeNameAndGroup.Value, Build);
-					if(MaxBuildBadgeChars != -1 && Badge.Label.Length > MaxBuildBadgeChars)
+					BadgeInfo BadgeInfo = CreateBadge(ChangeNumber, BadgeNameAndGroup.Key, BadgeNameAndGroup.Value, BadgeData);
+					if(MaxBuildBadgeChars != -1 && BadgeInfo.Label.Length > MaxBuildBadgeChars)
 					{
-						Badge.ToolTip = Badge.Label;
-						Badge.Label = Badge.Label.Substring(0, MaxBuildBadgeChars);
+						BadgeInfo.ToolTip = BadgeInfo.Label;
+						BadgeInfo.Label = BadgeInfo.Label.Substring(0, MaxBuildBadgeChars);
 					}
-					Badges.Add(Badge);
+					Badges.Add(BadgeInfo);
 				}
 			}
 
@@ -1991,27 +2010,37 @@ namespace UnrealGameSync
 			return Badges;
 		}
 
-		private BadgeInfo CreateBadge(int ChangeNumber, string BadgeName, string BadgeGroup, BuildData Build)
+		private BadgeInfo CreateBadge(int ChangeNumber, string BadgeName, string BadgeGroup, BadgeData BadgeData)
 		{
 			string BadgeLabel = BadgeName;
 			Color BadgeColor = Color.FromArgb(0, Color.White);
 
-			if(Build != null)
+			if(BadgeData != null)
 			{
-				BadgeLabel = Build.BadgeLabel;
-				BadgeColor = GetBuildBadgeColor(Build.Result);
+				BadgeLabel = BadgeData.BadgeLabel;
+				BadgeColor = GetBuildBadgeColor(BadgeData.Result);
 			}
 
 			Color HoverBadgeColor = Color.FromArgb(BadgeColor.A, Math.Min(BadgeColor.R + 32, 255), Math.Min(BadgeColor.G + 32, 255), Math.Min(BadgeColor.B + 32, 255));
 
+			Action ClickHandler;
+			if(BadgeData == null || BadgeData.Url == null)
+			{
+				ClickHandler = null;
+			}
+			else
+			{
+				ClickHandler = () => Process.Start(BadgeData.Url);
+			}
+
 			string UniqueId = String.Format("{0}:{1}", ChangeNumber, BadgeName);
-			return new BadgeInfo(BadgeLabel, BadgeGroup, UniqueId, BadgeColor, HoverBadgeColor, Build);
+			return new BadgeInfo(BadgeLabel, BadgeGroup, UniqueId, BadgeColor, HoverBadgeColor, ClickHandler);
 		}
 
 		private Dictionary<string, List<BadgeInfo>> CreateCustomBadges(int ChangeNumber, EventSummary Summary)
 		{
 			Dictionary<string, List<BadgeInfo>> ColumnNameToBadges = new Dictionary<string, List<BadgeInfo>>();
-			if(Summary != null && Summary.Builds.Count > 0)
+			if(Summary != null && Summary.Badges.Count > 0)
 			{
 				foreach(ColumnHeader CustomColumn in CustomColumns)
 				{
@@ -2023,7 +2052,7 @@ namespace UnrealGameSync
 						string[] BadgeNames = Config.GetValue("Badges", "").Split(new char[]{ ',' }, StringSplitOptions.RemoveEmptyEntries);
 						foreach(string BadgeName in BadgeNames)
 						{
-							BadgeInfo Badge = CreateBadge(ChangeNumber, BadgeName, "XXXX", Summary.Builds.FirstOrDefault(x => x.BadgeName == BadgeName));
+							BadgeInfo Badge = CreateBadge(ChangeNumber, BadgeName, "XXXX", Summary.Badges.FirstOrDefault(x => x.BadgeName == BadgeName));
 							Badges.Add(Badge);
 						}
 
@@ -2036,21 +2065,21 @@ namespace UnrealGameSync
 			return ColumnNameToBadges;
 		}
 
-		private static Color GetBuildBadgeColor(BuildDataResult Result)
+		private static Color GetBuildBadgeColor(BadgeResult Result)
 		{
-			if(Result == BuildDataResult.Starting)
+			if(Result == BadgeResult.Starting)
 			{
 				return Color.FromArgb(128, 192, 255);
 			}
-			else if(Result == BuildDataResult.Warning)
+			else if(Result == BadgeResult.Warning)
 			{
 				return Color.FromArgb(255, 192, 0);
 			}
-			else if(Result == BuildDataResult.Failure)
+			else if(Result == BadgeResult.Failure)
 			{
 				return Color.FromArgb(192, 64, 0);
 			}
-			else if(Result == BuildDataResult.Skipped)
+			else if(Result == BadgeResult.Skipped)
 			{
 				return Color.FromArgb(192, 192, 192);
 			}
@@ -2586,12 +2615,12 @@ namespace UnrealGameSync
 				ProgramsLine.AddText("  |  ");
 				ProgramsLine.AddLink("Windows Explorer", FontStyle.Regular, () => { Process.Start("explorer.exe", String.Format("\"{0}\"", Path.GetDirectoryName(SelectedFileName))); });
 
-				foreach(KeyValuePair<string, BuildData> ServiceBadge in ServiceBadges)
+				foreach(KeyValuePair<string, BadgeData> ServiceBadge in ServiceBadges)
 				{
 					ProgramsLine.AddText("  |  ");
 					if(ServiceBadge.Value == null)
 					{
-						ProgramsLine.AddBadge(ServiceBadge.Key, GetBuildBadgeColor(BuildDataResult.Skipped), null);
+						ProgramsLine.AddBadge(ServiceBadge.Key, GetBuildBadgeColor(BadgeResult.Skipped), null);
 					}
 					else
 					{
@@ -2741,6 +2770,7 @@ namespace UnrealGameSync
 					if(Idx + 1 < Text.Length && Text[Idx + 1] == '[')
 					{
 						ElementText.Append(Text[Idx]);
+						Idx += 2;
 						continue;
 					}
 
@@ -2798,41 +2828,39 @@ namespace UnrealGameSync
 			bool bShownContextMenu = false;
 			if(StreamName != null)
 			{
-				IReadOnlyList<string> OtherStreamFilter = Workspace.ProjectStreamFilter;
-				if(OtherStreamFilter != null)
+				IReadOnlyList<string> OtherStreamNames = Workspace.ProjectStreamFilter;
+				if(OtherStreamNames != null)
 				{
-					IReadOnlyList<string> OtherStreamNames = PerforceMonitor.OtherStreamNames.Where(x => OtherStreamFilter.Any(y => y.Equals(x, StringComparison.InvariantCultureIgnoreCase)) || x.Equals(StreamName, StringComparison.InvariantCultureIgnoreCase)).ToList().AsReadOnly();
-					if(OtherStreamNames.Count > 0)
-					{
-						StreamContextMenu.Items.Clear();
+					StreamContextMenu.Items.Clear();
 
-						foreach (string OtherStreamName in OtherStreamNames.OrderBy(x => x).Where(x => !x.EndsWith("/Dev-Binaries")))
-						{
-							string ThisStreamName = OtherStreamName; // Local for lambda capture
+					ToolStripMenuItem CurrentStreamItem = new ToolStripMenuItem(StreamName, null, new EventHandler((S, E) => SelectStream(StreamName)));
+					CurrentStreamItem.Checked = true;
+					StreamContextMenu.Items.Add(CurrentStreamItem);
+					
+					StreamContextMenu.Items.Add(new ToolStripSeparator());
 
+					foreach (string OtherStreamName in OtherStreamNames.OrderBy(x => x).Where(x => !x.EndsWith("/Dev-Binaries")))
+					{
+						string ThisStreamName = OtherStreamName; // Local for lambda capture
+						if(String.Compare(StreamName, OtherStreamName, StringComparison.InvariantCultureIgnoreCase) != 0)
+						{
 							ToolStripMenuItem Item = new ToolStripMenuItem(ThisStreamName, null, new EventHandler((S, E) => SelectStream(ThisStreamName)));
-							if(String.Compare(StreamName, OtherStreamName, StringComparison.InvariantCultureIgnoreCase) == 0)
-							{
-								Item.Checked = true;
-								StreamContextMenu.Items.Insert(0, Item);
-								StreamContextMenu.Items.Insert(1, new ToolStripSeparator());
-							}
-							else
-							{
-								StreamContextMenu.Items.Add(Item);
-							}
+							StreamContextMenu.Items.Add(Item);
 						}
+					}
 
+					if(StreamContextMenu.Items.Count > 2)
+					{
 						StreamContextMenu.Items.Add(new ToolStripSeparator());
-						StreamContextMenu.Items.Add(new ToolStripMenuItem("Select Other...", null, new EventHandler((S, E) => SelectOtherStreamDialog())));
+					}
 
-						int X = (Bounds.Left + Bounds.Right) / 2 + StreamContextMenu.Bounds.Width / 2;
-						int Y = Bounds.Bottom + 2;
-						StreamContextMenu.Show(StatusPanel, new Point(X, Y), ToolStripDropDownDirection.Left);
+					StreamContextMenu.Items.Add(new ToolStripMenuItem("Select Other...", null, new EventHandler((S, E) => SelectOtherStreamDialog())));
 
-						bShownContextMenu = true;
+					int X = (Bounds.Left + Bounds.Right) / 2 + StreamContextMenu.Bounds.Width / 2;
+					int Y = Bounds.Bottom + 2;
+					StreamContextMenu.Show(StatusPanel, new Point(X, Y), ToolStripDropDownDirection.Left);
 
-					}
+					bShownContextMenu = true;
 				}
 			}
 			if(!bShownContextMenu)
@@ -2989,9 +3017,9 @@ namespace UnrealGameSync
 							foreach (BadgeInfo Badge in LayoutInfo.DescriptionBadges)
 							{
 								Rectangle BadgeBounds = Badge.GetBounds(BuildListLocation);
-								if(BadgeBounds.Contains(Args.Location))
+								if(BadgeBounds.Contains(Args.Location) && Badge.ClickHandler != null)
 								{
-									Process.Start((string)Badge.UserData);
+									Badge.ClickHandler();
 									break;
 								}
 							}
@@ -3006,14 +3034,10 @@ namespace UnrealGameSync
 							Point BuildListLocation = GetBadgeListLocation(LayoutInfo.BuildBadges, HitTest.SubItem.Bounds, HorizontalAlign.Center, VerticalAlignment.Middle);
 							BuildListLocation.X = Math.Max(BuildListLocation.X, HitTest.SubItem.Bounds.Left);
 
-							BadgeInfo Badge = HitTestBadge(Args.Location, LayoutInfo.BuildBadges, BuildListLocation);
-							if(Badge != null)
+							BadgeInfo BadgeInfo = HitTestBadge(Args.Location, LayoutInfo.BuildBadges, BuildListLocation);
+							if(BadgeInfo != null && BadgeInfo.ClickHandler != null)
 							{
-								BuildData Build = (BuildData)Badge.UserData;
-								if(Build != null)
-								{
-									Process.Start(Build.Url);
-								}
+								BadgeInfo.ClickHandler();
 							}
 						}
 					}
@@ -3029,14 +3053,10 @@ namespace UnrealGameSync
 							{
 								Point ListLocation = GetBadgeListLocation(Badges, HitTest.SubItem.Bounds, HorizontalAlign.Center, VerticalAlignment.Middle);
 
-								BadgeInfo Badge = HitTestBadge(Args.Location, Badges, ListLocation);
-								if(Badge != null)
+								BadgeInfo BadgeInfo = HitTestBadge(Args.Location, Badges, ListLocation);
+								if(BadgeInfo != null && BadgeInfo.ClickHandler != null)
 								{
-									BuildData Build = (BuildData)Badge.UserData;
-									if(Build != null)
-									{
-										Process.Start(Build.Url);
-									}
+									BadgeInfo.ClickHandler();
 								}
 							}
 						}
@@ -3183,9 +3203,16 @@ namespace UnrealGameSync
 						BadgeInfo Badge = HitTestBadge(e.Location, LayoutInfo.BuildBadges, BuildListLocation);
 						NewHoverBadgeUniqueId = (Badge != null)? Badge.UniqueId : null;
 
-						if(Badge != null && Badge.ToolTip != null && HoverBadgeUniqueId != NewHoverBadgeUniqueId)
+						if(HoverBadgeUniqueId != NewHoverBadgeUniqueId)
 						{
-							BuildListToolTip.Show(Badge.ToolTip, BuildList, new Point(BuildListLocation.X + Badge.Offset, HitTest.Item.Bounds.Bottom + 2));
+							if(Badge != null && Badge.ToolTip != null && Badge.BackgroundColor.A != 0)
+							{
+								BuildListToolTip.Show(Badge.ToolTip, BuildList, new Point(BuildListLocation.X + Badge.Offset, HitTest.Item.Bounds.Bottom + 2));
+							}
+							else
+							{
+								BuildListToolTip.Hide(BuildList);
+							}
 						}
 					}
 				}
@@ -3957,13 +3984,16 @@ namespace UnrealGameSync
 				List<ConfigObject> ModifiedBuildSteps = new List<ConfigObject>();
 				foreach(BuildStep Step in UserSteps)
 				{
-					ConfigObject DefaultObject;
-					ProjectBuildStepObjects.TryGetValue(Step.UniqueId, out DefaultObject);
-
-					ConfigObject UserConfigObject = Step.ToConfigObject(DefaultObject);
-					if(UserConfigObject != null)
+					if(Step.IsValid())
 					{
-						ModifiedBuildSteps.Add(UserConfigObject);
+						ConfigObject DefaultObject;
+						ProjectBuildStepObjects.TryGetValue(Step.UniqueId, out DefaultObject);
+
+						ConfigObject UserConfigObject = Step.ToConfigObject(DefaultObject);
+						if(UserConfigObject != null && UserConfigObject.Pairs.Any(x => x.Key != "UniqueId"))
+						{
+							ModifiedBuildSteps.Add(UserConfigObject);
+						}
 					}
 				}
 
@@ -4188,7 +4218,7 @@ namespace UnrealGameSync
 
 		private void OptionsContextMenu_PerforceSettings_Click(object sender, EventArgs e)
 		{
-			PerforceSettingsWindow Perforce = new PerforceSettingsWindow(Settings);
+			PerforceSyncSettingsWindow Perforce = new PerforceSyncSettingsWindow(Settings);
 			Perforce.ShowDialog();
 		}
 
diff --git a/Engine/Source/Programs/UnrealGameSync/UnrealGameSync/DeleteFilesTask.cs b/Engine/Source/Programs/UnrealGameSync/UnrealGameSync/DeleteFilesTask.cs
--- a/Engine/Source/Programs/UnrealGameSync/UnrealGameSync/DeleteFilesTask.cs
+++ b/Engine/Source/Programs/UnrealGameSync/UnrealGameSync/DeleteFilesTask.cs
@@ -1,4 +1,4 @@
-// Copyright 1998-2019 Epic Games, Inc. All Rights Reserved.
+// Copyright 1998-2018 Epic Games, Inc. All Rights Reserved.
 
 using System;
 using System.Collections.Generic;
diff --git a/Engine/Source/Programs/UnrealGameSync/UnrealGameSync/DetectProjectSettingsTask.cs b/Engine/Source/Programs/UnrealGameSync/UnrealGameSync/DetectProjectSettingsTask.cs
--- a/Engine/Source/Programs/UnrealGameSync/UnrealGameSync/DetectProjectSettingsTask.cs
+++ b/Engine/Source/Programs/UnrealGameSync/UnrealGameSync/DetectProjectSettingsTask.cs
@@ -1,4 +1,4 @@
-// Copyright 1998-2019 Epic Games, Inc. All Rights Reserved.
+// Copyright 1998-2018 Epic Games, Inc. All Rights Reserved.
 
 using System;
 using System.Collections.Generic;
diff --git a/Engine/Source/Programs/UnrealGameSync/UnrealGameSync/EventMonitor.cs b/Engine/Source/Programs/UnrealGameSync/UnrealGameSync/EventMonitor.cs
--- a/Engine/Source/Programs/UnrealGameSync/UnrealGameSync/EventMonitor.cs
+++ b/Engine/Source/Programs/UnrealGameSync/UnrealGameSync/EventMonitor.cs
@@ -1,4 +1,4 @@
-// Copyright 1998-2019 Epic Games, Inc. All Rights Reserved.
+// Copyright 1998-2018 Epic Games, Inc. All Rights Reserved.
 
 using System;
 using System.Collections.Concurrent;
@@ -34,12 +34,14 @@ namespace UnrealGameSync
 		Investigating,
 		Resolved,
 	}
+
 	class LatestData
 	{
 		public long LastEventId { get; set; }
 		public long LastCommentId { get; set; }
 		public long LastBuildId { get; set; }
 	}
+
 	class EventData
 	{
 		public long Id { get; set; }
@@ -58,7 +60,7 @@ namespace UnrealGameSync
 		public string Project { get; set; }
 	}
 
-	enum BuildDataResult
+	enum BadgeResult
 	{
 		Starting,
 		Failure,
@@ -67,23 +69,23 @@ namespace UnrealGameSync
 		Skipped,
 	}
 
-	class BuildData
+	class BadgeData
 	{
 		public long Id { get; set; }
 		public int ChangeNumber { get; set; }
 		public string BuildType { get; set; }
-		public BuildDataResult Result { get; set; }
+		public BadgeResult Result { get; set; }
 		public string Url { get; set; }
 		public string Project { get; set; }
 
 		public bool IsSuccess
 		{
-			get { return Result == BuildDataResult.Success || Result == BuildDataResult.Warning; }
+			get { return Result == BadgeResult.Success || Result == BadgeResult.Warning; }
 		}
 
 		public bool IsFailure
 		{
-			get { return Result == BuildDataResult.Failure; }
+			get { return Result == BadgeResult.Failure; }
 		}
 
 		public string BadgeName
@@ -135,7 +137,7 @@ namespace UnrealGameSync
 		public List<EventData> Reviews = new List<EventData>();
 		public List<string> CurrentUsers = new List<string>();
 		public EventData LastStarReview;
-		public List<BuildData> Builds = new List<BuildData>();
+		public List<BadgeData> Badges = new List<BadgeData>();
 		public List<CommentData> Comments = new List<CommentData>();
 	}
 
@@ -150,10 +152,10 @@ namespace UnrealGameSync
 		ConcurrentQueue<EventData> IncomingEvents = new ConcurrentQueue<EventData>();
 		ConcurrentQueue<CommentData> OutgoingComments = new ConcurrentQueue<CommentData>();
 		ConcurrentQueue<CommentData> IncomingComments = new ConcurrentQueue<CommentData>();
-		ConcurrentQueue<BuildData> IncomingBuilds = new ConcurrentQueue<BuildData>();
+		ConcurrentQueue<BadgeData> IncomingBadges = new ConcurrentQueue<BadgeData>();
 		Dictionary<int, EventSummary> ChangeNumberToSummary = new Dictionary<int, EventSummary>();
 		Dictionary<string, EventData> UserNameToLastSyncEvent = new Dictionary<string, EventData>(StringComparer.InvariantCultureIgnoreCase);
-		Dictionary<string, BuildData> BadgeNameToLatestBuild = new Dictionary<string, BuildData>();
+		Dictionary<string, BadgeData> BadgeNameToLatestData = new Dictionary<string, BadgeData>();
 		BoundedLogWriter LogWriter;
 		bool bDisposing;
 		LatestData LatestIds;
@@ -258,10 +260,10 @@ namespace UnrealGameSync
 				ApplyEventUpdate(Event);
 			}
 
-			BuildData Build;
-			while(IncomingBuilds.TryDequeue(out Build))
+			BadgeData Badge;
+			while(IncomingBadges.TryDequeue(out Badge))
 			{
-				ApplyBuildUpdate(Build);
+				ApplyBadgeUpdate(Badge);
 			}
 
 			CommentData Comment;
@@ -320,7 +322,7 @@ namespace UnrealGameSync
 
 				// Add the new review, and find the new verdict for this change
 				Summary.Reviews.Add(Event);
-				Summary.Verdict = GetVerdict(Summary.Reviews, Summary.Builds);
+				Summary.Verdict = GetVerdict(Summary.Reviews, Summary.Badges);
 			}
 			else
 			{
@@ -328,16 +330,16 @@ namespace UnrealGameSync
 			}
 		}
 
-		void ApplyBuildUpdate(BuildData Build)
+		void ApplyBadgeUpdate(BadgeData Badge)
 		{
-			EventSummary Summary = FindOrAddSummary(Build.ChangeNumber);
+			EventSummary Summary = FindOrAddSummary(Badge.ChangeNumber);
 
-			BuildData ExistingBuild = Summary.Builds.Find(x => x.ChangeNumber == Build.ChangeNumber && x.BuildType == Build.BuildType);
-			if(ExistingBuild != null)
+			BadgeData ExistingBadge = Summary.Badges.Find(x => x.ChangeNumber == Badge.ChangeNumber && x.BuildType == Badge.BuildType);
+			if(ExistingBadge != null)
 			{
-				if(ExistingBuild.Id <= Build.Id)
+				if(ExistingBadge.Id <= Badge.Id)
 				{
-					Summary.Builds.Remove(ExistingBuild);
+					Summary.Badges.Remove(ExistingBadge);
 				}
 				else
 				{
@@ -345,13 +347,13 @@ namespace UnrealGameSync
 				}
 			}
 
-			Summary.Builds.Add(Build);
-			Summary.Verdict = GetVerdict(Summary.Reviews, Summary.Builds);
+			Summary.Badges.Add(Badge);
+			Summary.Verdict = GetVerdict(Summary.Reviews, Summary.Badges);
 
-			BuildData LatestBuild;
-			if(!BadgeNameToLatestBuild.TryGetValue(Build.BadgeName, out LatestBuild) || Build.ChangeNumber > LatestBuild.ChangeNumber || (Build.ChangeNumber == LatestBuild.ChangeNumber && Build.Id > LatestBuild.Id))
+			BadgeData LatestBadge;
+			if(!BadgeNameToLatestData.TryGetValue(Badge.BadgeName, out LatestBadge) || Badge.ChangeNumber > LatestBadge.ChangeNumber || (Badge.ChangeNumber == LatestBadge.ChangeNumber && Badge.Id > LatestBadge.Id))
 			{
-				BadgeNameToLatestBuild[Build.BadgeName] = Build;
+				BadgeNameToLatestData[Badge.BadgeName] = Badge;
 			}
 		}
 
@@ -391,7 +393,7 @@ namespace UnrealGameSync
 			return true;
 		}
 
-		static ReviewVerdict GetVerdict(IEnumerable<EventData> Events, IEnumerable<BuildData> Builds)
+		static ReviewVerdict GetVerdict(IEnumerable<EventData> Events, IEnumerable<BadgeData> Badges)
 		{
 			int NumPositiveReviews = Events.Count(x => x.Type == EventType.Good);
 			int NumNegativeReviews = Events.Count(x => x.Type == EventType.Bad);
@@ -407,11 +409,11 @@ namespace UnrealGameSync
 				return GetVerdict(NumCompiles, NumFailedCompiles);
 			}
 
-			int NumBuilds = Builds.Count(x => x.BuildType == "Editor" && x.IsSuccess);
-			int NumFailedBuilds = Builds.Count(x => x.BuildType == "Editor" && x.IsFailure);
-			if(NumBuilds > 0 || NumFailedBuilds > 0)
+			int NumBadges = Badges.Count(x => x.BuildType == "Editor" && x.IsSuccess);
+			int NumFailedBadges = Badges.Count(x => x.BuildType == "Editor" && x.IsFailure);
+			if(NumBadges > 0 || NumFailedBadges > 0)
 			{
-				return GetVerdict(NumBuilds, NumFailedBuilds);
+				return GetVerdict(NumBadges, NumFailedBadges);
 			}
 
 			return ReviewVerdict.Unknown;
@@ -460,10 +462,13 @@ namespace UnrealGameSync
 		{
 			EventData Event = null;
 			CommentData Comment = null;
-			while(!bDisposing)
+			bool bUpdateThrottledRequests = true;
+			double RequestThrottle = 90; // seconds to wait for throttled request;
+			Stopwatch Timer = Stopwatch.StartNew();
+			while (!bDisposing)
 			{
 				// If there's no connection string, just empty out the queue
-				if(ApiUrl != null)
+				if (ApiUrl != null)
 				{
 					// Post all the reviews to the database. We don't send them out of order, so keep the review outside the queue until the next update if it fails
 					while(Event != null || OutgoingEvents.TryDequeue(out Event))
@@ -479,18 +484,24 @@ namespace UnrealGameSync
 						Comment = null;
 					}
 
-					// Read all the new reviews
-					ReadEventsFromBackend();
+					if(Timer.Elapsed > TimeSpan.FromSeconds(RequestThrottle))
+					{
+						bUpdateThrottledRequests = true;
+						Timer.Reset();
+					}
+
+					// Read all the new reviews, pass whether or not to fire the throttled requests
+					ReadEventsFromBackend(bUpdateThrottledRequests);
 
 					// Send a notification that we're ready to update
-					if((IncomingEvents.Count > 0 || IncomingBuilds.Count > 0 || IncomingComments.Count > 0) && OnUpdatesReady != null)
+					if((IncomingEvents.Count > 0 || IncomingBadges.Count > 0 || IncomingComments.Count > 0) && OnUpdatesReady != null)
 					{
 						OnUpdatesReady();
 					}
 				}
 
 				// Wait for something else to do
-				RefreshEvent.WaitOne(30 * 1000);
+				bUpdateThrottledRequests = RefreshEvent.WaitOne(30 * 1000);
 			}
 		}
 
@@ -527,7 +538,7 @@ namespace UnrealGameSync
 			}
 		}
 
-		bool ReadEventsFromBackend()
+		bool ReadEventsFromBackend(bool bFireThrottledRequests)
 		{
 			try
 			{
@@ -546,33 +557,39 @@ namespace UnrealGameSync
 				}
 
 				//////////////
-				/// Reviews 
+				/// Bulids
 				//////////////
-				List<EventData> Events = RESTApi.GET<List<EventData>>(ApiUrl, "event", string.Format("project={0}", Project), string.Format("lasteventid={0}", LatestIds.LastEventId));
-				foreach(EventData Review in Events)
+				List<BadgeData> Builds = RESTApi.GET<List<BadgeData>>(ApiUrl, "build", string.Format("project={0}", Project), string.Format("lastbuildid={0}", LatestIds.LastBuildId));
+				foreach (BadgeData Build in Builds)
 				{
-					IncomingEvents.Enqueue(Review);
-					LatestIds.LastEventId = Math.Max(LatestIds.LastEventId, Review.Id);
+					IncomingBadges.Enqueue(Build);
+					LatestIds.LastBuildId = Math.Max(LatestIds.LastBuildId, Build.Id);
 				}
 
-				//////////////
-				/// Comments 
-				//////////////
-				List<CommentData> Comments = RESTApi.GET<List<CommentData>>(ApiUrl, "comment", string.Format("project={0}", Project), string.Format("lastcommentid={0}", LatestIds.LastCommentId));
-				foreach (CommentData Comment in Comments)
+				//////////////////////////
+				/// Throttled Requests
+				//////////////////////////
+				if (bFireThrottledRequests)
 				{
-					IncomingComments.Enqueue(Comment);
-					LatestIds.LastCommentId = Math.Max(LatestIds.LastCommentId, Comment.Id);
-				}
+					//////////////
+					/// Reviews 
+					//////////////
+					List<EventData> Events = RESTApi.GET<List<EventData>>(ApiUrl, "event", string.Format("project={0}", Project), string.Format("lasteventid={0}", LatestIds.LastEventId));
+					foreach (EventData Review in Events)
+					{
+						IncomingEvents.Enqueue(Review);
+						LatestIds.LastEventId = Math.Max(LatestIds.LastEventId, Review.Id);
+					}
 
-				//////////////
-				/// Bulids
-				//////////////
-				List<BuildData> Builds = RESTApi.GET<List<BuildData>>(ApiUrl, "build", string.Format("project={0}", Project), string.Format("lastbuildid={0}", LatestIds.LastBuildId));
-				foreach (BuildData Build in Builds)
-				{
-					IncomingBuilds.Enqueue(Build);
-					LatestIds.LastBuildId = Math.Max(LatestIds.LastBuildId, Build.Id);
+					//////////////
+					/// Comments 
+					//////////////
+					List<CommentData> Comments = RESTApi.GET<List<CommentData>>(ApiUrl, "comment", string.Format("project={0}", Project), string.Format("lastcommentid={0}", LatestIds.LastCommentId));
+					foreach (CommentData Comment in Comments)
+					{
+						IncomingComments.Enqueue(Comment);
+						LatestIds.LastCommentId = Math.Max(LatestIds.LastCommentId, Comment.Id);
+					}
 				}
 
 				LastStatusMessage = String.Format("Last update took {0}ms", Timer.ElapsedMilliseconds);
@@ -671,9 +688,9 @@ namespace UnrealGameSync
 			return Summary;
 		}
 
-		public bool TryGetLatestBuild(string BuildType, out BuildData BuildData)
+		public bool TryGetLatestBadge(string BuildType, out BadgeData BadgeData)
 		{
-			return BadgeNameToLatestBuild.TryGetValue(BuildType, out BuildData);
+			return BadgeNameToLatestData.TryGetValue(BuildType, out BadgeData);
 		}
 
 		public static bool IsReview(EventType Type)
diff --git a/Engine/Source/Programs/UnrealGameSync/UnrealGameSync/FileFilter.cs b/Engine/Source/Programs/UnrealGameSync/UnrealGameSync/FileFilter.cs
--- a/Engine/Source/Programs/UnrealGameSync/UnrealGameSync/FileFilter.cs
+++ b/Engine/Source/Programs/UnrealGameSync/UnrealGameSync/FileFilter.cs
@@ -1,4 +1,4 @@
-﻿// Copyright 1998-2019 Epic Games, Inc. All Rights Reserved.
+﻿// Copyright 1998-2018 Epic Games, Inc. All Rights Reserved.
 
 using System;
 using System.Collections.Generic;
diff --git a/Engine/Source/Programs/UnrealGameSync/UnrealGameSync/FindFoldersToCleanTask.cs b/Engine/Source/Programs/UnrealGameSync/UnrealGameSync/FindFoldersToCleanTask.cs
--- a/Engine/Source/Programs/UnrealGameSync/UnrealGameSync/FindFoldersToCleanTask.cs
+++ b/Engine/Source/Programs/UnrealGameSync/UnrealGameSync/FindFoldersToCleanTask.cs
@@ -1,4 +1,4 @@
-// Copyright 1998-2019 Epic Games, Inc. All Rights Reserved.
+// Copyright 1998-2018 Epic Games, Inc. All Rights Reserved.
 
 using System;
 using System.Collections.Generic;
diff --git a/Engine/Source/Programs/UnrealGameSync/UnrealGameSync/Forms/ArgumentsWindow.cs b/Engine/Source/Programs/UnrealGameSync/UnrealGameSync/Forms/ArgumentsWindow.cs
--- a/Engine/Source/Programs/UnrealGameSync/UnrealGameSync/Forms/ArgumentsWindow.cs
+++ b/Engine/Source/Programs/UnrealGameSync/UnrealGameSync/Forms/ArgumentsWindow.cs
@@ -1,4 +1,4 @@
-﻿// Copyright 1998-2019 Epic Games, Inc. All Rights Reserved.
+﻿// Copyright 1998-2018 Epic Games, Inc. All Rights Reserved.
 
 using System;
 using System.Collections.Generic;
diff --git a/Engine/Source/Programs/UnrealGameSync/UnrealGameSync/Forms/BuildStepWindow.cs b/Engine/Source/Programs/UnrealGameSync/UnrealGameSync/Forms/BuildStepWindow.cs
--- a/Engine/Source/Programs/UnrealGameSync/UnrealGameSync/Forms/BuildStepWindow.cs
+++ b/Engine/Source/Programs/UnrealGameSync/UnrealGameSync/Forms/BuildStepWindow.cs
@@ -1,4 +1,4 @@
-﻿// Copyright 1998-2019 Epic Games, Inc. All Rights Reserved.
+﻿// Copyright 1998-2018 Epic Games, Inc. All Rights Reserved.
 
 using System;
 using System.Collections.Generic;
diff --git a/Engine/Source/Programs/UnrealGameSync/UnrealGameSync/Forms/ChangelistWindow.cs b/Engine/Source/Programs/UnrealGameSync/UnrealGameSync/Forms/ChangelistWindow.cs
--- a/Engine/Source/Programs/UnrealGameSync/UnrealGameSync/Forms/ChangelistWindow.cs
+++ b/Engine/Source/Programs/UnrealGameSync/UnrealGameSync/Forms/ChangelistWindow.cs
@@ -1,4 +1,4 @@
-// Copyright 1998-2019 Epic Games, Inc. All Rights Reserved.
+// Copyright 1998-2018 Epic Games, Inc. All Rights Reserved.
 
 using System;
 using System.Collections.Generic;
diff --git a/Engine/Source/Programs/UnrealGameSync/UnrealGameSync/Forms/CleanWorkspaceWindow.cs b/Engine/Source/Programs/UnrealGameSync/UnrealGameSync/Forms/CleanWorkspaceWindow.cs
--- a/Engine/Source/Programs/UnrealGameSync/UnrealGameSync/Forms/CleanWorkspaceWindow.cs
+++ b/Engine/Source/Programs/UnrealGameSync/UnrealGameSync/Forms/CleanWorkspaceWindow.cs
@@ -1,4 +1,4 @@
-// Copyright 1998-2019 Epic Games, Inc. All Rights Reserved.
+// Copyright 1998-2018 Epic Games, Inc. All Rights Reserved.
 
 using System;
 using System.Collections.Generic;
diff --git a/Engine/Source/Programs/UnrealGameSync/UnrealGameSync/Forms/ClobberWindow.cs b/Engine/Source/Programs/UnrealGameSync/UnrealGameSync/Forms/ClobberWindow.cs
--- a/Engine/Source/Programs/UnrealGameSync/UnrealGameSync/Forms/ClobberWindow.cs
+++ b/Engine/Source/Programs/UnrealGameSync/UnrealGameSync/Forms/ClobberWindow.cs
@@ -1,4 +1,4 @@
-﻿// Copyright 1998-2019 Epic Games, Inc. All Rights Reserved.
+﻿// Copyright 1998-2018 Epic Games, Inc. All Rights Reserved.
 
 using System;
 using System.Collections.Generic;
diff --git a/Engine/Source/Programs/UnrealGameSync/UnrealGameSync/Forms/ConnectWindow.cs b/Engine/Source/Programs/UnrealGameSync/UnrealGameSync/Forms/ConnectWindow.cs
--- a/Engine/Source/Programs/UnrealGameSync/UnrealGameSync/Forms/ConnectWindow.cs
+++ b/Engine/Source/Programs/UnrealGameSync/UnrealGameSync/Forms/ConnectWindow.cs
@@ -1,4 +1,4 @@
-﻿// Copyright 1998-2019 Epic Games, Inc. All Rights Reserved.
+﻿// Copyright 1998-2018 Epic Games, Inc. All Rights Reserved.
 
 using System;
 using System.Collections.Generic;
diff --git a/Engine/Source/Programs/UnrealGameSync/UnrealGameSync/Forms/DeleteWindow.cs b/Engine/Source/Programs/UnrealGameSync/UnrealGameSync/Forms/DeleteWindow.cs
--- a/Engine/Source/Programs/UnrealGameSync/UnrealGameSync/Forms/DeleteWindow.cs
+++ b/Engine/Source/Programs/UnrealGameSync/UnrealGameSync/Forms/DeleteWindow.cs
@@ -1,4 +1,4 @@
-﻿// Copyright 1998-2019 Epic Games, Inc. All Rights Reserved.
+﻿// Copyright 1998-2018 Epic Games, Inc. All Rights Reserved.
 
 using System;
 using System.Collections.Generic;
diff --git a/Engine/Source/Programs/UnrealGameSync/UnrealGameSync/Forms/DiagnosticsWindow.cs b/Engine/Source/Programs/UnrealGameSync/UnrealGameSync/Forms/DiagnosticsWindow.cs
--- a/Engine/Source/Programs/UnrealGameSync/UnrealGameSync/Forms/DiagnosticsWindow.cs
+++ b/Engine/Source/Programs/UnrealGameSync/UnrealGameSync/Forms/DiagnosticsWindow.cs
@@ -1,4 +1,4 @@
-// Copyright 1998-2019 Epic Games, Inc. All Rights Reserved.
+// Copyright 1998-2018 Epic Games, Inc. All Rights Reserved.
 
 using System;
 using System.Collections.Generic;
diff --git a/Engine/Source/Programs/UnrealGameSync/UnrealGameSync/Forms/FailedToDeleteWindow.cs b/Engine/Source/Programs/UnrealGameSync/UnrealGameSync/Forms/FailedToDeleteWindow.cs
--- a/Engine/Source/Programs/UnrealGameSync/UnrealGameSync/Forms/FailedToDeleteWindow.cs
+++ b/Engine/Source/Programs/UnrealGameSync/UnrealGameSync/Forms/FailedToDeleteWindow.cs
@@ -1,4 +1,4 @@
-﻿// Copyright 1998-2019 Epic Games, Inc. All Rights Reserved.
+﻿// Copyright 1998-2018 Epic Games, Inc. All Rights Reserved.
 
 using System;
 using System.Collections.Generic;
diff --git a/Engine/Source/Programs/UnrealGameSync/UnrealGameSync/Forms/LeaveCommentWindow.cs b/Engine/Source/Programs/UnrealGameSync/UnrealGameSync/Forms/LeaveCommentWindow.cs
--- a/Engine/Source/Programs/UnrealGameSync/UnrealGameSync/Forms/LeaveCommentWindow.cs
+++ b/Engine/Source/Programs/UnrealGameSync/UnrealGameSync/Forms/LeaveCommentWindow.cs
@@ -1,4 +1,4 @@
-﻿// Copyright 1998-2019 Epic Games, Inc. All Rights Reserved.
+﻿// Copyright 1998-2018 Epic Games, Inc. All Rights Reserved.
 
 using System;
 using System.Collections.Generic;
diff --git a/Engine/Source/Programs/UnrealGameSync/UnrealGameSync/Forms/MainWindow.cs b/Engine/Source/Programs/UnrealGameSync/UnrealGameSync/Forms/MainWindow.cs
--- a/Engine/Source/Programs/UnrealGameSync/UnrealGameSync/Forms/MainWindow.cs
+++ b/Engine/Source/Programs/UnrealGameSync/UnrealGameSync/Forms/MainWindow.cs
@@ -1,4 +1,4 @@
-// Copyright 1998-2019 Epic Games, Inc. All Rights Reserved.
+// Copyright 1998-2018 Epic Games, Inc. All Rights Reserved.
 
 using System;
 using System.Collections.Generic;
diff --git a/Engine/Source/Programs/UnrealGameSync/UnrealGameSync/Forms/ModalTaskWindow.cs b/Engine/Source/Programs/UnrealGameSync/UnrealGameSync/Forms/ModalTaskWindow.cs
--- a/Engine/Source/Programs/UnrealGameSync/UnrealGameSync/Forms/ModalTaskWindow.cs
+++ b/Engine/Source/Programs/UnrealGameSync/UnrealGameSync/Forms/ModalTaskWindow.cs
@@ -1,4 +1,4 @@
-// Copyright 1998-2019 Epic Games, Inc. All Rights Reserved.
+// Copyright 1998-2018 Epic Games, Inc. All Rights Reserved.
 
 using System;
 using System.Collections.Generic;
diff --git a/Engine/Source/Programs/UnrealGameSync/UnrealGameSync/Forms/ModifyBuildStepsWindow.cs b/Engine/Source/Programs/UnrealGameSync/UnrealGameSync/Forms/ModifyBuildStepsWindow.cs
--- a/Engine/Source/Programs/UnrealGameSync/UnrealGameSync/Forms/ModifyBuildStepsWindow.cs
+++ b/Engine/Source/Programs/UnrealGameSync/UnrealGameSync/Forms/ModifyBuildStepsWindow.cs
@@ -1,4 +1,4 @@
-// Copyright 1998-2019 Epic Games, Inc. All Rights Reserved.
+// Copyright 1998-2018 Epic Games, Inc. All Rights Reserved.
 
 using System;
 using System.Collections.Generic;
diff --git a/Engine/Source/Programs/UnrealGameSync/UnrealGameSync/Forms/NewWorkspaceWindow.cs b/Engine/Source/Programs/UnrealGameSync/UnrealGameSync/Forms/NewWorkspaceWindow.cs
--- a/Engine/Source/Programs/UnrealGameSync/UnrealGameSync/Forms/NewWorkspaceWindow.cs
+++ b/Engine/Source/Programs/UnrealGameSync/UnrealGameSync/Forms/NewWorkspaceWindow.cs
@@ -1,4 +1,4 @@
-// Copyright 1998-2019 Epic Games, Inc. All Rights Reserved.
+// Copyright 1998-2018 Epic Games, Inc. All Rights Reserved.
 
 using System;
 using System.Collections.Generic;
diff --git a/Engine/Source/Programs/UnrealGameSync/UnrealGameSync/Forms/NotificationWindow.cs b/Engine/Source/Programs/UnrealGameSync/UnrealGameSync/Forms/NotificationWindow.cs
--- a/Engine/Source/Programs/UnrealGameSync/UnrealGameSync/Forms/NotificationWindow.cs
+++ b/Engine/Source/Programs/UnrealGameSync/UnrealGameSync/Forms/NotificationWindow.cs
@@ -1,4 +1,4 @@
-﻿// Copyright 1998-2019 Epic Games, Inc. All Rights Reserved.
+﻿// Copyright 1998-2018 Epic Games, Inc. All Rights Reserved.
 
 using System;
 using System.Collections.Generic;
diff --git a/Engine/Source/Programs/UnrealGameSync/UnrealGameSync/Forms/OpenProjectWindow.cs b/Engine/Source/Programs/UnrealGameSync/UnrealGameSync/Forms/OpenProjectWindow.cs
--- a/Engine/Source/Programs/UnrealGameSync/UnrealGameSync/Forms/OpenProjectWindow.cs
+++ b/Engine/Source/Programs/UnrealGameSync/UnrealGameSync/Forms/OpenProjectWindow.cs
@@ -1,4 +1,4 @@
-// Copyright 1998-2019 Epic Games, Inc. All Rights Reserved.
+// Copyright 1998-2018 Epic Games, Inc. All Rights Reserved.
 
 using System;
 using System.Collections.Generic;
@@ -325,12 +325,15 @@ namespace UnrealGameSync
 			Dialog.Filter = "Project files (*.uproject)|*.uproject|Project directory lists (*.uprojectdirs)|*.uprojectdirs|All supported files (*.uproject;*.uprojectdirs)|*.uproject;*.uprojectdirs|All files (*.*)|*.*" ;
 			Dialog.FilterIndex = Settings.FilterIndex;
 			
-			try
-			{
-				Dialog.InitialDirectory = Path.GetDirectoryName(LocalFileTextBox.Text);
-			}
-			catch
+			if(!String.IsNullOrEmpty(LocalFileTextBox.Text))
 			{
+				try
+				{
+					Dialog.InitialDirectory = Path.GetDirectoryName(LocalFileTextBox.Text);
+				}
+				catch
+				{
+				}
 			}
 
 			if(Dialog.ShowDialog(this) == System.Windows.Forms.DialogResult.OK)
diff --git a/Engine/Source/Programs/UnrealGameSync/UnrealGameSync/Forms/PasswordWindow.cs b/Engine/Source/Programs/UnrealGameSync/UnrealGameSync/Forms/PasswordWindow.cs
--- a/Engine/Source/Programs/UnrealGameSync/UnrealGameSync/Forms/PasswordWindow.cs
+++ b/Engine/Source/Programs/UnrealGameSync/UnrealGameSync/Forms/PasswordWindow.cs
@@ -1,4 +1,4 @@
-﻿// Copyright 1998-2019 Epic Games, Inc. All Rights Reserved.
+﻿// Copyright 1998-2018 Epic Games, Inc. All Rights Reserved.
 
 using System;
 using System.Collections.Generic;
diff --git a/Engine/Source/Programs/UnrealGameSync/UnrealGameSync/Forms/PerforceSettingsWindow.cs b/Engine/Source/Programs/UnrealGameSync/UnrealGameSync/Forms/PerforceSettingsWindow.cs
deleted file mode 100644
diff --git a/Engine/Source/Programs/UnrealGameSync/UnrealGameSync/Forms/PerforceSyncSettingsWindow.cs b/Engine/Source/Programs/UnrealGameSync/UnrealGameSync/Forms/PerforceSyncSettingsWindow.cs
new file mode 100644
--- /dev/null
+++ b/Engine/Source/Programs/UnrealGameSync/UnrealGameSync/Forms/PerforceSyncSettingsWindow.cs
@@ -0,0 +1,71 @@
+// Copyright 1998-2018 Epic Games, Inc. All Rights Reserved.
+
+using System;
+using System.Collections.Generic;
+using System.ComponentModel;
+using System.Data;
+using System.Drawing;
+using System.Linq;
+using System.Text;
+using System.Threading.Tasks;
+using System.Windows.Forms;
+
+namespace UnrealGameSync
+{
+	partial class PerforceSyncSettingsWindow : Form
+	{
+		UserSettings Settings;
+
+		public PerforceSyncSettingsWindow(UserSettings Settings)
+		{
+			this.Settings = Settings;
+			InitializeComponent();
+		}
+
+		private void PerforceSettingsWindow_Load(object sender, EventArgs e)
+		{
+			PerforceSyncOptions SyncOptions = Settings.SyncOptions;
+			NumRetriesTextBox.Text = (SyncOptions.NumRetries > 0)? SyncOptions.NumRetries.ToString() : "";
+			NumThreadsTextBox.Text = (SyncOptions.NumThreads > 0)? SyncOptions.NumThreads.ToString() : "";
+			TcpBufferSizeText.Text = (SyncOptions.TcpBufferSize > 0)? SyncOptions.TcpBufferSize.ToString() : "";
+		}
+
+		private void OkButton_Click(object sender, EventArgs e)
+		{
+			int NewNumRetries = 0;
+			if(NumRetriesTextBox.Text.Length > 0 && !int.TryParse(NumRetriesTextBox.Text, out NewNumRetries))
+			{
+				MessageBox.Show("Invalid value for number of retries");
+				return;
+			}
+
+			int NewNumThreads = 0;
+			if(NumThreadsTextBox.Text.Length > 0 && !int.TryParse(NumThreadsTextBox.Text, out NewNumThreads))
+			{
+				MessageBox.Show("Invalid value for number of threads");
+				return;
+			}
+
+			int NewTcpBufferSize = 0;
+			if(TcpBufferSizeText.Text.Length > 0 && !int.TryParse(TcpBufferSizeText.Text, out NewTcpBufferSize))
+			{
+				MessageBox.Show("Invalid value for TCP buffer size");
+				return;
+			}
+
+			Settings.SyncOptions.NumRetries = NewNumRetries;
+			Settings.SyncOptions.NumThreads = NewNumThreads;
+			Settings.SyncOptions.TcpBufferSize = NewTcpBufferSize;
+			Settings.Save();
+
+			DialogResult = System.Windows.Forms.DialogResult.OK;
+			Close();
+		}
+
+		private void CancButton_Click(object sender, EventArgs e)
+		{
+			DialogResult = System.Windows.Forms.DialogResult.Cancel;
+			Close();
+		}
+	}
+}
diff --git a/Engine/Source/Programs/UnrealGameSync/UnrealGameSync/Forms/ProgramsRunningWindow.cs b/Engine/Source/Programs/UnrealGameSync/UnrealGameSync/Forms/ProgramsRunningWindow.cs
--- a/Engine/Source/Programs/UnrealGameSync/UnrealGameSync/Forms/ProgramsRunningWindow.cs
+++ b/Engine/Source/Programs/UnrealGameSync/UnrealGameSync/Forms/ProgramsRunningWindow.cs
@@ -1,4 +1,4 @@
-﻿// Copyright 1998-2019 Epic Games, Inc. All Rights Reserved.
+﻿// Copyright 1998-2018 Epic Games, Inc. All Rights Reserved.
 
 using System;
 using System.Collections.Generic;
diff --git a/Engine/Source/Programs/UnrealGameSync/UnrealGameSync/Forms/ProjectWindow.cs b/Engine/Source/Programs/UnrealGameSync/UnrealGameSync/Forms/ProjectWindow.cs
--- a/Engine/Source/Programs/UnrealGameSync/UnrealGameSync/Forms/ProjectWindow.cs
+++ b/Engine/Source/Programs/UnrealGameSync/UnrealGameSync/Forms/ProjectWindow.cs
@@ -1,4 +1,4 @@
-﻿// Copyright 1998-2019 Epic Games, Inc. All Rights Reserved.
+﻿// Copyright 1998-2018 Epic Games, Inc. All Rights Reserved.
 
 using System;
 using System.Collections.Generic;
diff --git a/Engine/Source/Programs/UnrealGameSync/UnrealGameSync/Forms/ScheduleWindow.cs b/Engine/Source/Programs/UnrealGameSync/UnrealGameSync/Forms/ScheduleWindow.cs
--- a/Engine/Source/Programs/UnrealGameSync/UnrealGameSync/Forms/ScheduleWindow.cs
+++ b/Engine/Source/Programs/UnrealGameSync/UnrealGameSync/Forms/ScheduleWindow.cs
@@ -1,4 +1,4 @@
-// Copyright 1998-2019 Epic Games, Inc. All Rights Reserved.
+// Copyright 1998-2018 Epic Games, Inc. All Rights Reserved.
 
 using System;
 using System.Collections.Generic;
diff --git a/Engine/Source/Programs/UnrealGameSync/UnrealGameSync/Forms/SdkInfoWindow.cs b/Engine/Source/Programs/UnrealGameSync/UnrealGameSync/Forms/SdkInfoWindow.cs
--- a/Engine/Source/Programs/UnrealGameSync/UnrealGameSync/Forms/SdkInfoWindow.cs
+++ b/Engine/Source/Programs/UnrealGameSync/UnrealGameSync/Forms/SdkInfoWindow.cs
@@ -1,5 +1,3 @@
-// Copyright 1998-2019 Epic Games, Inc. All Rights Reserved.
-
 using System;
 using System.Collections.Generic;
 using System.ComponentModel;
diff --git a/Engine/Source/Programs/UnrealGameSync/UnrealGameSync/Forms/SelectProjectFromWorkspaceWindow.cs b/Engine/Source/Programs/UnrealGameSync/UnrealGameSync/Forms/SelectProjectFromWorkspaceWindow.cs
--- a/Engine/Source/Programs/UnrealGameSync/UnrealGameSync/Forms/SelectProjectFromWorkspaceWindow.cs
+++ b/Engine/Source/Programs/UnrealGameSync/UnrealGameSync/Forms/SelectProjectFromWorkspaceWindow.cs
@@ -1,4 +1,4 @@
-// Copyright 1998-2019 Epic Games, Inc. All Rights Reserved.
+// Copyright 1998-2018 Epic Games, Inc. All Rights Reserved.
 
 using System;
 using System.Collections.Generic;
diff --git a/Engine/Source/Programs/UnrealGameSync/UnrealGameSync/Forms/SelectStreamWindow.cs b/Engine/Source/Programs/UnrealGameSync/UnrealGameSync/Forms/SelectStreamWindow.cs
--- a/Engine/Source/Programs/UnrealGameSync/UnrealGameSync/Forms/SelectStreamWindow.cs
+++ b/Engine/Source/Programs/UnrealGameSync/UnrealGameSync/Forms/SelectStreamWindow.cs
@@ -1,4 +1,4 @@
-﻿// Copyright 1998-2019 Epic Games, Inc. All Rights Reserved.
+﻿// Copyright 1998-2018 Epic Games, Inc. All Rights Reserved.
 
 using System;
 using System.Collections.Generic;
diff --git a/Engine/Source/Programs/UnrealGameSync/UnrealGameSync/Forms/SelectWorkspaceWindow.cs b/Engine/Source/Programs/UnrealGameSync/UnrealGameSync/Forms/SelectWorkspaceWindow.cs
--- a/Engine/Source/Programs/UnrealGameSync/UnrealGameSync/Forms/SelectWorkspaceWindow.cs
+++ b/Engine/Source/Programs/UnrealGameSync/UnrealGameSync/Forms/SelectWorkspaceWindow.cs
@@ -1,4 +1,4 @@
-﻿// Copyright 1998-2019 Epic Games, Inc. All Rights Reserved.
+﻿// Copyright 1998-2018 Epic Games, Inc. All Rights Reserved.
 
 using System;
 using System.Collections.Generic;
diff --git a/Engine/Source/Programs/UnrealGameSync/UnrealGameSync/Forms/SyncFilter.cs b/Engine/Source/Programs/UnrealGameSync/UnrealGameSync/Forms/SyncFilter.cs
--- a/Engine/Source/Programs/UnrealGameSync/UnrealGameSync/Forms/SyncFilter.cs
+++ b/Engine/Source/Programs/UnrealGameSync/UnrealGameSync/Forms/SyncFilter.cs
@@ -1,4 +1,4 @@
-﻿// Copyright 1998-2019 Epic Games, Inc. All Rights Reserved.
+﻿// Copyright 1998-2018 Epic Games, Inc. All Rights Reserved.
 
 using System;
 using System.Collections.Generic;
diff --git a/Engine/Source/Programs/UnrealGameSync/UnrealGameSync/Forms/SyncFilterSyntax.cs b/Engine/Source/Programs/UnrealGameSync/UnrealGameSync/Forms/SyncFilterSyntax.cs
--- a/Engine/Source/Programs/UnrealGameSync/UnrealGameSync/Forms/SyncFilterSyntax.cs
+++ b/Engine/Source/Programs/UnrealGameSync/UnrealGameSync/Forms/SyncFilterSyntax.cs
@@ -1,4 +1,4 @@
-﻿// Copyright 1998-2019 Epic Games, Inc. All Rights Reserved.
+﻿// Copyright 1998-2018 Epic Games, Inc. All Rights Reserved.
 
 using System;
 using System.Collections.Generic;
diff --git a/Engine/Source/Programs/UnrealGameSync/UnrealGameSync/Forms/VariablesWindow.cs b/Engine/Source/Programs/UnrealGameSync/UnrealGameSync/Forms/VariablesWindow.cs
--- a/Engine/Source/Programs/UnrealGameSync/UnrealGameSync/Forms/VariablesWindow.cs
+++ b/Engine/Source/Programs/UnrealGameSync/UnrealGameSync/Forms/VariablesWindow.cs
@@ -1,4 +1,4 @@
-﻿// Copyright 1998-2019 Epic Games, Inc. All Rights Reserved.
+﻿// Copyright 1998-2018 Epic Games, Inc. All Rights Reserved.
 
 using System;
 using System.Collections;
diff --git a/Engine/Source/Programs/UnrealGameSync/UnrealGameSync/ModalTask.cs b/Engine/Source/Programs/UnrealGameSync/UnrealGameSync/ModalTask.cs
--- a/Engine/Source/Programs/UnrealGameSync/UnrealGameSync/ModalTask.cs
+++ b/Engine/Source/Programs/UnrealGameSync/UnrealGameSync/ModalTask.cs
@@ -1,4 +1,4 @@
-// Copyright 1998-2019 Epic Games, Inc. All Rights Reserved.
+// Copyright 1998-2018 Epic Games, Inc. All Rights Reserved.
 
 using System;
 using System.Collections.Generic;
diff --git a/Engine/Source/Programs/UnrealGameSync/UnrealGameSync/OutputAdapters.cs b/Engine/Source/Programs/UnrealGameSync/UnrealGameSync/OutputAdapters.cs
--- a/Engine/Source/Programs/UnrealGameSync/UnrealGameSync/OutputAdapters.cs
+++ b/Engine/Source/Programs/UnrealGameSync/UnrealGameSync/OutputAdapters.cs
@@ -1,4 +1,4 @@
-// Copyright 1998-2019 Epic Games, Inc. All Rights Reserved.
+// Copyright 1998-2018 Epic Games, Inc. All Rights Reserved.
 
 using System;
 using System.Collections.Generic;
diff --git a/Engine/Source/Programs/UnrealGameSync/UnrealGameSync/Perforce.cs b/Engine/Source/Programs/UnrealGameSync/UnrealGameSync/Perforce.cs
--- a/Engine/Source/Programs/UnrealGameSync/UnrealGameSync/Perforce.cs
+++ b/Engine/Source/Programs/UnrealGameSync/UnrealGameSync/Perforce.cs
@@ -1,4 +1,4 @@
-// Copyright 1998-2019 Epic Games, Inc. All Rights Reserved.
+// Copyright 1998-2018 Epic Games, Inc. All Rights Reserved.
 
 using System;
 using System.Collections.Generic;
@@ -155,7 +155,6 @@ namespace UnrealGameSync
 		public string UserName;
 		public string HostName;
 		public string ClientAddress;
-		public string ServerAddress;
 		public TimeSpan ServerTimeZone;
 
 		public PerforceInfoRecord(Dictionary<string, string> Tags)
@@ -163,7 +162,6 @@ namespace UnrealGameSync
 			Tags.TryGetValue("userName", out UserName);
 			Tags.TryGetValue("clientHost", out HostName);
 			Tags.TryGetValue("clientAddress", out ClientAddress);
-			Tags.TryGetValue("serverAddress", out ServerAddress);
 
 			string ServerDateTime;
 			if(Tags.TryGetValue("serverDate", out ServerDateTime))
@@ -475,6 +473,8 @@ namespace UnrealGameSync
 		delegate bool HandleRecordDelegate(Dictionary<string, string> Tags);
 		delegate bool HandleOutputDelegate(PerforceOutputLine Line);
 
+		public const string DefaultServerAndPort = "perforce:1666";
+
 		public readonly string ServerAndPort;
 		public readonly string UserName;
 		public readonly string ClientName;
diff --git a/Engine/Source/Programs/UnrealGameSync/UnrealGameSync/PerforceModalTask.cs b/Engine/Source/Programs/UnrealGameSync/UnrealGameSync/PerforceModalTask.cs
--- a/Engine/Source/Programs/UnrealGameSync/UnrealGameSync/PerforceModalTask.cs
+++ b/Engine/Source/Programs/UnrealGameSync/UnrealGameSync/PerforceModalTask.cs
@@ -1,4 +1,4 @@
-// Copyright 1998-2019 Epic Games, Inc. All Rights Reserved.
+// Copyright 1998-2018 Epic Games, Inc. All Rights Reserved.
 
 using System;
 using System.Collections.Generic;
@@ -67,28 +67,24 @@ namespace UnrealGameSync
 				}
 			}
 
+			// If the server setting is still null, user the default
+			if(ServerAndPort == null)
+			{
+				ServerAndPort = PerforceConnection.DefaultServerAndPort;
+			}
+
 			// Update the server and username from the reported server info if it's not set
-			if(ServerAndPort == null || UserName == null)
+			if(UserName == null)
 			{
 				PerforceConnection Perforce = new PerforceConnection(UserName, null, ServerAndPort);
 
 				PerforceInfoRecord PerforceInfo;
-				if(!Perforce.Info(out PerforceInfo, Log))
-				{
-					return false;
-				}
-				if(ServerAndPort == null && !String.IsNullOrEmpty(PerforceInfo.ServerAddress))
-				{
-					ServerAndPort = PerforceInfo.ServerAddress;
-				}
-				if(UserName == null && !String.IsNullOrEmpty(PerforceInfo.UserName))
-				{
-					UserName = PerforceInfo.UserName;
-				}
-				if(ServerAndPort == null || UserName == null)
+				if(!Perforce.Info(out PerforceInfo, Log) || String.IsNullOrEmpty(PerforceInfo.UserName))
 				{
 					return false;
 				}
+
+				UserName = PerforceInfo.UserName;
 			}
 
 			// Otherwise succeed
diff --git a/Engine/Source/Programs/UnrealGameSync/UnrealGameSync/PerforceMonitor.cs b/Engine/Source/Programs/UnrealGameSync/UnrealGameSync/PerforceMonitor.cs
--- a/Engine/Source/Programs/UnrealGameSync/UnrealGameSync/PerforceMonitor.cs
+++ b/Engine/Source/Programs/UnrealGameSync/UnrealGameSync/PerforceMonitor.cs
@@ -1,4 +1,4 @@
-// Copyright 1998-2019 Epic Games, Inc. All Rights Reserved.
+// Copyright 1998-2018 Epic Games, Inc. All Rights Reserved.
 
 using System;
 using System.Collections.Concurrent;
@@ -59,7 +59,6 @@ namespace UnrealGameSync
 			PendingMaxChangesValue = 100;
 			LastChangeByCurrentUser = -1;
 			LastCodeChangeByCurrentUser = -1;
-			OtherStreamNames = new List<string>();
 			bIsEnterpriseProject = bInIsEnterpriseProject;
 			LatestProjectConfigFile = InProjectConfigFile;
 			LocalConfigFiles = InLocalConfigFiles;
@@ -116,12 +115,6 @@ namespace UnrealGameSync
 			set { lock(this){ if(value != PendingMaxChangesValue){ PendingMaxChangesValue = value; RefreshEvent.Set(); } } }
 		}
 
-		public IReadOnlyList<string> OtherStreamNames
-		{
-			get;
-			private set;
-		}
-
 		void PollForUpdates()
 		{
 			string StreamName;
@@ -155,17 +148,6 @@ namespace UnrealGameSync
 							OnStreamChange();
 						}
 
-						// Update the stream list
-						if(StreamName != null)
-						{
-							List<string> NewOtherStreamNames;
-							if(!Perforce.FindStreams(PerforceUtils.GetClientOrDepotDirectoryName(StreamName) + "/*", out NewOtherStreamNames, LogWriter))
-							{
-								NewOtherStreamNames = new List<string>();
-							}
-							OtherStreamNames = NewOtherStreamNames;
-						}
-
 						// Check for any p4 changes
 						if(!UpdateChanges())
 						{
diff --git a/Engine/Source/Programs/UnrealGameSync/UnrealGameSync/Program.cs b/Engine/Source/Programs/UnrealGameSync/UnrealGameSync/Program.cs
--- a/Engine/Source/Programs/UnrealGameSync/UnrealGameSync/Program.cs
+++ b/Engine/Source/Programs/UnrealGameSync/UnrealGameSync/Program.cs
@@ -1,4 +1,4 @@
-// Copyright 1998-2019 Epic Games, Inc. All Rights Reserved.
+// Copyright 1998-2018 Epic Games, Inc. All Rights Reserved.
 
 using System;
 using System.Collections.Generic;
diff --git a/Engine/Source/Programs/UnrealGameSync/UnrealGameSync/ProgramApplicationContext.cs b/Engine/Source/Programs/UnrealGameSync/UnrealGameSync/ProgramApplicationContext.cs
--- a/Engine/Source/Programs/UnrealGameSync/UnrealGameSync/ProgramApplicationContext.cs
+++ b/Engine/Source/Programs/UnrealGameSync/UnrealGameSync/ProgramApplicationContext.cs
@@ -1,4 +1,4 @@
-﻿// Copyright 1998-2019 Epic Games, Inc. All Rights Reserved.
+﻿// Copyright 1998-2018 Epic Games, Inc. All Rights Reserved.
 
 using System;
 using System.Collections.Generic;
diff --git a/Engine/Source/Programs/UnrealGameSync/UnrealGameSync/Properties/AssemblyInfo.cs b/Engine/Source/Programs/UnrealGameSync/UnrealGameSync/Properties/AssemblyInfo.cs
--- a/Engine/Source/Programs/UnrealGameSync/UnrealGameSync/Properties/AssemblyInfo.cs
+++ b/Engine/Source/Programs/UnrealGameSync/UnrealGameSync/Properties/AssemblyInfo.cs
@@ -1,4 +1,4 @@
-// Copyright 1998-2019 Epic Games, Inc. All Rights Reserved.
+// Copyright 1998-2018 Epic Games, Inc. All Rights Reserved.
 
 using System.Reflection;
 using System.Runtime.CompilerServices;
@@ -12,7 +12,7 @@ using System.Runtime.InteropServices;
 [assembly: AssemblyConfiguration("")]
 [assembly: AssemblyCompany("")]
 [assembly: AssemblyProduct("UnrealGameSync")]
-[assembly: AssemblyCopyright("Copyright 1998-2019 Epic Games, Inc. All Rights Reserved.")]
+[assembly: AssemblyCopyright("Copyright 1998-2018 Epic Games, Inc. All Rights Reserved.")]
 [assembly: AssemblyTrademark("")]
 [assembly: AssemblyCulture("")]
 
@@ -34,5 +34,5 @@ using System.Runtime.InteropServices;
 // You can specify all the values or you can default the Build and Revision Numbers 
 // by using the '*' as shown below:
 // [assembly: AssemblyVersion("1.0.*")]
-[assembly: AssemblyVersion("1.155.*")]
-[assembly: AssemblyFileVersion("1.155.0.0")]
+[assembly: AssemblyVersion("1.159.*")]
+[assembly: AssemblyFileVersion("1.159.0.0")]
diff --git a/Engine/Source/Programs/UnrealGameSync/UnrealGameSync/RESTApi.cs b/Engine/Source/Programs/UnrealGameSync/UnrealGameSync/RESTApi.cs
--- a/Engine/Source/Programs/UnrealGameSync/UnrealGameSync/RESTApi.cs
+++ b/Engine/Source/Programs/UnrealGameSync/UnrealGameSync/RESTApi.cs
@@ -1,4 +1,4 @@
-﻿// Copyright 1998-2019 Epic Games, Inc. All Rights Reserved.
+﻿// Copyright 1998-2018 Epic Games, Inc. All Rights Reserved.
 
 using System;
 using System.Collections.Generic;
diff --git a/Engine/Source/Programs/UnrealGameSync/UnrealGameSync/Taskbar.cs b/Engine/Source/Programs/UnrealGameSync/UnrealGameSync/Taskbar.cs
--- a/Engine/Source/Programs/UnrealGameSync/UnrealGameSync/Taskbar.cs
+++ b/Engine/Source/Programs/UnrealGameSync/UnrealGameSync/Taskbar.cs
@@ -1,4 +1,4 @@
-﻿// Copyright 1998-2019 Epic Games, Inc. All Rights Reserved.
+﻿// Copyright 1998-2018 Epic Games, Inc. All Rights Reserved.
 
 using System;
 using System.Collections.Generic;
diff --git a/Engine/Source/Programs/UnrealGameSync/UnrealGameSync/Telemetry.cs b/Engine/Source/Programs/UnrealGameSync/UnrealGameSync/Telemetry.cs
--- a/Engine/Source/Programs/UnrealGameSync/UnrealGameSync/Telemetry.cs
+++ b/Engine/Source/Programs/UnrealGameSync/UnrealGameSync/Telemetry.cs
@@ -1,4 +1,4 @@
-// Copyright 1998-2019 Epic Games, Inc. All Rights Reserved.
+// Copyright 1998-2018 Epic Games, Inc. All Rights Reserved.
 
 using System;
 using System.Collections.Concurrent;
diff --git a/Engine/Source/Programs/UnrealGameSync/UnrealGameSync/UnrealGameSync.csproj b/Engine/Source/Programs/UnrealGameSync/UnrealGameSync/UnrealGameSync.csproj
--- a/Engine/Source/Programs/UnrealGameSync/UnrealGameSync/UnrealGameSync.csproj
+++ b/Engine/Source/Programs/UnrealGameSync/UnrealGameSync/UnrealGameSync.csproj
@@ -206,11 +206,11 @@
     <Compile Include="Forms\PasswordWindow.Designer.cs">
       <DependentUpon>PasswordWindow.cs</DependentUpon>
     </Compile>
-    <Compile Include="Forms\PerforceSettingsWindow.cs">
+    <Compile Include="Forms\PerforceSyncSettingsWindow.cs">
       <SubType>Form</SubType>
     </Compile>
-    <Compile Include="Forms\PerforceSettingsWindow.Designer.cs">
-      <DependentUpon>PerforceSettingsWindow.cs</DependentUpon>
+    <Compile Include="Forms\PerforceSyncSettingsWindow.Designer.cs">
+      <DependentUpon>PerforceSyncSettingsWindow.cs</DependentUpon>
     </Compile>
     <Compile Include="Forms\ProgramsRunningWindow.cs">
       <SubType>Form</SubType>
@@ -369,8 +369,8 @@
     <EmbeddedResource Include="Forms\PasswordWindow.resx">
       <DependentUpon>PasswordWindow.cs</DependentUpon>
     </EmbeddedResource>
-    <EmbeddedResource Include="Forms\PerforceSettingsWindow.resx">
-      <DependentUpon>PerforceSettingsWindow.cs</DependentUpon>
+    <EmbeddedResource Include="Forms\PerforceSyncSettingsWindow.resx">
+      <DependentUpon>PerforceSyncSettingsWindow.cs</DependentUpon>
     </EmbeddedResource>
     <EmbeddedResource Include="Forms\ProgramsRunningWindow.resx">
       <DependentUpon>ProgramsRunningWindow.cs</DependentUpon>
diff --git a/Engine/Source/Programs/UnrealGameSync/UnrealGameSync/UpdateMonitor.cs b/Engine/Source/Programs/UnrealGameSync/UnrealGameSync/UpdateMonitor.cs
--- a/Engine/Source/Programs/UnrealGameSync/UnrealGameSync/UpdateMonitor.cs
+++ b/Engine/Source/Programs/UnrealGameSync/UnrealGameSync/UpdateMonitor.cs
@@ -1,4 +1,4 @@
-﻿// Copyright 1998-2019 Epic Games, Inc. All Rights Reserved.
+﻿// Copyright 1998-2018 Epic Games, Inc. All Rights Reserved.
 
 using System;
 using System.Collections.Generic;
diff --git a/Engine/Source/Programs/UnrealGameSync/UnrealGameSync/UserSettings.cs b/Engine/Source/Programs/UnrealGameSync/UnrealGameSync/UserSettings.cs
--- a/Engine/Source/Programs/UnrealGameSync/UnrealGameSync/UserSettings.cs
+++ b/Engine/Source/Programs/UnrealGameSync/UnrealGameSync/UserSettings.cs
@@ -1,4 +1,4 @@
-// Copyright 1998-2019 Epic Games, Inc. All Rights Reserved.
+// Copyright 1998-2018 Epic Games, Inc. All Rights Reserved.
 
 using System;
 using System.Collections.Generic;
@@ -69,6 +69,12 @@ namespace UnrealGameSync
 					ServerAndPort = null;
 				}
 
+				// Fixup for code that was saving server host name rather than DNS entry
+				if(ServerAndPort != null && ServerAndPort.Equals("p4-nodeb.epicgames.net:1666", StringComparison.OrdinalIgnoreCase))
+				{
+					ServerAndPort = "perforce:1666";
+				}
+
 				string UserName = Object.GetValue("UserName", null);
 				if(String.IsNullOrWhiteSpace(UserName))
 				{
diff --git a/Engine/Source/Programs/UnrealGameSync/UnrealGameSync/Utility.cs b/Engine/Source/Programs/UnrealGameSync/UnrealGameSync/Utility.cs
--- a/Engine/Source/Programs/UnrealGameSync/UnrealGameSync/Utility.cs
+++ b/Engine/Source/Programs/UnrealGameSync/UnrealGameSync/Utility.cs
@@ -1,4 +1,4 @@
-// Copyright 1998-2019 Epic Games, Inc. All Rights Reserved.
+// Copyright 1998-2018 Epic Games, Inc. All Rights Reserved.
 
 #define USE_NEW_PROCESS_JOBS
 
diff --git a/Engine/Source/Programs/UnrealGameSync/UnrealGameSync/Workspace.cs b/Engine/Source/Programs/UnrealGameSync/UnrealGameSync/Workspace.cs
--- a/Engine/Source/Programs/UnrealGameSync/UnrealGameSync/Workspace.cs
+++ b/Engine/Source/Programs/UnrealGameSync/UnrealGameSync/Workspace.cs
@@ -1,4 +1,4 @@
-// Copyright 1998-2019 Epic Games, Inc. All Rights Reserved.
+// Copyright 1998-2018 Epic Games, Inc. All Rights Reserved.
 
 using System;
 using System.Collections.Generic;
@@ -697,8 +697,8 @@ namespace UnrealGameSync
 							Log.WriteLine("Executing post-sync steps...");
 
 							Dictionary<string, string> PostSyncVariables = new Dictionary<string, string>(Context.Variables);
-							PostSyncVariables.Add("Change", PendingChangeNumber.ToString());
-							PostSyncVariables.Add("CodeChange", VersionChangeNumber.ToString());
+							PostSyncVariables["Change"] = PendingChangeNumber.ToString();
+							PostSyncVariables["CodeChange"] = VersionChangeNumber.ToString();
 
 							foreach (string PostSyncStep in PostSyncSteps.Select(x => x.Trim()))
 							{
@@ -905,79 +905,82 @@ namespace UnrealGameSync
 
 						Log.WriteLine(Step.StatusText);
 
-						switch(Step.Type)
+						if(Step.IsValid())
 						{
-							case BuildStepType.Compile:
-								using(TelemetryStopwatch StepStopwatch = new TelemetryStopwatch("Compile:" + Step.Target, TelemetryProjectPath))
-								{
-									string CommandLine = String.Format("{0} {1} {2} {3} -NoHotReloadFromIDE", Step.Target, Step.Platform, Step.Configuration, Utility.ExpandVariables(Step.Arguments, Context.Variables));
-									if(!Context.Options.HasFlag(WorkspaceUpdateOptions.UseIncrementalBuilds) || bForceClean)
+							switch(Step.Type)
+							{
+								case BuildStepType.Compile:
+									using(TelemetryStopwatch StepStopwatch = new TelemetryStopwatch("Compile:" + Step.Target, TelemetryProjectPath))
 									{
-										Log.WriteLine("ubt> Running {0} {1} -clean", UnrealBuildToolPath, CommandLine);
-										Utility.ExecuteProcess(UnrealBuildToolPath, null, CommandLine + " -clean", null, new ProgressTextWriter(Progress, new PrefixedTextWriter("ubt> ", Log)));
-									}
-
-									Log.WriteLine("ubt> Running {0} {1} -progress", UnrealBuildToolPath, CommandLine);
+										string CommandLine = String.Format("{0} {1} {2} {3} -NoHotReloadFromIDE", Step.Target, Step.Platform, Step.Configuration, Utility.ExpandVariables(Step.Arguments ?? "", Context.Variables));
+										if(!Context.Options.HasFlag(WorkspaceUpdateOptions.UseIncrementalBuilds) || bForceClean)
+										{
+											Log.WriteLine("ubt> Running {0} {1} -clean", UnrealBuildToolPath, CommandLine);
+											Utility.ExecuteProcess(UnrealBuildToolPath, null, CommandLine + " -clean", null, new ProgressTextWriter(Progress, new PrefixedTextWriter("ubt> ", Log)));
+										}
 
-									int ResultFromBuild = Utility.ExecuteProcess(UnrealBuildToolPath, null, CommandLine + " -progress", null, new ProgressTextWriter(Progress, new PrefixedTextWriter("ubt> ", Log)));
-									if(ResultFromBuild != 0)
-									{
-										StepStopwatch.Stop("Failed");
-										StatusMessage = String.Format("Failed to compile {0}.", Step.Target);
-										return (HasModifiedSourceFiles() || Context.UserBuildStepObjects.Count > 0)? WorkspaceUpdateResult.FailedToCompile : WorkspaceUpdateResult.FailedToCompileWithCleanWorkspace;
-									}
+										Log.WriteLine("ubt> Running {0} {1} -progress", UnrealBuildToolPath, CommandLine);
 
-									StepStopwatch.Stop("Success");
-								}
-								break;
-							case BuildStepType.Cook:
-								using(TelemetryStopwatch StepStopwatch = new TelemetryStopwatch("Cook/Launch: " + Path.GetFileNameWithoutExtension(Step.FileName), TelemetryProjectPath))
-								{
-									string LocalRunUAT = Path.Combine(LocalRootPath, "Engine", "Build", "BatchFiles", "RunUAT.bat");
-									string Arguments = String.Format("/C \"\"{0}\" -profile=\"{1}\"\"", LocalRunUAT, Path.Combine(LocalRootPath, Step.FileName));
-									Log.WriteLine("uat> Running {0} {1}", LocalRunUAT, Arguments);
+										int ResultFromBuild = Utility.ExecuteProcess(UnrealBuildToolPath, null, CommandLine + " -progress", null, new ProgressTextWriter(Progress, new PrefixedTextWriter("ubt> ", Log)));
+										if(ResultFromBuild != 0)
+										{
+											StepStopwatch.Stop("Failed");
+											StatusMessage = String.Format("Failed to compile {0}.", Step.Target);
+											return (HasModifiedSourceFiles() || Context.UserBuildStepObjects.Count > 0)? WorkspaceUpdateResult.FailedToCompile : WorkspaceUpdateResult.FailedToCompileWithCleanWorkspace;
+										}
 
-									int ResultFromUAT = Utility.ExecuteProcess(CmdExe, null, Arguments, null, new ProgressTextWriter(Progress, new PrefixedTextWriter("uat> ", Log)));
-									if(ResultFromUAT != 0)
-									{
-										StepStopwatch.Stop("Failed");
-										StatusMessage = String.Format("Cook failed. ({0})", ResultFromUAT);
-										return WorkspaceUpdateResult.FailedToCompile;
+										StepStopwatch.Stop("Success");
 									}
-
-									StepStopwatch.Stop("Success");
-								}
-								break;
-							case BuildStepType.Other:
-								using(TelemetryStopwatch StepStopwatch = new TelemetryStopwatch("Custom: " + Path.GetFileNameWithoutExtension(Step.FileName), TelemetryProjectPath))
-								{
-									string ToolFileName = Path.Combine(LocalRootPath, Utility.ExpandVariables(Step.FileName, Context.Variables));
-									string ToolWorkingDir = String.IsNullOrWhiteSpace(Step.WorkingDir) ? Path.GetDirectoryName(ToolFileName) : Utility.ExpandVariables(Step.WorkingDir, Context.Variables);
-									string ToolArguments = Utility.ExpandVariables(Step.Arguments, Context.Variables);
-									Log.WriteLine("tool> Running {0} {1}", ToolFileName, ToolArguments);
-
-									if(Step.bUseLogWindow)
+									break;
+								case BuildStepType.Cook:
+									using(TelemetryStopwatch StepStopwatch = new TelemetryStopwatch("Cook/Launch: " + Path.GetFileNameWithoutExtension(Step.FileName), TelemetryProjectPath))
 									{
-										int ResultFromTool = Utility.ExecuteProcess(ToolFileName, ToolWorkingDir, ToolArguments, null, new ProgressTextWriter(Progress, new PrefixedTextWriter("tool> ", Log)));
-										if(ResultFromTool != 0)
+										string LocalRunUAT = Path.Combine(LocalRootPath, "Engine", "Build", "BatchFiles", "RunUAT.bat");
+										string Arguments = String.Format("/C \"\"{0}\" -profile=\"{1}\"\"", LocalRunUAT, Path.Combine(LocalRootPath, Step.FileName));
+										Log.WriteLine("uat> Running {0} {1}", LocalRunUAT, Arguments);
+
+										int ResultFromUAT = Utility.ExecuteProcess(CmdExe, null, Arguments, null, new ProgressTextWriter(Progress, new PrefixedTextWriter("uat> ", Log)));
+										if(ResultFromUAT != 0)
 										{
 											StepStopwatch.Stop("Failed");
-											StatusMessage = String.Format("Tool terminated with exit code {0}.", ResultFromTool);
+											StatusMessage = String.Format("Cook failed. ({0})", ResultFromUAT);
 											return WorkspaceUpdateResult.FailedToCompile;
 										}
+
+										StepStopwatch.Stop("Success");
 									}
-									else
+									break;
+								case BuildStepType.Other:
+									using(TelemetryStopwatch StepStopwatch = new TelemetryStopwatch("Custom: " + Path.GetFileNameWithoutExtension(Step.FileName), TelemetryProjectPath))
 									{
-										ProcessStartInfo StartInfo = new ProcessStartInfo(ToolFileName, ToolArguments);
-										StartInfo.WorkingDirectory = ToolWorkingDir;
-										using(Process.Start(StartInfo))
+										string ToolFileName = Path.Combine(LocalRootPath, Utility.ExpandVariables(Step.FileName, Context.Variables));
+										string ToolWorkingDir = String.IsNullOrWhiteSpace(Step.WorkingDir) ? Path.GetDirectoryName(ToolFileName) : Utility.ExpandVariables(Step.WorkingDir, Context.Variables);
+										string ToolArguments = Utility.ExpandVariables(Step.Arguments ?? "", Context.Variables);
+										Log.WriteLine("tool> Running {0} {1}", ToolFileName, ToolArguments);
+
+										if(Step.bUseLogWindow)
 										{
+											int ResultFromTool = Utility.ExecuteProcess(ToolFileName, ToolWorkingDir, ToolArguments, null, new ProgressTextWriter(Progress, new PrefixedTextWriter("tool> ", Log)));
+											if(ResultFromTool != 0)
+											{
+												StepStopwatch.Stop("Failed");
+												StatusMessage = String.Format("Tool terminated with exit code {0}.", ResultFromTool);
+												return WorkspaceUpdateResult.FailedToCompile;
+											}
+										}
+										else
+										{
+											ProcessStartInfo StartInfo = new ProcessStartInfo(ToolFileName, ToolArguments);
+											StartInfo.WorkingDirectory = ToolWorkingDir;
+											using(Process.Start(StartInfo))
+											{
+											}
 										}
-									}
 
-									StepStopwatch.Stop("Success");
-								}
-								break;
+										StepStopwatch.Stop("Success");
+									}
+									break;
+							}
 						}
 
 						Log.WriteLine();
diff --git a/Engine/Source/Programs/UnrealGameSync/UnrealGameSyncAutomation/App.config b/Engine/Source/Programs/UnrealGameSync/UnrealGameSyncAutomation/App.config
deleted file mode 100644
diff --git a/Engine/Source/Programs/UnrealGameSync/UnrealGameSyncAutomation/Program.cs b/Engine/Source/Programs/UnrealGameSync/UnrealGameSyncAutomation/Program.cs
deleted file mode 100644
diff --git a/Engine/Source/Programs/UnrealGameSync/UnrealGameSyncAutomation/Properties/AssemblyInfo.cs b/Engine/Source/Programs/UnrealGameSync/UnrealGameSyncAutomation/Properties/AssemblyInfo.cs
deleted file mode 100644
diff --git a/Engine/Source/Programs/UnrealGameSync/UnrealGameSyncAutomation/UnrealGameSyncAutomation.csproj b/Engine/Source/Programs/UnrealGameSync/UnrealGameSyncAutomation/UnrealGameSyncAutomation.csproj
deleted file mode 100644
diff --git a/Engine/Source/Programs/UnrealGameSync/UnrealGameSyncLauncher/Forms/LogWindow.cs b/Engine/Source/Programs/UnrealGameSync/UnrealGameSyncLauncher/Forms/LogWindow.cs
--- a/Engine/Source/Programs/UnrealGameSync/UnrealGameSyncLauncher/Forms/LogWindow.cs
+++ b/Engine/Source/Programs/UnrealGameSync/UnrealGameSyncLauncher/Forms/LogWindow.cs
@@ -1,4 +1,4 @@
-// Copyright 1998-2019 Epic Games, Inc. All Rights Reserved.
+// Copyright 1998-2018 Epic Games, Inc. All Rights Reserved.
 
 using System;
 using System.Collections.Generic;
diff --git a/Engine/Source/Programs/UnrealGameSync/UnrealGameSyncLauncher/Forms/SettingsWindow.cs b/Engine/Source/Programs/UnrealGameSync/UnrealGameSyncLauncher/Forms/SettingsWindow.cs
--- a/Engine/Source/Programs/UnrealGameSync/UnrealGameSyncLauncher/Forms/SettingsWindow.cs
+++ b/Engine/Source/Programs/UnrealGameSync/UnrealGameSyncLauncher/Forms/SettingsWindow.cs
@@ -1,4 +1,4 @@
-// Copyright 1998-2019 Epic Games, Inc. All Rights Reserved.
+// Copyright 1998-2018 Epic Games, Inc. All Rights Reserved.
 
 using System;
 using System.Collections.Generic;
@@ -82,7 +82,7 @@ namespace UnrealGameSyncLauncher
 				DepotPath = null;
 			}
 
-			Utility.SaveGlobalPerforceSettings(ServerAndPort, UserName, DepotPath);
+			Program.SaveSettings(ServerAndPort, UserName, DepotPath);
 
 			// Create the task for connecting to this server
 			StringWriter Log = new StringWriter();
@@ -93,7 +93,7 @@ namespace UnrealGameSyncLauncher
 			ModalTaskResult Result = PerforceModalTask.Execute(this, null, ServerAndPort, UserName, SyncApplication, "Updating", "Checking for updates, please wait...", Log, out ErrorMessage);
 			if(Result == ModalTaskResult.Succeeded)
 			{
-				Utility.SaveGlobalPerforceSettings(ServerAndPort, UserName, DepotPath);
+				Program.SaveSettings(ServerAndPort, UserName, DepotPath);
 				DialogResult = DialogResult.OK;
 				Close();
 			}
diff --git a/Engine/Source/Programs/UnrealGameSync/UnrealGameSyncLauncher/Program.cs b/Engine/Source/Programs/UnrealGameSync/UnrealGameSyncLauncher/Program.cs
--- a/Engine/Source/Programs/UnrealGameSync/UnrealGameSyncLauncher/Program.cs
+++ b/Engine/Source/Programs/UnrealGameSync/UnrealGameSyncLauncher/Program.cs
@@ -1,4 +1,4 @@
-// Copyright 1998-2019 Epic Games, Inc. All Rights Reserved.
+// Copyright 1998-2018 Epic Games, Inc. All Rights Reserved.
 
 using Microsoft.Win32;
 using System;
@@ -15,8 +15,14 @@ using UnrealGameSync;
 
 namespace UnrealGameSyncLauncher
 {
-	static class Program
+	static partial class Program
 	{
+		/// <summary>
+		/// Specifies the depot path to sync down the stable version of UGS from, without a trailing slash (eg. //depot/UnrealGameSync/bin). This is a site-specific setting. 
+		/// The UnrealGameSync executable should be located at Release/UnrealGameSync.exe under this path, with any dependent DLLs.
+		/// </summary>
+		static readonly string DefaultDepotPath;
+
 		[STAThread]
 		static int Main(string[] Args)
 		{
@@ -64,8 +70,8 @@ namespace UnrealGameSyncLauncher
 				// Read the settings
 				string ServerAndPort = null;
 				string UserName = null;
-				string DepotPath = DeploymentSettings.DefaultDepotPath;
-				Utility.ReadGlobalPerforceSettings(ref ServerAndPort, ref UserName, ref DepotPath);
+				string DepotPath = DefaultDepotPath;
+				ReadSettings(ref ServerAndPort, ref UserName, ref DepotPath);
 
 				// If the shift key is held down, immediately show the settings window
 				if((Control.ModifierKeys & Keys.Shift) != 0)
@@ -98,6 +104,62 @@ namespace UnrealGameSyncLauncher
 			return 1;
 		}
 
+		public static void ReadSettings(ref string ServerAndPort, ref string UserName, ref string DepotPath)
+		{
+			using (RegistryKey Key = Registry.CurrentUser.OpenSubKey("SOFTWARE\\Epic Games\\UnrealGameSync", false))
+			{
+				if(Key != null)
+				{
+					ServerAndPort = Key.GetValue("ServerAndPort", ServerAndPort) as string;
+					UserName = Key.GetValue("UserName", UserName) as string;
+					DepotPath = Key.GetValue("DepotPath", DepotPath) as string;
+				}
+			}
+		}
+
+		public static void SaveSettings(string ServerAndPort, string UserName, string DepotPath)
+		{
+			try
+			{
+				using (RegistryKey Key = Registry.CurrentUser.CreateSubKey("SOFTWARE\\Epic Games\\UnrealGameSync"))
+				{
+					// Delete this legacy setting
+					try { Key.DeleteValue("Server"); } catch(Exception) { }
+
+					if(String.IsNullOrEmpty(ServerAndPort))
+					{
+						try { Key.DeleteValue("ServerAndPort"); } catch(Exception) { }
+					}
+					else
+					{
+						Key.SetValue("ServerAndPort", ServerAndPort);
+					}
+
+					if(String.IsNullOrEmpty(UserName))
+					{
+						try { Key.DeleteValue("UserName"); } catch(Exception) { }
+					}
+					else
+					{
+						Key.SetValue("UserName", UserName);
+					}
+
+					if(String.IsNullOrEmpty(DepotPath) || (DefaultDepotPath != null && String.Equals(DepotPath, DefaultDepotPath, StringComparison.InvariantCultureIgnoreCase)))
+					{
+						try { Key.DeleteValue("DepotPath"); } catch(Exception) { }
+					}
+					else
+					{
+						Key.SetValue("DepotPath", DepotPath);
+					}
+				}
+			}
+			catch(Exception Ex)
+			{
+				MessageBox.Show("Unable to save settings.\n\n" + Ex.ToString());
+			}
+		}
+
 		public static bool SyncAndRun(PerforceConnection Perforce, string BaseDepotPath, bool bUnstable, string[] Args, Mutex InstanceMutex, TextWriter LogWriter)
 		{
 			try
diff --git a/Engine/Source/Programs/UnrealGameSync/UnrealGameSyncLauncher/Properties/AssemblyInfo.cs b/Engine/Source/Programs/UnrealGameSync/UnrealGameSyncLauncher/Properties/AssemblyInfo.cs
--- a/Engine/Source/Programs/UnrealGameSync/UnrealGameSyncLauncher/Properties/AssemblyInfo.cs
+++ b/Engine/Source/Programs/UnrealGameSync/UnrealGameSyncLauncher/Properties/AssemblyInfo.cs
@@ -1,4 +1,4 @@
-// Copyright 1998-2019 Epic Games, Inc. All Rights Reserved.
+// Copyright 1998-2018 Epic Games, Inc. All Rights Reserved.
 
 using System.Reflection;
 using System.Runtime.CompilerServices;
@@ -12,7 +12,7 @@ using System.Runtime.InteropServices;
 [assembly: AssemblyConfiguration("")]
 [assembly: AssemblyCompany("Epic Games, Inc.")]
 [assembly: AssemblyProduct("UnrealGameSync")]
-[assembly: AssemblyCopyright("Copyright 1998-2019 Epic Games, Inc. All Rights Reserved.")]
+[assembly: AssemblyCopyright("Copyright 1998-2018 Epic Games, Inc. All Rights Reserved.")]
 [assembly: AssemblyTrademark("")]
 [assembly: AssemblyCulture("")]
 
diff --git a/Engine/Source/Programs/UnrealGameSync/UnrealGameSyncLauncher/SyncAndRunModalTask.cs b/Engine/Source/Programs/UnrealGameSync/UnrealGameSyncLauncher/SyncAndRunModalTask.cs
--- a/Engine/Source/Programs/UnrealGameSync/UnrealGameSyncLauncher/SyncAndRunModalTask.cs
+++ b/Engine/Source/Programs/UnrealGameSync/UnrealGameSyncLauncher/SyncAndRunModalTask.cs
@@ -1,4 +1,4 @@
-// Copyright 1998-2019 Epic Games, Inc. All Rights Reserved.
+// Copyright 1998-2018 Epic Games, Inc. All Rights Reserved.
 
 using System;
 using System.Collections.Generic;
diff --git a/Engine/Source/Programs/UnrealGameSync/UnrealGameSyncLauncher/SyncAndRunPerforceTask.cs b/Engine/Source/Programs/UnrealGameSync/UnrealGameSyncLauncher/SyncAndRunPerforceTask.cs
--- a/Engine/Source/Programs/UnrealGameSync/UnrealGameSyncLauncher/SyncAndRunPerforceTask.cs
+++ b/Engine/Source/Programs/UnrealGameSync/UnrealGameSyncLauncher/SyncAndRunPerforceTask.cs
@@ -1,4 +1,4 @@
-// Copyright 1998-2019 Epic Games, Inc. All Rights Reserved.
+// Copyright 1998-2018 Epic Games, Inc. All Rights Reserved.
 
 using System;
 using System.Collections.Generic;
diff --git a/Engine/Source/Programs/UnrealGameSync/UnrealGameSyncLauncher/UnrealGameSyncLauncher.csproj b/Engine/Source/Programs/UnrealGameSync/UnrealGameSyncLauncher/UnrealGameSyncLauncher.csproj
--- a/Engine/Source/Programs/UnrealGameSync/UnrealGameSyncLauncher/UnrealGameSyncLauncher.csproj
+++ b/Engine/Source/Programs/UnrealGameSync/UnrealGameSyncLauncher/UnrealGameSyncLauncher.csproj
@@ -89,9 +89,6 @@
     <Compile Include="..\UnrealGameSync\Forms\PasswordWindow.Designer.cs">
       <Link>UnrealGameSync\PasswordWindow.Designer.cs</Link>
     </Compile>
-    <Compile Include="..\UnrealGameSync\DeploymentSettings.cs">
-      <Link>Shared\DeploymentSettings.cs</Link>
-    </Compile>
     <Compile Include="..\UnrealGameSync\ModalTask.cs">
       <Link>Shared\ModalTask.cs</Link>
     </Compile>
@@ -154,11 +151,9 @@
   </ItemGroup>
   <ItemGroup />
   <Choose>
-    <When Condition="Exists('..\UnrealGameSync\NotForLicensees\DeploymentSettingsInternal.cs')">
+    <When Condition="Exists('NotForLicensees\ProgramSettings.cs')">
       <ItemGroup>
-        <Compile Include="..\UnrealGameSync\NotForLicensees\DeploymentSettingsInternal.cs">
-			<Link>Shared\NotForLicensees\DeploymentSettingsInternal.cs</Link>
-		</Compile>
+        <Compile Include="NotForLicensees\ProgramSettings.cs" />
       </ItemGroup>
     </When>
   </Choose>
diff --git a/Engine/Source/Programs/UnrealGameSync/UnrealGameSyncMetadataServer/App_Start/WebApiConfig.cs b/Engine/Source/Programs/UnrealGameSync/UnrealGameSyncMetadataServer/App_Start/WebApiConfig.cs
--- a/Engine/Source/Programs/UnrealGameSync/UnrealGameSyncMetadataServer/App_Start/WebApiConfig.cs
+++ b/Engine/Source/Programs/UnrealGameSync/UnrealGameSyncMetadataServer/App_Start/WebApiConfig.cs
@@ -1,4 +1,4 @@
-﻿// Copyright 1998-2019 Epic Games, Inc. All Rights Reserved.
+﻿// Copyright 1998-2018 Epic Games, Inc. All Rights Reserved.
 
 using System.Web.Http;
 
diff --git a/Engine/Source/Programs/UnrealGameSync/UnrealGameSyncMetadataServer/Connectors/SqlConnector.cs b/Engine/Source/Programs/UnrealGameSync/UnrealGameSyncMetadataServer/Connectors/SqlConnector.cs
--- a/Engine/Source/Programs/UnrealGameSync/UnrealGameSyncMetadataServer/Connectors/SqlConnector.cs
+++ b/Engine/Source/Programs/UnrealGameSync/UnrealGameSyncMetadataServer/Connectors/SqlConnector.cs
@@ -1,4 +1,4 @@
-﻿// Copyright 1998-2019 Epic Games, Inc. All Rights Reserved.
+﻿// Copyright 1998-2018 Epic Games, Inc. All Rights Reserved.
 
 using System;
 using System.Collections.Generic;
diff --git a/Engine/Source/Programs/UnrealGameSync/UnrealGameSyncMetadataServer/Controllers/BuildController.cs b/Engine/Source/Programs/UnrealGameSync/UnrealGameSyncMetadataServer/Controllers/BuildController.cs
--- a/Engine/Source/Programs/UnrealGameSync/UnrealGameSyncMetadataServer/Controllers/BuildController.cs
+++ b/Engine/Source/Programs/UnrealGameSync/UnrealGameSyncMetadataServer/Controllers/BuildController.cs
@@ -1,4 +1,4 @@
-﻿// Copyright 1998-2019 Epic Games, Inc. All Rights Reserved.
+﻿// Copyright 1998-2018 Epic Games, Inc. All Rights Reserved.
 
 using System.Collections.Generic;
 using System.Web.Http;
diff --git a/Engine/Source/Programs/UnrealGameSync/UnrealGameSyncMetadataServer/Controllers/CISController.cs b/Engine/Source/Programs/UnrealGameSync/UnrealGameSyncMetadataServer/Controllers/CISController.cs
--- a/Engine/Source/Programs/UnrealGameSync/UnrealGameSyncMetadataServer/Controllers/CISController.cs
+++ b/Engine/Source/Programs/UnrealGameSync/UnrealGameSyncMetadataServer/Controllers/CISController.cs
@@ -1,4 +1,4 @@
-﻿// Copyright 1998-2019 Epic Games, Inc. All Rights Reserved.
+﻿// Copyright 1998-2018 Epic Games, Inc. All Rights Reserved.
 
 using System;
 using System.Collections.Generic;
diff --git a/Engine/Source/Programs/UnrealGameSync/UnrealGameSyncMetadataServer/Controllers/CommentController.cs b/Engine/Source/Programs/UnrealGameSync/UnrealGameSyncMetadataServer/Controllers/CommentController.cs
--- a/Engine/Source/Programs/UnrealGameSync/UnrealGameSyncMetadataServer/Controllers/CommentController.cs
+++ b/Engine/Source/Programs/UnrealGameSync/UnrealGameSyncMetadataServer/Controllers/CommentController.cs
@@ -1,4 +1,4 @@
-﻿// Copyright 1998-2019 Epic Games, Inc. All Rights Reserved.
+﻿// Copyright 1998-2018 Epic Games, Inc. All Rights Reserved.
 
 using System.Collections.Generic;
 using System.Web.Http;
diff --git a/Engine/Source/Programs/UnrealGameSync/UnrealGameSyncMetadataServer/Controllers/ErrorController.cs b/Engine/Source/Programs/UnrealGameSync/UnrealGameSyncMetadataServer/Controllers/ErrorController.cs
--- a/Engine/Source/Programs/UnrealGameSync/UnrealGameSyncMetadataServer/Controllers/ErrorController.cs
+++ b/Engine/Source/Programs/UnrealGameSync/UnrealGameSyncMetadataServer/Controllers/ErrorController.cs
@@ -1,4 +1,4 @@
-﻿// Copyright 1998-2019 Epic Games, Inc. All Rights Reserved.
+﻿// Copyright 1998-2018 Epic Games, Inc. All Rights Reserved.
 
 using System.Collections.Generic;
 using System.Web.Http;
diff --git a/Engine/Source/Programs/UnrealGameSync/UnrealGameSyncMetadataServer/Controllers/EventController.cs b/Engine/Source/Programs/UnrealGameSync/UnrealGameSyncMetadataServer/Controllers/EventController.cs
--- a/Engine/Source/Programs/UnrealGameSync/UnrealGameSyncMetadataServer/Controllers/EventController.cs
+++ b/Engine/Source/Programs/UnrealGameSync/UnrealGameSyncMetadataServer/Controllers/EventController.cs
@@ -1,4 +1,4 @@
-﻿// Copyright 1998-2019 Epic Games, Inc. All Rights Reserved.
+﻿// Copyright 1998-2018 Epic Games, Inc. All Rights Reserved.
 
 using System.Collections.Generic;
 using System.Web.Http;
diff --git a/Engine/Source/Programs/UnrealGameSync/UnrealGameSyncMetadataServer/Controllers/LatestController.cs b/Engine/Source/Programs/UnrealGameSync/UnrealGameSyncMetadataServer/Controllers/LatestController.cs
--- a/Engine/Source/Programs/UnrealGameSync/UnrealGameSyncMetadataServer/Controllers/LatestController.cs
+++ b/Engine/Source/Programs/UnrealGameSync/UnrealGameSyncMetadataServer/Controllers/LatestController.cs
@@ -1,4 +1,4 @@
-﻿// Copyright 1998-2019 Epic Games, Inc. All Rights Reserved.
+﻿// Copyright 1998-2018 Epic Games, Inc. All Rights Reserved.
 
 using System;
 using System.Collections.Generic;
diff --git a/Engine/Source/Programs/UnrealGameSync/UnrealGameSyncMetadataServer/Controllers/TelemetryController.cs b/Engine/Source/Programs/UnrealGameSync/UnrealGameSyncMetadataServer/Controllers/TelemetryController.cs
--- a/Engine/Source/Programs/UnrealGameSync/UnrealGameSyncMetadataServer/Controllers/TelemetryController.cs
+++ b/Engine/Source/Programs/UnrealGameSync/UnrealGameSyncMetadataServer/Controllers/TelemetryController.cs
@@ -1,4 +1,4 @@
-﻿// Copyright 1998-2019 Epic Games, Inc. All Rights Reserved.
+﻿// Copyright 1998-2018 Epic Games, Inc. All Rights Reserved.
 
 using System.Web.Http;
 using UnrealGameSyncMetadataServer.Connectors;
diff --git a/Engine/Source/Programs/UnrealGameSync/UnrealGameSyncMetadataServer/Global.asax.cs b/Engine/Source/Programs/UnrealGameSync/UnrealGameSyncMetadataServer/Global.asax.cs
--- a/Engine/Source/Programs/UnrealGameSync/UnrealGameSyncMetadataServer/Global.asax.cs
+++ b/Engine/Source/Programs/UnrealGameSync/UnrealGameSyncMetadataServer/Global.asax.cs
@@ -1,4 +1,4 @@
-﻿// Copyright 1998-2019 Epic Games, Inc. All Rights Reserved.
+﻿// Copyright 1998-2018 Epic Games, Inc. All Rights Reserved.
 using System;
 using System.Data.SQLite;
 using System.IO;
@@ -25,7 +25,60 @@ namespace UnrealGameSyncMetadataServer
 				using (SQLiteConnection Connection = new SQLiteConnection(ConnectionString, true))
 				{
 					Connection.Open();
-					using (SQLiteCommand Command = new SQLiteCommand(Properties.Resources.Setup, Connection))
+					string CreateTables = @"CREATE TABLE IF NOT EXISTS
+										[CIS] (
+											[Id]           INTEGER       NOT NULL PRIMARY KEY AUTOINCREMENT,
+											[ChangeNumber] INTEGER       NOT NULL,
+											[BuildType]    NCHAR (32)    NOT NULL,
+											[Result]       NCHAR (10)    NOT NULL,
+											[Url]          VARCHAR (512) NOT NULL,
+											[Project]      VARCHAR (256) NULL,
+											[ArchivePath]  VARCHAR (512) NULL
+										);
+										
+										CREATE TABLE IF NOT EXISTS
+										[Comments] (
+											[Id]           INTEGER    	  NOT NULL PRIMARY KEY AUTOINCREMENT,
+											[ChangeNumber] INTEGER        NOT NULL,
+											[UserName]     VARCHAR (128)  NOT NULL,
+											[Text]         VARCHAR (140)  NOT NULL,
+											[Project]      NVARCHAR (256) NOT NULL
+										);
+
+										CREATE TABLE IF NOT EXISTS
+										[Errors] (
+											[Id]        INTEGER		   NOT NULL PRIMARY KEY AUTOINCREMENT,
+											[Type]      VARCHAR (50)   NOT NULL,
+											[Text]      VARCHAR (1024) NOT NULL,
+											[UserName]  NVARCHAR (128) NOT NULL,
+											[Project]   VARCHAR (128)  NULL,
+											[Timestamp] DATETIME       NOT NULL,
+											[Version]   VARCHAR (64)   NOT NULL,
+											[IpAddress] VARCHAR (64)   NOT NULL
+										);
+										
+										CREATE TABLE IF NOT EXISTS
+										[Telemetry.v2] (
+											[Id]        INTEGER		   NOT NULL PRIMARY KEY AUTOINCREMENT,
+											[Action]    VARCHAR (128)  NOT NULL,
+											[Result]    VARCHAR (128)  NOT NULL,
+											[UserName]  NVARCHAR (128) NOT NULL,
+											[Project]   VARCHAR (128)  NOT NULL,
+											[Timestamp] DATETIME       NOT NULL,
+											[Duration]  REAL           NOT NULL,
+											[Version]   VARCHAR (64)   NULL,
+											[IpAddress] VARCHAR (64)   NULL
+										);
+
+										CREATE TABLE IF NOT EXISTS
+										[UserVotes] (
+											[Id]         INTEGER		NOT NULL PRIMARY KEY AUTOINCREMENT,
+											[Changelist] INTEGER        NOT NULL,
+											[UserName]   NVARCHAR (128) NOT NULL,
+											[Verdict]    NCHAR (32)     NOT NULL,
+											[Project]    NVARCHAR (256) NULL
+										);";
+					using (SQLiteCommand Command = new SQLiteCommand(CreateTables, Connection))
 					{
 						Command.ExecuteNonQuery();
 					}
diff --git a/Engine/Source/Programs/UnrealGameSync/UnrealGameSyncMetadataServer/Models/BuildData.cs b/Engine/Source/Programs/UnrealGameSync/UnrealGameSyncMetadataServer/Models/BuildData.cs
--- a/Engine/Source/Programs/UnrealGameSync/UnrealGameSyncMetadataServer/Models/BuildData.cs
+++ b/Engine/Source/Programs/UnrealGameSync/UnrealGameSyncMetadataServer/Models/BuildData.cs
@@ -1,4 +1,4 @@
-﻿// Copyright 1998-2019 Epic Games, Inc. All Rights Reserved.
+﻿// Copyright 1998-2018 Epic Games, Inc. All Rights Reserved.
 
 using System;
 using System.Collections.Generic;
diff --git a/Engine/Source/Programs/UnrealGameSync/UnrealGameSyncMetadataServer/Models/CommentData.cs b/Engine/Source/Programs/UnrealGameSync/UnrealGameSyncMetadataServer/Models/CommentData.cs
--- a/Engine/Source/Programs/UnrealGameSync/UnrealGameSyncMetadataServer/Models/CommentData.cs
+++ b/Engine/Source/Programs/UnrealGameSync/UnrealGameSyncMetadataServer/Models/CommentData.cs
@@ -1,4 +1,4 @@
-﻿// Copyright 1998-2019 Epic Games, Inc. All Rights Reserved.
+﻿// Copyright 1998-2018 Epic Games, Inc. All Rights Reserved.
 
 using System;
 using System.Collections.Generic;
diff --git a/Engine/Source/Programs/UnrealGameSync/UnrealGameSyncMetadataServer/Models/EventData.cs b/Engine/Source/Programs/UnrealGameSync/UnrealGameSyncMetadataServer/Models/EventData.cs
--- a/Engine/Source/Programs/UnrealGameSync/UnrealGameSyncMetadataServer/Models/EventData.cs
+++ b/Engine/Source/Programs/UnrealGameSync/UnrealGameSyncMetadataServer/Models/EventData.cs
@@ -1,4 +1,4 @@
-﻿// Copyright 1998-2019 Epic Games, Inc. All Rights Reserved.
+﻿// Copyright 1998-2018 Epic Games, Inc. All Rights Reserved.
 
 using System;
 using System.Collections.Generic;
diff --git a/Engine/Source/Programs/UnrealGameSync/UnrealGameSyncMetadataServer/Models/EventSummary.cs b/Engine/Source/Programs/UnrealGameSync/UnrealGameSyncMetadataServer/Models/EventSummary.cs
--- a/Engine/Source/Programs/UnrealGameSync/UnrealGameSyncMetadataServer/Models/EventSummary.cs
+++ b/Engine/Source/Programs/UnrealGameSync/UnrealGameSyncMetadataServer/Models/EventSummary.cs
@@ -1,4 +1,4 @@
-﻿// Copyright 1998-2019 Epic Games, Inc. All Rights Reserved.
+﻿// Copyright 1998-2018 Epic Games, Inc. All Rights Reserved.
 
 using System;
 using System.Collections.Generic;
diff --git a/Engine/Source/Programs/UnrealGameSync/UnrealGameSyncMetadataServer/Models/LatestData.cs b/Engine/Source/Programs/UnrealGameSync/UnrealGameSyncMetadataServer/Models/LatestData.cs
--- a/Engine/Source/Programs/UnrealGameSync/UnrealGameSyncMetadataServer/Models/LatestData.cs
+++ b/Engine/Source/Programs/UnrealGameSync/UnrealGameSyncMetadataServer/Models/LatestData.cs
@@ -1,4 +1,4 @@
-﻿// Copyright 1998-2019 Epic Games, Inc. All Rights Reserved.
+﻿// Copyright 1998-2018 Epic Games, Inc. All Rights Reserved.
 
 using System;
 using System.Collections.Generic;
diff --git a/Engine/Source/Programs/UnrealGameSync/UnrealGameSyncMetadataServer/Models/TelemetryErrorData.cs b/Engine/Source/Programs/UnrealGameSync/UnrealGameSyncMetadataServer/Models/TelemetryErrorData.cs
--- a/Engine/Source/Programs/UnrealGameSync/UnrealGameSyncMetadataServer/Models/TelemetryErrorData.cs
+++ b/Engine/Source/Programs/UnrealGameSync/UnrealGameSyncMetadataServer/Models/TelemetryErrorData.cs
@@ -1,4 +1,4 @@
-﻿// Copyright 1998-2019 Epic Games, Inc. All Rights Reserved.
+﻿// Copyright 1998-2018 Epic Games, Inc. All Rights Reserved.
 
 using System;
 using System.Collections.Generic;
diff --git a/Engine/Source/Programs/UnrealGameSync/UnrealGameSyncMetadataServer/Models/TelemetryTimingData.cs b/Engine/Source/Programs/UnrealGameSync/UnrealGameSyncMetadataServer/Models/TelemetryTimingData.cs
--- a/Engine/Source/Programs/UnrealGameSync/UnrealGameSyncMetadataServer/Models/TelemetryTimingData.cs
+++ b/Engine/Source/Programs/UnrealGameSync/UnrealGameSyncMetadataServer/Models/TelemetryTimingData.cs
@@ -1,4 +1,4 @@
-﻿// Copyright 1998-2019 Epic Games, Inc. All Rights Reserved.
+﻿// Copyright 1998-2018 Epic Games, Inc. All Rights Reserved.
 
 using System;
 using System.Collections.Generic;
diff --git a/Engine/Source/Programs/UnrealGameSync/UnrealGameSyncMetadataServer/Properties/AssemblyInfo.cs b/Engine/Source/Programs/UnrealGameSync/UnrealGameSyncMetadataServer/Properties/AssemblyInfo.cs
--- a/Engine/Source/Programs/UnrealGameSync/UnrealGameSyncMetadataServer/Properties/AssemblyInfo.cs
+++ b/Engine/Source/Programs/UnrealGameSync/UnrealGameSyncMetadataServer/Properties/AssemblyInfo.cs
@@ -1,4 +1,4 @@
-﻿// Copyright 1998-2019 Epic Games, Inc. All Rights Reserved.
+﻿// Copyright 1998-2018 Epic Games, Inc. All Rights Reserved.
 
 using System.Reflection;
 using System.Runtime.CompilerServices;
@@ -12,7 +12,7 @@ using System.Runtime.InteropServices;
 [assembly: AssemblyConfiguration("")]
 [assembly: AssemblyCompany("")]
 [assembly: AssemblyProduct("UnrealGameSync")]
-[assembly: AssemblyCopyright("Copyright 1998-2019 Epic Games, Inc. All Rights Reserved.")]
+[assembly: AssemblyCopyright("Copyright 1998-2018 Epic Games, Inc. All Rights Reserved.")]
 [assembly: AssemblyTrademark("")]
 [assembly: AssemblyCulture("")]
 
diff --git a/Engine/Source/Programs/UnrealGameSync/UnrealGameSyncMetadataServer/UnrealGameSyncMetadataServer.csproj b/Engine/Source/Programs/UnrealGameSync/UnrealGameSyncMetadataServer/UnrealGameSyncMetadataServer.csproj
--- a/Engine/Source/Programs/UnrealGameSync/UnrealGameSyncMetadataServer/UnrealGameSyncMetadataServer.csproj
+++ b/Engine/Source/Programs/UnrealGameSync/UnrealGameSyncMetadataServer/UnrealGameSyncMetadataServer.csproj
@@ -24,7 +24,6 @@
     <NuGetPackageImportStamp>
     </NuGetPackageImportStamp>
     <TargetFrameworkProfile />
-    <Use64BitIISExpress />
   </PropertyGroup>
   <PropertyGroup Condition=" '$(Configuration)|$(Platform)' == 'Debug|AnyCPU' ">
     <TransformWebConfigEnabled>False</TransformWebConfigEnabled>
@@ -90,12 +89,6 @@
   </ItemGroup>
   <ItemGroup>
     <Compile Include="App_Start\WebApiConfig.cs" />
-    <Compile Include="Properties\Resources.Designer.cs">
-      <AutoGen>True</AutoGen>
-      <DesignTime>True</DesignTime>
-      <DependentUpon>Resources.resx</DependentUpon>
-    </Compile>
-    <Content Include="Setup.sql" />
     <Compile Include="Controllers\BuildController.cs" />
     <Compile Include="Controllers\CISController.cs" />
     <Compile Include="Controllers\CommentController.cs" />
@@ -131,17 +124,12 @@
     <Content Include="NotForLicensees\Web.config">
       <SubType>Designer</SubType>
     </Content>
+    <None Include="Properties\PublishProfiles\FilePublish.pubxml" />
+    <None Include="Properties\PublishProfiles\Profile.pubxml" />
   </ItemGroup>
   <ItemGroup>
     <Content Include="packages.config" />
   </ItemGroup>
-  <ItemGroup>
-    <EmbeddedResource Include="Properties\Resources.resx">
-      <Generator>ResXFileCodeGenerator</Generator>
-      <LastGenOutput>Resources.Designer.cs</LastGenOutput>
-      <SubType>Designer</SubType>
-    </EmbeddedResource>
-  </ItemGroup>
   <PropertyGroup>
     <VisualStudioVersion Condition="'$(VisualStudioVersion)' == ''">10.0</VisualStudioVersion>
     <VSToolsPath Condition="'$(VSToolsPath)' == ''">$(MSBuildExtensionsPath32)\Microsoft\VisualStudio\v$(VisualStudioVersion)</VSToolsPath>
