Copy UGS from Dev-Build (pt 2)

#rnx

diff --git a/Engine/Source/Programs/UnrealGameSync/InstallerCustomAction/CustomAction.cs b/Engine/Source/Programs/UnrealGameSync/InstallerCustomAction/CustomAction.cs
--- a/Engine/Source/Programs/UnrealGameSync/InstallerCustomAction/CustomAction.cs
+++ b/Engine/Source/Programs/UnrealGameSync/InstallerCustomAction/CustomAction.cs
@@ -1,4 +1,4 @@
-﻿// Copyright 1998-2018 Epic Games, Inc. All Rights Reserved.
+﻿// Copyright 1998-2019 Epic Games, Inc. All Rights Reserved.
 
 using System;
 using System.Collections.Generic;
diff --git a/Engine/Source/Programs/UnrealGameSync/InstallerCustomAction/Properties/AssemblyInfo.cs b/Engine/Source/Programs/UnrealGameSync/InstallerCustomAction/Properties/AssemblyInfo.cs
--- a/Engine/Source/Programs/UnrealGameSync/InstallerCustomAction/Properties/AssemblyInfo.cs
+++ b/Engine/Source/Programs/UnrealGameSync/InstallerCustomAction/Properties/AssemblyInfo.cs
@@ -1,4 +1,4 @@
-﻿// Copyright 1998-2018 Epic Games, Inc. All Rights Reserved.
+﻿// Copyright 1998-2019 Epic Games, Inc. All Rights Reserved.
 
 using System.Reflection;
 using System.Runtime.CompilerServices;
@@ -11,7 +11,7 @@ using System.Runtime.InteropServices;
 [assembly: AssemblyDescription("")]
 [assembly: AssemblyCompany("")]
 [assembly: AssemblyProduct("InstallerCustomAction")]
-[assembly: AssemblyCopyright("Copyright ©  2015")]
+[assembly: AssemblyCopyright("Copyright 1998-2019 Epic Games, Inc. All rights reserved.")]
 [assembly: AssemblyTrademark("")]
 [assembly: AssemblyCulture("")]
 
diff --git a/Engine/Source/Programs/UnrealGameSync/PostBadgeStatus/Program.cs b/Engine/Source/Programs/UnrealGameSync/PostBadgeStatus/Program.cs
--- a/Engine/Source/Programs/UnrealGameSync/PostBadgeStatus/Program.cs
+++ b/Engine/Source/Programs/UnrealGameSync/PostBadgeStatus/Program.cs
@@ -1,4 +1,4 @@
-﻿// Copyright 1998-2018 Epic Games, Inc. All Rights Reserved.
+﻿// Copyright 1998-2019 Epic Games, Inc. All Rights Reserved.
 
 using System;
 using System.Collections.Generic;
diff --git a/Engine/Source/Programs/UnrealGameSync/PostBadgeStatus/Properties/AssemblyInfo.cs b/Engine/Source/Programs/UnrealGameSync/PostBadgeStatus/Properties/AssemblyInfo.cs
--- a/Engine/Source/Programs/UnrealGameSync/PostBadgeStatus/Properties/AssemblyInfo.cs
+++ b/Engine/Source/Programs/UnrealGameSync/PostBadgeStatus/Properties/AssemblyInfo.cs
@@ -1,4 +1,4 @@
-﻿// Copyright 1998-2018 Epic Games, Inc. All Rights Reserved.
+﻿// Copyright 1998-2019 Epic Games, Inc. All Rights Reserved.
 
 using System.Reflection;
 using System.Runtime.CompilerServices;
@@ -12,7 +12,7 @@ using System.Runtime.InteropServices;
 [assembly: AssemblyConfiguration("")]
 [assembly: AssemblyCompany("")]
 [assembly: AssemblyProduct("WriteBadgeStatus")]
-[assembly: AssemblyCopyright("Copyright 1998-2018 Epic Games, Inc. All Rights Reserved.")]
+[assembly: AssemblyCopyright("Copyright 1998-2019 Epic Games, Inc. All Rights Reserved.")]
 [assembly: AssemblyTrademark("")]
 [assembly: AssemblyCulture("")]
 
diff --git a/Engine/Source/Programs/UnrealGameSync/UnrealGameSync/ActivationListener.cs b/Engine/Source/Programs/UnrealGameSync/UnrealGameSync/ActivationListener.cs
--- a/Engine/Source/Programs/UnrealGameSync/UnrealGameSync/ActivationListener.cs
+++ b/Engine/Source/Programs/UnrealGameSync/UnrealGameSync/ActivationListener.cs
@@ -1,4 +1,4 @@
-﻿// Copyright 1998-2018 Epic Games, Inc. All Rights Reserved.
+﻿// Copyright 1998-2019 Epic Games, Inc. All Rights Reserved.
 
 using System;
 using System.Collections.Generic;
diff --git a/Engine/Source/Programs/UnrealGameSync/UnrealGameSync/ArchiveManifest.cs b/Engine/Source/Programs/UnrealGameSync/UnrealGameSync/ArchiveManifest.cs
--- a/Engine/Source/Programs/UnrealGameSync/UnrealGameSync/ArchiveManifest.cs
+++ b/Engine/Source/Programs/UnrealGameSync/UnrealGameSync/ArchiveManifest.cs
@@ -1,4 +1,4 @@
-﻿// Copyright 1998-2018 Epic Games, Inc. All Rights Reserved.
+﻿// Copyright 1998-2019 Epic Games, Inc. All Rights Reserved.
 
 using System;
 using System.Collections.Generic;
diff --git a/Engine/Source/Programs/UnrealGameSync/UnrealGameSync/AutomationServer.cs b/Engine/Source/Programs/UnrealGameSync/UnrealGameSync/AutomationServer.cs
new file mode 100644
--- /dev/null
+++ b/Engine/Source/Programs/UnrealGameSync/UnrealGameSync/AutomationServer.cs
@@ -0,0 +1,239 @@
+// Copyright 1998-2019 Epic Games, Inc. All Rights Reserved.
+
+using Microsoft.Win32;
+using System;
+using System.Collections.Generic;
+using System.IO;
+using System.Linq;
+using System.Net;
+using System.Net.Sockets;
+using System.Text;
+using System.Threading;
+using System.Threading.Tasks;
+
+namespace UnrealGameSync
+{
+	enum AutomationRequestType
+	{
+		SyncProject,
+		FindProject,
+		OpenProject,
+	}
+
+	class AutomationRequestInput
+	{
+		public AutomationRequestType Type;
+		public byte[] Data;
+
+		public AutomationRequestInput(AutomationRequestType Type, byte[] Data)
+		{
+			this.Type = Type;
+			this.Data = Data;
+		}
+
+		public static AutomationRequestInput Read(Stream InputStream)
+		{
+			BinaryReader Reader = new BinaryReader(InputStream);
+			
+			int Type = Reader.ReadInt32();
+			int InputSize = Reader.ReadInt32();
+			byte[] Input = Reader.ReadBytes(InputSize);
+
+			return new AutomationRequestInput((AutomationRequestType)Type, Input);
+		}
+
+		public void Write(Stream OutputStream)
+		{
+			BinaryWriter Writer = new BinaryWriter(OutputStream);
+
+			Writer.Write((int)Type);
+			Writer.Write(Data.Length);
+			Writer.Write(Data);
+		}
+	}
+
+	enum AutomationRequestResult
+	{
+		Ok,
+		Invalid,
+		Busy,
+		Canceled,
+		Error,
+		NotFound
+	}
+
+	class AutomationRequestOutput
+	{
+		public AutomationRequestResult Result;
+		public byte[] Data;
+
+		public AutomationRequestOutput(AutomationRequestResult Result)
+		{
+			this.Result = Result;
+			this.Data = new byte[0];
+		}
+
+		public AutomationRequestOutput(AutomationRequestResult Result, byte[] Data)
+		{
+			this.Result = Result;
+			this.Data = Data;
+		}
+
+		public static AutomationRequestOutput Read(Stream InputStream)
+		{
+			using(BinaryReader Reader = new BinaryReader(InputStream))
+			{
+				AutomationRequestResult Result = (AutomationRequestResult)Reader.ReadInt32();
+				int DataSize = Reader.ReadInt32();
+				byte[] Data = Reader.ReadBytes(DataSize);
+				return new AutomationRequestOutput(Result, Data);
+			}
+		}
+
+		public void Write(Stream OutputStream)
+		{
+			using(BinaryWriter Writer = new BinaryWriter(OutputStream))
+			{
+				Writer.Write((int)Result);
+				Writer.Write(Data.Length);
+				Writer.Write(Data);
+			}
+		}
+	}
+
+	class AutomationRequest : IDisposable
+	{
+		public AutomationRequestInput Input;
+		public AutomationRequestOutput Output;
+		public ManualResetEventSlim Complete;
+
+		public AutomationRequest(AutomationRequestInput Input)
+		{
+			this.Input = Input;
+			this.Complete = new ManualResetEventSlim(false);
+		}
+
+		public void SetOutput(AutomationRequestOutput Output)
+		{
+			this.Output = Output;
+			Complete.Set();
+		}
+
+		public void Dispose()
+		{
+			if(Complete != null)
+			{
+				Complete.Dispose();
+				Complete = null;
+			}
+		}
+	}
+
+	class AutomationServer : IDisposable
+	{
+		TcpListener Listener;
+		TcpClient CurrentClient;
+		Thread BackgroundThread;
+		Action<AutomationRequest> PostRequest;
+		bool bDisposing;
+		TextWriter Log;
+
+		public AutomationServer(Action<AutomationRequest> PostRequest, TextWriter Log)
+		{
+			this.PostRequest = PostRequest;
+			this.Log = Log;
+
+			object PortValue = Registry.GetValue("HKEY_CURRENT_USER\\Software\\Epic Games\\UnrealGameSync", "AutomationPort", null);
+			if(PortValue != null)
+			{
+				try
+				{
+					int PortNumber = (int)PortValue;
+
+					Listener = new TcpListener(IPAddress.Loopback, PortNumber);
+					Listener.Start();
+
+					BackgroundThread = new Thread(() => Run());
+					BackgroundThread.Start();
+				}
+				catch(Exception Ex)
+				{
+					Log.WriteLine("Unable to start automation server: {0}", Ex.ToString());
+				}
+			}
+		}
+
+		public void Run()
+		{
+			try
+			{
+				for(;;)
+				{
+					TcpClient Client = CurrentClient = Listener.AcceptTcpClient();
+					try
+					{
+						Log.WriteLine("Accepted connection from {0}", Client.Client.RemoteEndPoint);
+
+						NetworkStream Stream = Client.GetStream();
+
+						AutomationRequestInput Input = AutomationRequestInput.Read(Stream);
+						Log.WriteLine("Received input: {0} (+{1} bytes)", Input.Type, Input.Data.Length);
+
+						AutomationRequestOutput Output;
+						using(AutomationRequest Request = new AutomationRequest(Input))
+						{
+							PostRequest(Request);
+							Request.Complete.Wait();
+							Output = Request.Output;
+						}
+
+						Output.Write(Stream);
+						Log.WriteLine("Sent output: {0} (+{1} bytes)", Output.Result, Output.Data.Length);
+					}
+					catch(Exception Ex)
+					{
+						Log.WriteLine("Exception: {0}", Ex.ToString());
+					}
+					finally
+					{
+						Client.Close();
+						Log.WriteLine("Closed connection.");
+					}
+					CurrentClient = null;
+				}
+			}
+			catch(Exception Ex)
+			{
+				if(!bDisposing)
+				{
+					Log.WriteLine("Exception: {0}", Ex.ToString());
+				}
+			}
+			Log.WriteLine("Closing socket.");
+		}
+		
+		public void Dispose()
+		{
+			bDisposing = true;
+
+			TcpClient Client = CurrentClient;
+			if(Client != null)
+			{
+				try { Client.Close(); } catch { }
+				Client = null;
+			}
+
+			if(Listener != null)
+			{
+				Listener.Stop();
+				Listener = null;
+			}
+
+			if(BackgroundThread != null)
+			{
+				BackgroundThread.Join();
+				BackgroundThread = null;
+			}
+		}
+	}
+}
diff --git a/Engine/Source/Programs/UnrealGameSync/UnrealGameSync/BuildStep.cs b/Engine/Source/Programs/UnrealGameSync/UnrealGameSync/BuildStep.cs
--- a/Engine/Source/Programs/UnrealGameSync/UnrealGameSync/BuildStep.cs
+++ b/Engine/Source/Programs/UnrealGameSync/UnrealGameSync/BuildStep.cs
@@ -1,4 +1,4 @@
-// Copyright 1998-2018 Epic Games, Inc. All Rights Reserved.
+﻿// Copyright 1998-2019 Epic Games, Inc. All Rights Reserved.
 
 using System;
 using System.Collections.Generic;
diff --git a/Engine/Source/Programs/UnrealGameSync/UnrealGameSync/ChildProcess.cs b/Engine/Source/Programs/UnrealGameSync/UnrealGameSync/ChildProcess.cs
--- a/Engine/Source/Programs/UnrealGameSync/UnrealGameSync/ChildProcess.cs
+++ b/Engine/Source/Programs/UnrealGameSync/UnrealGameSync/ChildProcess.cs
@@ -1,4 +1,4 @@
-// Copyright 1998-2018 Epic Games, Inc. All Rights Reserved.
+// Copyright 1998-2019 Epic Games, Inc. All Rights Reserved.
 
 using Microsoft.Win32.SafeHandles;
 using System;
diff --git a/Engine/Source/Programs/UnrealGameSync/UnrealGameSync/ConfigFile.cs b/Engine/Source/Programs/UnrealGameSync/UnrealGameSync/ConfigFile.cs
--- a/Engine/Source/Programs/UnrealGameSync/UnrealGameSync/ConfigFile.cs
+++ b/Engine/Source/Programs/UnrealGameSync/UnrealGameSync/ConfigFile.cs
@@ -1,4 +1,4 @@
-// Copyright 1998-2018 Epic Games, Inc. All Rights Reserved.
+﻿// Copyright 1998-2019 Epic Games, Inc. All Rights Reserved.
 
 using System;
 using System.Collections.Generic;
diff --git a/Engine/Source/Programs/UnrealGameSync/UnrealGameSync/Controls/BuildListControl.cs b/Engine/Source/Programs/UnrealGameSync/UnrealGameSync/Controls/BuildListControl.cs
--- a/Engine/Source/Programs/UnrealGameSync/UnrealGameSync/Controls/BuildListControl.cs
+++ b/Engine/Source/Programs/UnrealGameSync/UnrealGameSync/Controls/BuildListControl.cs
@@ -1,4 +1,4 @@
-// Copyright 1998-2018 Epic Games, Inc. All Rights Reserved.
+// Copyright 1998-2019 Epic Games, Inc. All Rights Reserved.
 
 using System;
 using System.Collections.Generic;
diff --git a/Engine/Source/Programs/UnrealGameSync/UnrealGameSync/Controls/CustomListViewControl.cs b/Engine/Source/Programs/UnrealGameSync/UnrealGameSync/Controls/CustomListViewControl.cs
--- a/Engine/Source/Programs/UnrealGameSync/UnrealGameSync/Controls/CustomListViewControl.cs
+++ b/Engine/Source/Programs/UnrealGameSync/UnrealGameSync/Controls/CustomListViewControl.cs
@@ -1,3 +1,5 @@
+// Copyright 1998-2019 Epic Games, Inc. All Rights Reserved.
+
 using System;
 using System.Collections.Generic;
 using System.ComponentModel;
diff --git a/Engine/Source/Programs/UnrealGameSync/UnrealGameSync/Controls/ErrorPanel.cs b/Engine/Source/Programs/UnrealGameSync/UnrealGameSync/Controls/ErrorPanel.cs
--- a/Engine/Source/Programs/UnrealGameSync/UnrealGameSync/Controls/ErrorPanel.cs
+++ b/Engine/Source/Programs/UnrealGameSync/UnrealGameSync/Controls/ErrorPanel.cs
@@ -1,4 +1,4 @@
-﻿// Copyright 1998-2018 Epic Games, Inc. All Rights Reserved.
+﻿// Copyright 1998-2019 Epic Games, Inc. All Rights Reserved.
 
 using System;
 using System.Collections.Generic;
diff --git a/Engine/Source/Programs/UnrealGameSync/UnrealGameSync/Controls/LogControl.cs b/Engine/Source/Programs/UnrealGameSync/UnrealGameSync/Controls/LogControl.cs
--- a/Engine/Source/Programs/UnrealGameSync/UnrealGameSync/Controls/LogControl.cs
+++ b/Engine/Source/Programs/UnrealGameSync/UnrealGameSync/Controls/LogControl.cs
@@ -1,4 +1,4 @@
-﻿// Copyright 1998-2018 Epic Games, Inc. All Rights Reserved.
+﻿// Copyright 1998-2019 Epic Games, Inc. All Rights Reserved.
 
 using System;
 using System.Collections.Generic;
diff --git a/Engine/Source/Programs/UnrealGameSync/UnrealGameSync/Controls/LogSplitContainer.cs b/Engine/Source/Programs/UnrealGameSync/UnrealGameSync/Controls/LogSplitContainer.cs
--- a/Engine/Source/Programs/UnrealGameSync/UnrealGameSync/Controls/LogSplitContainer.cs
+++ b/Engine/Source/Programs/UnrealGameSync/UnrealGameSync/Controls/LogSplitContainer.cs
@@ -1,4 +1,4 @@
-﻿// Copyright 1998-2018 Epic Games, Inc. All Rights Reserved.
+﻿// Copyright 1998-2019 Epic Games, Inc. All Rights Reserved.
 
 using System;
 using System.Collections.Generic;
diff --git a/Engine/Source/Programs/UnrealGameSync/UnrealGameSync/Controls/StatusPanel.cs b/Engine/Source/Programs/UnrealGameSync/UnrealGameSync/Controls/StatusPanel.cs
--- a/Engine/Source/Programs/UnrealGameSync/UnrealGameSync/Controls/StatusPanel.cs
+++ b/Engine/Source/Programs/UnrealGameSync/UnrealGameSync/Controls/StatusPanel.cs
@@ -1,4 +1,4 @@
-// Copyright 1998-2018 Epic Games, Inc. All Rights Reserved.
+// Copyright 1998-2019 Epic Games, Inc. All Rights Reserved.
 
 using System;
 using System.Collections.Generic;
@@ -6,12 +6,24 @@ using System.ComponentModel;
 using System.Drawing;
 using System.Drawing.Drawing2D;
 using System.Linq;
+using System.Runtime.InteropServices;
 using System.Text;
 using System.Threading.Tasks;
 using System.Windows.Forms;
 
 namespace UnrealGameSync
 {
+	// For some reason, WinForms uses embedded resources for certain cursors, which don't scale correctly for high DPI modes
+	static class NativeCursors
+	{
+		[DllImport("user32.dll")]
+		public static extern IntPtr LoadCursor(IntPtr hInstance, IntPtr lpCursorName);
+
+		const int IDC_HAND = 32649;
+
+		public static readonly Cursor Hand = new Cursor(LoadCursor(IntPtr.Zero, new IntPtr(IDC_HAND)));
+	}
+
 	class StatusElementResources
 	{
 		Dictionary<FontStyle, Font> FontCache = new Dictionary<FontStyle, Font>();
@@ -148,7 +160,7 @@ namespace UnrealGameSync
 			Text = InText;
 			Style = InStyle;
 			LinkAction = InLinkAction;
-			Cursor = Cursors.Hand;
+			Cursor = NativeCursors.Hand;
 		}
 
 		public override void OnClick(Point Location)
@@ -198,7 +210,7 @@ namespace UnrealGameSync
 			ClickAction = InClickAction;
 			if(ClickAction != null)
 			{
-				Cursor = Cursors.Hand;
+				Cursor = NativeCursors.Hand;
 			}
 		}
 
@@ -461,6 +473,11 @@ namespace UnrealGameSync
 				Resources = new StatusElementResources(Font);
 			}
 
+			if(TintColor != NewTintColor)
+			{
+				Invalidate();
+			}
+
 			InvalidateElements();
 			Lines.Clear();
 			Lines.AddRange(NewLines);
diff --git a/Engine/Source/Programs/UnrealGameSync/UnrealGameSync/Controls/SyncFilterControl.cs b/Engine/Source/Programs/UnrealGameSync/UnrealGameSync/Controls/SyncFilterControl.cs
--- a/Engine/Source/Programs/UnrealGameSync/UnrealGameSync/Controls/SyncFilterControl.cs
+++ b/Engine/Source/Programs/UnrealGameSync/UnrealGameSync/Controls/SyncFilterControl.cs
@@ -1,4 +1,4 @@
-﻿// Copyright 1998-2018 Epic Games, Inc. All Rights Reserved.
+﻿// Copyright 1998-2019 Epic Games, Inc. All Rights Reserved.
 
 using System;
 using System.Collections.Generic;
diff --git a/Engine/Source/Programs/UnrealGameSync/UnrealGameSync/Controls/TabControl.cs b/Engine/Source/Programs/UnrealGameSync/UnrealGameSync/Controls/TabControl.cs
--- a/Engine/Source/Programs/UnrealGameSync/UnrealGameSync/Controls/TabControl.cs
+++ b/Engine/Source/Programs/UnrealGameSync/UnrealGameSync/Controls/TabControl.cs
@@ -1,4 +1,4 @@
-// Copyright 1998-2018 Epic Games, Inc. All Rights Reserved.
+// Copyright 1998-2019 Epic Games, Inc. All Rights Reserved.
 
 using System;
 using System.Collections.Generic;
diff --git a/Engine/Source/Programs/UnrealGameSync/UnrealGameSync/Controls/TextBoxWithCueBanner.cs b/Engine/Source/Programs/UnrealGameSync/UnrealGameSync/Controls/TextBoxWithCueBanner.cs
--- a/Engine/Source/Programs/UnrealGameSync/UnrealGameSync/Controls/TextBoxWithCueBanner.cs
+++ b/Engine/Source/Programs/UnrealGameSync/UnrealGameSync/Controls/TextBoxWithCueBanner.cs
@@ -1,4 +1,4 @@
-// Copyright 1998-2018 Epic Games, Inc. All Rights Reserved.
+// Copyright 1998-2019 Epic Games, Inc. All Rights Reserved.
 
 using System;
 using System.Collections.Generic;
diff --git a/Engine/Source/Programs/UnrealGameSync/UnrealGameSync/Controls/WorkspaceControl.cs b/Engine/Source/Programs/UnrealGameSync/UnrealGameSync/Controls/WorkspaceControl.cs
--- a/Engine/Source/Programs/UnrealGameSync/UnrealGameSync/Controls/WorkspaceControl.cs
+++ b/Engine/Source/Programs/UnrealGameSync/UnrealGameSync/Controls/WorkspaceControl.cs
@@ -1,4 +1,4 @@
-// Copyright 1998-2018 Epic Games, Inc. All Rights Reserved.
+// Copyright 1998-2019 Epic Games, Inc. All Rights Reserved.
 
 using System;
 using System.Collections.Generic;
@@ -41,8 +41,12 @@ namespace UnrealGameSync
 		void SetTabNames(TabLabels TabNames);
 		void SetupScheduledSync();
 		void UpdateProgress();
+		void ModifyApplicationSettings();
 	}
 
+	delegate void WorkspaceStartupCallback(WorkspaceControl Workspace, bool bCancel);
+	delegate void WorkspaceUpdateCallback(WorkspaceUpdateResult Result);
+
 	partial class WorkspaceControl : UserControl, IMainWindowTabPanel
 	{
 		enum HorizontalAlignment
@@ -132,6 +136,8 @@ namespace UnrealGameSync
 		[DllImport("uxtheme.dll", CharSet = CharSet.Unicode)]
 		static extern int SetWindowTheme(IntPtr hWnd, string pszSubAppName, string pszSubIdList);
 
+		const int BuildListExpandCount = 250;
+
 		const string EditorArchiveType = "Editor";
 
 		IWorkspaceControlOwner Owner;
@@ -178,6 +184,7 @@ namespace UnrealGameSync
 		SynchronizationContext MainThreadSynchronizationContext;
 		bool bIsDisposing;
 
+		bool bUnstable;
 		string EditorTargetName;
 		bool bIsEnterpriseProject;
 		PerforceMonitor PerforceMonitor;
@@ -196,7 +203,7 @@ namespace UnrealGameSync
 
 		Dictionary<string, int> NotifiedBuildTypeToChangeNumber = new Dictionary<string,int>();
 
-		TimeSpan ServerTimeZone;
+		bool bMouseOverExpandLink;
 
 		string HoverBadgeUniqueId = null;
 		bool bHoverSync;
@@ -226,12 +233,18 @@ namespace UnrealGameSync
 		string LastColumnSettings;
 		List<ColumnHeader> CustomColumns;
 		int MaxBuildBadgeChars;
+		ListViewItem ExpandItem;
 
 		bool bUpdateBuildListPosted;
 		bool bUpdateBuildMetadataPosted;
 		bool bUpdateReviewsPosted;
 
-		public WorkspaceControl(IWorkspaceControlOwner InOwner, string InApiUrl, string InOriginalExecutableFileName, DetectProjectSettingsTask DetectSettings, LineBasedTextWriter InLog, UserSettings InSettings)
+		WorkspaceUpdateCallback UpdateCallback;
+
+		System.Threading.Timer StartupTimer;
+		List<WorkspaceStartupCallback> StartupCallbacks;
+
+		public WorkspaceControl(IWorkspaceControlOwner InOwner, string InApiUrl, string InOriginalExecutableFileName, bool bInUnstable, DetectProjectSettingsTask DetectSettings, LineBasedTextWriter InLog, UserSettings InSettings)
 		{
 			InitializeComponent();
 
@@ -241,6 +254,7 @@ namespace UnrealGameSync
 			ApiUrl = InApiUrl;
 			DataFolder = DetectSettings.DataFolder;
 			OriginalExecutableFileName = InOriginalExecutableFileName;
+			bUnstable = bInUnstable;
 			Log = InLog;
 			Settings = InSettings;
 			WorkspaceSettings = InSettings.FindOrAddWorkspace(DetectSettings.BranchClientPath);
@@ -280,7 +294,6 @@ namespace UnrealGameSync
 			EditorTargetName = DetectSettings.NewProjectEditorTarget;
 			bIsEnterpriseProject = DetectSettings.bIsEnterpriseProject;
 			StreamName = DetectSettings.StreamName;
-			ServerTimeZone = DetectSettings.ServerTimeZone;
 
 			// Update the branch directory
 			BranchDirectoryName = DetectSettings.BranchDirectoryName;
@@ -301,7 +314,7 @@ namespace UnrealGameSync
 
 			string ProjectLogBaseName = Path.Combine(DataFolder, String.Format("{0}@{1}", PerforceClient.ClientName, DetectSettings.BranchClientPath.Replace("//" + PerforceClient.ClientName + "/", "").Trim('/').Replace("/", "$")));
 
-			PerforceMonitor = new PerforceMonitor(PerforceClient, DetectSettings.BranchClientPath, DetectSettings.NewSelectedClientFileName, DetectSettings.NewSelectedProjectIdentifier, ProjectLogBaseName + ".p4.log", DetectSettings.bIsEnterpriseProject, DetectSettings.LatestProjectConfigFile, DetectSettings.LocalConfigFiles);
+			PerforceMonitor = new PerforceMonitor(PerforceClient, DetectSettings.BranchClientPath, DetectSettings.NewSelectedClientFileName, DetectSettings.NewSelectedProjectIdentifier, ProjectLogBaseName + ".p4.log", DetectSettings.bIsEnterpriseProject, DetectSettings.LatestProjectConfigFile, DetectSettings.CacheFolder, DetectSettings.LocalConfigFiles);
 			PerforceMonitor.OnUpdate += UpdateBuildListCallback;
 			PerforceMonitor.OnUpdateMetadata += UpdateBuildMetadataCallback;
 			PerforceMonitor.OnStreamChange += StreamChangedCallback;
@@ -324,13 +337,53 @@ namespace UnrealGameSync
 			UpdateStatusPanel();
 			UpdateServiceBadges();
 
-			if(CurrentChangeNumber != -1)
+			PerforceMonitor.Start();
+			EventMonitor.Start();
+
+			StartupTimer = new System.Threading.Timer(x => MainThreadSynchronizationContext.Post((o) => { if(!IsDisposed){ StartupTimerElapsed(false); } }, null), null, TimeSpan.FromSeconds(20.0), TimeSpan.FromMilliseconds(-1.0));
+			StartupCallbacks = new List<WorkspaceStartupCallback>();
+		}
+
+		private void CheckForStartupComplete()
+		{
+			if(StartupTimer != null)
 			{
-				SelectChange(CurrentChangeNumber);
+				int LatestChangeNumber;
+				if(FindChangeToSync(Settings.SyncType, out LatestChangeNumber))
+				{
+					StartupTimerElapsed(false);
+				}
 			}
+		}
 
-			PerforceMonitor.Start();
-			EventMonitor.Start();
+		private void StartupTimerElapsed(bool bCancel)
+		{
+			if(StartupTimer != null)
+			{
+				StartupTimer.Dispose();
+				StartupTimer = null;
+			}
+
+			if(StartupCallbacks != null)
+			{
+				foreach(WorkspaceStartupCallback StartupCallback in StartupCallbacks)
+				{
+					StartupCallback(this, bCancel);
+				}
+				StartupCallbacks = null;
+			}
+		}
+
+		public void AddStartupCallback(WorkspaceStartupCallback StartupCallback)
+		{
+			if(StartupTimer == null)
+			{
+				StartupCallback(this, false);
+			}
+			else
+			{
+				StartupCallbacks.Add(StartupCallback);
+			}
 		}
 
 		private void UpdateColumnSettings()
@@ -517,6 +570,20 @@ namespace UnrealGameSync
 
 			UpdateTimer.Stop();
 
+			if(StartupCallbacks != null)
+			{
+				foreach(WorkspaceStartupCallback StartupCallback in StartupCallbacks)
+				{
+					StartupCallback(this, true);
+				}
+				StartupCallbacks = null;
+			}
+
+			if(StartupTimer != null)
+			{
+				StartupTimer.Dispose();
+				StartupTimer = null;
+			}
 			if(NotificationWindow != null)
 			{
 				NotificationWindow.Dispose();
@@ -638,65 +705,54 @@ namespace UnrealGameSync
 		private void BuildList_OnScroll()
 		{
 			PendingSelectedChangeNumber = -1;
-			UpdateNumRequestedBuilds(false);
 		}
 
-		void UpdateNumRequestedBuilds(bool bAllowShrink)
+		void ShrinkNumRequestedBuilds()
 		{
-			if(PerforceMonitor != null)
+			if(PerforceMonitor != null && BuildList.Items.Count > 0 && PendingSelectedChangeNumber == -1)
 			{
-				if(!Settings.bShowUnreviewedChanges)
+				// Find the number of visible items using a (slightly wasteful) binary search
+				int VisibleItemCount = 1;
+				for(int StepSize = BuildList.Items.Count / 2; StepSize >= 1; )
 				{
-					PerforceMonitor.PendingMaxChanges = 1000;
-				}
-				else
-				{
-					int NumItemsPerPage = Math.Max(BuildList.GetVisibleItemsPerPage(), 10);
-
-					// Find the number of visible items using a (slightly wasteful) binary search
-					int VisibleItemCount = 1;
-					for(int StepSize = BuildList.Items.Count / 2; StepSize >= 1; )
-					{
-						int TestIndex = VisibleItemCount + StepSize;
-						if(TestIndex < BuildList.Items.Count && BuildList.GetItemRect(TestIndex).Top < BuildList.Height)
-						{
-							VisibleItemCount += StepSize;
-						}
-						else
-						{
-							StepSize /= 2;
-						}
-					}
-
-					// Increase or decrease the number of builds we want, with a bit of rubber-banding
-					const int IncreaseStep = 50;
-					if(VisibleItemCount > BuildList.Items.Count - 20)
+					int TestIndex = VisibleItemCount + StepSize;
+					if(TestIndex < BuildList.Items.Count && BuildList.GetItemRect(TestIndex).Top < BuildList.Height)
 					{
-						PerforceMonitor.PendingMaxChanges = NumChanges + IncreaseStep;
+						VisibleItemCount += StepSize;
 					}
-					else if(bAllowShrink)
+					else
 					{
-						int NewNumChanges = ListIndexToChangeIndex[VisibleItemCount - 1] + IncreaseStep;
-						if(NewNumChanges < NumChanges)
-						{
-							PerforceMonitor.PendingMaxChanges = NewNumChanges;
-						}
+						StepSize /= 2;
 					}
 				}
+
+				// Figure out the last index to ensure is visible
+				int LastVisibleIndex = VisibleItemCount;
+				if(LastVisibleIndex >= ListIndexToChangeIndex.Count)
+				{
+					LastVisibleIndex = ListIndexToChangeIndex.Count - 1;
+				}
+
+				// Get the max number of changes to ensure this
+				int NewPendingMaxChanges = ListIndexToChangeIndex[LastVisibleIndex];
+				NewPendingMaxChanges = PerforceMonitor.InitialMaxChangesValue + ((Math.Max(NewPendingMaxChanges - PerforceMonitor.InitialMaxChangesValue, 0) + BuildListExpandCount - 1) / BuildListExpandCount) * BuildListExpandCount;
+
+				// Shrink the number of changes retained by the PerforceMonitor class
+				if(PerforceMonitor.PendingMaxChanges > NewPendingMaxChanges)
+				{
+					PerforceMonitor.PendingMaxChanges = NewPendingMaxChanges;
+				}
 			}
 		}
 
 		void StartSync(int ChangeNumber)
+		{
+			StartSync(ChangeNumber, null);
+		}
+
+		void StartSync(int ChangeNumber, WorkspaceUpdateCallback Callback)
 		{
 			WorkspaceUpdateOptions Options = WorkspaceUpdateOptions.Sync | WorkspaceUpdateOptions.SyncArchives | WorkspaceUpdateOptions.GenerateProjectFiles;
-			if(Settings.bAutoResolveConflicts)
-			{
-				Options |= WorkspaceUpdateOptions.AutoResolveChanges;
-			}
-			if(Settings.bUseIncrementalBuilds)
-			{
-				Options |= WorkspaceUpdateOptions.UseIncrementalBuilds;
-			}
 			if(Settings.bBuildAfterSync)
 			{
 				Options |= WorkspaceUpdateOptions.Build;
@@ -709,14 +765,15 @@ namespace UnrealGameSync
 			{
 				Options |= WorkspaceUpdateOptions.OpenSolutionAfterSync;
 			}
-			if(WorkspaceSettings.bSyncAllProjects ?? Settings.bSyncAllProjects)
-			{
-				Options |= WorkspaceUpdateOptions.SyncAllProjects;
-			}
-			StartWorkspaceUpdate(ChangeNumber, Options);
+			StartWorkspaceUpdate(ChangeNumber, Options, Callback);
 		}
 
 		void StartWorkspaceUpdate(int ChangeNumber, WorkspaceUpdateOptions Options)
+		{
+			StartWorkspaceUpdate(ChangeNumber, Options, null);
+		}
+
+		void StartWorkspaceUpdate(int ChangeNumber, WorkspaceUpdateOptions Options, WorkspaceUpdateCallback Callback)
 		{
 			if((Options & (WorkspaceUpdateOptions.Sync | WorkspaceUpdateOptions.Build)) != 0 && GetProcessesRunningInWorkspace().Length > 0)
 			{
@@ -753,11 +810,30 @@ namespace UnrealGameSync
 				}
 				Context.ArchiveTypeToDepotPath.Add(EditorArchiveType, EditorArchivePath);
 			}
-			StartWorkspaceUpdate(Context);
+			StartWorkspaceUpdate(Context, Callback);
 		}
 
-		void StartWorkspaceUpdate(WorkspaceUpdateContext Context)
+		void StartWorkspaceUpdate(WorkspaceUpdateContext Context, WorkspaceUpdateCallback Callback)
 		{
+			if(Settings.bAutoResolveConflicts)
+			{
+				Context.Options |= WorkspaceUpdateOptions.AutoResolveChanges;
+			}
+			if(Settings.bUseIncrementalBuilds)
+			{
+				Context.Options |= WorkspaceUpdateOptions.UseIncrementalBuilds;
+			}
+			if(WorkspaceSettings.bSyncAllProjects ?? Settings.bSyncAllProjects)
+			{
+				Context.Options |= WorkspaceUpdateOptions.SyncAllProjects | WorkspaceUpdateOptions.IncludeAllProjectsInSolution;
+			}
+			if(WorkspaceSettings.bIncludeAllProjectsInSolution ?? Settings.bIncludeAllProjectsInSolution)
+			{
+				Context.Options |= WorkspaceUpdateOptions.IncludeAllProjectsInSolution;
+			}
+
+			UpdateCallback = Callback;
+
 			Context.StartTime = DateTime.UtcNow;
 			Context.PerforceSyncOptions = (PerforceSyncOptions)Settings.SyncOptions.Clone();
 
@@ -810,6 +886,12 @@ namespace UnrealGameSync
 
 				Workspace.CancelUpdate();
 
+				if(UpdateCallback != null)
+				{
+					UpdateCallback(WorkspaceUpdateResult.Canceled);
+					UpdateCallback = null;
+				}
+
 				UpdateTimer.Stop();
 
 				UpdateSyncActionCheckboxes();
@@ -863,7 +945,7 @@ namespace UnrealGameSync
 				DeleteWindow Window = new DeleteWindow(Context.DeleteFiles);
 				if(Window.ShowDialog(this) == DialogResult.OK)
 				{
-					StartWorkspaceUpdate(Context);
+					StartWorkspaceUpdate(Context, UpdateCallback);
 					return;
 				}
 			}
@@ -875,7 +957,7 @@ namespace UnrealGameSync
 				ClobberWindow Window = new ClobberWindow(Context.ClobberFiles);
 				if(Window.ShowDialog(this) == DialogResult.OK)
 				{
-					StartWorkspaceUpdate(Context);
+					StartWorkspaceUpdate(Context, UpdateCallback);
 					return;
 				}
 			}
@@ -899,6 +981,12 @@ namespace UnrealGameSync
 				}
 			}
 
+			if(UpdateCallback != null)
+			{
+				UpdateCallback(Result);
+				UpdateCallback = null;
+			}
+
 			DesiredTaskbarState = Tuple.Create((Result == WorkspaceUpdateResult.Success)? TaskbarState.NoProgress : TaskbarState.Error, 0.0f);
 			Owner.UpdateProgress();
 
@@ -937,11 +1025,11 @@ namespace UnrealGameSync
 					}
 				}
 
-				int SelectedChange = (BuildList.SelectedItems.Count > 0)? ((PerforceChangeSummary)BuildList.SelectedItems[0].Tag).Number : -1;
-
 				ChangeNumberToArchivePath.Clear();
 				ChangeNumberToLayoutInfo.Clear();
 
+				ExpandItem = null;
+
 				BuildList.BeginUpdate();
 
 				foreach(ListViewGroup Group in BuildList.Groups)
@@ -987,7 +1075,13 @@ namespace UnrealGameSync
 
 							ListViewGroup Group;
 
-							string GroupName = Change.Date.ToString("D");//"dddd\\,\\ h\\.mmtt");
+							DateTime DisplayTime = Change.Date;
+							if(Settings.bShowLocalTimes)
+							{
+								DisplayTime = (DisplayTime - PerforceMonitor.ServerTimeZone).ToLocalTime();
+							}
+
+							string GroupName = DisplayTime.ToString("D");//"dddd\\,\\ h\\.mmtt");
 							for(int Idx = 0;;Idx++)
 							{
 								if(Idx == BuildList.Groups.Count)
@@ -1004,12 +1098,6 @@ namespace UnrealGameSync
 								}
 							}
 
-							DateTime DisplayTime = Change.Date;
-							if(Settings.bShowLocalTimes)
-							{
-								DisplayTime = (DisplayTime - ServerTimeZone).ToLocalTime();
-							}
-
 							string[] SubItemLabels = new string[BuildList.Columns.Count];
 							SubItemLabels[ChangeColumn.Index] = String.Format("{0}", Change.Number);
 							SubItemLabels[TimeColumn.Index] = DisplayTime.ToString("h\\.mmtt");
@@ -1052,6 +1140,19 @@ namespace UnrealGameSync
 					BuildList.Groups.RemoveAt(0);
 				}
 
+				if(Changes.Count > 0)
+				{
+					ExpandItem = new ListViewItem((BuildList.Groups.Count > 0)? BuildList.Groups[BuildList.Groups.Count - 1] : null);
+					ExpandItem.Tag = null;
+					ExpandItem.Selected = false;
+					ExpandItem.Text = "";
+					for(int ColumnIdx = 1; ColumnIdx < BuildList.Columns.Count; ColumnIdx++)
+					{
+						ExpandItem.SubItems.Add(new ListViewItem.ListViewSubItem(ExpandItem, ""));
+					}
+					BuildList.Items.Add(ExpandItem);
+				}
+
 				BuildList.EndUpdate();
 
 				if(PendingSelectedChangeNumber != -1)
@@ -1073,6 +1174,30 @@ namespace UnrealGameSync
 
 		bool ShouldShowChange(PerforceChangeSummary Change, string[] ExcludeChanges)
 		{
+			if(ProjectSettings.FilterType != FilterType.None)
+			{
+				PerforceChangeDetails Details;
+				if(!PerforceMonitor.TryGetChangeDetails(Change.Number, out Details))
+				{
+					return false;
+				}
+				if(ProjectSettings.FilterType == FilterType.Code && !Details.bContainsCode)
+				{
+					return false;
+				}
+				if(ProjectSettings.FilterType == FilterType.Content && !Details.bContainsContent)
+				{
+					return false;
+				}
+			}
+			if(ProjectSettings.FilterBadges.Count > 0)
+			{
+				EventSummary Summary = EventMonitor.GetSummaryForChange(Change.Number);
+				if(Summary == null || !Summary.Badges.Any(x => ProjectSettings.FilterBadges.Contains(x.BadgeName)))
+				{
+					return false;
+				}
+			}
 			if(!Settings.bShowAutomatedChanges)
 			{
 				foreach(string ExcludeChange in ExcludeChanges)
@@ -1133,15 +1258,18 @@ namespace UnrealGameSync
 			// Add a dummy group for any other badges we have
 			foreach(ListViewItem Item in BuildList.Items)
 			{
-				EventSummary Summary = EventMonitor.GetSummaryForChange(((PerforceChangeSummary)Item.Tag).Number);
-				if(Summary != null)
+				if(Item.Tag != null)
 				{
-					foreach(BadgeData Badge in Summary.Badges)
+					EventSummary Summary = EventMonitor.GetSummaryForChange(((PerforceChangeSummary)Item.Tag).Number);
+					if(Summary != null)
 					{
-						string BadgeSlot = Badge.BadgeName;
-						if(!BadgeNameToGroup.ContainsKey(BadgeSlot))
+						foreach(BadgeData Badge in Summary.Badges)
 						{
-							BadgeNameToGroup.Add(BadgeSlot, "XXXX");
+							string BadgeSlot = Badge.BadgeName;
+							if(!BadgeNameToGroup.ContainsKey(BadgeSlot))
+							{
+								BadgeNameToGroup.Add(BadgeSlot, "XXXX");
+							}
 						}
 					}
 				}
@@ -1176,6 +1304,13 @@ namespace UnrealGameSync
 			UpdateServiceBadges();
 			UpdateStatusPanel();
 			UpdateBuildFailureNotification();
+			CheckForStartupComplete();
+
+			// If we are filtering by badges, we may also need to update the build list
+			if(ProjectSettings.FilterType != FilterType.None || ProjectSettings.FilterBadges.Count > 0)
+			{
+				UpdateBuildList();
+			}
 		}
 
 		void UpdateMaxBuildBadgeChars()
@@ -1270,6 +1405,7 @@ namespace UnrealGameSync
 			EventMonitor.ApplyUpdates();
 			Refresh();
 			UpdateBuildFailureNotification();
+			CheckForStartupComplete();
 		}
 
 		void UpdateBuildFailureNotification()
@@ -1387,23 +1523,106 @@ namespace UnrealGameSync
 			e.DrawDefault = true;
 		}
 
-		private void BuildList_DrawItem(object sender, DrawListViewItemEventArgs e)
+		class ExpandRowLayout
+		{
+			public string MainText;
+			public Rectangle MainRect;
+			public string LinkText;
+			public Rectangle LinkRect;
+		}
+
+		private ExpandRowLayout LayoutExpandRow(Graphics Graphics, Rectangle Bounds)
 		{
-			if(e.State.HasFlag(ListViewItemStates.Selected))
+			ExpandRowLayout Layout = new ExpandRowLayout();
+
+			int CurrentMaxChanges = PerforceMonitor.CurrentMaxChanges;
+
+			string ShowingChanges;
+			if(ListIndexToChangeIndex.Count == CurrentMaxChanges)
 			{
-				BuildList.DrawSelectedBackground(e.Graphics, e.Bounds);
+				ShowingChanges = String.Format("Showing {0} changes.", ListIndexToChangeIndex.Count);
 			}
-			else if(e.ItemIndex == BuildList.HoverItem)
+			else
 			{
-				BuildList.DrawTrackedBackground(e.Graphics, e.Bounds);
+				ShowingChanges = String.Format("Showing {0}/{1} changes.", ListIndexToChangeIndex.Count, CurrentMaxChanges);
 			}
-			else if(((PerforceChangeSummary)e.Item.Tag).Number == Workspace.PendingChangeNumber)
+
+			if(PerforceMonitor.PendingMaxChanges > CurrentMaxChanges)
 			{
-				BuildList.DrawTrackedBackground(e.Graphics, e.Bounds);
+				Layout.MainText = String.Format("{0}. Fetching {1} more from server...  ", ShowingChanges, PerforceMonitor.PendingMaxChanges - CurrentMaxChanges);
+				Layout.LinkText = "Cancel";
 			}
 			else
+			{
+				Layout.MainText = String.Format("{0}  ", ShowingChanges);
+				Layout.LinkText = "Show more...";
+			}
+
+			Size MainTextSize = TextRenderer.MeasureText(Graphics, Layout.MainText, BuildFont, new Size(1000, 1000), TextFormatFlags.NoPadding);
+			Size LinkTextSize = TextRenderer.MeasureText(Graphics, Layout.LinkText, BuildFont, new Size(1000, 1000), TextFormatFlags.NoPadding);
+
+			int MinX = Bounds.Left + (Bounds.Width - MainTextSize.Width - LinkTextSize.Width) / 2;
+			int MinY = Bounds.Bottom - MainTextSize.Height - 1;
+
+			Layout.MainRect = new Rectangle(new Point(MinX, MinY), MainTextSize);
+			Layout.LinkRect = new Rectangle(new Point(MinX + MainTextSize.Width, MinY), LinkTextSize);
+
+			return Layout;
+		}
+
+		private void DrawExpandRow(Graphics Graphics, ExpandRowLayout Layout)
+		{
+			TextRenderer.DrawText(Graphics, Layout.MainText, BuildFont, Layout.MainRect, SystemColors.WindowText, TextFormatFlags.SingleLine | TextFormatFlags.VerticalCenter | TextFormatFlags.NoPrefix | TextFormatFlags.NoPadding);
+
+			Color LinkColor = SystemColors.HotTrack;
+			if(bMouseOverExpandLink)
+			{
+				LinkColor = Color.FromArgb(LinkColor.B, LinkColor.G, LinkColor.R);
+			}
+
+			TextRenderer.DrawText(Graphics, Layout.LinkText, BuildFont, Layout.LinkRect, LinkColor, TextFormatFlags.SingleLine | TextFormatFlags.VerticalCenter | TextFormatFlags.NoPrefix | TextFormatFlags.NoPadding);
+		}
+
+		private bool HitTestExpandLink(Point Location)
+		{
+			if(ExpandItem == null)
+			{
+				return false;
+			}
+			using(Graphics Graphics = Graphics.FromHwnd(IntPtr.Zero))
+			{
+				ExpandRowLayout Layout = LayoutExpandRow(Graphics, ExpandItem.Bounds);
+				return Layout.LinkRect.Contains(Location);
+			}
+		}
+
+		private void BuildList_DrawItem(object sender, DrawListViewItemEventArgs e)
+		{
+			if(e.Item == ExpandItem)
 			{
 				BuildList.DrawDefaultBackground(e.Graphics, e.Bounds);
+
+				ExpandRowLayout ExpandLayout = LayoutExpandRow(e.Graphics, e.Bounds);
+				DrawExpandRow(e.Graphics, ExpandLayout);
+			}
+			else
+			{
+				if(e.State.HasFlag(ListViewItemStates.Selected))
+				{
+					BuildList.DrawSelectedBackground(e.Graphics, e.Bounds);
+				}
+				else if(e.ItemIndex == BuildList.HoverItem)
+				{
+					BuildList.DrawTrackedBackground(e.Graphics, e.Bounds);
+				}
+				else if(((PerforceChangeSummary)e.Item.Tag).Number == Workspace.PendingChangeNumber)
+				{
+					BuildList.DrawTrackedBackground(e.Graphics, e.Bounds);
+				}
+				else
+				{
+					BuildList.DrawDefaultBackground(e.Graphics, e.Bounds);
+				}
 			}
 		}
 
@@ -1487,6 +1706,10 @@ namespace UnrealGameSync
 		private void BuildList_DrawSubItem(object sender, DrawListViewSubItemEventArgs e)
 		{
 			PerforceChangeSummary Change = (PerforceChangeSummary)e.Item.Tag;
+			if(Change == null)
+			{
+				return;
+			}
 
 			float DpiScaleX = e.Graphics.DpiX / 96.0f;
 			float DpiScaleY = e.Graphics.DpiY / 96.0f;
@@ -1504,6 +1727,14 @@ namespace UnrealGameSync
 			int BadgeAlpha = bAllowSync? 255 : 128;
 			Color TextColor = (bAllowSync || Change.Number == Workspace.PendingChangeNumber || Change.Number == Workspace.CurrentChangeNumber || (WorkspaceSettings != null && WorkspaceSettings.AdditionalChangeNumbers.Contains(Change.Number)))? SystemColors.WindowText : Blend(SystemColors.Window, SystemColors.WindowText, 0.25f);
 
+			const int FadeRange = 6;
+			if(e.ItemIndex >= BuildList.Items.Count - FadeRange)
+			{
+				float Opacity = (float)(BuildList.Items.Count - e.ItemIndex - 0.9f) / FadeRange;
+				BadgeAlpha = (int)(BadgeAlpha * Opacity);
+				TextColor = Blend(SystemColors.Window, TextColor, Opacity);
+			}
+
 			if(e.ColumnIndex == IconColumn.Index)
 			{
 				EventSummary Summary = EventMonitor.GetSummaryForChange(Change.Number);
@@ -2024,7 +2255,7 @@ namespace UnrealGameSync
 			Color HoverBadgeColor = Color.FromArgb(BadgeColor.A, Math.Min(BadgeColor.R + 32, 255), Math.Min(BadgeColor.G + 32, 255), Math.Min(BadgeColor.B + 32, 255));
 
 			Action ClickHandler;
-			if(BadgeData == null || BadgeData.Url == null)
+			if(BadgeData == null || String.IsNullOrEmpty(BadgeData.Url))
 			{
 				ClickHandler = null;
 			}
@@ -2253,13 +2484,16 @@ namespace UnrealGameSync
 				if(HitTest.Item != null)
 				{
 					PerforceChangeSummary Change = (PerforceChangeSummary)HitTest.Item.Tag;
-					if(Change.Number == Workspace.CurrentChangeNumber)
-					{
-						LaunchEditor();
-					}
-					else
+					if(Change != null)
 					{
-						StartSync(Change.Number);
+						if(Change.Number == Workspace.CurrentChangeNumber)
+						{
+							LaunchEditor();
+						}
+						else
+						{
+							StartSync(Change.Number, null);
+						}
 					}
 				}
 			}
@@ -2316,7 +2550,7 @@ namespace UnrealGameSync
 						}
 		
 						WorkspaceUpdateContext Context = new WorkspaceUpdateContext(Workspace.CurrentChangeNumber, Options, null, GetDefaultBuildStepObjects(), ProjectSettings.BuildSteps, null, GetWorkspaceVariables(Workspace.CurrentChangeNumber));
-						StartWorkspaceUpdate(Context);
+						StartWorkspaceUpdate(Context, null);
 					}
 				}
 			}
@@ -2888,7 +3122,7 @@ namespace UnrealGameSync
 				}
 				else if(Workspace.Perforce != null)
 				{
-					if(!Workspace.Perforce.HasOpenFiles(Log) || MessageBox.Show("You have files open for edit in this workspace. If you continue, you will not be able to submit them until you switch back.\n\nContinue switching workspaces?", "Files checked out", MessageBoxButtons.YesNo) == DialogResult.Yes)
+					if(!Workspace.Perforce.HasOpenFiles(Log) || MessageBox.Show("You have files open for edit in this workspace. If you continue, you will not be able to submit them until you switch back.\n\nContinue switching streams?", "Files checked out", MessageBoxButtons.YesNo) == DialogResult.Yes)
 					{
 						if(Workspace.Perforce.SwitchStream(NewStreamName, Log))
 						{
@@ -2948,12 +3182,15 @@ namespace UnrealGameSync
 
 			PendingSelectedChangeNumber = ChangeNumber;
 
-			if(BuildList.Items.Count > 0)
+			int CurrentMaxChanges = PerforceMonitor.CurrentMaxChanges;
+			if(PerforceMonitor.PendingMaxChanges <= CurrentMaxChanges && BuildList.SelectedItems.Count == 0)
 			{
-				BuildList.Items[BuildList.Items.Count - 1].EnsureVisible();
+				PerforceMonitor.PendingMaxChanges = CurrentMaxChanges + BuildListExpandCount;
+				if(ExpandItem != null)
+				{
+					ExpandItem.EnsureVisible();
+				}
 			}
-
-			UpdateNumRequestedBuilds(false);
 		}
 
 		private void BuildList_MouseClick(object Sender, MouseEventArgs Args)
@@ -2963,100 +3200,119 @@ namespace UnrealGameSync
 				ListViewHitTestInfo HitTest = BuildList.HitTest(Args.Location);
 				if(HitTest.Item != null)
 				{
-					PerforceChangeSummary Change = (PerforceChangeSummary)HitTest.Item.Tag;
-					if(Workspace.PendingChangeNumber == Change.Number)
+					if(HitTest.Item == ExpandItem)
 					{
-						Rectangle SubItemRect = HitTest.Item.SubItems[StatusColumn.Index].Bounds;
-
-						if(Workspace.IsBusy())
+						if(HitTestExpandLink(Args.Location))
 						{
-							Rectangle CancelRect = new Rectangle(SubItemRect.Right - 16, SubItemRect.Top, 16, SubItemRect.Height);
-							Rectangle InfoRect = new Rectangle(SubItemRect.Right - 32, SubItemRect.Top, 16, SubItemRect.Height);
-							if(CancelRect.Contains(Args.Location))
+							int CurrentMaxChanges = PerforceMonitor.CurrentMaxChanges;
+							if(PerforceMonitor.PendingMaxChanges > CurrentMaxChanges)
 							{
-								CancelWorkspaceUpdate();
+								PerforceMonitor.PendingMaxChanges = CurrentMaxChanges;
 							}
-							else if(InfoRect.Contains(Args.Location) && !Splitter.IsLogVisible())
+							else
 							{
-								ToggleLogVisibility();
+								PerforceMonitor.PendingMaxChanges = CurrentMaxChanges + BuildListExpandCount;
 							}
+							BuildList.Invalidate();
 						}
-						else
+					}
+					else
+					{
+						PerforceChangeSummary Change = (PerforceChangeSummary)HitTest.Item.Tag;
+						if(Workspace.PendingChangeNumber == Change.Number)
 						{
-							Rectangle HappyRect = new Rectangle(SubItemRect.Right - 32, SubItemRect.Top, 16, SubItemRect.Height);
-							Rectangle FrownRect = new Rectangle(SubItemRect.Right - 16, SubItemRect.Top, 16, SubItemRect.Height);
-							if(HappyRect.Contains(Args.Location))
+							Rectangle SubItemRect = HitTest.Item.SubItems[StatusColumn.Index].Bounds;
+
+							if(Workspace.IsBusy())
 							{
-								EventMonitor.PostEvent(Change.Number, EventType.Good);
-								BuildList.Invalidate();
+								Rectangle CancelRect = new Rectangle(SubItemRect.Right - 16, SubItemRect.Top, 16, SubItemRect.Height);
+								Rectangle InfoRect = new Rectangle(SubItemRect.Right - 32, SubItemRect.Top, 16, SubItemRect.Height);
+								if(CancelRect.Contains(Args.Location))
+								{
+									CancelWorkspaceUpdate();
+								}
+								else if(InfoRect.Contains(Args.Location) && !Splitter.IsLogVisible())
+								{
+									ToggleLogVisibility();
+								}
 							}
-							else if(FrownRect.Contains(Args.Location))
+							else
 							{
-								EventMonitor.PostEvent(Change.Number, EventType.Bad);
-								BuildList.Invalidate();
+								Rectangle HappyRect = new Rectangle(SubItemRect.Right - 32, SubItemRect.Top, 16, SubItemRect.Height);
+								Rectangle FrownRect = new Rectangle(SubItemRect.Right - 16, SubItemRect.Top, 16, SubItemRect.Height);
+								if(HappyRect.Contains(Args.Location))
+								{
+									EventMonitor.PostEvent(Change.Number, EventType.Good);
+									BuildList.Invalidate();
+								}
+								else if(FrownRect.Contains(Args.Location))
+								{
+									EventMonitor.PostEvent(Change.Number, EventType.Bad);
+									BuildList.Invalidate();
+								}
 							}
 						}
-					}
-					else
-					{
-						Rectangle SyncBadgeRectangle = GetSyncBadgeRectangle(HitTest.Item.SubItems[StatusColumn.Index].Bounds);
-						if(SyncBadgeRectangle.Contains(Args.Location) && CanSyncChange(Change.Number))
+						else
 						{
-							StartSync(Change.Number);
+							Rectangle SyncBadgeRectangle = GetSyncBadgeRectangle(HitTest.Item.SubItems[StatusColumn.Index].Bounds);
+							if(SyncBadgeRectangle.Contains(Args.Location) && CanSyncChange(Change.Number))
+							{
+								StartSync(Change.Number, null);
+							}
 						}
-					}
 
-					if(DescriptionColumn.Index < HitTest.Item.SubItems.Count && HitTest.Item.SubItems[DescriptionColumn.Index] == HitTest.SubItem)
-					{
-						ChangeLayoutInfo LayoutInfo = GetChangeLayoutInfo(Change);
-						if(LayoutInfo.DescriptionBadges.Count > 0)
+						if(DescriptionColumn.Index < HitTest.Item.SubItems.Count && HitTest.Item.SubItems[DescriptionColumn.Index] == HitTest.SubItem)
 						{
-							Point BuildListLocation = GetBadgeListLocation(LayoutInfo.DescriptionBadges, HitTest.SubItem.Bounds, HorizontalAlign.Right, VerticalAlignment.Middle);
-							BuildListLocation.Offset(-2, 0);
-
-							foreach (BadgeInfo Badge in LayoutInfo.DescriptionBadges)
+							ChangeLayoutInfo LayoutInfo = GetChangeLayoutInfo(Change);
+							if(LayoutInfo.DescriptionBadges.Count > 0)
 							{
-								Rectangle BadgeBounds = Badge.GetBounds(BuildListLocation);
-								if(BadgeBounds.Contains(Args.Location) && Badge.ClickHandler != null)
+								Point BuildListLocation = GetBadgeListLocation(LayoutInfo.DescriptionBadges, HitTest.SubItem.Bounds, HorizontalAlign.Right, VerticalAlignment.Middle);
+								BuildListLocation.Offset(-2, 0);
+
+								foreach (BadgeInfo Badge in LayoutInfo.DescriptionBadges)
 								{
-									Badge.ClickHandler();
-									break;
+									Rectangle BadgeBounds = Badge.GetBounds(BuildListLocation);
+									if(BadgeBounds.Contains(Args.Location) && Badge.ClickHandler != null)
+									{
+										Badge.ClickHandler();
+										break;
+									}
 								}
 							}
 						}
-					}
 
-					if(CISColumn.Index < HitTest.Item.SubItems.Count && HitTest.Item.SubItems[CISColumn.Index] == HitTest.SubItem)
-					{
-						ChangeLayoutInfo LayoutInfo = GetChangeLayoutInfo(Change);
-						if(LayoutInfo.BuildBadges.Count > 0)
+						if(CISColumn.Index < HitTest.Item.SubItems.Count && HitTest.Item.SubItems[CISColumn.Index] == HitTest.SubItem)
 						{
-							Point BuildListLocation = GetBadgeListLocation(LayoutInfo.BuildBadges, HitTest.SubItem.Bounds, HorizontalAlign.Center, VerticalAlignment.Middle);
-							BuildListLocation.X = Math.Max(BuildListLocation.X, HitTest.SubItem.Bounds.Left);
-
-							BadgeInfo BadgeInfo = HitTestBadge(Args.Location, LayoutInfo.BuildBadges, BuildListLocation);
-							if(BadgeInfo != null && BadgeInfo.ClickHandler != null)
+							ChangeLayoutInfo LayoutInfo = GetChangeLayoutInfo(Change);
+							if(LayoutInfo.BuildBadges.Count > 0)
 							{
-								BadgeInfo.ClickHandler();
+								Point BuildListLocation = GetBadgeListLocation(LayoutInfo.BuildBadges, HitTest.SubItem.Bounds, HorizontalAlign.Center, VerticalAlignment.Middle);
+								BuildListLocation.X = Math.Max(BuildListLocation.X, HitTest.SubItem.Bounds.Left);
+
+								BadgeInfo BadgeInfo = HitTestBadge(Args.Location, LayoutInfo.BuildBadges, BuildListLocation);
+								if(BadgeInfo != null && BadgeInfo.ClickHandler != null)
+								{
+									BadgeInfo.ClickHandler();
+								}
 							}
 						}
-					}
 
-					foreach(ColumnHeader CustomColumn in CustomColumns)
-					{
-						if(CustomColumn.Index < HitTest.Item.SubItems.Count && HitTest.Item.SubItems[CustomColumn.Index] == HitTest.SubItem)
+						foreach(ColumnHeader CustomColumn in CustomColumns)
 						{
-							ChangeLayoutInfo LayoutInfo = GetChangeLayoutInfo((PerforceChangeSummary)HitTest.Item.Tag);
-
-							List<BadgeInfo> Badges;
-							if(LayoutInfo.CustomBadges.TryGetValue(CustomColumn.Text, out Badges) && Badges.Count > 0)
+							if(CustomColumn.Index < HitTest.Item.SubItems.Count && HitTest.Item.SubItems[CustomColumn.Index] == HitTest.SubItem)
 							{
-								Point ListLocation = GetBadgeListLocation(Badges, HitTest.SubItem.Bounds, HorizontalAlign.Center, VerticalAlignment.Middle);
+								ChangeLayoutInfo LayoutInfo = GetChangeLayoutInfo((PerforceChangeSummary)HitTest.Item.Tag);
 
-								BadgeInfo BadgeInfo = HitTestBadge(Args.Location, Badges, ListLocation);
-								if(BadgeInfo != null && BadgeInfo.ClickHandler != null)
+								List<BadgeInfo> Badges;
+								if(LayoutInfo.CustomBadges.TryGetValue(CustomColumn.Text, out Badges) && Badges.Count > 0)
 								{
-									BadgeInfo.ClickHandler();
+									Point ListLocation = GetBadgeListLocation(Badges, HitTest.SubItem.Bounds, HorizontalAlign.Center, VerticalAlignment.Middle);
+
+									BadgeInfo BadgeInfo = HitTestBadge(Args.Location, Badges, ListLocation);
+									if(BadgeInfo != null && BadgeInfo.ClickHandler != null)
+									{
+										BadgeInfo.ClickHandler();
+									}
 								}
 							}
 						}
@@ -3179,8 +3435,16 @@ namespace UnrealGameSync
 		{
 			ListViewHitTestInfo HitTest = BuildList.HitTest(e.Location);
 
+			bool bNewMouseOverExpandLink = HitTest.Item != null && HitTestExpandLink(e.Location);
+			if(bMouseOverExpandLink != bNewMouseOverExpandLink)
+			{
+				bMouseOverExpandLink = bNewMouseOverExpandLink;
+				Cursor = bMouseOverExpandLink? NativeCursors.Hand : Cursors.Arrow;
+				BuildList.Invalidate();
+			}
+
 			string NewHoverBadgeUniqueId = null;
-			if(HitTest.Item != null)
+			if(HitTest.Item != null && HitTest.Item.Tag is PerforceChangeSummary)
 			{
 				int ColumnIndex = HitTest.Item.SubItems.IndexOf(HitTest.SubItem);
 				if(ColumnIndex == DescriptionColumn.Index)
@@ -3276,8 +3540,6 @@ namespace UnrealGameSync
 			OptionsContextMenu_ScheduledSync.Checked = Settings.bScheduleEnabled;
 			OptionsContextMenu_TimeZone_Local.Checked = Settings.bShowLocalTimes;
 			OptionsContextMenu_TimeZone_PerforceServer.Checked = !Settings.bShowLocalTimes;
-			OptionsContextMenu_AutomaticallyRunAtStartup.Checked = IsAutomaticallyRunAtStartup();
-			OptionsContextMenu_KeepInTray.Checked = Settings.bKeepInTray;
 			OptionsContextMenu_ShowChanges_ShowUnreviewed.Checked = Settings.bShowUnreviewedChanges;
 			OptionsContextMenu_ShowChanges_ShowAutomated.Checked = Settings.bShowAutomatedChanges;
 			OptionsContextMenu_TabNames_Stream.Checked = Settings.TabLabels == TabLabels.Stream;
@@ -3355,6 +3617,11 @@ namespace UnrealGameSync
 		}
 
 		public void SyncLatestChange()
+		{
+			SyncLatestChange(null);
+		}
+
+		public void SyncLatestChange(WorkspaceUpdateCallback Callback)
 		{
 			if(Workspace != null)
 			{
@@ -3371,7 +3638,7 @@ namespace UnrealGameSync
 				{
 					Owner.ShowAndActivate();
 					SelectChange(ChangeNumber);
-					StartSync(ChangeNumber);
+					StartSync(ChangeNumber, Callback);
 				}
 			}
 		}
@@ -3495,7 +3762,7 @@ namespace UnrealGameSync
 		private void OnlyShowReviewedCheckBox_CheckedChanged(object sender, EventArgs e)
 		{
 			UpdateBuildList();
-			UpdateNumRequestedBuilds(true);
+			ShrinkNumRequestedBuilds();
 		}
 
 		public void Activate()
@@ -3516,7 +3783,7 @@ namespace UnrealGameSync
 
 		public void Deactivate()
 		{
-			UpdateNumRequestedBuilds(true);
+			ShrinkNumRequestedBuilds();
 		}
 
 		private void BuildList_ItemMouseHover(object sender, ListViewItemMouseHoverEventArgs Args)
@@ -3527,6 +3794,10 @@ namespace UnrealGameSync
 				if(Args.Item.Bounds.Contains(ClientPoint))
 				{
 					PerforceChangeSummary Change = (PerforceChangeSummary)Args.Item.Tag;
+					if(Change == null)
+					{
+						return;
+					}
 
 					Rectangle CISBounds = Args.Item.SubItems[CISColumn.Index].Bounds;
 					if(CISBounds.Contains(ClientPoint))
@@ -3593,6 +3864,13 @@ namespace UnrealGameSync
 				HoverBadgeUniqueId = null;
 				BuildList.Invalidate();
 			}
+
+			if(bMouseOverExpandLink)
+			{
+				Cursor = Cursors.Arrow;
+				bMouseOverExpandLink = false;
+				BuildList.Invalidate();
+			}
 		}
 
 		private void OptionsContextMenu_ShowLog_Click(object sender, EventArgs e)
@@ -3753,30 +4031,33 @@ namespace UnrealGameSync
 			for(int Idx = 0; Idx < BuildList.Items.Count; Idx++)
 			{
 				PerforceChangeSummary Change = (PerforceChangeSummary)BuildList.Items[Idx].Tag;
-				if(ChangeType == LatestChangeType.Any)
+				if(Change != null)
 				{
-					if(CanSyncChange(Change.Number))
+					if(ChangeType == LatestChangeType.Any)
 					{
-						ChangeNumber = FindNewestGoodContentChange(Change.Number);
-						return true;
+						if(CanSyncChange(Change.Number))
+						{
+							ChangeNumber = FindNewestGoodContentChange(Change.Number);
+							return true;
+						}
 					}
-				}
-				else if(ChangeType == LatestChangeType.Good)
-				{
-					EventSummary Summary = EventMonitor.GetSummaryForChange(Change.Number);
-					if(Summary != null && Summary.Verdict == ReviewVerdict.Good && CanSyncChange(Change.Number))
+					else if(ChangeType == LatestChangeType.Good)
 					{
-						ChangeNumber = FindNewestGoodContentChange(Change.Number);
-						return true;
+						EventSummary Summary = EventMonitor.GetSummaryForChange(Change.Number);
+						if(Summary != null && Summary.Verdict == ReviewVerdict.Good && CanSyncChange(Change.Number))
+						{
+							ChangeNumber = FindNewestGoodContentChange(Change.Number);
+							return true;
+						}
 					}
-				}
-				else if(ChangeType == LatestChangeType.Starred)
-				{
-					EventSummary Summary = EventMonitor.GetSummaryForChange(Change.Number);
-					if(((Summary != null && Summary.LastStarReview != null) || PromotedChangeNumbers.Contains(Change.Number)) && CanSyncChange(Change.Number))
+					else if(ChangeType == LatestChangeType.Starred)
 					{
-						ChangeNumber = FindNewestGoodContentChange(Change.Number);
-						return true;
+						EventSummary Summary = EventMonitor.GetSummaryForChange(Change.Number);
+						if(((Summary != null && Summary.LastStarReview != null) || PromotedChangeNumbers.Contains(Change.Number)) && CanSyncChange(Change.Number))
+						{
+							ChangeNumber = FindNewestGoodContentChange(Change.Number);
+							return true;
+						}
 					}
 				}
 			}
@@ -3922,7 +4203,7 @@ namespace UnrealGameSync
 				else
 				{
 					WorkspaceUpdateContext Context = new WorkspaceUpdateContext(Workspace.CurrentChangeNumber, WorkspaceUpdateOptions.Build, null, GetDefaultBuildStepObjects(), ProjectSettings.BuildSteps, new HashSet<Guid>{ UniqueId }, GetWorkspaceVariables(Workspace.CurrentChangeNumber));
-					StartWorkspaceUpdate(Context);
+					StartWorkspaceUpdate(Context, null);
 				}
 			}
 		}
@@ -4089,25 +4370,6 @@ namespace UnrealGameSync
 			return UserBuildStepObjects.Values.Select(x => new BuildStep(x)).OrderBy(x => x.OrderIndex).ToList();
 		}
 
-		private void OptionsContextMenu_AutomaticallyRunAtStartup_Click(object sender, EventArgs e)
-		{
-			RegistryKey Key = Registry.CurrentUser.CreateSubKey("SOFTWARE\\Microsoft\\Windows\\CurrentVersion\\Run");
-			if(IsAutomaticallyRunAtStartup())
-			{
-	            Key.DeleteValue("UnrealGameSync", false);
-			}
-			else
-			{
-				Key.SetValue("UnrealGameSync", String.Format("\"{0}\" -RestoreState", OriginalExecutableFileName));
-			}
-		}
-
-		private bool IsAutomaticallyRunAtStartup()
-		{
-			RegistryKey Key = Registry.CurrentUser.OpenSubKey("SOFTWARE\\Microsoft\\Windows\\CurrentVersion\\Run");
-			return (Key.GetValue("UnrealGameSync") != null);
-		}
-
 		private void OptionsContextMenu_SyncPrecompiledEditor_Click(object sender, EventArgs e)
 		{
 			OptionsContextMenu_SyncPrecompiledEditor.Checked ^= true;
@@ -4147,16 +4409,18 @@ namespace UnrealGameSync
 
 		private void OptionsContextMenu_SyncFilter_Click(object sender, EventArgs e)
 		{
-			SyncFilter Filter = new SyncFilter(Workspace.GetSyncCategories(), Settings.SyncView, Settings.SyncExcludedCategories, Settings.bSyncAllProjects, WorkspaceSettings.SyncView, WorkspaceSettings.SyncIncludedCategories, WorkspaceSettings.SyncExcludedCategories, WorkspaceSettings.bSyncAllProjects);
+			SyncFilter Filter = new SyncFilter(Workspace.GetSyncCategories(), Settings.SyncView, Settings.SyncExcludedCategories, Settings.bSyncAllProjects, Settings.bIncludeAllProjectsInSolution, WorkspaceSettings.SyncView, WorkspaceSettings.SyncIncludedCategories, WorkspaceSettings.SyncExcludedCategories, WorkspaceSettings.bSyncAllProjects, WorkspaceSettings.bIncludeAllProjectsInSolution);
 			if(Filter.ShowDialog() == DialogResult.OK)
 			{
 				Settings.SyncExcludedCategories = Filter.GlobalExcludedCategories;
 				Settings.SyncView = Filter.GlobalView;
 				Settings.bSyncAllProjects = Filter.bGlobalSyncAllProjects;
+				Settings.bIncludeAllProjectsInSolution = Filter.bGlobalIncludeAllProjectsInSolution;
 				WorkspaceSettings.SyncIncludedCategories = Filter.WorkspaceIncludedCategories;
 				WorkspaceSettings.SyncExcludedCategories = Filter.WorkspaceExcludedCategories;
 				WorkspaceSettings.SyncView = Filter.WorkspaceView;
 				WorkspaceSettings.bSyncAllProjects = Filter.bWorkspaceSyncAllProjects;
+				WorkspaceSettings.bIncludeAllProjectsInSolution = Filter.bWorkspaceIncludeAllProjectsInSolution;
                 Settings.Save();
 			}
 		}
@@ -4201,12 +4465,6 @@ namespace UnrealGameSync
 			}
 		}
 
-		private void OptionsContextMenu_KeepInTray_Click(object sender, EventArgs e)
-		{
-			Settings.bKeepInTray ^= true;
-			Settings.Save();
-		}
-
 		private void BuildList_KeyDown(object Sender, KeyEventArgs Args)
 		{
 			if(Args.Control && Args.KeyCode == Keys.C && BuildList.SelectedItems.Count > 0)
@@ -4216,10 +4474,15 @@ namespace UnrealGameSync
 			}
 		}
 
+		private void OptionsContextMenu_ApplicationSettings_Click(object sender, EventArgs e)
+		{
+			Owner.ModifyApplicationSettings();
+		}
+
 		private void OptionsContextMenu_PerforceSettings_Click(object sender, EventArgs e)
 		{
-			PerforceSyncSettingsWindow Perforce = new PerforceSyncSettingsWindow(Settings);
-			Perforce.ShowDialog();
+			PerforceSyncSettingsWindow Window = new PerforceSyncSettingsWindow(Settings);
+			Window.ShowDialog();
 		}
 
 		private void OptionsContextMenu_TabNames_Stream_Click(object sender, EventArgs e)
@@ -4276,7 +4539,7 @@ namespace UnrealGameSync
 			Settings.Save();
 
 			UpdateBuildList();
-			UpdateNumRequestedBuilds(true);
+			ShrinkNumRequestedBuilds();
 		}
 
 		private void OptionsContextMenu_ShowChanges_ShowAutomated_Click(object sender, EventArgs e)
@@ -4284,8 +4547,7 @@ namespace UnrealGameSync
 			Settings.bShowAutomatedChanges ^= true;
 			Settings.Save();
 
-			UpdateBuildList();
-			UpdateNumRequestedBuilds(true);
+			UpdateBuildListFilter();
 		}
 
 		private void BuildList_ColumnWidthChanging(object sender, ColumnWidthChangingEventArgs e)
@@ -4581,5 +4843,110 @@ namespace UnrealGameSync
 				PerforceMonitor.IsActive = Visible;
 			}
 		}
+
+		private void FilterButton_Click(object sender, EventArgs e)
+		{
+			FilterContextMenu_Default.Checked = !Settings.bShowAutomatedChanges && ProjectSettings.FilterType == FilterType.None && ProjectSettings.FilterBadges.Count == 0;
+
+			FilterContextMenu_Type.Checked = ProjectSettings.FilterType != FilterType.None;
+			FilterContextMenu_Type_ShowAll.Checked = ProjectSettings.FilterType == FilterType.None;
+			FilterContextMenu_Type_Code.Checked = ProjectSettings.FilterType == FilterType.Code;
+			FilterContextMenu_Type_Content.Checked = ProjectSettings.FilterType == FilterType.Content;
+
+			FilterContextMenu_Badges.DropDownItems.Clear();
+			FilterContextMenu_Badges.Checked = ProjectSettings.FilterBadges.Count > 0;
+
+			HashSet<string> BadgeNames = new HashSet<string>(ProjectSettings.FilterBadges, StringComparer.OrdinalIgnoreCase);
+			BadgeNames.ExceptWith(BadgeNameAndGroupPairs.Select(x => x.Key));
+
+			List<KeyValuePair<string, string>> DisplayBadgeNameAndGroupPairs = new List<KeyValuePair<string, string>>(BadgeNameAndGroupPairs);
+			DisplayBadgeNameAndGroupPairs.AddRange(BadgeNames.Select(x => new KeyValuePair<string, string>(x, "User")));
+
+			string LastGroup = null;
+			foreach(KeyValuePair<string, string> BadgeNameAndGroupPair in DisplayBadgeNameAndGroupPairs)
+			{
+				if(LastGroup != BadgeNameAndGroupPair.Value)
+				{
+					if(LastGroup != null)
+					{
+						FilterContextMenu_Badges.DropDownItems.Add(new ToolStripSeparator());
+					}
+					LastGroup = BadgeNameAndGroupPair.Value;
+				}
+
+				ToolStripMenuItem Item = new ToolStripMenuItem(BadgeNameAndGroupPair.Key);
+				Item.Checked = ProjectSettings.FilterBadges.Contains(BadgeNameAndGroupPair.Key, StringComparer.OrdinalIgnoreCase);
+				Item.Click += (Sender, Args) => FilterContextMenu_Badge_Click(BadgeNameAndGroupPair.Key);
+				FilterContextMenu_Badges.DropDownItems.Add(Item);
+			}
+
+			FilterContextMenu_Badges.Enabled = FilterContextMenu_Badges.DropDownItems.Count > 0;
+
+			FilterContextMenu_ShowBuildMachineChanges.Checked = Settings.bShowAutomatedChanges;
+			FilterContextMenu.Show(FilterButton, new Point(0, FilterButton.Height));
+		}
+
+		private void FilterContextMenu_Badge_Click(string BadgeName)
+		{
+			if(ProjectSettings.FilterBadges.Contains(BadgeName))
+			{
+				ProjectSettings.FilterBadges.Remove(BadgeName);
+			}
+			else
+			{
+				ProjectSettings.FilterBadges.Add(BadgeName);
+			}
+
+			UpdateBuildListFilter();
+		}
+
+		private void FilterContextMenu_Default_Click(object sender, EventArgs e)
+		{
+			ProjectSettings.FilterBadges.Clear();
+			ProjectSettings.FilterType = FilterType.None;
+
+			Settings.bShowAutomatedChanges = false;
+			Settings.Save();
+
+			UpdateBuildListFilter();
+		}
+
+		private void FilterContextMenu_ShowBuildMachineChanges_Click(object sender, EventArgs e)
+		{
+			Settings.bShowAutomatedChanges ^= true;
+			Settings.Save();
+
+			UpdateBuildListFilter();
+		}
+
+		private void FilterContextMenu_Type_ShowAll_Click(object sender, EventArgs e)
+		{
+			ProjectSettings.FilterType = FilterType.None;
+			Settings.Save();
+
+			UpdateBuildListFilter();
+		}
+
+		private void FilterContextMenu_Type_Code_Click(object sender, EventArgs e)
+		{
+			ProjectSettings.FilterType = FilterType.Code;
+			Settings.Save();
+
+			UpdateBuildListFilter();
+		}
+
+		private void FilterContextMenu_Type_Content_Click(object sender, EventArgs e)
+		{
+			ProjectSettings.FilterType = FilterType.Content;
+			Settings.Save();
+
+			UpdateBuildListFilter();
+		}
+
+		private void UpdateBuildListFilter()
+		{
+			UpdateBuildList();
+			ShrinkNumRequestedBuilds();
+		}
 	}
 }
diff --git a/Engine/Source/Programs/UnrealGameSync/UnrealGameSync/DeleteFilesTask.cs b/Engine/Source/Programs/UnrealGameSync/UnrealGameSync/DeleteFilesTask.cs
--- a/Engine/Source/Programs/UnrealGameSync/UnrealGameSync/DeleteFilesTask.cs
+++ b/Engine/Source/Programs/UnrealGameSync/UnrealGameSync/DeleteFilesTask.cs
@@ -1,4 +1,4 @@
-// Copyright 1998-2018 Epic Games, Inc. All Rights Reserved.
+// Copyright 1998-2019 Epic Games, Inc. All Rights Reserved.
 
 using System;
 using System.Collections.Generic;
diff --git a/Engine/Source/Programs/UnrealGameSync/UnrealGameSync/DeploymentSettings.cs b/Engine/Source/Programs/UnrealGameSync/UnrealGameSync/DeploymentSettings.cs
new file mode 100644
--- /dev/null
+++ b/Engine/Source/Programs/UnrealGameSync/UnrealGameSync/DeploymentSettings.cs
@@ -0,0 +1,28 @@
+// Copyright 1998-2019 Epic Games, Inc. All Rights Reserved.
+
+using System;
+using System.Collections.Generic;
+using System.Linq;
+using System.Text;
+using System.Threading.Tasks;
+
+namespace UnrealGameSync
+{
+	/// <summary>
+	/// This class contains settings for a site-specific deployment of UGS. Epic's internal implementation uses a static constructor in a NotForLicensees folder to initialize these values.
+	/// </summary>
+	static partial class DeploymentSettings
+	{
+		/// <summary>
+		/// SQL connection string used to connect to the database for telemetry and review data. The 'Program' class is a partial class, to allow an
+		/// opportunistically included C# source file in NotForLicensees/ProgramSettings.cs to override this value in a static constructor.
+		/// </summary>
+		public static readonly string ApiUrl = null;
+
+		/// <summary>
+		/// Specifies the depot path to sync down the stable version of UGS from, without a trailing slash (eg. //depot/UnrealGameSync/bin). This is a site-specific setting. 
+		/// The UnrealGameSync executable should be located at Release/UnrealGameSync.exe under this path, with any dependent DLLs.
+		/// </summary>
+		public static readonly string DefaultDepotPath = null;
+	}
+}
diff --git a/Engine/Source/Programs/UnrealGameSync/UnrealGameSync/DetectProjectSettingsTask.cs b/Engine/Source/Programs/UnrealGameSync/UnrealGameSync/DetectProjectSettingsTask.cs
--- a/Engine/Source/Programs/UnrealGameSync/UnrealGameSync/DetectProjectSettingsTask.cs
+++ b/Engine/Source/Programs/UnrealGameSync/UnrealGameSync/DetectProjectSettingsTask.cs
@@ -1,4 +1,4 @@
-// Copyright 1998-2018 Epic Games, Inc. All Rights Reserved.
+// Copyright 1998-2019 Epic Games, Inc. All Rights Reserved.
 
 using System;
 using System.Collections.Generic;
@@ -23,17 +23,18 @@ namespace UnrealGameSync
 		public string NewSelectedClientFileName;
 		public string StreamName;
 		public Image ProjectLogo;
-		public TimeSpan ServerTimeZone;
 		public string DataFolder;
+		public string CacheFolder;
 		public bool bIsEnterpriseProject;
 		public ConfigFile LatestProjectConfigFile;
 		public List<KeyValuePair<string, DateTime>> LocalConfigFiles;
 		TextWriter Log;
 
-		public DetectProjectSettingsTask(UserSelectedProjectSettings SelectedProject, string InDataFolder, TextWriter InLog)
+		public DetectProjectSettingsTask(UserSelectedProjectSettings SelectedProject, string InDataFolder, string InCacheFolder, TextWriter InLog)
 		{
 			this.SelectedProject = SelectedProject;
 			DataFolder = InDataFolder;
+			CacheFolder = InCacheFolder;
 			Log = InLog;
 		}
 
@@ -86,19 +87,8 @@ namespace UnrealGameSync
 
 		public bool Run(PerforceConnection Perforce, TextWriter Log, out string ErrorMessage)
 		{
-			// Get the perforce server settings
-			PerforceInfoRecord PerforceInfo;
-			if(!Perforce.Info(out PerforceInfo, Log))
-			{
-				ErrorMessage = String.Format("Couldn't get Perforce server info");
-				return false;
-			}
-
-			// Configure the time zone
-			ServerTimeZone = PerforceInfo.ServerTimeZone;
-
-			// If we're using the legacy path of specifying a file, figure out the workspace name now
-			if(SelectedProject.Type == UserSelectedProjectType.Client)
+			// Use the cached client path to the file if it's available; it's much quicker than trying to find the correct workspace.
+			if(!String.IsNullOrEmpty(SelectedProject.ClientPath))
 			{
 				// Get the client path
 				NewSelectedClientFileName = SelectedProject.ClientPath;
@@ -114,15 +104,30 @@ namespace UnrealGameSync
 				// Create the client
 				PerforceClient = new PerforceConnection(Perforce.UserName, ClientName, Perforce.ServerAndPort);
 
-				// Figure out the path on the client
-				if(!PerforceClient.ConvertToLocalPath(NewSelectedClientFileName, out NewSelectedFileName, Log))
+				// Figure out the path on the client. Use the cached location if it's valid.
+				if(SelectedProject.LocalPath != null && File.Exists(SelectedProject.LocalPath))
 				{
-					ErrorMessage = String.Format("Couldn't get client path for {0}", NewSelectedFileName);
-					return false;
+					NewSelectedFileName = SelectedProject.LocalPath;
+				}
+				else
+				{
+					if(!PerforceClient.ConvertToLocalPath(NewSelectedClientFileName, out NewSelectedFileName, Log))
+					{
+						ErrorMessage = String.Format("Couldn't get client path for {0}", NewSelectedFileName);
+						return false;
+					}
 				}
 			}
-			else if(SelectedProject.Type == UserSelectedProjectType.Local)
+			else
 			{
+				// Get the perforce server settings
+				PerforceInfoRecord PerforceInfo;
+				if(!Perforce.Info(out PerforceInfo, Log))
+				{
+					ErrorMessage = String.Format("Couldn't get Perforce server info");
+					return false;
+				}
+
 				// Use the path as the selected filename
 				NewSelectedFileName = SelectedProject.LocalPath;
 
@@ -133,67 +138,40 @@ namespace UnrealGameSync
 					return false;
 				}
 
-				// Find all the clients on this machine
-				Log.WriteLine("Enumerating clients on local machine...");
+				// Find all the clients for this user
+				Log.WriteLine("Enumerating clients for {0}...", PerforceInfo.UserName);
+
 				List<PerforceClientRecord> Clients;
-				if(!Perforce.FindClients(out Clients, Log))
+				if(!Perforce.FindClients(PerforceInfo.UserName, out Clients, Log))
 				{
 					ErrorMessage = String.Format("Couldn't find any clients for this host.");
 					return false;
 				}
 
-				// Find any clients which are valid. If this is not exactly one, we should fail.
-				List<PerforceConnection> CandidateClients = new List<PerforceConnection>();
-				foreach(PerforceClientRecord Client in Clients)
+				List<PerforceConnection> CandidateClients = FilterClients(Clients, Perforce.ServerAndPort, PerforceInfo.HostName, PerforceInfo.UserName);
+				if(CandidateClients.Count == 0)
 				{
-					// Make sure the client is well formed
-					if(!String.IsNullOrEmpty(Client.Name) && (!String.IsNullOrEmpty(Client.Host) || !String.IsNullOrEmpty(Client.Owner)) && !String.IsNullOrEmpty(Client.Root))
+					// Search through all workspaces. We may find a suitable workspace which is for any user.
+					Log.WriteLine("Enumerating shared clients...");
+					if(!Perforce.FindClients("", out Clients, Log))
 					{
-						// Require either a username or host name match
-						if((String.IsNullOrEmpty(Client.Host) || String.Compare(Client.Host, PerforceInfo.HostName, StringComparison.InvariantCultureIgnoreCase) == 0) && (String.IsNullOrEmpty(Client.Owner) || String.Compare(Client.Owner, PerforceInfo.UserName, StringComparison.InvariantCultureIgnoreCase) == 0))
-						{
-							if(!Utility.SafeIsFileUnderDirectory(NewSelectedFileName, Client.Root))
-							{
-								Log.WriteLine("Rejecting {0} due to root mismatch ({1})", Client.Name, Client.Root);
-								continue;
-							}
-
-							PerforceConnection CandidateClient = new PerforceConnection(PerforceInfo.UserName, Client.Name, Perforce.ServerAndPort);
-
-							bool bFileExists;
-							if(!CandidateClient.FileExists(NewSelectedFileName, out bFileExists, Log) || !bFileExists)
-							{
-								Log.WriteLine("Rejecting {0} due to file not existing in workspace", Client.Name);
-								continue;
-							}
-
-							List<PerforceFileRecord> Records;
-							if(!CandidateClient.Stat(NewSelectedFileName, out Records, Log))
-							{
-								Log.WriteLine("Rejecting {0} due to {1} not in depot", Client.Name, NewSelectedFileName);
-								continue;
-							}
-
-							Records.RemoveAll(x => !x.IsMapped);
-							if(Records.Count == 0)
-							{
-								Log.WriteLine("Rejecting {0} due to {1} matching records", Client.Name, Records.Count);
-								continue;
-							}
-
-							Log.WriteLine("Found valid client {0}", Client.Name);
-							CandidateClients.Add(CandidateClient);
-						}
+						ErrorMessage = "Failed to enumerate clients.";
+						return false;
+					}
+
+					// Filter this list of clients
+					CandidateClients = FilterClients(Clients, Perforce.ServerAndPort, PerforceInfo.HostName, PerforceInfo.UserName);
+
+					// If we still couldn't find any, fail.
+					if(CandidateClients.Count == 0)
+					{
+						ErrorMessage = String.Format("Couldn't find any Perforce workspace containing {0}. Check your connection settings.", NewSelectedFileName);
+						return false;
 					}
 				}
 
 				// Check there's only one client
-				if(CandidateClients.Count == 0)
-				{
-					ErrorMessage = String.Format("Couldn't find any Perforce workspace containing {0}. Check your connection settings.", NewSelectedFileName);
-					return false;
-				}
-				else if(CandidateClients.Count > 1)
+				if(CandidateClients.Count > 1)
 				{
 					ErrorMessage = String.Format("Found multiple workspaces containing {0}:\n\n{1}\n\nCannot determine which to use.", Path.GetFileName(NewSelectedFileName), String.Join("\n", CandidateClients.Select(x => x.ClientName)));
 					return false;
@@ -209,10 +187,6 @@ namespace UnrealGameSync
 					return false;
 				}
 			}
-			else
-			{
-				throw new InvalidDataException("Invalid selected project type");
-			}
 
 			// Normalize the filename
 			NewSelectedFileName = Path.GetFullPath(NewSelectedFileName).Replace('/', Path.DirectorySeparatorChar);
@@ -224,7 +198,12 @@ namespace UnrealGameSync
 			SelectedProject = new UserSelectedProjectSettings(Perforce.ServerAndPort, Perforce.UserName, SelectedProject.Type, NewSelectedClientFileName, NewSelectedFileName);
 
 			// Figure out where the engine is in relation to it
-			for(int EndIdx = NewSelectedClientFileName.Length - 1;;EndIdx--)
+			int EndIdx = NewSelectedClientFileName.Length - 1;
+			if(EndIdx != -1 && NewSelectedClientFileName.EndsWith(".uproject", StringComparison.InvariantCultureIgnoreCase))
+			{
+				EndIdx = NewSelectedClientFileName.LastIndexOf('/') - 1;
+			}
+			for(;;EndIdx--)
 			{
 				if(EndIdx < 2)
 				{
@@ -233,38 +212,32 @@ namespace UnrealGameSync
 				}
 				if(NewSelectedClientFileName[EndIdx] == '/')
 				{
-					bool bFileExists;
-					if(PerforceClient.FileExists(NewSelectedClientFileName.Substring(0, EndIdx) + "/Engine/Build/Build.version", out bFileExists, Log) && bFileExists)
+					List<PerforceFileRecord> FileRecords;
+					if(PerforceClient.Stat(NewSelectedClientFileName.Substring(0, EndIdx) + "/Engine/Build/Build.version", out FileRecords, Log) && FileRecords.Count > 0)
 					{
+						if(FileRecords[0].ClientPath == null)
+						{
+							ErrorMessage = String.Format("Missing client path for {0}", FileRecords[0].DepotPath);
+							return false;
+						}
+
 						BranchClientPath = NewSelectedClientFileName.Substring(0, EndIdx);
+						BranchDirectoryName = Path.GetFullPath(Path.Combine(Path.GetDirectoryName(FileRecords[0].ClientPath), "..", ".."));
 						break;
 					}
 				}
 			}
 			Log.WriteLine("Found branch root at {0}", BranchClientPath);
 
-			// Get the local branch root
-			string BuildVersionPath;
-			if(!PerforceClient.ConvertToLocalPath(BranchClientPath + "/Engine/Build/Build.version", out BuildVersionPath, Log))
-			{
-				ErrorMessage = String.Format("Couldn't get local path for Engine/Build/Build.version");
-				return false;
-			}
-			BranchDirectoryName = Path.GetFullPath(Path.Combine(Path.GetDirectoryName(BuildVersionPath), "..", ".."));
-
 			// Find the editor target for this project
 			if(NewSelectedFileName.EndsWith(".uproject", StringComparison.InvariantCultureIgnoreCase))
 			{
 				List<PerforceFileRecord> Files;
 				if(PerforceClient.FindFiles(PerforceUtils.GetClientOrDepotDirectoryName(NewSelectedClientFileName) + "/Source/*Editor.Target.cs", out Files, Log) && Files.Count >= 1)
 				{
-					PerforceFileRecord File = Files.FirstOrDefault(x => x.Action == null || !x.Action.Contains("delete"));
-					if(File != null)
-					{
-						string DepotPath = File.DepotPath;
-						NewProjectEditorTarget = Path.GetFileNameWithoutExtension(Path.GetFileNameWithoutExtension(DepotPath.Substring(DepotPath.LastIndexOf('/') + 1)));
-						Log.WriteLine("Using {0} as editor target name (from {1})", NewProjectEditorTarget, Files[0]);
-					}
+					string DepotPath = Files[0].DepotPath;
+					NewProjectEditorTarget = Path.GetFileNameWithoutExtension(Path.GetFileNameWithoutExtension(DepotPath.Substring(DepotPath.LastIndexOf('/') + 1)));
+					Log.WriteLine("Using {0} as editor target name (from {1})", NewProjectEditorTarget, Files[0]);
 				}
 				if (NewProjectEditorTarget == null)
 				{
@@ -319,23 +292,84 @@ namespace UnrealGameSync
 				}
 			}
 
-			// Figure out if it's an enterprise project
-			List<string> ProjectLines;
-			if(NewSelectedClientFileName.EndsWith(".uproject", StringComparison.InvariantCultureIgnoreCase) && PerforceClient.Print(NewSelectedClientFileName, out ProjectLines, Log))
+			// Figure out if it's an enterprise project. Use the synced version if we have it.
+			if(NewSelectedClientFileName.EndsWith(".uproject", StringComparison.InvariantCultureIgnoreCase))
 			{
-				string Text = String.Join("\n", ProjectLines);
+				string Text;
+				if(File.Exists(NewSelectedFileName))
+				{
+					Text = File.ReadAllText(NewSelectedFileName);
+				}
+				else
+				{
+					List<string> ProjectLines;
+					if(PerforceClient.Print(NewSelectedClientFileName, out ProjectLines, Log))
+					{
+						ErrorMessage = String.Format("Unable to get contents of {0}", NewSelectedClientFileName);
+						return false;
+					}
+					Text = String.Join("\n", ProjectLines);
+				}
 				bIsEnterpriseProject = Utility.IsEnterpriseProjectFromText(Text);
 			}
 
 			// Read the initial config file
 			LocalConfigFiles = new List<KeyValuePair<string, DateTime>>();
-			LatestProjectConfigFile = PerforceMonitor.ReadProjectConfigFile(PerforceClient, BranchClientPath, NewSelectedClientFileName, LocalConfigFiles, Log);
+			LatestProjectConfigFile = PerforceMonitor.ReadProjectConfigFile(PerforceClient, BranchClientPath, NewSelectedClientFileName, CacheFolder, LocalConfigFiles, Log);
 
 			// Succeed!
 			ErrorMessage = null;
 			return true;
 		}
 
+		List<PerforceConnection> FilterClients(List<PerforceClientRecord> Clients, string ServerAndPort, string HostName, string UserName)
+		{
+			List<PerforceConnection> CandidateClients = new List<PerforceConnection>();
+			foreach(PerforceClientRecord Client in Clients)
+			{
+				// Make sure the client is well formed
+				if(!String.IsNullOrEmpty(Client.Name) && (!String.IsNullOrEmpty(Client.Host) || !String.IsNullOrEmpty(Client.Owner)) && !String.IsNullOrEmpty(Client.Root))
+				{
+					// Require either a username or host name match
+					if((String.IsNullOrEmpty(Client.Host) || String.Compare(Client.Host, HostName, StringComparison.InvariantCultureIgnoreCase) == 0) && (String.IsNullOrEmpty(Client.Owner) || String.Compare(Client.Owner, UserName, StringComparison.InvariantCultureIgnoreCase) == 0))
+					{
+						if(!Utility.SafeIsFileUnderDirectory(NewSelectedFileName, Client.Root))
+						{
+							Log.WriteLine("Rejecting {0} due to root mismatch ({1})", Client.Name, Client.Root);
+							continue;
+						}
+
+						PerforceConnection CandidateClient = new PerforceConnection(UserName, Client.Name, ServerAndPort);
+
+						bool bFileExists;
+						if(!CandidateClient.FileExists(NewSelectedFileName, out bFileExists, Log) || !bFileExists)
+						{
+							Log.WriteLine("Rejecting {0} due to file not existing in workspace", Client.Name);
+							continue;
+						}
+
+						List<PerforceFileRecord> Records;
+						if(!CandidateClient.Stat(NewSelectedFileName, out Records, Log))
+						{
+							Log.WriteLine("Rejecting {0} due to {1} not in depot", Client.Name, NewSelectedFileName);
+							continue;
+						}
+
+						Records.RemoveAll(x => !x.IsMapped);
+						if(Records.Count == 0)
+						{
+							Log.WriteLine("Rejecting {0} due to {1} matching records", Client.Name, Records.Count);
+							continue;
+						}
+
+						Log.WriteLine("Found valid client {0}", Client.Name);
+						CandidateClients.Add(CandidateClient);
+					}
+				}
+			}
+			return CandidateClients;
+		}
+
 		bool TryGetStreamPrefix(PerforceConnection Perforce, string StreamName, TextWriter Log, out string StreamPrefix)
 		{ 
 			string CurrentStreamName = StreamName;
@@ -382,12 +416,12 @@ namespace UnrealGameSync
 
 	class DetectMultipleProjectSettingsTask : IModalTask, IDisposable
 	{
-		public List<DetectProjectSettingsTask> Tasks;
-		public List<DetectProjectSettingsResult> Results;
+		public DetectProjectSettingsTask[] Tasks;
+		public DetectProjectSettingsResult[] Results;
 
 		public DetectMultipleProjectSettingsTask(IEnumerable<DetectProjectSettingsTask> Tasks)
 		{
-			this.Tasks = new List<DetectProjectSettingsTask>(Tasks);
+			this.Tasks = Tasks.ToArray();
 		}
 
 		public void Dispose()
@@ -398,7 +432,7 @@ namespace UnrealGameSync
 				{
 					Task.Dispose();
 				}
-				foreach(DetectProjectSettingsResult Result in Results)
+				foreach(DetectProjectSettingsResult Result in Results.Where(x => x != null))
 				{
 					Result.Dispose();
 				}
@@ -408,17 +442,19 @@ namespace UnrealGameSync
 
 		public bool Run(out string ErrorMessage)
 		{
-			Results = new List<DetectProjectSettingsResult>();
-			for(int Idx = 0; Idx < Tasks.Count; Idx++)
-			{
-				string TaskErrorMessage;
-				bool bTaskSucceeded = Tasks[Idx].Run(out TaskErrorMessage);
-				Results.Add(new DetectProjectSettingsResult(Tasks[Idx], bTaskSucceeded, TaskErrorMessage));
-				Tasks[Idx] = null;
-			}
+			Results = new DetectProjectSettingsResult[Tasks.Length];
+			Parallel.For(0, Tasks.Length, new ParallelOptions(){ MaxDegreeOfParallelism = 4 }, Idx => RunTask(Idx));
 
 			ErrorMessage = null;
 			return true;
 		}
+
+		void RunTask(int Idx)
+		{
+			string TaskErrorMessage;
+			bool bTaskSucceeded = Tasks[Idx].Run(out TaskErrorMessage);
+			Results[Idx] = new DetectProjectSettingsResult(Tasks[Idx], bTaskSucceeded, TaskErrorMessage);
+			Tasks[Idx] = null;
+		}
 	}
 }
diff --git a/Engine/Source/Programs/UnrealGameSync/UnrealGameSync/EventMonitor.cs b/Engine/Source/Programs/UnrealGameSync/UnrealGameSync/EventMonitor.cs
--- a/Engine/Source/Programs/UnrealGameSync/UnrealGameSync/EventMonitor.cs
+++ b/Engine/Source/Programs/UnrealGameSync/UnrealGameSync/EventMonitor.cs
@@ -1,4 +1,4 @@
-// Copyright 1998-2018 Epic Games, Inc. All Rights Reserved.
+// Copyright 1998-2019 Epic Games, Inc. All Rights Reserved.
 
 using System;
 using System.Collections.Concurrent;
diff --git a/Engine/Source/Programs/UnrealGameSync/UnrealGameSync/FileFilter.cs b/Engine/Source/Programs/UnrealGameSync/UnrealGameSync/FileFilter.cs
--- a/Engine/Source/Programs/UnrealGameSync/UnrealGameSync/FileFilter.cs
+++ b/Engine/Source/Programs/UnrealGameSync/UnrealGameSync/FileFilter.cs
@@ -1,4 +1,4 @@
-﻿// Copyright 1998-2018 Epic Games, Inc. All Rights Reserved.
+﻿// Copyright 1998-2019 Epic Games, Inc. All Rights Reserved.
 
 using System;
 using System.Collections.Generic;
diff --git a/Engine/Source/Programs/UnrealGameSync/UnrealGameSync/FindFoldersToCleanTask.cs b/Engine/Source/Programs/UnrealGameSync/UnrealGameSync/FindFoldersToCleanTask.cs
--- a/Engine/Source/Programs/UnrealGameSync/UnrealGameSync/FindFoldersToCleanTask.cs
+++ b/Engine/Source/Programs/UnrealGameSync/UnrealGameSync/FindFoldersToCleanTask.cs
@@ -1,4 +1,4 @@
-// Copyright 1998-2018 Epic Games, Inc. All Rights Reserved.
+// Copyright 1998-2019 Epic Games, Inc. All Rights Reserved.
 
 using System;
 using System.Collections.Generic;
@@ -152,7 +152,13 @@ namespace UnrealGameSync
 				}
 
 				// Find all the files that should be deleted
-				LocalFolder.FilesToDelete.AddRange(LocalFolder.NameToFile.Values.Where(x => !PerforceFolder.NameToFile.ContainsKey(x.Name)));
+				foreach(FileInfo LocalFileInfo in LocalFolder.NameToFile.Values)
+				{
+					if(!PerforceFolder.NameToFile.ContainsKey(LocalFileInfo.Name) && !OpenClientPaths.Contains(LocalFileInfo.FullName))
+					{
+						LocalFolder.FilesToDelete.Add(LocalFileInfo);
+					}
+				}
 			}
 
 			// Figure out if this folder is just an empty directory that needs to be removed
diff --git a/Engine/Source/Programs/UnrealGameSync/UnrealGameSync/Forms/ApplicationSettingsWindow.cs b/Engine/Source/Programs/UnrealGameSync/UnrealGameSync/Forms/ApplicationSettingsWindow.cs
new file mode 100644
--- /dev/null
+++ b/Engine/Source/Programs/UnrealGameSync/UnrealGameSync/Forms/ApplicationSettingsWindow.cs
@@ -0,0 +1,213 @@
+// Copyright 1998-2019 Epic Games, Inc. All Rights Reserved.
+
+using Microsoft.Win32;
+using System;
+using System.Collections.Generic;
+using System.ComponentModel;
+using System.Data;
+using System.Drawing;
+using System.IO;
+using System.Linq;
+using System.Runtime.InteropServices;
+using System.Text;
+using System.Threading;
+using System.Threading.Tasks;
+using System.Windows.Forms;
+using UnrealGameSync;
+
+namespace UnrealGameSync
+{
+	partial class ApplicationSettingsWindow : Form
+	{
+		class GetDefaultSettingsTask : IModalTask
+		{
+			TextWriter Log;
+
+			public string ServerAndPort;
+			public string UserName;
+
+			public GetDefaultSettingsTask(TextWriter Log)
+			{
+				this.Log = Log;
+			}
+
+			public bool Run(out string ErrorMessage)
+			{
+				if(PerforceModalTask.TryGetServerSettings(null, ref ServerAndPort, ref UserName, Log))
+				{
+					ErrorMessage = null;
+					return true;
+				}
+				else
+				{
+					ErrorMessage = "Unable to query server settings";
+					return false;
+				}
+			}
+		}
+
+		class PerforceTestConnectionTask : IPerforceModalTask
+		{
+			string DepotPath;
+
+			public PerforceTestConnectionTask(string DepotPath)
+			{
+				this.DepotPath = DepotPath ?? DeploymentSettings.DefaultDepotPath;
+			}
+
+			public bool Run(PerforceConnection Perforce, TextWriter Log, out string ErrorMessage)
+			{
+				string CheckFilePath = String.Format("{0}/Release/UnrealGameSync.exe", DepotPath);
+
+				List<PerforceFileRecord> FileRecords;
+				if(!Perforce.FindFiles(CheckFilePath, out FileRecords, Log) || FileRecords.Count == 0)
+				{
+					ErrorMessage = String.Format("Unable to find {0}", CheckFilePath);
+					return false;
+				}
+
+				ErrorMessage = null;
+				return true;
+			}
+		}
+
+		string OriginalExecutableFileName;
+		UserSettings Settings;
+		TextWriter Log;
+
+		string InitialServerAndPort;
+		string InitialUserName;
+		string InitialDepotPath;
+		bool bInitialUnstable;
+
+		bool? bRestartUnstable;
+
+		private ApplicationSettingsWindow(string DefaultServerAndPort, string DefaultUserName, bool bUnstable, string OriginalExecutableFileName, UserSettings Settings, TextWriter Log)
+		{
+			InitializeComponent();
+
+			this.OriginalExecutableFileName = OriginalExecutableFileName;
+			this.Settings = Settings;
+			this.Log = Log;
+
+			Utility.ReadGlobalPerforceSettings(ref InitialServerAndPort, ref InitialUserName, ref InitialDepotPath);
+			bInitialUnstable = bUnstable;
+
+			this.AutomaticallyRunAtStartupCheckBox.Checked = IsAutomaticallyRunAtStartup();
+			this.KeepInTrayCheckBox.Checked = Settings.bKeepInTray;
+						
+			this.ServerTextBox.Text = InitialServerAndPort;
+			this.ServerTextBox.Select(ServerTextBox.TextLength, 0);
+			this.ServerTextBox.CueBanner = (DefaultServerAndPort == null)? "Default" : String.Format("Default ({0})", DefaultServerAndPort);
+
+			this.UserNameTextBox.Text = InitialUserName;
+			this.UserNameTextBox.Select(UserNameTextBox.TextLength, 0);
+			this.UserNameTextBox.CueBanner = (DefaultUserName == null)? "Default" : String.Format("Default ({0})", DefaultUserName);
+
+			this.DepotPathTextBox.Text = InitialDepotPath;
+			this.DepotPathTextBox.Select(DepotPathTextBox.TextLength, 0);
+			this.DepotPathTextBox.CueBanner = DeploymentSettings.DefaultDepotPath;
+
+			this.UseUnstableBuildCheckBox.Checked = bUnstable;
+		}
+
+		public static bool? ShowModal(IWin32Window Owner, bool bUnstable, string OriginalExecutableFileName, UserSettings Settings, TextWriter Log)
+		{
+			GetDefaultSettingsTask DefaultSettings = new GetDefaultSettingsTask(Log);
+
+			string ErrorMessage;
+			ModalTask.Execute(Owner, DefaultSettings, "Checking Settings", "Checking settings, please wait...", out ErrorMessage);
+
+			ApplicationSettingsWindow ApplicationSettings = new ApplicationSettingsWindow(DefaultSettings.ServerAndPort, DefaultSettings.UserName, bUnstable, OriginalExecutableFileName, Settings, Log);
+			if(ApplicationSettings.ShowDialog() == DialogResult.OK)
+			{
+				return ApplicationSettings.bRestartUnstable;
+			}
+			else
+			{
+				return null;
+			}
+		}
+
+		private bool IsAutomaticallyRunAtStartup()
+		{
+			RegistryKey Key = Registry.CurrentUser.OpenSubKey("SOFTWARE\\Microsoft\\Windows\\CurrentVersion\\Run");
+			return (Key.GetValue("UnrealGameSync") != null);
+		}
+
+		private void OkBtn_Click(object sender, EventArgs e)
+		{
+			// Update the settings
+			string ServerAndPort = ServerTextBox.Text.Trim();
+			if(ServerAndPort.Length == 0)
+			{
+				ServerAndPort = null;
+			}
+
+			string UserName = UserNameTextBox.Text.Trim();
+			if(UserName.Length == 0)
+			{
+				UserName = null;
+			}
+
+			string DepotPath = DepotPathTextBox.Text.Trim();
+			if(DepotPath.Length == 0 || DepotPath == DeploymentSettings.DefaultDepotPath)
+			{
+				DepotPath = null;
+			}
+
+			bool bUnstable = UseUnstableBuildCheckBox.Checked;
+
+			if(ServerAndPort != InitialServerAndPort || UserName != InitialUserName || DepotPath != InitialDepotPath || bUnstable != bInitialUnstable)
+			{
+				// Try to log in to the new server, and check the application is there
+				if(ServerAndPort != InitialServerAndPort || UserName != InitialUserName || DepotPath != InitialDepotPath)
+				{
+					string ErrorMessage;
+					ModalTaskResult Result = PerforceModalTask.Execute(this, null, ServerAndPort, UserName, new PerforceTestConnectionTask(DepotPath), "Connecting", "Checking connection, please wait...", Log, out ErrorMessage);
+					if(Result != ModalTaskResult.Succeeded)
+					{
+						if(Result == ModalTaskResult.Failed)
+						{
+							MessageBox.Show(ErrorMessage, "Unable to connect");
+						}
+						return;
+					}
+				}
+
+				if(MessageBox.Show("UnrealGameSync must be restarted to apply these settings.\n\nWould you like to restart now?", "Restart Required", MessageBoxButtons.OKCancel) != DialogResult.OK)
+				{
+					return;
+				}
+
+				bRestartUnstable = UseUnstableBuildCheckBox.Checked;
+				Utility.SaveGlobalPerforceSettings(ServerAndPort, UserName, DepotPath);
+			}
+
+			RegistryKey Key = Registry.CurrentUser.CreateSubKey("SOFTWARE\\Microsoft\\Windows\\CurrentVersion\\Run");
+			if(IsAutomaticallyRunAtStartup())
+			{
+	            Key.DeleteValue("UnrealGameSync", false);
+			}
+			else
+			{
+				Key.SetValue("UnrealGameSync", String.Format("\"{0}\" -RestoreState", OriginalExecutableFileName));
+			}
+
+			if(Settings.bKeepInTray != KeepInTrayCheckBox.Checked)
+			{
+				Settings.bKeepInTray = KeepInTrayCheckBox.Checked;
+				Settings.Save();
+			}
+
+			DialogResult = DialogResult.OK;
+			Close();
+		}
+
+		private void CancelBtn_Click(object sender, EventArgs e)
+		{
+			DialogResult = DialogResult.Cancel;
+			Close();
+		}
+	}
+}
diff --git a/Engine/Source/Programs/UnrealGameSync/UnrealGameSync/Forms/ArgumentsWindow.cs b/Engine/Source/Programs/UnrealGameSync/UnrealGameSync/Forms/ArgumentsWindow.cs
--- a/Engine/Source/Programs/UnrealGameSync/UnrealGameSync/Forms/ArgumentsWindow.cs
+++ b/Engine/Source/Programs/UnrealGameSync/UnrealGameSync/Forms/ArgumentsWindow.cs
@@ -1,4 +1,4 @@
-﻿// Copyright 1998-2018 Epic Games, Inc. All Rights Reserved.
+﻿// Copyright 1998-2019 Epic Games, Inc. All Rights Reserved.
 
 using System;
 using System.Collections.Generic;
diff --git a/Engine/Source/Programs/UnrealGameSync/UnrealGameSync/Forms/AutomatedSyncWindow.cs b/Engine/Source/Programs/UnrealGameSync/UnrealGameSync/Forms/AutomatedSyncWindow.cs
new file mode 100644
--- /dev/null
+++ b/Engine/Source/Programs/UnrealGameSync/UnrealGameSync/Forms/AutomatedSyncWindow.cs
@@ -0,0 +1,243 @@
+// Copyright 1998-2019 Epic Games, Inc. All Rights Reserved.
+
+using System;
+using System.Collections.Generic;
+using System.ComponentModel;
+using System.Data;
+using System.Drawing;
+using System.IO;
+using System.Linq;
+using System.Text;
+using System.Threading.Tasks;
+using System.Windows.Forms;
+
+namespace UnrealGameSync
+{
+	partial class AutomatedSyncWindow : Form
+	{
+		class FindDefaultWorkspaceTask : IPerforceModalTask
+		{
+			string StreamName;
+			public string WorkspaceName;
+
+			public FindDefaultWorkspaceTask(string StreamName)
+			{
+				this.StreamName = StreamName;
+			}
+
+			public bool Run(PerforceConnection Perforce, TextWriter Log, out string ErrorMessage)
+			{
+				PerforceInfoRecord Info;
+				if(!Perforce.Info(out Info, Log))
+				{
+					ErrorMessage = "Unable to query Perforce info.";
+					return false;
+				}
+				
+				List<PerforceClientRecord> Clients;
+				if(!Perforce.FindClients(Info.UserName, out Clients, Log))
+				{
+					ErrorMessage = "Unable to enumerate clients from Perforce";
+					return false;
+				}
+
+				List<PerforceClientRecord> CandidateClients = new List<PerforceClientRecord>();
+				foreach(PerforceClientRecord Client in Clients)
+				{
+					if(Client.Host == null || Client.Host.Equals(Info.HostName, StringComparison.OrdinalIgnoreCase))
+					{
+						if(Client.Stream != null && Client.Stream.Equals(StreamName, StringComparison.OrdinalIgnoreCase))
+						{
+							CandidateClients.Add(Client);
+						}
+					}
+				}
+
+				if(CandidateClients.Count == 1)
+				{
+					WorkspaceName = CandidateClients[0].Name;
+				}
+
+				ErrorMessage = null;
+				return true;
+			}
+		}
+
+		class ValidateWorkspaceTask : IPerforceModalTask
+		{
+			public string WorkspaceName;
+			public string StreamName;
+			public string ServerAndPort;
+			public string UserName;
+			public bool bRequiresStreamSwitch;
+			public bool bHasOpenFiles;
+
+			public ValidateWorkspaceTask(string WorkspaceName, string StreamName)
+			{
+				this.WorkspaceName = WorkspaceName;
+				this.StreamName = StreamName;
+			}
+
+			public bool Run(PerforceConnection Perforce, TextWriter Log, out string ErrorMessage)
+			{
+				this.ServerAndPort = Perforce.ServerAndPort;
+				this.UserName = Perforce.UserName;
+
+				PerforceSpec Spec;
+				if(!Perforce.TryGetClientSpec(WorkspaceName, out Spec, Log))
+				{
+					ErrorMessage = String.Format("Unable to get info for client '{0}'", WorkspaceName);
+					return false;
+				}
+
+				string CurrentStreamName = Spec.GetField("Stream");
+				if(CurrentStreamName == null || CurrentStreamName != StreamName)
+				{
+					bRequiresStreamSwitch = true;
+					bHasOpenFiles = Perforce.HasOpenFiles(Log);
+				}
+			
+				ErrorMessage = null;
+				return true;
+			}
+		}
+
+		public class WorkspaceInfo
+		{
+			public string ServerAndPort;
+			public string UserName;
+			public string WorkspaceName;
+			public bool bRequiresStreamSwitch;
+		}
+		
+		string StreamName;
+		TextWriter Log;
+
+		string ServerAndPort;
+		string UserName;
+		WorkspaceInfo SelectedWorkspaceInfo;
+
+		private AutomatedSyncWindow(string StreamName, string ProjectPath, string WorkspaceName, TextWriter Log)
+		{
+			this.StreamName = StreamName;
+			this.Log = Log;
+
+			InitializeComponent();
+
+			ActiveControl = WorkspaceNameTextBox;
+
+			MinimumSize = Size;
+			MaximumSize = new Size(32768, Size.Height);
+
+			ProjectTextBox.Text = StreamName + ProjectPath;
+
+			if(WorkspaceName != null)
+			{
+				WorkspaceNameTextBox.Text = WorkspaceName;
+				WorkspaceNameTextBox.Select(WorkspaceNameTextBox.Text.Length, 0);
+			}
+
+			UpdateServerLabel();
+			UpdateOkButton();
+		}
+
+		public static bool ShowModal(IWin32Window Owner, string StreamName, string ProjectPath, out WorkspaceInfo WorkspaceInfo, TextWriter Log)
+		{
+			FindDefaultWorkspaceTask FindWorkspace = new FindDefaultWorkspaceTask(StreamName);
+
+			string ErrorMessage;
+			PerforceModalTask.Execute(Owner, null, null, null, FindWorkspace, "Finding workspace", "Finding default workspace, please wait...", Log, out ErrorMessage);
+
+			AutomatedSyncWindow Window = new AutomatedSyncWindow(StreamName, ProjectPath, FindWorkspace.WorkspaceName, Log);
+			if(Window.ShowDialog() == DialogResult.OK)
+			{
+				WorkspaceInfo = Window.SelectedWorkspaceInfo;
+				return true;
+			}
+			else
+			{
+				WorkspaceInfo = null;
+				return false;
+			}
+		}
+
+		private void ChangeLink_LinkClicked(object sender, LinkLabelLinkClickedEventArgs e)
+		{
+			string NewServerAndPort;
+			string NewUserName;
+			if(ConnectWindow.ShowModal(this, ServerAndPort, UserName, out NewServerAndPort, out NewUserName))
+			{
+				ServerAndPort = NewServerAndPort;
+				UserName = NewUserName;
+				UpdateServerLabel();
+			}
+		}
+
+		private void UpdateServerLabel()
+		{
+			ServerLabel.Text = OpenProjectWindow.GetServerLabelText(ServerAndPort, UserName);
+			ChangeLink.Location = new Point(ServerLabel.Right + 5, ChangeLink.Location.Y);
+		}
+
+		private void WorkspaceNameNewBtn_Click(object sender, EventArgs e)
+		{
+			string WorkspaceName;
+			if(NewWorkspaceWindow.ShowModal(this, ServerAndPort, UserName, StreamName, WorkspaceNameTextBox.Text, Log, out WorkspaceName))
+			{
+				WorkspaceNameTextBox.Text = WorkspaceName;
+				UpdateOkButton();
+			}
+		}
+
+		private void WorkspaceNameBrowseBtn_Click(object sender, EventArgs e)
+		{
+			string WorkspaceName = WorkspaceNameTextBox.Text;
+			if(SelectWorkspaceWindow.ShowModal(this, ServerAndPort, UserName, WorkspaceName, Log, out WorkspaceName))
+			{
+				WorkspaceNameTextBox.Text = WorkspaceName;
+				UpdateOkButton();
+			}
+		}
+
+		private void UpdateOkButton()
+		{
+			OkBtn.Enabled = (WorkspaceNameTextBox.Text.Length > 0);
+		}
+
+		private void OkBtn_Click(object sender, EventArgs e)
+		{
+			ValidateWorkspaceTask ValidateWorkspace = new ValidateWorkspaceTask(WorkspaceNameTextBox.Text, StreamName);
+
+			string ErrorMessage;
+			ModalTaskResult Result = PerforceModalTask.Execute(Owner, null, ServerAndPort, UserName, ValidateWorkspace, "Checking workspace", "Checking workspace, please wait...", Log, out ErrorMessage);
+			if(Result == ModalTaskResult.Failed)
+			{
+				MessageBox.Show(ErrorMessage);
+			}
+			else if(Result == ModalTaskResult.Succeeded)
+			{
+				if(ValidateWorkspace.bRequiresStreamSwitch)
+				{
+					string Message;
+					if(ValidateWorkspace.bHasOpenFiles)
+					{
+						Message = String.Format("You have files open for edit in this workspace. If you switch this workspace to {0}, you will not be able to submit them until you switch back.\n\nContinue switching streams?", StreamName);
+					}
+					else
+					{
+						Message = String.Format("Switch this workspace to {0}?", StreamName);
+					}
+					if(MessageBox.Show(Message, "Switch Streams", MessageBoxButtons.YesNo) != DialogResult.Yes)
+					{
+						return;
+					}
+				}
+
+				SelectedWorkspaceInfo = new WorkspaceInfo(){ ServerAndPort = ValidateWorkspace.ServerAndPort, UserName = ValidateWorkspace.UserName, WorkspaceName = ValidateWorkspace.WorkspaceName, bRequiresStreamSwitch = ValidateWorkspace.bRequiresStreamSwitch };
+
+				DialogResult = DialogResult.OK;
+				Close();
+			}
+		}
+	}
+}
diff --git a/Engine/Source/Programs/UnrealGameSync/UnrealGameSync/Forms/BuildStepWindow.cs b/Engine/Source/Programs/UnrealGameSync/UnrealGameSync/Forms/BuildStepWindow.cs
--- a/Engine/Source/Programs/UnrealGameSync/UnrealGameSync/Forms/BuildStepWindow.cs
+++ b/Engine/Source/Programs/UnrealGameSync/UnrealGameSync/Forms/BuildStepWindow.cs
@@ -1,4 +1,4 @@
-﻿// Copyright 1998-2018 Epic Games, Inc. All Rights Reserved.
+﻿// Copyright 1998-2019 Epic Games, Inc. All Rights Reserved.
 
 using System;
 using System.Collections.Generic;
diff --git a/Engine/Source/Programs/UnrealGameSync/UnrealGameSync/Forms/ChangelistWindow.cs b/Engine/Source/Programs/UnrealGameSync/UnrealGameSync/Forms/ChangelistWindow.cs
--- a/Engine/Source/Programs/UnrealGameSync/UnrealGameSync/Forms/ChangelistWindow.cs
+++ b/Engine/Source/Programs/UnrealGameSync/UnrealGameSync/Forms/ChangelistWindow.cs
@@ -1,4 +1,4 @@
-// Copyright 1998-2018 Epic Games, Inc. All Rights Reserved.
+// Copyright 1998-2019 Epic Games, Inc. All Rights Reserved.
 
 using System;
 using System.Collections.Generic;
diff --git a/Engine/Source/Programs/UnrealGameSync/UnrealGameSync/Forms/CleanWorkspaceWindow.cs b/Engine/Source/Programs/UnrealGameSync/UnrealGameSync/Forms/CleanWorkspaceWindow.cs
--- a/Engine/Source/Programs/UnrealGameSync/UnrealGameSync/Forms/CleanWorkspaceWindow.cs
+++ b/Engine/Source/Programs/UnrealGameSync/UnrealGameSync/Forms/CleanWorkspaceWindow.cs
@@ -1,4 +1,4 @@
-// Copyright 1998-2018 Epic Games, Inc. All Rights Reserved.
+// Copyright 1998-2019 Epic Games, Inc. All Rights Reserved.
 
 using System;
 using System.Collections.Generic;
diff --git a/Engine/Source/Programs/UnrealGameSync/UnrealGameSync/Forms/ClobberWindow.cs b/Engine/Source/Programs/UnrealGameSync/UnrealGameSync/Forms/ClobberWindow.cs
--- a/Engine/Source/Programs/UnrealGameSync/UnrealGameSync/Forms/ClobberWindow.cs
+++ b/Engine/Source/Programs/UnrealGameSync/UnrealGameSync/Forms/ClobberWindow.cs
@@ -1,4 +1,4 @@
-﻿// Copyright 1998-2018 Epic Games, Inc. All Rights Reserved.
+﻿// Copyright 1998-2019 Epic Games, Inc. All Rights Reserved.
 
 using System;
 using System.Collections.Generic;
diff --git a/Engine/Source/Programs/UnrealGameSync/UnrealGameSync/Forms/ConnectWindow.cs b/Engine/Source/Programs/UnrealGameSync/UnrealGameSync/Forms/ConnectWindow.cs
--- a/Engine/Source/Programs/UnrealGameSync/UnrealGameSync/Forms/ConnectWindow.cs
+++ b/Engine/Source/Programs/UnrealGameSync/UnrealGameSync/Forms/ConnectWindow.cs
@@ -1,4 +1,4 @@
-﻿// Copyright 1998-2018 Epic Games, Inc. All Rights Reserved.
+﻿// Copyright 1998-2019 Epic Games, Inc. All Rights Reserved.
 
 using System;
 using System.Collections.Generic;
diff --git a/Engine/Source/Programs/UnrealGameSync/UnrealGameSync/Forms/DeleteWindow.cs b/Engine/Source/Programs/UnrealGameSync/UnrealGameSync/Forms/DeleteWindow.cs
--- a/Engine/Source/Programs/UnrealGameSync/UnrealGameSync/Forms/DeleteWindow.cs
+++ b/Engine/Source/Programs/UnrealGameSync/UnrealGameSync/Forms/DeleteWindow.cs
@@ -1,4 +1,4 @@
-﻿// Copyright 1998-2018 Epic Games, Inc. All Rights Reserved.
+﻿// Copyright 1998-2019 Epic Games, Inc. All Rights Reserved.
 
 using System;
 using System.Collections.Generic;
diff --git a/Engine/Source/Programs/UnrealGameSync/UnrealGameSync/Forms/DiagnosticsWindow.cs b/Engine/Source/Programs/UnrealGameSync/UnrealGameSync/Forms/DiagnosticsWindow.cs
--- a/Engine/Source/Programs/UnrealGameSync/UnrealGameSync/Forms/DiagnosticsWindow.cs
+++ b/Engine/Source/Programs/UnrealGameSync/UnrealGameSync/Forms/DiagnosticsWindow.cs
@@ -1,4 +1,4 @@
-// Copyright 1998-2018 Epic Games, Inc. All Rights Reserved.
+// Copyright 1998-2019 Epic Games, Inc. All Rights Reserved.
 
 using System;
 using System.Collections.Generic;
diff --git a/Engine/Source/Programs/UnrealGameSync/UnrealGameSync/Forms/FailedToDeleteWindow.cs b/Engine/Source/Programs/UnrealGameSync/UnrealGameSync/Forms/FailedToDeleteWindow.cs
--- a/Engine/Source/Programs/UnrealGameSync/UnrealGameSync/Forms/FailedToDeleteWindow.cs
+++ b/Engine/Source/Programs/UnrealGameSync/UnrealGameSync/Forms/FailedToDeleteWindow.cs
@@ -1,4 +1,4 @@
-﻿// Copyright 1998-2018 Epic Games, Inc. All Rights Reserved.
+﻿// Copyright 1998-2019 Epic Games, Inc. All Rights Reserved.
 
 using System;
 using System.Collections.Generic;
diff --git a/Engine/Source/Programs/UnrealGameSync/UnrealGameSync/Forms/LeaveCommentWindow.cs b/Engine/Source/Programs/UnrealGameSync/UnrealGameSync/Forms/LeaveCommentWindow.cs
--- a/Engine/Source/Programs/UnrealGameSync/UnrealGameSync/Forms/LeaveCommentWindow.cs
+++ b/Engine/Source/Programs/UnrealGameSync/UnrealGameSync/Forms/LeaveCommentWindow.cs
@@ -1,4 +1,4 @@
-﻿// Copyright 1998-2018 Epic Games, Inc. All Rights Reserved.
+﻿// Copyright 1998-2019 Epic Games, Inc. All Rights Reserved.
 
 using System;
 using System.Collections.Generic;
diff --git a/Engine/Source/Programs/UnrealGameSync/UnrealGameSync/Forms/MainWindow.cs b/Engine/Source/Programs/UnrealGameSync/UnrealGameSync/Forms/MainWindow.cs
--- a/Engine/Source/Programs/UnrealGameSync/UnrealGameSync/Forms/MainWindow.cs
+++ b/Engine/Source/Programs/UnrealGameSync/UnrealGameSync/Forms/MainWindow.cs
@@ -1,4 +1,4 @@
-// Copyright 1998-2018 Epic Games, Inc. All Rights Reserved.
+// Copyright 1998-2019 Epic Games, Inc. All Rights Reserved.
 
 using System;
 using System.Collections.Generic;
@@ -9,6 +9,9 @@ using System.Linq;
 using System.Runtime.InteropServices;
 using System.Windows.Forms;
 using System.Threading;
+using System.Reflection;
+using System.Text;
+using System.Diagnostics;
 
 namespace UnrealGameSync
 {
@@ -53,10 +56,12 @@ namespace UnrealGameSync
 
 		private const int WM_SETREDRAW = 11; 
 
+		UpdateMonitor UpdateMonitor;
 		SynchronizationContext MainThreadSynchronizationContext;
 
 		string ApiUrl;
 		string DataFolder;
+		string CacheFolder;
 		LineBasedTextWriter Log;
 		UserSettings Settings;
 		int TabMenu_TabIdx = -1;
@@ -70,21 +75,34 @@ namespace UnrealGameSync
 		System.Threading.Timer ScheduleSettledTimer;
 
 		string OriginalExecutableFileName;
+		bool bUnstable;
 
 		IMainWindowTabPanel CurrentTabPanel;
 
-		public MainWindow(string InApiUrl, string InDataFolder, bool bInRestoreStateOnLoad, string InOriginalExecutableFileName, List<DetectProjectSettingsResult> StartupProjects, LineBasedTextWriter InLog, UserSettings InSettings)
+		AutomationServer AutomationServer;
+		TextWriter AutomationLog;
+
+		bool bAllowCreatingHandle;
+
+		public MainWindow(UpdateMonitor InUpdateMonitor, string InApiUrl, string InDataFolder, string InCacheFolder, bool bInRestoreStateOnLoad, string InOriginalExecutableFileName, bool bInUnstable, DetectProjectSettingsResult[] StartupProjects, LineBasedTextWriter InLog, UserSettings InSettings)
 		{
 			InitializeComponent();
 
+			UpdateMonitor = InUpdateMonitor;
 			MainThreadSynchronizationContext = SynchronizationContext.Current;
 			ApiUrl = InApiUrl;
 			DataFolder = InDataFolder;
+			CacheFolder = InCacheFolder;
 			bRestoreStateOnLoad = bInRestoreStateOnLoad;
 			OriginalExecutableFileName = InOriginalExecutableFileName;
+			bUnstable = bInUnstable;
 			Log = InLog;
 			Settings = InSettings;
 
+			// While creating tab controls during startup, we need to prevent layout calls resulting in the window handle being created too early. Disable layout calls here.
+			SuspendLayout();
+			TabPanel.SuspendLayout();
+
 			TabControl.OnTabChanged += TabControl_OnTabChanged;
 			TabControl.OnNewTabClick += TabControl_OnNewTabClick;
 			TabControl.OnTabClicked += TabControl_OnTabClicked;
@@ -124,11 +142,187 @@ namespace UnrealGameSync
 			}
 
 			StartScheduleTimer();
+
+			if(bUnstable)
+			{
+				Text += String.Format(" (UNSTABLE BUILD {0})", Assembly.GetExecutingAssembly().GetName().Version);
+			}
+
+			AutomationLog = new TimestampLogWriter(new BoundedLogWriter(Path.Combine(DataFolder, "Automation.log")));
+			AutomationServer = new AutomationServer(Request => { MainThreadSynchronizationContext.Post(Obj => PostAutomationRequest(Request), null); }, AutomationLog);
+
+			// Allow creating controls from now on
+			TabPanel.ResumeLayout(false);
+			ResumeLayout(false);
+
+			bAllowCreatingHandle = true;
+		}
+
+		void PostAutomationRequest(AutomationRequest Request)
+		{
+			try
+			{
+				if(!CanFocus)
+				{
+					Request.SetOutput(new AutomationRequestOutput(AutomationRequestResult.Busy));
+				}
+				else if(Request.Input.Type == AutomationRequestType.SyncProject)
+				{
+					AutomationRequestOutput Output = StartAutomatedSync(Request, true);
+					if(Output != null)
+					{
+						Request.SetOutput(Output);
+					}
+				}
+				else if(Request.Input.Type == AutomationRequestType.FindProject)
+				{
+					AutomationRequestOutput Output = FindProject(Request);
+					Request.SetOutput(Output);
+				}
+				else if(Request.Input.Type == AutomationRequestType.OpenProject)
+				{
+					AutomationRequestOutput Output = StartAutomatedSync(Request, false);
+					if(Output != null)
+					{
+						Request.SetOutput(Output);
+					}
+				}
+				else
+				{
+					Request.SetOutput(new AutomationRequestOutput(AutomationRequestResult.Invalid));
+				}
+			}
+			catch(Exception Ex)
+			{
+				Log.WriteLine("Exception running automation request: {0}", Ex);
+				Request.SetOutput(new AutomationRequestOutput(AutomationRequestResult.Invalid));
+			}
+		}
+
+		AutomationRequestOutput StartAutomatedSync(AutomationRequest Request, bool bForceSync)
+		{
+			ShowAndActivate();
+
+			BinaryReader Reader = new BinaryReader(new MemoryStream(Request.Input.Data));
+			string StreamName = Reader.ReadString();
+			string ProjectPath = Reader.ReadString();
+
+			AutomatedSyncWindow.WorkspaceInfo WorkspaceInfo;
+			if(!AutomatedSyncWindow.ShowModal(this, StreamName, ProjectPath, out WorkspaceInfo, Log))
+			{
+				return new AutomationRequestOutput(AutomationRequestResult.Canceled);
+			}
+
+			if(WorkspaceInfo.bRequiresStreamSwitch)
+			{
+				// Close any tab containing this window
+				for(int ExistingTabIdx = 0; ExistingTabIdx < TabControl.GetTabCount(); ExistingTabIdx++)
+				{
+					WorkspaceControl ExistingWorkspace = TabControl.GetTabData(ExistingTabIdx) as WorkspaceControl;
+					if(ExistingWorkspace != null && ExistingWorkspace.ClientName.Equals(WorkspaceInfo.WorkspaceName))
+					{
+						TabControl.RemoveTab(ExistingTabIdx);
+						break;
+					}
+				}
+
+				// Switch the stream
+				PerforceConnection Perforce = new PerforceConnection(WorkspaceInfo.UserName, WorkspaceInfo.WorkspaceName, WorkspaceInfo.ServerAndPort);
+				if(!Perforce.SwitchStream(StreamName, Log))
+				{
+					Log.WriteLine("Unable to switch stream");
+					return new AutomationRequestOutput(AutomationRequestResult.Error);
+				}
+			}
+
+			UserSelectedProjectSettings SelectedProject = new UserSelectedProjectSettings(WorkspaceInfo.ServerAndPort, WorkspaceInfo.UserName, UserSelectedProjectType.Client, String.Format("//{0}{1}", WorkspaceInfo.WorkspaceName, ProjectPath), null);
+
+			int TabIdx = TryOpenProject(SelectedProject, -1, OpenProjectOptions.None);
+			if(TabIdx == -1)
+			{
+				Log.WriteLine("Unable to open project");
+				return new AutomationRequestOutput(AutomationRequestResult.Error);
+			}
+
+			WorkspaceControl Workspace = TabControl.GetTabData(TabIdx) as WorkspaceControl;
+			if(Workspace == null)
+			{
+				Log.WriteLine("Workspace was unable to open");
+				return new AutomationRequestOutput(AutomationRequestResult.Error);
+			}
+
+			if(!bForceSync && Workspace.CanLaunchEditor())
+			{
+				return new AutomationRequestOutput(AutomationRequestResult.Ok, Encoding.UTF8.GetBytes(Workspace.SelectedFileName));
+			}
+
+			Workspace.AddStartupCallback((Control, bCancel) => StartAutomatedSyncAfterStartup(Control, bCancel, Request));
+			return null;
+		}
+
+		private void StartAutomatedSyncAfterStartup(WorkspaceControl Workspace, bool bCancel, AutomationRequest Request)
+		{
+			if(bCancel)
+			{
+				Request.SetOutput(new AutomationRequestOutput(AutomationRequestResult.Canceled));
+			}
+			else
+			{
+				Workspace.SyncLatestChange(Result => CompleteAutomatedSync(Result, Workspace.SelectedFileName, Request));
+			}
+		}
+
+		void CompleteAutomatedSync(WorkspaceUpdateResult Result, string SelectedFileName, AutomationRequest Request)
+		{
+			if(Result == WorkspaceUpdateResult.Success)
+			{
+				Request.SetOutput(new AutomationRequestOutput(AutomationRequestResult.Ok, Encoding.UTF8.GetBytes(SelectedFileName)));
+			}
+			else if(Result == WorkspaceUpdateResult.Canceled)
+			{
+				Request.SetOutput(new AutomationRequestOutput(AutomationRequestResult.Canceled));
+			}
+			else
+			{
+				Request.SetOutput(new AutomationRequestOutput(AutomationRequestResult.Error));
+			}
+		}
+
+		AutomationRequestOutput FindProject(AutomationRequest Request)
+		{
+			BinaryReader Reader = new BinaryReader(new MemoryStream(Request.Input.Data));
+			string StreamName = Reader.ReadString();
+			string ProjectPath = Reader.ReadString();
+
+			for(int ExistingTabIdx = 0; ExistingTabIdx < TabControl.GetTabCount(); ExistingTabIdx++)
+			{
+				WorkspaceControl ExistingWorkspace = TabControl.GetTabData(ExistingTabIdx) as WorkspaceControl;
+				if(ExistingWorkspace != null && String.Compare(ExistingWorkspace.StreamName, StreamName, StringComparison.OrdinalIgnoreCase) == 0 && ExistingWorkspace.SelectedProject != null)
+				{
+					string ClientPath = ExistingWorkspace.SelectedProject.ClientPath;
+					if(ClientPath != null && ClientPath.StartsWith("//"))
+					{
+						int SlashIdx = ClientPath.IndexOf('/', 2);
+						if(SlashIdx != -1)
+						{
+							string ExistingProjectPath = ClientPath.Substring(SlashIdx);
+							if(String.Compare(ExistingProjectPath, ProjectPath, StringComparison.OrdinalIgnoreCase) == 0)
+							{
+								return new AutomationRequestOutput(AutomationRequestResult.Ok, Encoding.UTF8.GetBytes(ExistingWorkspace.SelectedFileName));
+							}
+						}
+					}
+				}
+			}
+
+			return new AutomationRequestOutput(AutomationRequestResult.NotFound);
 		}
 
 		protected override void OnHandleCreated(EventArgs e)
 		{
 			base.OnHandleCreated(e);
+
+			Debug.Assert(bAllowCreatingHandle, "Window handle should not be created before constructor has run.");
 		}
 
 		void TabControl_OnButtonClick(int ButtonIdx, Point Location, MouseButtons Buttons)
@@ -222,6 +416,18 @@ namespace UnrealGameSync
 
 			StopScheduleTimer();
 
+			if(AutomationServer != null)
+			{
+				AutomationServer.Dispose();
+				AutomationServer = null;
+			}
+
+			if(AutomationLog != null)
+			{
+				AutomationLog.Close();
+				AutomationLog = null;
+			}
+
 			base.Dispose(disposing);
 		}
 
@@ -231,9 +437,6 @@ namespace UnrealGameSync
 			{
 				Hide();
 				EventArgs.Cancel = true;
-
-				Settings.bWindowVisible = Visible;
-				Settings.Save();
 			}
 			else
 			{
@@ -248,10 +451,20 @@ namespace UnrealGameSync
 				}
 
 				StopScheduleTimer();
+			}
 
-				Settings.bWindowVisible = Visible;
-				Settings.Save();
+			Settings.bWindowVisible = Visible;
+			Settings.WindowState = WindowState;
+			if(WindowState == FormWindowState.Normal)
+			{
+				Settings.WindowBounds = new Rectangle(Location, Size);
 			}
+			else
+			{
+				Settings.WindowBounds = RestoreBounds;
+			}
+
+			Settings.Save();
 		}
 
 		private void SetupDefaultControl()
@@ -278,8 +491,7 @@ namespace UnrealGameSync
 			ErrorPanel.BorderStyle = BorderStyle.FixedSingle;
 			ErrorPanel.BackColor = System.Drawing.Color.FromArgb(((int)(((byte)(250)))), ((int)(((byte)(250)))), ((int)(((byte)(250)))));
 			ErrorPanel.Location = new Point(0, 0);
-			ErrorPanel.Size = new Size(TabPanel.Width, TabPanel.Height);
-			ErrorPanel.Anchor = AnchorStyles.Left | AnchorStyles.Top | AnchorStyles.Right | AnchorStyles.Bottom;
+			ErrorPanel.Dock = DockStyle.Fill;
 			ErrorPanel.Hide();
 
 			string SummaryText = String.Format("Unable to open '{0}'.", Project.ToString());
@@ -606,7 +818,7 @@ namespace UnrealGameSync
 		public void OpenNewProject()
 		{
 			DetectProjectSettingsTask DetectedProjectSettings;
-			if(OpenProjectWindow.ShowModal(this, null, out DetectedProjectSettings, Settings, DataFolder, Log))
+			if(OpenProjectWindow.ShowModal(this, null, out DetectedProjectSettings, Settings, DataFolder, CacheFolder, Log))
 			{
 				int NewTabIdx = TryOpenProject(DetectedProjectSettings, -1, OpenProjectOptions.None);
 				if(NewTabIdx != -1)
@@ -658,7 +870,7 @@ namespace UnrealGameSync
 		public void EditSelectedProject(int TabIdx, UserSelectedProjectSettings SelectedProject)
 		{
 			DetectProjectSettingsTask DetectedProjectSettings;
-			if(OpenProjectWindow.ShowModal(this, SelectedProject, out DetectedProjectSettings, Settings, DataFolder, Log))
+			if(OpenProjectWindow.ShowModal(this, SelectedProject, out DetectedProjectSettings, Settings, DataFolder, CacheFolder, Log))
 			{
 				int NewTabIdx = TryOpenProject(DetectedProjectSettings, TabIdx, OpenProjectOptions.None);
 				if(NewTabIdx != -1)
@@ -676,7 +888,7 @@ namespace UnrealGameSync
 		int TryOpenProject(UserSelectedProjectSettings Project, int ReplaceTabIdx, OpenProjectOptions Options = OpenProjectOptions.None)
 		{
 			Log.WriteLine("Detecting settings for {0}", Project);
-			using(DetectProjectSettingsTask DetectProjectSettings = new DetectProjectSettingsTask(Project, DataFolder, new PrefixedTextWriter("  ", Log)))
+			using(DetectProjectSettingsTask DetectProjectSettings = new DetectProjectSettingsTask(Project, DataFolder, CacheFolder, new PrefixedTextWriter("  ", Log)))
 			{
 				string ErrorMessage;
 
@@ -753,11 +965,9 @@ namespace UnrealGameSync
 			}
 
 			// Now that we have the project settings, we can construct the tab
-			WorkspaceControl NewWorkspace = new WorkspaceControl(this, ApiUrl, OriginalExecutableFileName, ProjectSettings, Log, Settings);
+			WorkspaceControl NewWorkspace = new WorkspaceControl(this, ApiUrl, OriginalExecutableFileName, bUnstable, ProjectSettings, Log, Settings);
 			NewWorkspace.Parent = TabPanel;
-			NewWorkspace.Location = new Point(0, 0);
-			NewWorkspace.Size = new Size(TabPanel.Width, TabPanel.Height);
-			NewWorkspace.Anchor = AnchorStyles.Left | AnchorStyles.Top | AnchorStyles.Right | AnchorStyles.Bottom;
+			NewWorkspace.Dock = DockStyle.Fill;
 			NewWorkspace.Hide();
 
 			// Add the tab
@@ -945,5 +1155,24 @@ namespace UnrealGameSync
 				}
 			}
 		}
+
+		public void ModifyApplicationSettings()
+		{
+			bool? bRelaunchUnstable = ApplicationSettingsWindow.ShowModal(this, bUnstable, OriginalExecutableFileName, Settings, Log);
+			if(bRelaunchUnstable.HasValue)
+			{
+				UpdateMonitor.TriggerUpdate(UpdateType.UserInitiated, bRelaunchUnstable);
+			}
+		}
+
+		private void MainWindow_Load(object sender, EventArgs e)
+		{
+			if(Settings.WindowBounds != null)
+			{
+				Location = Settings.WindowBounds.Value.Location;
+				Size = Settings.WindowBounds.Value.Size;
+			}
+			WindowState = Settings.WindowState;
+		}
 	}
 }
diff --git a/Engine/Source/Programs/UnrealGameSync/UnrealGameSync/Forms/ModalTaskWindow.cs b/Engine/Source/Programs/UnrealGameSync/UnrealGameSync/Forms/ModalTaskWindow.cs
--- a/Engine/Source/Programs/UnrealGameSync/UnrealGameSync/Forms/ModalTaskWindow.cs
+++ b/Engine/Source/Programs/UnrealGameSync/UnrealGameSync/Forms/ModalTaskWindow.cs
@@ -1,4 +1,4 @@
-// Copyright 1998-2018 Epic Games, Inc. All Rights Reserved.
+// Copyright 1998-2019 Epic Games, Inc. All Rights Reserved.
 
 using System;
 using System.Collections.Generic;
diff --git a/Engine/Source/Programs/UnrealGameSync/UnrealGameSync/Forms/ModifyBuildStepsWindow.cs b/Engine/Source/Programs/UnrealGameSync/UnrealGameSync/Forms/ModifyBuildStepsWindow.cs
--- a/Engine/Source/Programs/UnrealGameSync/UnrealGameSync/Forms/ModifyBuildStepsWindow.cs
+++ b/Engine/Source/Programs/UnrealGameSync/UnrealGameSync/Forms/ModifyBuildStepsWindow.cs
@@ -1,4 +1,4 @@
-// Copyright 1998-2018 Epic Games, Inc. All Rights Reserved.
+// Copyright 1998-2019 Epic Games, Inc. All Rights Reserved.
 
 using System;
 using System.Collections.Generic;
diff --git a/Engine/Source/Programs/UnrealGameSync/UnrealGameSync/Forms/NewWorkspaceWindow.cs b/Engine/Source/Programs/UnrealGameSync/UnrealGameSync/Forms/NewWorkspaceWindow.cs
--- a/Engine/Source/Programs/UnrealGameSync/UnrealGameSync/Forms/NewWorkspaceWindow.cs
+++ b/Engine/Source/Programs/UnrealGameSync/UnrealGameSync/Forms/NewWorkspaceWindow.cs
@@ -1,4 +1,4 @@
-// Copyright 1998-2018 Epic Games, Inc. All Rights Reserved.
+// Copyright 1998-2019 Epic Games, Inc. All Rights Reserved.
 
 using System;
 using System.Collections.Generic;
@@ -112,7 +112,7 @@ namespace UnrealGameSync
 		NewWorkspaceSettings Settings;
 		string DefaultRootPath;
 
-		private NewWorkspaceWindow(string ServerAndPort, string UserName, string DefaultStream, PerforceInfoRecord Info, List<PerforceClientRecord> Clients, TextWriter Log)
+		private NewWorkspaceWindow(string ServerAndPort, string UserName, string ForceStream, string DefaultStream, PerforceInfoRecord Info, List<PerforceClientRecord> Clients, TextWriter Log)
 		{
 			InitializeComponent();
 
@@ -159,17 +159,28 @@ namespace UnrealGameSync
 				}
 			}
 
-			StreamTextBox.Text = DefaultStream ?? "";
+			if(ForceStream != null)
+			{
+				StreamTextBox.Text = ForceStream;
+				StreamTextBox.Enabled = false;
+			}
+			else
+			{
+				StreamTextBox.Text = DefaultStream ?? "";
+				StreamTextBox.Enabled = true;
+			}
 			StreamTextBox.SelectionStart = StreamTextBox.Text.Length;
 			StreamTextBox.SelectionLength = 0;
 			StreamTextBox.Focus();
 
+			StreamBrowseBtn.Enabled = (ForceStream == null);
+
 			UpdateOkButton();
 			UpdateNameCueBanner();
 			UpdateRootDirCueBanner();
 		}
 
-		public static bool ShowModal(IWin32Window Owner, string ServerAndPort, string UserName, string CurrentWorkspaceName, TextWriter Log, out string WorkspaceName)
+		public static bool ShowModal(IWin32Window Owner, string ServerAndPort, string UserName, string ForceStreamName, string CurrentWorkspaceName, TextWriter Log, out string WorkspaceName)
 		{
 			FindWorkspaceSettingsTask FindSettings = new FindWorkspaceSettingsTask(CurrentWorkspaceName);
 
@@ -180,9 +191,12 @@ namespace UnrealGameSync
 				{
 					MessageBox.Show(ErrorMessage);
 				}
+
+				WorkspaceName = null;
+				return false;
 			}
 
-			NewWorkspaceWindow Window = new NewWorkspaceWindow(ServerAndPort, UserName, FindSettings.CurrentStream, FindSettings.Info, FindSettings.Clients, Log);
+			NewWorkspaceWindow Window = new NewWorkspaceWindow(ServerAndPort, UserName, ForceStreamName, FindSettings.CurrentStream, FindSettings.Info, FindSettings.Clients, Log);
 			if(Window.ShowDialog(Owner) == DialogResult.OK)
 			{
 				WorkspaceName = Window.Settings.Name;
diff --git a/Engine/Source/Programs/UnrealGameSync/UnrealGameSync/Forms/NotificationWindow.cs b/Engine/Source/Programs/UnrealGameSync/UnrealGameSync/Forms/NotificationWindow.cs
--- a/Engine/Source/Programs/UnrealGameSync/UnrealGameSync/Forms/NotificationWindow.cs
+++ b/Engine/Source/Programs/UnrealGameSync/UnrealGameSync/Forms/NotificationWindow.cs
@@ -1,4 +1,4 @@
-﻿// Copyright 1998-2018 Epic Games, Inc. All Rights Reserved.
+﻿// Copyright 1998-2019 Epic Games, Inc. All Rights Reserved.
 
 using System;
 using System.Collections.Generic;
diff --git a/Engine/Source/Programs/UnrealGameSync/UnrealGameSync/Forms/OpenProjectWindow.cs b/Engine/Source/Programs/UnrealGameSync/UnrealGameSync/Forms/OpenProjectWindow.cs
--- a/Engine/Source/Programs/UnrealGameSync/UnrealGameSync/Forms/OpenProjectWindow.cs
+++ b/Engine/Source/Programs/UnrealGameSync/UnrealGameSync/Forms/OpenProjectWindow.cs
@@ -1,4 +1,4 @@
-// Copyright 1998-2018 Epic Games, Inc. All Rights Reserved.
+// Copyright 1998-2019 Epic Games, Inc. All Rights Reserved.
 
 using System;
 using System.Collections.Generic;
@@ -19,16 +19,18 @@ namespace UnrealGameSync
 		string UserName;
 		DetectProjectSettingsTask DetectedProjectSettings;
 		string DataFolder;
+		string CacheFolder;
 		TextWriter Log;
 		UserSettings Settings;
 
-		private OpenProjectWindow(UserSelectedProjectSettings Project, UserSettings Settings, string DataFolder, TextWriter Log)
+		private OpenProjectWindow(UserSelectedProjectSettings Project, UserSettings Settings, string DataFolder, string CacheFolder, TextWriter Log)
 		{
 			InitializeComponent();
 
 			this.Settings = Settings;
 			this.DetectedProjectSettings = null;
 			this.DataFolder = DataFolder;
+			this.CacheFolder = CacheFolder;
 			this.Log = Log;
 
 			if(Project == null)
@@ -77,9 +79,9 @@ namespace UnrealGameSync
 			UpdateOkButton();
 		}
 
-		public static bool ShowModal(IWin32Window Owner, UserSelectedProjectSettings Project, out DetectProjectSettingsTask NewDetectedProjectSettings, UserSettings Settings, string DataFolder, TextWriter Log)
+		public static bool ShowModal(IWin32Window Owner, UserSelectedProjectSettings Project, out DetectProjectSettingsTask NewDetectedProjectSettings, UserSettings Settings, string DataFolder, string CacheFolder, TextWriter Log)
 		{
-			OpenProjectWindow Window = new OpenProjectWindow(Project, Settings, DataFolder, Log);
+			OpenProjectWindow Window = new OpenProjectWindow(Project, Settings, DataFolder, CacheFolder, Log);
 			if(Window.ShowDialog(Owner) == DialogResult.OK)
 			{
 				NewDetectedProjectSettings = Window.DetectedProjectSettings;
@@ -111,11 +113,11 @@ namespace UnrealGameSync
 			UpdateWorkspacePathBrowseButton();
 		}
 
-		private void UpdateServerLabel()
+		public static string GetServerLabelText(string ServerAndPort, string UserName)
 		{
 			if(ServerAndPort == null && UserName == null)
 			{
-				ServerLabel.Text = "Using default Perforce server settings.";
+				return "Using default Perforce server settings.";
 			}
 			else
 			{
@@ -137,9 +139,13 @@ namespace UnrealGameSync
 				{
 					Text.AppendFormat("server '{0}'.", ServerAndPort);
 				}
-				ServerLabel.Text = Text.ToString();
+				return Text.ToString();
 			}
+		}
 
+		private void UpdateServerLabel()
+		{
+			ServerLabel.Text = GetServerLabelText(ServerAndPort, UserName);
 			ChangeLink.Location = new Point(ServerLabel.Right + 5, ChangeLink.Location.Y);
 		}
 
@@ -160,7 +166,7 @@ namespace UnrealGameSync
 			WorkspaceRadioBtn.Checked = true;
 			
 			string WorkspaceName;
-			if(NewWorkspaceWindow.ShowModal(this, ServerAndPort, UserName, WorkspaceNameTextBox.Text, Log, out WorkspaceName))
+			if(NewWorkspaceWindow.ShowModal(this, ServerAndPort, UserName, null, WorkspaceNameTextBox.Text, Log, out WorkspaceName))
 			{
 				WorkspaceNameTextBox.Text = WorkspaceName;
 				UpdateOkButton();
@@ -270,7 +276,7 @@ namespace UnrealGameSync
 			UserSelectedProjectSettings SelectedProject;
 			if(TryGetSelectedProject(out SelectedProject))
 			{
-				DetectProjectSettingsTask NewDetectedProjectSettings = new DetectProjectSettingsTask(SelectedProject, DataFolder, Log);
+				DetectProjectSettingsTask NewDetectedProjectSettings = new DetectProjectSettingsTask(SelectedProject, DataFolder, CacheFolder, Log);
 				try
 				{
 					string ProjectFileName = null;
diff --git a/Engine/Source/Programs/UnrealGameSync/UnrealGameSync/Forms/PasswordWindow.cs b/Engine/Source/Programs/UnrealGameSync/UnrealGameSync/Forms/PasswordWindow.cs
--- a/Engine/Source/Programs/UnrealGameSync/UnrealGameSync/Forms/PasswordWindow.cs
+++ b/Engine/Source/Programs/UnrealGameSync/UnrealGameSync/Forms/PasswordWindow.cs
@@ -1,4 +1,4 @@
-﻿// Copyright 1998-2018 Epic Games, Inc. All Rights Reserved.
+﻿// Copyright 1998-2019 Epic Games, Inc. All Rights Reserved.
 
 using System;
 using System.Collections.Generic;
diff --git a/Engine/Source/Programs/UnrealGameSync/UnrealGameSync/Forms/PerforceSyncSettingsWindow.cs b/Engine/Source/Programs/UnrealGameSync/UnrealGameSync/Forms/PerforceSyncSettingsWindow.cs
--- a/Engine/Source/Programs/UnrealGameSync/UnrealGameSync/Forms/PerforceSyncSettingsWindow.cs
+++ b/Engine/Source/Programs/UnrealGameSync/UnrealGameSync/Forms/PerforceSyncSettingsWindow.cs
@@ -1,4 +1,4 @@
-// Copyright 1998-2018 Epic Games, Inc. All Rights Reserved.
+﻿// Copyright 1998-2019 Epic Games, Inc. All Rights Reserved.
 
 using System;
 using System.Collections.Generic;
diff --git a/Engine/Source/Programs/UnrealGameSync/UnrealGameSync/Forms/ProgramsRunningWindow.cs b/Engine/Source/Programs/UnrealGameSync/UnrealGameSync/Forms/ProgramsRunningWindow.cs
--- a/Engine/Source/Programs/UnrealGameSync/UnrealGameSync/Forms/ProgramsRunningWindow.cs
+++ b/Engine/Source/Programs/UnrealGameSync/UnrealGameSync/Forms/ProgramsRunningWindow.cs
@@ -1,4 +1,4 @@
-﻿// Copyright 1998-2018 Epic Games, Inc. All Rights Reserved.
+﻿// Copyright 1998-2019 Epic Games, Inc. All Rights Reserved.
 
 using System;
 using System.Collections.Generic;
diff --git a/Engine/Source/Programs/UnrealGameSync/UnrealGameSync/Forms/ProjectWindow.cs b/Engine/Source/Programs/UnrealGameSync/UnrealGameSync/Forms/ProjectWindow.cs
--- a/Engine/Source/Programs/UnrealGameSync/UnrealGameSync/Forms/ProjectWindow.cs
+++ b/Engine/Source/Programs/UnrealGameSync/UnrealGameSync/Forms/ProjectWindow.cs
@@ -1,4 +1,4 @@
-﻿// Copyright 1998-2018 Epic Games, Inc. All Rights Reserved.
+﻿// Copyright 1998-2019 Epic Games, Inc. All Rights Reserved.
 
 using System;
 using System.Collections.Generic;
diff --git a/Engine/Source/Programs/UnrealGameSync/UnrealGameSync/Forms/ScheduleWindow.cs b/Engine/Source/Programs/UnrealGameSync/UnrealGameSync/Forms/ScheduleWindow.cs
--- a/Engine/Source/Programs/UnrealGameSync/UnrealGameSync/Forms/ScheduleWindow.cs
+++ b/Engine/Source/Programs/UnrealGameSync/UnrealGameSync/Forms/ScheduleWindow.cs
@@ -1,4 +1,4 @@
-// Copyright 1998-2018 Epic Games, Inc. All Rights Reserved.
+// Copyright 1998-2019 Epic Games, Inc. All Rights Reserved.
 
 using System;
 using System.Collections.Generic;
diff --git a/Engine/Source/Programs/UnrealGameSync/UnrealGameSync/Forms/SdkInfoWindow.cs b/Engine/Source/Programs/UnrealGameSync/UnrealGameSync/Forms/SdkInfoWindow.cs
--- a/Engine/Source/Programs/UnrealGameSync/UnrealGameSync/Forms/SdkInfoWindow.cs
+++ b/Engine/Source/Programs/UnrealGameSync/UnrealGameSync/Forms/SdkInfoWindow.cs
@@ -1,3 +1,5 @@
+// Copyright 1998-2019 Epic Games, Inc. All Rights Reserved.
+
 using System;
 using System.Collections.Generic;
 using System.ComponentModel;
diff --git a/Engine/Source/Programs/UnrealGameSync/UnrealGameSync/Forms/SelectProjectFromWorkspaceWindow.cs b/Engine/Source/Programs/UnrealGameSync/UnrealGameSync/Forms/SelectProjectFromWorkspaceWindow.cs
--- a/Engine/Source/Programs/UnrealGameSync/UnrealGameSync/Forms/SelectProjectFromWorkspaceWindow.cs
+++ b/Engine/Source/Programs/UnrealGameSync/UnrealGameSync/Forms/SelectProjectFromWorkspaceWindow.cs
@@ -1,4 +1,4 @@
-// Copyright 1998-2018 Epic Games, Inc. All Rights Reserved.
+// Copyright 1998-2019 Epic Games, Inc. All Rights Reserved.
 
 using System;
 using System.Collections.Generic;
@@ -84,6 +84,7 @@ namespace UnrealGameSync
 		[DllImport("Shell32.dll", EntryPoint = "ExtractIconExW", CharSet = CharSet.Unicode, ExactSpelling = true, CallingConvention = CallingConvention.StdCall)]
 		private static extern int ExtractIconEx(string sFile, int iIndex, IntPtr piLargeVersion, out IntPtr piSmallVersion, int amountIcons);
 
+		List<string> ProjectFiles;
 		string SelectedProjectFile;
 
 		class ProjectNode
@@ -106,6 +107,7 @@ namespace UnrealGameSync
 		{
 			InitializeComponent();
 			
+			this.ProjectFiles = ProjectFiles;
 			this.SelectedProjectFile = SelectedProjectFile;
 
 			// Make the image strip containing icons for nodes in the tree
@@ -135,8 +137,25 @@ namespace UnrealGameSync
 			RootNode.Expand();
 			ProjectTreeView.Nodes.Add(RootNode);
 
+			// Populate the tree
+			Populate();
+		}
+
+		private void Populate()
+		{
+			// Clear out the existing nodes
+			TreeNode RootNode = ProjectTreeView.Nodes[0];
+			RootNode.Nodes.Clear();
+
+			// Filter the project files
+			List<string> FilteredProjectFiles = new List<string>(ProjectFiles); 
+			if(!ShowProjectDirsFiles.Checked)
+			{
+				FilteredProjectFiles.RemoveAll(x => x.EndsWith(".uprojectdirs", StringComparison.OrdinalIgnoreCase));
+			}
+
 			// Sort by paths, then files
-			List<ProjectNode> ProjectNodes = ProjectFiles.Select(x => new ProjectNode(x)).OrderBy(x => x.Folder).ThenBy(x => x.Name).ToList();
+			List<ProjectNode> ProjectNodes = FilteredProjectFiles.Select(x => new ProjectNode(x)).OrderBy(x => x.Folder).ThenBy(x => x.Name).ToList();
 
 			// Add the folders for each project
 			TreeNode[] ProjectParentNodes = new TreeNode[ProjectNodes.Count];
@@ -236,5 +255,10 @@ namespace UnrealGameSync
 			DialogResult = DialogResult.Cancel;
 			Close();
 		}
+
+		private void ShowProjectDirsFiles_CheckedChanged(object sender, EventArgs e)
+		{
+			Populate();
+		}
 	}
 }
diff --git a/Engine/Source/Programs/UnrealGameSync/UnrealGameSync/Forms/SelectStreamWindow.cs b/Engine/Source/Programs/UnrealGameSync/UnrealGameSync/Forms/SelectStreamWindow.cs
--- a/Engine/Source/Programs/UnrealGameSync/UnrealGameSync/Forms/SelectStreamWindow.cs
+++ b/Engine/Source/Programs/UnrealGameSync/UnrealGameSync/Forms/SelectStreamWindow.cs
@@ -1,4 +1,4 @@
-﻿// Copyright 1998-2018 Epic Games, Inc. All Rights Reserved.
+﻿// Copyright 1998-2019 Epic Games, Inc. All Rights Reserved.
 
 using System;
 using System.Collections.Generic;
diff --git a/Engine/Source/Programs/UnrealGameSync/UnrealGameSync/Forms/SelectWorkspaceWindow.cs b/Engine/Source/Programs/UnrealGameSync/UnrealGameSync/Forms/SelectWorkspaceWindow.cs
--- a/Engine/Source/Programs/UnrealGameSync/UnrealGameSync/Forms/SelectWorkspaceWindow.cs
+++ b/Engine/Source/Programs/UnrealGameSync/UnrealGameSync/Forms/SelectWorkspaceWindow.cs
@@ -1,4 +1,4 @@
-﻿// Copyright 1998-2018 Epic Games, Inc. All Rights Reserved.
+﻿// Copyright 1998-2019 Epic Games, Inc. All Rights Reserved.
 
 using System;
 using System.Collections.Generic;
diff --git a/Engine/Source/Programs/UnrealGameSync/UnrealGameSync/Forms/SyncFilter.cs b/Engine/Source/Programs/UnrealGameSync/UnrealGameSync/Forms/SyncFilter.cs
--- a/Engine/Source/Programs/UnrealGameSync/UnrealGameSync/Forms/SyncFilter.cs
+++ b/Engine/Source/Programs/UnrealGameSync/UnrealGameSync/Forms/SyncFilter.cs
@@ -1,4 +1,4 @@
-﻿// Copyright 1998-2018 Epic Games, Inc. All Rights Reserved.
+// Copyright 1998-2019 Epic Games, Inc. All Rights Reserved.
 
 using System;
 using System.Collections.Generic;
@@ -18,12 +18,14 @@ namespace UnrealGameSync
 		public string[] GlobalView;
 		public Guid[] GlobalExcludedCategories;
 		public bool bGlobalSyncAllProjects;
+		public bool bGlobalIncludeAllProjectsInSolution;
 		public string[] WorkspaceView;
 		public Guid[] WorkspaceIncludedCategories;
 		public Guid[] WorkspaceExcludedCategories;
 		public bool? bWorkspaceSyncAllProjects;
+		public bool? bWorkspaceIncludeAllProjectsInSolution;
 
-		public SyncFilter(Dictionary<Guid, WorkspaceSyncCategory> InUniqueIdToCategory, string[] InGlobalView, Guid[] InGlobalExcludedCategories, bool bInGlobalProjectOnly, string[] InWorkspaceView, Guid[] InWorkspaceIncludedCategories, Guid[] InWorkspaceExcludedCategories, bool? bInWorkspaceProjectOnly)
+		public SyncFilter(Dictionary<Guid, WorkspaceSyncCategory> InUniqueIdToCategory, string[] InGlobalView, Guid[] InGlobalExcludedCategories, bool bInGlobalProjectOnly, bool bInGlobalIncludeAllProjectsInSolution, string[] InWorkspaceView, Guid[] InWorkspaceIncludedCategories, Guid[] InWorkspaceExcludedCategories, bool? bInWorkspaceProjectOnly, bool? bInWorkspaceIncludeAllProjectsInSolution)
 		{
 			InitializeComponent();
 
@@ -31,21 +33,26 @@ namespace UnrealGameSync
 			GlobalExcludedCategories = InGlobalExcludedCategories;
 			GlobalView = InGlobalView;
 			bGlobalSyncAllProjects = bInGlobalProjectOnly;
+			bGlobalIncludeAllProjectsInSolution = bInGlobalIncludeAllProjectsInSolution;
 			WorkspaceIncludedCategories = InWorkspaceIncludedCategories;
 			WorkspaceExcludedCategories = InWorkspaceExcludedCategories;
 			WorkspaceView = InWorkspaceView;
-			bWorkspaceSyncAllProjects = bInWorkspaceProjectOnly ?? bGlobalSyncAllProjects;
+			bWorkspaceSyncAllProjects = bInWorkspaceProjectOnly;
+			bWorkspaceIncludeAllProjectsInSolution = bInWorkspaceIncludeAllProjectsInSolution;
 
 			GlobalControl.SetView(GlobalView);
 			SetExcludedCategories(GlobalControl.CategoriesCheckList, UniqueIdToCategory, GlobalExcludedCategories);
 			GlobalControl.SyncAllProjects.Checked = bGlobalSyncAllProjects;
+			GlobalControl.IncludeAllProjectsInSolution.Checked = bGlobalIncludeAllProjectsInSolution;
 
 			WorkspaceControl.SetView(WorkspaceView);
 			SetExcludedCategories(WorkspaceControl.CategoriesCheckList, UniqueIdToCategory, UserSettings.GetEffectiveExcludedCategories(GlobalExcludedCategories, WorkspaceIncludedCategories, WorkspaceExcludedCategories));
 			WorkspaceControl.SyncAllProjects.Checked = bWorkspaceSyncAllProjects ?? bGlobalSyncAllProjects;
+			WorkspaceControl.IncludeAllProjectsInSolution.Checked = bWorkspaceIncludeAllProjectsInSolution ?? bGlobalIncludeAllProjectsInSolution;
 
 			GlobalControl.CategoriesCheckList.ItemCheck += GlobalControl_CategoriesCheckList_ItemCheck;
 			GlobalControl.SyncAllProjects.CheckStateChanged += GlobalControl_SyncAllProjects_CheckStateChanged;
+			GlobalControl.IncludeAllProjectsInSolution.CheckStateChanged += GlobalControl_IncludeAllProjectsInSolution_CheckStateChanged;
 		}
 
 		private void GlobalControl_CategoriesCheckList_ItemCheck(object sender, ItemCheckEventArgs e)
@@ -58,6 +65,11 @@ namespace UnrealGameSync
 			WorkspaceControl.SyncAllProjects.Checked = GlobalControl.SyncAllProjects.Checked;
 		}
 
+		private void GlobalControl_IncludeAllProjectsInSolution_CheckStateChanged(object sender, EventArgs e)
+		{
+			WorkspaceControl.IncludeAllProjectsInSolution.Checked = GlobalControl.IncludeAllProjectsInSolution.Checked;
+		}
+
 		private static void SetExcludedCategories(CheckedListBox ListBox, Dictionary<Guid, WorkspaceSyncCategory> UniqueIdToFilter, Guid[] ExcludedCategories)
 		{
 			ListBox.Items.Clear();
@@ -124,9 +136,11 @@ namespace UnrealGameSync
 		
 			GlobalView = NewGlobalView;
 			bGlobalSyncAllProjects = GlobalControl.SyncAllProjects.Checked;
+			bGlobalIncludeAllProjectsInSolution = GlobalControl.IncludeAllProjectsInSolution.Checked;
 
 			WorkspaceView = NewWorkspaceView;
 			bWorkspaceSyncAllProjects = (WorkspaceControl.SyncAllProjects.Checked == bGlobalSyncAllProjects)? (bool?)null : WorkspaceControl.SyncAllProjects.Checked;
+			bWorkspaceIncludeAllProjectsInSolution = (WorkspaceControl.IncludeAllProjectsInSolution.Checked == bGlobalIncludeAllProjectsInSolution)? (bool?)null : WorkspaceControl.IncludeAllProjectsInSolution.Checked;
 
 			GetExcludedCategories(out GlobalExcludedCategories, out WorkspaceIncludedCategories, out WorkspaceExcludedCategories);
 
diff --git a/Engine/Source/Programs/UnrealGameSync/UnrealGameSync/Forms/SyncFilterSyntax.cs b/Engine/Source/Programs/UnrealGameSync/UnrealGameSync/Forms/SyncFilterSyntax.cs
--- a/Engine/Source/Programs/UnrealGameSync/UnrealGameSync/Forms/SyncFilterSyntax.cs
+++ b/Engine/Source/Programs/UnrealGameSync/UnrealGameSync/Forms/SyncFilterSyntax.cs
@@ -1,4 +1,4 @@
-﻿// Copyright 1998-2018 Epic Games, Inc. All Rights Reserved.
+﻿// Copyright 1998-2019 Epic Games, Inc. All Rights Reserved.
 
 using System;
 using System.Collections.Generic;
diff --git a/Engine/Source/Programs/UnrealGameSync/UnrealGameSync/Forms/VariablesWindow.cs b/Engine/Source/Programs/UnrealGameSync/UnrealGameSync/Forms/VariablesWindow.cs
--- a/Engine/Source/Programs/UnrealGameSync/UnrealGameSync/Forms/VariablesWindow.cs
+++ b/Engine/Source/Programs/UnrealGameSync/UnrealGameSync/Forms/VariablesWindow.cs
@@ -1,4 +1,4 @@
-﻿// Copyright 1998-2018 Epic Games, Inc. All Rights Reserved.
+﻿// Copyright 1998-2019 Epic Games, Inc. All Rights Reserved.
 
 using System;
 using System.Collections;
diff --git a/Engine/Source/Programs/UnrealGameSync/UnrealGameSync/ModalTask.cs b/Engine/Source/Programs/UnrealGameSync/UnrealGameSync/ModalTask.cs
--- a/Engine/Source/Programs/UnrealGameSync/UnrealGameSync/ModalTask.cs
+++ b/Engine/Source/Programs/UnrealGameSync/UnrealGameSync/ModalTask.cs
@@ -1,4 +1,4 @@
-// Copyright 1998-2018 Epic Games, Inc. All Rights Reserved.
+// Copyright 1998-2019 Epic Games, Inc. All Rights Reserved.
 
 using System;
 using System.Collections.Generic;
diff --git a/Engine/Source/Programs/UnrealGameSync/UnrealGameSync/OutputAdapters.cs b/Engine/Source/Programs/UnrealGameSync/UnrealGameSync/OutputAdapters.cs
--- a/Engine/Source/Programs/UnrealGameSync/UnrealGameSync/OutputAdapters.cs
+++ b/Engine/Source/Programs/UnrealGameSync/UnrealGameSync/OutputAdapters.cs
@@ -1,4 +1,4 @@
-// Copyright 1998-2018 Epic Games, Inc. All Rights Reserved.
+// Copyright 1998-2019 Epic Games, Inc. All Rights Reserved.
 
 using System;
 using System.Collections.Generic;
@@ -53,6 +53,16 @@ namespace UnrealGameSync
 		}
 	}
 
+	class BufferedTextWriter : LineBasedTextWriter
+	{
+		public List<string> Lines = new List<string>();
+
+		protected override void FlushLine(string Line)
+		{
+			Lines.Add(Line);
+		}
+	}
+
 	class PrefixedTextWriter : LineBasedTextWriter
 	{
 		string Prefix;
diff --git a/Engine/Source/Programs/UnrealGameSync/UnrealGameSync/Perforce.cs b/Engine/Source/Programs/UnrealGameSync/UnrealGameSync/Perforce.cs
--- a/Engine/Source/Programs/UnrealGameSync/UnrealGameSync/Perforce.cs
+++ b/Engine/Source/Programs/UnrealGameSync/UnrealGameSync/Perforce.cs
@@ -1,4 +1,4 @@
-// Copyright 1998-2018 Epic Games, Inc. All Rights Reserved.
+// Copyright 1998-2019 Epic Games, Inc. All Rights Reserved.
 
 using System;
 using System.Collections.Generic;
@@ -199,6 +199,7 @@ namespace UnrealGameSync
 		public int Revision;
 		public bool IsMapped;
 		public bool Unmap;
+		public string Digest;
 
 		public PerforceFileRecord(Dictionary<string, string> Tags)
 		{
@@ -245,6 +246,8 @@ namespace UnrealGameSync
 			{
 				int.TryParse(RevisionString, out Revision);
 			}
+
+			Tags.TryGetValue("digest", out Digest);
 		}
 	}
 
@@ -468,6 +471,7 @@ namespace UnrealGameSync
 			IgnoreFilesNotInClientViewError = 0x80,
 			IgnoreEnterPassword = 0x100,
 			IgnoreFilesNotOnClientError = 0x200,
+			IgnoreProtectedNamespaceError = 0x400,
 		}
 
 		delegate bool HandleRecordDelegate(Dictionary<string, string> Tags);
@@ -506,7 +510,7 @@ namespace UnrealGameSync
 				bIsLoggedIn = true;
 				return true;
 			}
-			else if(Lines[0].Channel == PerforceOutputChannel.Error && Lines[0].Text.Contains("P4PASSWD"))
+			else if(Lines[0].Channel == PerforceOutputChannel.Error && (Lines[0].Text.Contains("P4PASSWD") || Lines[0].Text.Contains("has expired")))
 			{
 				bIsLoggedIn = false;
 				return true;
@@ -599,13 +603,13 @@ namespace UnrealGameSync
 
 		public bool FindClients(string ForUserName, out List<PerforceClientRecord> Clients, TextWriter Log)
 		{
-			return RunCommand(String.Format("clients -u{0}", ForUserName), out Clients, CommandOptions.NoClient, Log);
+			return RunCommand(String.Format("clients -u \"{0}\"", ForUserName), out Clients, CommandOptions.NoClient, Log);
 		}
 
 		public bool FindClients(string ForUserName, out List<string> ClientNames, TextWriter Log)
 		{
 			List<string> Lines;
-			if(!RunCommand(String.Format("clients -u{0}", ForUserName), out Lines, CommandOptions.None, Log))
+			if(!RunCommand(String.Format("clients -u \"{0}\"", ForUserName), out Lines, CommandOptions.None, Log))
 			{
 				ClientNames = null;
 				return false;
@@ -685,9 +689,14 @@ namespace UnrealGameSync
 
 		public bool FindFiles(string Filter, out List<PerforceFileRecord> FileRecords, TextWriter Log)
 		{
-			return RunCommand(String.Format("fstat \"{0}\"", Filter), out FileRecords, CommandOptions.None, Log);
+			bool bResult = RunCommand(String.Format("fstat \"{0}\"", Filter), out FileRecords, CommandOptions.None, Log);
+			if(bResult)
+			{
+                 FileRecords.RemoveAll(x => x.Action != null && x.Action.Contains("delete"));
+			}
+			return bResult;
 		}
-		
+
 		public bool Print(string DepotPath, out List<string> Lines, TextWriter Log)
 		{
 			string TempFileName = Path.GetTempFileName();
@@ -731,9 +740,9 @@ namespace UnrealGameSync
 		public bool FileExists(string Filter, out bool bExists, TextWriter Log)
 		{
 			List<PerforceFileRecord> FileRecords;
-			if(RunCommand(String.Format("fstat \"{0}\"", Filter), out FileRecords, CommandOptions.IgnoreNoSuchFilesError | CommandOptions.IgnoreFilesNotInClientViewError, Log))
+			if(RunCommand(String.Format("fstat \"{0}\"", Filter), out FileRecords, CommandOptions.IgnoreNoSuchFilesError | CommandOptions.IgnoreFilesNotInClientViewError | CommandOptions.IgnoreProtectedNamespaceError, Log))
 			{
-				bExists = (FileRecords.Count > 0);
+				bExists = (FileRecords.Exists(x => x.Action == null || !x.Action.Contains("delete")));
 				return true;
 			}
 			else
@@ -1077,7 +1086,21 @@ namespace UnrealGameSync
 
 		public bool Stat(string Filter, out List<PerforceFileRecord> FileRecords, TextWriter Log)
 		{
-			return RunCommand(String.Format("fstat \"{0}\"", Filter), out FileRecords, CommandOptions.IgnoreFilesNotOnClientError | CommandOptions.IgnoreNoSuchFilesError, Log);
+			return RunCommand(String.Format("fstat \"{0}\"", Filter), out FileRecords, CommandOptions.IgnoreFilesNotOnClientError | CommandOptions.IgnoreNoSuchFilesError | CommandOptions.IgnoreProtectedNamespaceError, Log);
+		}
+
+		public bool Stat(string Options, List<string> Files, out List<PerforceFileRecord> FileRecords, TextWriter Log)
+		{
+			StringBuilder Arguments = new StringBuilder("fstat");
+			if(!String.IsNullOrEmpty(Options))
+			{
+				Arguments.AppendFormat(" {0}", Options);
+			}
+			foreach(string File in Files)
+			{
+				Arguments.AppendFormat(" \"{0}\"", File);
+			}
+			return RunCommand(Arguments.ToString(), out FileRecords, CommandOptions.IgnoreFilesNotOnClientError | CommandOptions.IgnoreNoSuchFilesError | CommandOptions.IgnoreProtectedNamespaceError, Log);
 		}
 
 		public bool Sync(string Filter, TextWriter Log)
@@ -1442,6 +1465,10 @@ namespace UnrealGameSync
 			{
 				return true;
 			}
+			else if(Options.HasFlag(CommandOptions.IgnoreProtectedNamespaceError) && Text.StartsWith("error: ") && Text.EndsWith(" - protected namespace - access denied."))
+			{
+				return true;
+			}
 			else if(Options.HasFlag(CommandOptions.IgnoreEnterPassword) && Text.StartsWith("Enter password:"))
 			{
 				return true;
diff --git a/Engine/Source/Programs/UnrealGameSync/UnrealGameSync/PerforceModalTask.cs b/Engine/Source/Programs/UnrealGameSync/UnrealGameSync/PerforceModalTask.cs
--- a/Engine/Source/Programs/UnrealGameSync/UnrealGameSync/PerforceModalTask.cs
+++ b/Engine/Source/Programs/UnrealGameSync/UnrealGameSync/PerforceModalTask.cs
@@ -1,4 +1,4 @@
-// Copyright 1998-2018 Epic Games, Inc. All Rights Reserved.
+// Copyright 1998-2019 Epic Games, Inc. All Rights Reserved.
 
 using System;
 using System.Collections.Generic;
diff --git a/Engine/Source/Programs/UnrealGameSync/UnrealGameSync/PerforceMonitor.cs b/Engine/Source/Programs/UnrealGameSync/UnrealGameSync/PerforceMonitor.cs
--- a/Engine/Source/Programs/UnrealGameSync/UnrealGameSync/PerforceMonitor.cs
+++ b/Engine/Source/Programs/UnrealGameSync/UnrealGameSync/PerforceMonitor.cs
@@ -1,4 +1,4 @@
-// Copyright 1998-2018 Epic Games, Inc. All Rights Reserved.
+// Copyright 1998-2019 Epic Games, Inc. All Rights Reserved.
 
 using System;
 using System.Collections.Concurrent;
@@ -29,6 +29,8 @@ namespace UnrealGameSync
 			}
 		}
 
+		public int InitialMaxChangesValue = 100;
+
 		PerforceConnection Perforce;
 		readonly string BranchClientPath;
 		readonly string SelectedClientFileName;
@@ -43,6 +45,7 @@ namespace UnrealGameSync
 		BoundedLogWriter LogWriter;
 		bool bIsEnterpriseProject;
 		bool bDisposing;
+		string CacheFolder;
 		List<KeyValuePair<string, DateTime>> LocalConfigFiles;
 
 		public event Action OnUpdate;
@@ -50,17 +53,24 @@ namespace UnrealGameSync
 		public event Action OnStreamChange;
 		public event Action OnLoginExpired;
 
-		public PerforceMonitor(PerforceConnection InPerforce, string InBranchClientPath, string InSelectedClientFileName, string InSelectedProjectIdentifier, string InLogPath, bool bInIsEnterpriseProject, ConfigFile InProjectConfigFile, List<KeyValuePair<string, DateTime>> InLocalConfigFiles)
+		public TimeSpan ServerTimeZone
+		{
+			get;
+			private set;
+		}
+
+		public PerforceMonitor(PerforceConnection InPerforce, string InBranchClientPath, string InSelectedClientFileName, string InSelectedProjectIdentifier, string InLogPath, bool bInIsEnterpriseProject, ConfigFile InProjectConfigFile, string InCacheFolder, List<KeyValuePair<string, DateTime>> InLocalConfigFiles)
 		{
 			Perforce = InPerforce;
 			BranchClientPath = InBranchClientPath;
 			SelectedClientFileName = InSelectedClientFileName;
 			SelectedProjectIdentifier = InSelectedProjectIdentifier;
-			PendingMaxChangesValue = 100;
+			PendingMaxChangesValue = InitialMaxChangesValue;
 			LastChangeByCurrentUser = -1;
 			LastCodeChangeByCurrentUser = -1;
 			bIsEnterpriseProject = bInIsEnterpriseProject;
 			LatestProjectConfigFile = InProjectConfigFile;
+			CacheFolder = InCacheFolder;
 			LocalConfigFiles = InLocalConfigFiles;
 
 			LogWriter = new BoundedLogWriter(InLogPath);
@@ -123,6 +133,13 @@ namespace UnrealGameSync
 				StreamName = null;
 			}
 
+			// Get the perforce server settings
+			PerforceInfoRecord PerforceInfo;
+			if(Perforce.Info(out PerforceInfo, LogWriter))
+			{
+				ServerTimeZone = PerforceInfo.ServerTimeZone;
+			}
+
 			// Try to update the zipped binaries list before anything else, because it causes a state change in the UI
 			UpdateZippedBinaries();
 
@@ -347,6 +364,12 @@ namespace UnrealGameSync
 			{
 				string[] CodeExtensions = { ".cs", ".h", ".cpp", ".inl", ".usf", ".ush", ".uproject", ".uplugin" };
 
+				// Skip this stuff if the user wants us to query for more changes
+				if(PendingMaxChanges > CurrentMaxChanges)
+				{
+					break;
+				}
+
 				// If there's something to check for, find all the content changes after this changelist
 				PerforceDescribeRecord DescribeRecord;
 				if(Perforce.Describe(QueryChangeNumber, out DescribeRecord, LogWriter))
@@ -499,57 +522,62 @@ namespace UnrealGameSync
 		void UpdateProjectConfigFile()
 		{
 			LocalConfigFiles.Clear();
-			LatestProjectConfigFile = ReadProjectConfigFile(Perforce, BranchClientPath, SelectedClientFileName, LocalConfigFiles, LogWriter);
+			LatestProjectConfigFile = ReadProjectConfigFile(Perforce, BranchClientPath, SelectedClientFileName, CacheFolder, LocalConfigFiles, LogWriter);
 		}
 
-		public static ConfigFile ReadProjectConfigFile(PerforceConnection Perforce, string BranchClientPath, string SelectedClientFileName, List<KeyValuePair<string, DateTime>> LocalConfigFiles, TextWriter Log)
+		public static ConfigFile ReadProjectConfigFile(PerforceConnection Perforce, string BranchClientPath, string SelectedClientFileName, string CacheFolder, List<KeyValuePair<string, DateTime>> LocalConfigFiles, TextWriter Log)
 		{
 			List<string> ConfigFilePaths = Utility.GetConfigFileLocations(BranchClientPath, SelectedClientFileName, '/');
 
-			List<PerforceFileRecord> OpenFiles;
-			Perforce.GetOpenFiles(String.Format("{0}/....ini", BranchClientPath), out OpenFiles, Log);
-
 			ConfigFile ProjectConfig = new ConfigFile();
-			foreach(string ConfigFilePath in ConfigFilePaths)
-			{
-				List<string> Lines = null;
 
-				// If this file is open for edit, read the local version
-				if(OpenFiles != null && OpenFiles.Any(x => x.ClientPath.Equals(ConfigFilePath, StringComparison.InvariantCultureIgnoreCase)))
+			List<PerforceFileRecord> FileRecords;
+			if(Perforce.Stat("-Ol", ConfigFilePaths, out FileRecords, Log))
+			{
+				foreach(PerforceFileRecord FileRecord in FileRecords)
 				{
-					try
+					List<string> Lines = null;
+
+					// Skip file records which are still in the workspace, but were synced from a different branch. For these files, the action seems to be empty, so filter against that.
+					if(FileRecord.Action == null)
 					{
-						string LocalFileName;
-						if(Perforce.ConvertToLocalPath(ConfigFilePath, out LocalFileName, Log))
+						continue;
+					}
+
+					// If this file is open for edit, read the local version
+					string LocalFileName = FileRecord.ClientPath;
+					if(LocalFileName != null && File.Exists(LocalFileName) && (File.GetAttributes(LocalFileName) & FileAttributes.ReadOnly) == 0)
+					{
+						try
 						{
 							DateTime LastModifiedTime = File.GetLastWriteTimeUtc(LocalFileName);
 							LocalConfigFiles.Add(new KeyValuePair<string, DateTime>(LocalFileName, LastModifiedTime));
 							Lines = File.ReadAllLines(LocalFileName).ToList();
 						}
+						catch(Exception Ex)
+						{
+							Log.WriteLine("Failed to read local config file for {0}: {1}", LocalFileName, Ex.ToString());
+						}
 					}
-					catch(Exception Ex)
-					{
-						Log.WriteLine("Failed to read local config file for {0}: {1}", ConfigFilePath, Ex.ToString());
-					}
-				}
-
-				// Otherwise try to get it from perforce
-				if(Lines == null)
-				{
-					Perforce.Print(ConfigFilePath, out Lines, Log);
-				}
 
-				// Merge the text with the config file
-				if(Lines != null)
-				{
-					try
+					// Otherwise try to get it from perforce
+					if(Lines == null)
 					{
-						ProjectConfig.Parse(Lines.ToArray());
-						Log.WriteLine("Read config file from {0}", ConfigFilePath);
+						Utility.TryPrintFileUsingCache(Perforce, FileRecord.DepotPath, CacheFolder, FileRecord.Digest, out Lines, Log);
 					}
-					catch(Exception Ex)
+
+					// Merge the text with the config file
+					if(Lines != null)
 					{
-						Log.WriteLine("Failed to read config file from {0}: {1}", ConfigFilePath, Ex.ToString());
+						try
+						{
+							ProjectConfig.Parse(Lines.ToArray());
+							Log.WriteLine("Read config file from {0}", FileRecord.DepotPath);
+						}
+						catch(Exception Ex)
+						{
+							Log.WriteLine("Failed to read config file from {0}: {1}", FileRecord.DepotPath, Ex.ToString());
+						}
 					}
 				}
 			}
diff --git a/Engine/Source/Programs/UnrealGameSync/UnrealGameSync/Program.cs b/Engine/Source/Programs/UnrealGameSync/UnrealGameSync/Program.cs
--- a/Engine/Source/Programs/UnrealGameSync/UnrealGameSync/Program.cs
+++ b/Engine/Source/Programs/UnrealGameSync/UnrealGameSync/Program.cs
@@ -1,4 +1,4 @@
-// Copyright 1998-2018 Epic Games, Inc. All Rights Reserved.
+// Copyright 1998-2019 Epic Games, Inc. All Rights Reserved.
 
 using System;
 using System.Collections.Generic;
@@ -15,14 +15,8 @@ using System.Windows.Forms;
 
 namespace UnrealGameSync
 {
-	static partial class Program
+	static class Program
 	{
-		/// <summary>
-		/// SQL connection string used to connect to the database for telemetry and review data. The 'Program' class is a partial class, to allow an
-		/// opportunistically included C# source file in NotForLicensees/ProgramSettings.cs to override this value in a static constructor.
-		/// </summary>
-		public static readonly string ApiUrl = null;
-
 		public static string SyncVersion = null;
 
 		[STAThread]
@@ -92,18 +86,19 @@ namespace UnrealGameSync
 			string DataFolder = Path.Combine(Environment.GetFolderPath(Environment.SpecialFolder.LocalApplicationData), "UnrealGameSync");
 			Directory.CreateDirectory(DataFolder);
 
-			using(TelemetryWriter Telemetry = new TelemetryWriter(ApiUrl, Path.Combine(DataFolder, "Telemetry.log")))
+			using(TelemetryWriter Telemetry = new TelemetryWriter(DeploymentSettings.ApiUrl, Path.Combine(DataFolder, "Telemetry.log")))
 			{
 				AppDomain.CurrentDomain.UnhandledException += CurrentDomain_UnhandledException;
 				using(UpdateMonitor UpdateMonitor = new UpdateMonitor(new PerforceConnection(UserName, null, ServerAndPort), UpdatePath))
 				{
-					ProgramApplicationContext Context = new ProgramApplicationContext(UpdateMonitor, ApiUrl, DataFolder, ActivateEvent, bRestoreState, UpdateSpawn, ProjectFileName, bUnstable);
+					ProgramApplicationContext Context = new ProgramApplicationContext(UpdateMonitor, DeploymentSettings.ApiUrl, DataFolder, ActivateEvent, bRestoreState, UpdateSpawn, ProjectFileName, bUnstable);
 					Application.Run(Context);
 
 					if(UpdateMonitor.IsUpdateAvailable && UpdateSpawn != null)
 					{
 						InstanceMutex.Close();
-						Utility.SpawnProcess(UpdateSpawn, "-restorestate" + (bUnstable? " -unstable" : ""));
+						bool bLaunchUnstable = UpdateMonitor.RelaunchUnstable ?? bUnstable;
+						Utility.SpawnProcess(UpdateSpawn, "-restorestate" + (bLaunchUnstable? " -unstable" : ""));
 					}
 				}
 			}
diff --git a/Engine/Source/Programs/UnrealGameSync/UnrealGameSync/ProgramApplicationContext.cs b/Engine/Source/Programs/UnrealGameSync/UnrealGameSync/ProgramApplicationContext.cs
--- a/Engine/Source/Programs/UnrealGameSync/UnrealGameSync/ProgramApplicationContext.cs
+++ b/Engine/Source/Programs/UnrealGameSync/UnrealGameSync/ProgramApplicationContext.cs
@@ -1,4 +1,4 @@
-﻿// Copyright 1998-2018 Epic Games, Inc. All Rights Reserved.
+﻿// Copyright 1998-2019 Epic Games, Inc. All Rights Reserved.
 
 using System;
 using System.Collections.Generic;
@@ -21,6 +21,7 @@ namespace UnrealGameSync
 		UpdateMonitor UpdateMonitor;
 		string ApiUrl;
 		string DataFolder;
+		string CacheFolder;
 		bool bRestoreState;
 		string UpdateSpawn;
 		bool bUnstable;
@@ -40,6 +41,7 @@ namespace UnrealGameSync
 		ToolStripSeparator NotifyMenu_ExitSeparator;
 		ToolStripMenuItem NotifyMenu_Exit;
 
+		List<BufferedTextWriter> StartupLogs = new List<BufferedTextWriter>();
 		DetectMultipleProjectSettingsTask DetectStartupProjectSettingsTask;
 		ModalTaskWindow DetectStartupProjectSettingsWindow;
 		MainWindow MainWindowInstance;
@@ -49,10 +51,15 @@ namespace UnrealGameSync
 			this.UpdateMonitor = UpdateMonitor;
 			this.ApiUrl = ApiUrl;
 			this.DataFolder = DataFolder;
+			this.CacheFolder = Path.Combine(DataFolder, "Cache");
 			this.bRestoreState = bRestoreState;
 			this.UpdateSpawn = UpdateSpawn;
 			this.bUnstable = bUnstable;
 
+			// Create the directories
+			Directory.CreateDirectory(DataFolder);
+			Directory.CreateDirectory(CacheFolder);
+
 			// Make sure a synchronization context is set. We spawn a bunch of threads (eg. UpdateMonitor) at startup, and need to make sure we can post messages 
 			// back to the main thread at any time.
 			if(SynchronizationContext.Current == null)
@@ -147,8 +154,11 @@ namespace UnrealGameSync
 			List<DetectProjectSettingsTask> Tasks = new List<DetectProjectSettingsTask>();
 			foreach(UserSelectedProjectSettings OpenProject in Settings.OpenProjects)
 			{
-				Log.WriteLine("Detecting settings for {0}", OpenProject);
-				Tasks.Add(new DetectProjectSettingsTask(OpenProject, DataFolder, new PrefixedTextWriter("  ", Log)));
+				BufferedTextWriter StartupLog = new BufferedTextWriter();
+				StartupLog.WriteLine("Detecting settings for {0}", OpenProject);
+				StartupLogs.Add(StartupLog);
+
+				Tasks.Add(new DetectProjectSettingsTask(OpenProject, DataFolder, CacheFolder, new TimestampLogWriter(new PrefixedTextWriter("  ", StartupLog))));
 			}
 
 			// Detect settings for the project we want to open
@@ -177,12 +187,23 @@ namespace UnrealGameSync
 			DetectStartupProjectSettingsWindow.Close();
 			DetectStartupProjectSettingsWindow = null;
 
-			// Create the main window
-			MainWindowInstance = new MainWindow(ApiUrl, DataFolder, bRestoreState, UpdateSpawn ?? Assembly.GetExecutingAssembly().Location, DetectStartupProjectSettingsTask.Results, Log, Settings);
-			if(bUnstable)
+			// Copy all the logs to the main log
+			foreach(BufferedTextWriter StartupLog in StartupLogs)
 			{
-				MainWindowInstance.Text += String.Format(" (UNSTABLE BUILD {0})", Assembly.GetExecutingAssembly().GetName().Version);
+				foreach(string Line in StartupLog.Lines)
+				{
+					Log.WriteLine("{0}", Line);
+				}
 			}
+
+			// Clear out the cache folder
+			Utility.ClearPrintCache(CacheFolder);
+
+			// Get a list of all the valid projects to open
+			DetectProjectSettingsResult[] StartupProjects = DetectStartupProjectSettingsTask.Results.Where(x => x != null).ToArray();
+
+			// Create the main window
+			MainWindowInstance = new MainWindow(UpdateMonitor, ApiUrl, DataFolder, CacheFolder, bRestoreState, UpdateSpawn ?? Assembly.GetExecutingAssembly().Location, bUnstable, StartupProjects, Log, Settings);
 			if(bVisible)
 			{
 				MainWindowInstance.Show();
@@ -205,12 +226,7 @@ namespace UnrealGameSync
 
 		private void OnActivationListenerCallback()
 		{
-			// Check if we're trying to reopen with the unstable version; if so, trigger an update to trigger a restart with the new executable
-			if(!bUnstable && (Control.ModifierKeys & Keys.Shift) != 0)
-			{
-				UpdateMonitor.TriggerUpdate();
-			}
-			else if(MainWindowInstance != null)
+			if(MainWindowInstance != null)
 			{
 				MainWindowInstance.ShowAndActivate();
 			}
@@ -221,18 +237,21 @@ namespace UnrealGameSync
 			MainThreadSynchronizationContext.Post((o) => OnActivationListenerCallback(), null);
 		}
 
-		private void OnUpdateAvailable()
+		private void OnUpdateAvailable(UpdateType Type)
 		{
-			if(MainWindowInstance != null && !bIsClosing && MainWindowInstance.CanPerformUpdate())
+			if(MainWindowInstance != null && !bIsClosing)
 			{
-				bIsClosing = true;
-				MainWindowInstance.ForceClose();
+				if(Type == UpdateType.UserInitiated || MainWindowInstance.CanPerformUpdate())
+				{
+					bIsClosing = true;
+					MainWindowInstance.ForceClose();
+				}
 			}
 		}
 
-		private void OnUpdateAvailableCallback()
+		private void OnUpdateAvailableCallback(UpdateType Type)
 		{ 
-			MainThreadSynchronizationContext.Post((o) => OnUpdateAvailable(), null);
+			MainThreadSynchronizationContext.Post((o) => OnUpdateAvailable(Type), null);
 		}
 
 		protected override void Dispose(bool bDisposing)
diff --git a/Engine/Source/Programs/UnrealGameSync/UnrealGameSync/Properties/AssemblyInfo.cs b/Engine/Source/Programs/UnrealGameSync/UnrealGameSync/Properties/AssemblyInfo.cs
--- a/Engine/Source/Programs/UnrealGameSync/UnrealGameSync/Properties/AssemblyInfo.cs
+++ b/Engine/Source/Programs/UnrealGameSync/UnrealGameSync/Properties/AssemblyInfo.cs
@@ -1,4 +1,4 @@
-// Copyright 1998-2018 Epic Games, Inc. All Rights Reserved.
+// Copyright 1998-2019 Epic Games, Inc. All Rights Reserved.
 
 using System.Reflection;
 using System.Runtime.CompilerServices;
@@ -12,7 +12,7 @@ using System.Runtime.InteropServices;
 [assembly: AssemblyConfiguration("")]
 [assembly: AssemblyCompany("")]
 [assembly: AssemblyProduct("UnrealGameSync")]
-[assembly: AssemblyCopyright("Copyright 1998-2018 Epic Games, Inc. All Rights Reserved.")]
+[assembly: AssemblyCopyright("Copyright 1998-2019 Epic Games, Inc. All Rights Reserved.")]
 [assembly: AssemblyTrademark("")]
 [assembly: AssemblyCulture("")]
 
@@ -34,5 +34,5 @@ using System.Runtime.InteropServices;
 // You can specify all the values or you can default the Build and Revision Numbers 
 // by using the '*' as shown below:
 // [assembly: AssemblyVersion("1.0.*")]
-[assembly: AssemblyVersion("1.159.*")]
-[assembly: AssemblyFileVersion("1.159.0.0")]
+[assembly: AssemblyVersion("1.167.*")]
+[assembly: AssemblyFileVersion("1.167.0.0")]
diff --git a/Engine/Source/Programs/UnrealGameSync/UnrealGameSync/RESTApi.cs b/Engine/Source/Programs/UnrealGameSync/UnrealGameSync/RESTApi.cs
--- a/Engine/Source/Programs/UnrealGameSync/UnrealGameSync/RESTApi.cs
+++ b/Engine/Source/Programs/UnrealGameSync/UnrealGameSync/RESTApi.cs
@@ -1,4 +1,4 @@
-﻿// Copyright 1998-2018 Epic Games, Inc. All Rights Reserved.
+﻿// Copyright 1998-2019 Epic Games, Inc. All Rights Reserved.
 
 using System;
 using System.Collections.Generic;
diff --git a/Engine/Source/Programs/UnrealGameSync/UnrealGameSync/Taskbar.cs b/Engine/Source/Programs/UnrealGameSync/UnrealGameSync/Taskbar.cs
--- a/Engine/Source/Programs/UnrealGameSync/UnrealGameSync/Taskbar.cs
+++ b/Engine/Source/Programs/UnrealGameSync/UnrealGameSync/Taskbar.cs
@@ -1,4 +1,4 @@
-﻿// Copyright 1998-2018 Epic Games, Inc. All Rights Reserved.
+﻿// Copyright 1998-2019 Epic Games, Inc. All Rights Reserved.
 
 using System;
 using System.Collections.Generic;
diff --git a/Engine/Source/Programs/UnrealGameSync/UnrealGameSync/Telemetry.cs b/Engine/Source/Programs/UnrealGameSync/UnrealGameSync/Telemetry.cs
--- a/Engine/Source/Programs/UnrealGameSync/UnrealGameSync/Telemetry.cs
+++ b/Engine/Source/Programs/UnrealGameSync/UnrealGameSync/Telemetry.cs
@@ -1,4 +1,4 @@
-// Copyright 1998-2018 Epic Games, Inc. All Rights Reserved.
+// Copyright 1998-2019 Epic Games, Inc. All Rights Reserved.
 
 using System;
 using System.Collections.Concurrent;
diff --git a/Engine/Source/Programs/UnrealGameSync/UnrealGameSync/UnrealGameSync.csproj b/Engine/Source/Programs/UnrealGameSync/UnrealGameSync/UnrealGameSync.csproj
--- a/Engine/Source/Programs/UnrealGameSync/UnrealGameSync/UnrealGameSync.csproj
+++ b/Engine/Source/Programs/UnrealGameSync/UnrealGameSync/UnrealGameSync.csproj
@@ -112,6 +112,7 @@
   <ItemGroup>
     <Compile Include="ActivationListener.cs" />
     <Compile Include="ArchiveManifest.cs" />
+    <Compile Include="AutomationServer.cs" />
     <Compile Include="BuildStep.cs" />
     <Compile Include="ChildProcess.cs" />
     <Compile Include="Controls\CustomListViewControl.cs">
@@ -149,6 +150,7 @@
       <SubType>Component</SubType>
     </Compile>
     <Compile Include="DeleteFilesTask.cs" />
+    <Compile Include="DeploymentSettings.cs" />
     <Compile Include="DetectProjectSettingsTask.cs" />
     <Compile Include="FileFilter.cs" />
     <Compile Include="FindFoldersToCleanTask.cs" />
@@ -158,6 +160,12 @@
     <Compile Include="Forms\ArgumentsWindow.Designer.cs">
       <DependentUpon>ArgumentsWindow.cs</DependentUpon>
     </Compile>
+    <Compile Include="Forms\AutomatedSyncWindow.cs">
+      <SubType>Form</SubType>
+    </Compile>
+    <Compile Include="Forms\AutomatedSyncWindow.Designer.cs">
+      <DependentUpon>AutomatedSyncWindow.cs</DependentUpon>
+    </Compile>
     <Compile Include="Forms\CleanWorkspaceWindow.cs">
       <SubType>Form</SubType>
     </Compile>
@@ -206,6 +214,12 @@
     <Compile Include="Forms\PasswordWindow.Designer.cs">
       <DependentUpon>PasswordWindow.cs</DependentUpon>
     </Compile>
+    <Compile Include="Forms\ApplicationSettingsWindow.cs">
+      <SubType>Form</SubType>
+    </Compile>
+    <Compile Include="Forms\ApplicationSettingsWindow.Designer.cs">
+      <DependentUpon>ApplicationSettingsWindow.cs</DependentUpon>
+    </Compile>
     <Compile Include="Forms\PerforceSyncSettingsWindow.cs">
       <SubType>Form</SubType>
     </Compile>
@@ -317,6 +331,7 @@
       <DependentUpon>SyncFilter.cs</DependentUpon>
     </Compile>
     <Compile Include="ModalTask.cs" />
+    <Compile Include="NotForLicensees\DeploymentSettingsInternal.cs" />
     <Compile Include="PerforceModalTask.cs" />
     <Compile Include="PerforceMonitor.cs" />
     <Compile Include="ProgramApplicationContext.cs" />
@@ -345,6 +360,9 @@
     <EmbeddedResource Include="Forms\ArgumentsWindow.resx">
       <DependentUpon>ArgumentsWindow.cs</DependentUpon>
     </EmbeddedResource>
+    <EmbeddedResource Include="Forms\AutomatedSyncWindow.resx">
+      <DependentUpon>AutomatedSyncWindow.cs</DependentUpon>
+    </EmbeddedResource>
     <EmbeddedResource Include="Forms\CleanWorkspaceWindow.resx">
       <DependentUpon>CleanWorkspaceWindow.cs</DependentUpon>
     </EmbeddedResource>
@@ -369,6 +387,9 @@
     <EmbeddedResource Include="Forms\PasswordWindow.resx">
       <DependentUpon>PasswordWindow.cs</DependentUpon>
     </EmbeddedResource>
+    <EmbeddedResource Include="Forms\ApplicationSettingsWindow.resx">
+      <DependentUpon>ApplicationSettingsWindow.cs</DependentUpon>
+    </EmbeddedResource>
     <EmbeddedResource Include="Forms\PerforceSyncSettingsWindow.resx">
       <DependentUpon>PerforceSyncSettingsWindow.cs</DependentUpon>
     </EmbeddedResource>
@@ -441,11 +462,7 @@
     <None Include="Resources\App.config" />
   </ItemGroup>
   <Choose>
-    <When Condition="Exists('NotForLicensees\ProgramSettings.cs')">
-      <ItemGroup>
-        <Compile Include="NotForLicensees\ProgramSettings.cs" />
-      </ItemGroup>
-    </When>
+    <When Condition="Exists('NotForLicensees\ProgramSettings.cs')" />
   </Choose>
   <ItemGroup>
     <EmbeddedResource Include="Resources\Icons.png" />
diff --git a/Engine/Source/Programs/UnrealGameSync/UnrealGameSync/UpdateMonitor.cs b/Engine/Source/Programs/UnrealGameSync/UnrealGameSync/UpdateMonitor.cs
--- a/Engine/Source/Programs/UnrealGameSync/UnrealGameSync/UpdateMonitor.cs
+++ b/Engine/Source/Programs/UnrealGameSync/UnrealGameSync/UpdateMonitor.cs
@@ -1,4 +1,4 @@
-﻿// Copyright 1998-2018 Epic Games, Inc. All Rights Reserved.
+﻿// Copyright 1998-2019 Epic Games, Inc. All Rights Reserved.
 
 using System;
 using System.Collections.Generic;
@@ -10,14 +10,31 @@ using System.Threading.Tasks;
 
 namespace UnrealGameSync
 {
+	enum UpdateType
+	{
+		Background,
+		UserInitiated,
+	}
+
 	class UpdateMonitor : IDisposable
 	{
-		PerforceConnection Perforce;
 		string WatchPath;
 		Thread WorkerThread;
 		ManualResetEvent QuitEvent;
 
-		public event Action OnUpdateAvailable;
+		public event Action<UpdateType> OnUpdateAvailable;
+
+		public PerforceConnection Perforce
+		{
+			get;
+			private set;
+		}
+
+		public bool? RelaunchUnstable
+		{
+			get;
+			private set;
+		}
 
 		public UpdateMonitor(PerforceConnection InPerforce, string InWatchPath)
 		{
@@ -68,17 +85,18 @@ namespace UnrealGameSync
 				List<PerforceChangeSummary> Changes;
 				if(Perforce.FindChanges(WatchPath, 1, out Changes, Log) && Changes.Count > 0)
 				{
-					TriggerUpdate();
+					TriggerUpdate(UpdateType.Background, null);
 				}
 			}
 		}
 
-		public void TriggerUpdate()
+		public void TriggerUpdate(UpdateType UpdateType, bool? RelaunchUnstable)
 		{
+			this.RelaunchUnstable = RelaunchUnstable;
 			IsUpdateAvailable = true;
 			if(OnUpdateAvailable != null)
 			{
-				OnUpdateAvailable();
+				OnUpdateAvailable(UpdateType);
 			}
 		}
 	}
diff --git a/Engine/Source/Programs/UnrealGameSync/UnrealGameSync/UserSettings.cs b/Engine/Source/Programs/UnrealGameSync/UnrealGameSync/UserSettings.cs
--- a/Engine/Source/Programs/UnrealGameSync/UnrealGameSync/UserSettings.cs
+++ b/Engine/Source/Programs/UnrealGameSync/UnrealGameSync/UserSettings.cs
@@ -1,4 +1,4 @@
-// Copyright 1998-2018 Epic Games, Inc. All Rights Reserved.
+// Copyright 1998-2019 Epic Games, Inc. All Rights Reserved.
 
 using System;
 using System.Collections.Generic;
@@ -7,6 +7,7 @@ using System.IO;
 using System.Linq;
 using System.Text;
 using System.Threading.Tasks;
+using System.Windows.Forms;
 
 namespace UnrealGameSync
 {
@@ -39,6 +40,13 @@ namespace UnrealGameSync
 		Local
 	}
 
+	enum FilterType
+	{
+		None,
+		Code,
+		Content
+	}
+
 	class UserSelectedProjectSettings
 	{
 		public readonly string ServerAndPort;
@@ -166,11 +174,14 @@ namespace UnrealGameSync
 		public Guid[] SyncIncludedCategories;
 		public Guid[] SyncExcludedCategories;
 		public bool? bSyncAllProjects;
+		public bool? bIncludeAllProjectsInSolution;
 	}
 
 	class UserProjectSettings
 	{
 		public List<ConfigObject> BuildSteps = new List<ConfigObject>();
+		public FilterType FilterType;
+		public HashSet<string> FilterBadges = new HashSet<string>(StringComparer.OrdinalIgnoreCase);
 	}
 
 	class UserSettings
@@ -197,12 +208,15 @@ namespace UnrealGameSync
 		public string[] SyncView;
 		public Guid[] SyncExcludedCategories;
 		public bool bSyncAllProjects;
+		public bool bIncludeAllProjectsInSolution;
 		public LatestChangeType SyncType;
 		public BuildConfig CompiledEditorBuildConfig; // NB: This assumes not using precompiled editor. See CurrentBuildConfig.
 		public TabLabels TabLabels;
 
 		// Window settings
 		public bool bWindowVisible;
+		public FormWindowState WindowState;
+		public Rectangle? WindowBounds;
 		
 		// Schedule settings
 		public bool bScheduleEnabled;
@@ -298,6 +312,7 @@ namespace UnrealGameSync
 			SyncView = ConfigFile.GetValues("General.SyncFilter", new string[0]);
 			SyncExcludedCategories = ConfigFile.GetGuidValues("General.SyncExcludedCategories", new Guid[0]);
 			bSyncAllProjects = ConfigFile.GetValue("General.SyncAllProjects", false);
+			bIncludeAllProjectsInSolution = ConfigFile.GetValue("General.IncludeAllProjectsInSolution", false);
 			if(!Enum.TryParse(ConfigFile.GetValue("General.SyncType", ""), out SyncType))
 			{
 				SyncType = LatestChangeType.Good;
@@ -338,6 +353,11 @@ namespace UnrealGameSync
 
 			// Window settings
 			bWindowVisible = ConfigFile.GetValue("Window.Visible", true);
+			if(!Enum.TryParse(ConfigFile.GetValue("Window.State", ""), true, out WindowState))
+			{
+				WindowState = FormWindowState.Normal;
+			}
+			WindowBounds = ParseRectangleValue(ConfigFile.GetValue("Window.Bounds", ""));
 
 			// Schedule settings
 			bScheduleEnabled = ConfigFile.GetValue("Schedule.Enabled", false);
@@ -367,6 +387,37 @@ namespace UnrealGameSync
 			}
 		}
 
+		static Rectangle? ParseRectangleValue(string Text)
+		{
+			ConfigObject Object = new ConfigObject(Text);
+
+			int X = Object.GetValue("X", -1);
+			int Y = Object.GetValue("Y", -1);
+			int W = Object.GetValue("W", -1);
+			int H = Object.GetValue("H", -1);
+
+			if(X == -1 || Y == -1 || W == -1 || H == -1)
+			{
+				return null;
+			}
+			else
+			{
+				return new Rectangle(X, Y, W, H);
+			}
+		}
+
+		static string FormatRectangleValue(Rectangle Value)
+		{
+			ConfigObject Object = new ConfigObject();
+
+			Object.SetValue("X", Value.X);
+			Object.SetValue("Y", Value.Y);
+			Object.SetValue("W", Value.Width);
+			Object.SetValue("H", Value.Height);
+
+			return Object.ToString();
+		}
+
 		public UserWorkspaceSettings FindOrAddWorkspace(string ClientBranchPath)
 		{
 			// Update the current workspace
@@ -429,6 +480,7 @@ namespace UnrealGameSync
 					CurrentWorkspace.SyncIncludedCategories = new Guid[0];
 					CurrentWorkspace.SyncExcludedCategories = new Guid[0];
 					CurrentWorkspace.bSyncAllProjects = null;
+					CurrentWorkspace.bIncludeAllProjectsInSolution = null;
 				}
 				else
 				{
@@ -464,6 +516,9 @@ namespace UnrealGameSync
 					int SyncAllProjects = WorkspaceSection.GetValue("SyncAllProjects", -1);
 					CurrentWorkspace.bSyncAllProjects = (SyncAllProjects == 0)? (bool?)false : (SyncAllProjects == 1)? (bool?)true : (bool?)null;
 
+					int IncludeAllProjectsInSolution = WorkspaceSection.GetValue("IncludeAllProjectsInSolution", -1);
+					CurrentWorkspace.bIncludeAllProjectsInSolution = (IncludeAllProjectsInSolution == 0)? (bool?)false : (IncludeAllProjectsInSolution == 1)? (bool?)true : (bool?)null;
+
 					string[] BisectEntries = WorkspaceSection.GetValues("Bisect", new string[0]);
 					foreach(string BisectEntry in BisectEntries)
 					{
@@ -495,6 +550,11 @@ namespace UnrealGameSync
 	
 				ConfigSection ProjectSection = ConfigFile.FindOrAddSection(ClientProjectFileName);
 				CurrentProject.BuildSteps.AddRange(ProjectSection.GetValues("BuildStep", new string[0]).Select(x => new ConfigObject(x)));
+				if(!Enum.TryParse(ProjectSection.GetValue("FilterType", ""), true, out CurrentProject.FilterType))
+				{
+					CurrentProject.FilterType = FilterType.None;
+				}
+				CurrentProject.FilterBadges.UnionWith(ProjectSection.GetValues("FilterBadges", new string[0]));
 			}
 			return CurrentProject;
 		}
@@ -525,6 +585,7 @@ namespace UnrealGameSync
 			GeneralSection.SetValues("SyncFilter", SyncView);
 			GeneralSection.SetValues("SyncExcludedCategories", SyncExcludedCategories);
 			GeneralSection.SetValue("SyncAllProjects", bSyncAllProjects);
+			GeneralSection.SetValue("IncludeAllProjectsInSolution", bIncludeAllProjectsInSolution);
 			GeneralSection.SetValue("SyncType", SyncType.ToString());
 
 			// Build configuration
@@ -555,6 +616,11 @@ namespace UnrealGameSync
 			ConfigSection WindowSection = ConfigFile.FindOrAddSection("Window");
 			WindowSection.Clear();
 			WindowSection.SetValue("Visible", bWindowVisible);
+			WindowSection.SetValue("State", WindowState.ToString());
+			if(WindowBounds != null)
+			{
+				WindowSection.SetValue("Bounds", FormatRectangleValue(WindowBounds.Value));
+			}
 
 			// Current workspace settings
 			foreach(KeyValuePair<string, UserWorkspaceSettings> Pair in WorkspaceKeyToSettings)
@@ -591,6 +657,10 @@ namespace UnrealGameSync
 				{
 					WorkspaceSection.SetValue("SyncAllProjects", CurrentWorkspace.bSyncAllProjects.Value);
 				}
+				if(CurrentWorkspace.bIncludeAllProjectsInSolution.HasValue)
+				{
+					WorkspaceSection.SetValue("IncludeAllProjectsInSolution", CurrentWorkspace.bIncludeAllProjectsInSolution.Value);
+				}
 
 				List<ConfigObject> BisectEntryObjects = new List<ConfigObject>();
 				foreach(KeyValuePair<int, BisectState> BisectPair in CurrentWorkspace.ChangeNumberToBisectState)
@@ -612,6 +682,11 @@ namespace UnrealGameSync
 				ConfigSection ProjectSection = ConfigFile.FindOrAddSection(CurrentProjectKey);
 				ProjectSection.Clear();
 				ProjectSection.SetValues("BuildStep", CurrentProject.BuildSteps.Select(x => x.ToString()).ToArray());
+				if(CurrentProject.FilterType != FilterType.None)
+				{
+					ProjectSection.SetValue("FilterType", CurrentProject.FilterType.ToString());
+				}
+				ProjectSection.SetValues("FilterBadges", CurrentProject.FilterBadges.ToArray());
 			}
 
 			// Perforce settings
diff --git a/Engine/Source/Programs/UnrealGameSync/UnrealGameSync/Utility.cs b/Engine/Source/Programs/UnrealGameSync/UnrealGameSync/Utility.cs
--- a/Engine/Source/Programs/UnrealGameSync/UnrealGameSync/Utility.cs
+++ b/Engine/Source/Programs/UnrealGameSync/UnrealGameSync/Utility.cs
@@ -1,7 +1,8 @@
-// Copyright 1998-2018 Epic Games, Inc. All Rights Reserved.
+// Copyright 1998-2019 Epic Games, Inc. All Rights Reserved.
 
 #define USE_NEW_PROCESS_JOBS
 
+using Microsoft.Win32;
 using System;
 using System.Collections.Generic;
 using System.Diagnostics;
@@ -11,6 +12,7 @@ using System.Runtime.InteropServices;
 using System.Text;
 using System.Threading.Tasks;
 using System.Web.Script.Serialization;
+using System.Windows.Forms;
 
 namespace UnrealGameSync
 {
@@ -289,5 +291,138 @@ namespace UnrealGameSync
 			}
 			return ProjectConfigFileNames;
 		}
+
+		public static void ReadGlobalPerforceSettings(ref string ServerAndPort, ref string UserName, ref string DepotPath)
+		{
+			using (RegistryKey Key = Registry.CurrentUser.OpenSubKey("SOFTWARE\\Epic Games\\UnrealGameSync", false))
+			{
+				if(Key != null)
+				{
+					ServerAndPort = Key.GetValue("ServerAndPort", ServerAndPort) as string;
+					UserName = Key.GetValue("UserName", UserName) as string;
+					DepotPath = Key.GetValue("DepotPath", DepotPath) as string;
+				}
+			}
+		}
+
+		public static void SaveGlobalPerforceSettings(string ServerAndPort, string UserName, string DepotPath)
+		{
+			try
+			{
+				using (RegistryKey Key = Registry.CurrentUser.CreateSubKey("SOFTWARE\\Epic Games\\UnrealGameSync"))
+				{
+					// Delete this legacy setting
+					try { Key.DeleteValue("Server"); } catch(Exception) { }
+
+					if(String.IsNullOrEmpty(ServerAndPort))
+					{
+						try { Key.DeleteValue("ServerAndPort"); } catch(Exception) { }
+					}
+					else
+					{
+						Key.SetValue("ServerAndPort", ServerAndPort);
+					}
+
+					if(String.IsNullOrEmpty(UserName))
+					{
+						try { Key.DeleteValue("UserName"); } catch(Exception) { }
+					}
+					else
+					{
+						Key.SetValue("UserName", UserName);
+					}
+
+					if(String.IsNullOrEmpty(DepotPath) || (DeploymentSettings.DefaultDepotPath != null && String.Equals(DepotPath, DeploymentSettings.DefaultDepotPath, StringComparison.InvariantCultureIgnoreCase)))
+					{
+						try { Key.DeleteValue("DepotPath"); } catch(Exception) { }
+					}
+					else
+					{
+						Key.SetValue("DepotPath", DepotPath);
+					}
+				}
+			}
+			catch(Exception Ex)
+			{
+				MessageBox.Show("Unable to save settings.\n\n" + Ex.ToString());
+			}
+		}
+
+		public static bool TryPrintFileUsingCache(PerforceConnection Perforce, string DepotPath, string CacheFolder, string Digest, out List<string> Lines, TextWriter Log)
+		{
+			if(Digest == null)
+			{
+				return Perforce.Print(DepotPath, out Lines, Log);
+			}
+
+			string CacheFile = Path.Combine(CacheFolder, Digest);
+			if(File.Exists(CacheFile))
+			{
+				Log.WriteLine("Reading cached copy of {0} from {1}", DepotPath, CacheFile);
+				Lines = new List<string>(File.ReadAllLines(CacheFile));
+				try
+				{
+					File.SetLastWriteTimeUtc(CacheFile, DateTime.UtcNow);
+				}
+				catch(Exception Ex)
+				{
+					Log.WriteLine("Exception touching cache file {0}: {1}", CacheFile, Ex.ToString());
+				}
+				return true;
+			}
+			else
+			{
+				string TempFile = String.Format("{0}.{1}.temp", CacheFile, Guid.NewGuid());
+				if(!Perforce.PrintToFile(DepotPath, TempFile, Log))
+				{
+					Lines = null;
+					return false;
+				}
+				else
+				{
+					Lines = new List<string>(File.ReadAllLines(TempFile));
+					try
+					{
+						File.SetAttributes(TempFile, FileAttributes.Normal);
+						File.SetLastWriteTimeUtc(TempFile, DateTime.UtcNow);
+						File.Move(TempFile, CacheFile);
+					}
+					catch
+					{
+						try
+						{
+							File.Delete(TempFile);
+						}
+						catch
+						{
+						}
+					}
+					return true;
+				}
+			}
+		}
+
+		public static void ClearPrintCache(string CacheFolder)
+		{
+			DirectoryInfo CacheDir = new DirectoryInfo(CacheFolder);
+			if(CacheDir.Exists)
+			{
+				DateTime DeleteTime = DateTime.UtcNow - TimeSpan.FromDays(5.0);
+				foreach(FileInfo CacheFile in CacheDir.EnumerateFiles())
+				{
+					if(CacheFile.LastWriteTimeUtc < DeleteTime || CacheFile.Name.EndsWith(".temp", StringComparison.OrdinalIgnoreCase))
+					{
+						try
+						{
+							CacheFile.Attributes = FileAttributes.Normal;
+							CacheFile.Delete();
+						}
+						catch
+						{
+						}
+					}
+				}
+			}
+		}
 	}
 }
diff --git a/Engine/Source/Programs/UnrealGameSync/UnrealGameSync/Workspace.cs b/Engine/Source/Programs/UnrealGameSync/UnrealGameSync/Workspace.cs
--- a/Engine/Source/Programs/UnrealGameSync/UnrealGameSync/Workspace.cs
+++ b/Engine/Source/Programs/UnrealGameSync/UnrealGameSync/Workspace.cs
@@ -1,4 +1,4 @@
-// Copyright 1998-2018 Epic Games, Inc. All Rights Reserved.
+// Copyright 1998-2019 Epic Games, Inc. All Rights Reserved.
 
 using System;
 using System.Collections.Generic;
@@ -30,6 +30,7 @@ namespace UnrealGameSync
 		ContentOnly = 0x400,
 		UpdateFilter = 0x800,
 		SyncAllProjects = 0x1000,
+		IncludeAllProjectsInSolution = 0x2000,
 	}
 
 	enum WorkspaceUpdateResult
@@ -814,18 +815,21 @@ namespace UnrealGameSync
 				{
 					Progress.Set("Generating project files...", 0.0f);
 
-					string ProjectFileArgument = "";
-					if(SelectedLocalFileName.EndsWith(".uproject", StringComparison.InvariantCultureIgnoreCase) && (Context.Options & WorkspaceUpdateOptions.SyncAllProjects) == 0)
+					StringBuilder CommandLine = new StringBuilder();
+					CommandLine.AppendFormat("/C \"\"{0}\"", Path.Combine(LocalRootPath, "GenerateProjectFiles.bat"));
+					if((Context.Options & WorkspaceUpdateOptions.SyncAllProjects) == 0 && (Context.Options & WorkspaceUpdateOptions.IncludeAllProjectsInSolution) == 0)
 					{
-						ProjectFileArgument = String.Format("\"{0}\" ", SelectedLocalFileName);
+						if(SelectedLocalFileName.EndsWith(".uproject", StringComparison.InvariantCultureIgnoreCase))
+						{
+							CommandLine.AppendFormat(" \"{0}\"", SelectedLocalFileName);
+						}
 					}
-
-					string CommandLine = String.Format("/C \"\"{0}\" {1}-progress\"", Path.Combine(LocalRootPath, "GenerateProjectFiles.bat"), ProjectFileArgument);
+					CommandLine.Append(" -progress\"");
 
 					Log.WriteLine("Generating project files...");
 					Log.WriteLine("gpf> Running {0} {1}", CmdExe, CommandLine);
 
-					int GenerateProjectFilesResult = Utility.ExecuteProcess(CmdExe, null, CommandLine, null, new ProgressTextWriter(Progress, new PrefixedTextWriter("gpf> ", Log)));
+					int GenerateProjectFilesResult = Utility.ExecuteProcess(CmdExe, null, CommandLine.ToString(), null, new ProgressTextWriter(Progress, new PrefixedTextWriter("gpf> ", Log)));
 					if(GenerateProjectFilesResult != 0)
 					{
 						StatusMessage = String.Format("Failed to generate project files (exit code {0}).", GenerateProjectFilesResult);
diff --git a/Engine/Source/Programs/UnrealGameSync/UnrealGameSyncAutomation/App.config b/Engine/Source/Programs/UnrealGameSync/UnrealGameSyncAutomation/App.config
new file mode 100644
--- /dev/null
+++ b/Engine/Source/Programs/UnrealGameSync/UnrealGameSyncAutomation/App.config
@@ -0,0 +1,6 @@
+<?xml version="1.0" encoding="utf-8"?>
+<configuration>
+    <startup> 
+        <supportedRuntime version="v4.0" sku=".NETFramework,Version=v4.5"/>
+    </startup>
+</configuration>
diff --git a/Engine/Source/Programs/UnrealGameSync/UnrealGameSyncAutomation/Program.cs b/Engine/Source/Programs/UnrealGameSync/UnrealGameSyncAutomation/Program.cs
new file mode 100644
--- /dev/null
+++ b/Engine/Source/Programs/UnrealGameSync/UnrealGameSyncAutomation/Program.cs
@@ -0,0 +1,77 @@
+// Copyright 1998-2019 Epic Games, Inc. All Rights Reserved.
+
+using Microsoft.Win32;
+using System;
+using System.Collections.Generic;
+using System.IO;
+using System.Linq;
+using System.Net;
+using System.Net.Sockets;
+using System.Text;
+using System.Threading.Tasks;
+
+namespace UnrealGameSyncAutomation
+{
+	class Program
+	{
+		enum AutomationRequestType
+		{
+			SyncProject,
+			FindProject,
+			OpenProject
+		}
+
+		enum AutomationRequestResult
+		{
+			Ok,
+			Invalid,
+			Busy,
+			Canceled,
+			Error,
+			NotFound
+		}
+
+		static void Main(string[] Arguments)
+		{
+			// Create the request data
+			MemoryStream InputDataStream = new MemoryStream();
+			using(BinaryWriter Writer = new BinaryWriter(InputDataStream))
+			{
+				Writer.Write("//UE4/Main");
+				Writer.Write("/Samples/Games/ShooterGame/ShooterGame.uproject");
+			}
+
+			Tuple<AutomationRequestResult, byte[]> Response = SendRequest(AutomationRequestType.OpenProject, InputDataStream.ToArray());
+		
+			string ResponseString = Encoding.UTF8.GetString(Response.Item2);
+			Console.WriteLine("{0}: {1}", Response.Item1, ResponseString);
+		}
+
+		static Tuple<AutomationRequestResult, byte[]> SendRequest(AutomationRequestType Request, byte[] RequestData)
+		{
+			int PortNumber = (int)Registry.GetValue("HKEY_CURRENT_USER\\Software\\Epic Games\\UnrealGameSync", "AutomationPort", null);
+			using(TcpClient Client = new TcpClient())
+			{
+				Client.Connect(new IPEndPoint(IPAddress.Loopback, PortNumber));
+
+				using(NetworkStream Stream = Client.GetStream())
+				{
+					// Send the request
+					BinaryWriter Writer = new BinaryWriter(Stream);
+					Writer.Write((int)Request);
+					Writer.Write(RequestData.Length);
+					Writer.Write(RequestData);
+
+					// Read the response
+					BinaryReader Reader = new BinaryReader(Stream);
+					AutomationRequestResult Result = (AutomationRequestResult)Reader.ReadInt32();
+					int ResponseLength = Reader.ReadInt32();
+					byte[] ResponseData = Reader.ReadBytes(ResponseLength);
+
+					// Return the response
+					return Tuple.Create(Result, ResponseData);
+				}
+			}
+		}
+	}
+}
diff --git a/Engine/Source/Programs/UnrealGameSync/UnrealGameSyncAutomation/Properties/AssemblyInfo.cs b/Engine/Source/Programs/UnrealGameSync/UnrealGameSyncAutomation/Properties/AssemblyInfo.cs
new file mode 100644
--- /dev/null
+++ b/Engine/Source/Programs/UnrealGameSync/UnrealGameSyncAutomation/Properties/AssemblyInfo.cs
@@ -0,0 +1,38 @@
+// Copyright 1998-2019 Epic Games, Inc. All Rights Reserved.
+
+using System.Reflection;
+using System.Runtime.CompilerServices;
+using System.Runtime.InteropServices;
+
+// General Information about an assembly is controlled through the following
+// set of attributes. Change these attribute values to modify the information
+// associated with an assembly.
+[assembly: AssemblyTitle("UnrealGameSyncAutomation")]
+[assembly: AssemblyDescription("")]
+[assembly: AssemblyConfiguration("")]
+[assembly: AssemblyCompany("")]
+[assembly: AssemblyProduct("UnrealGameSyncAutomation")]
+[assembly: AssemblyCopyright("Copyright 1998-2019 Epic Games, Inc. All Rights Reserved.")]
+[assembly: AssemblyTrademark("")]
+[assembly: AssemblyCulture("")]
+
+// Setting ComVisible to false makes the types in this assembly not visible
+// to COM components.  If you need to access a type in this assembly from
+// COM, set the ComVisible attribute to true on that type.
+[assembly: ComVisible(false)]
+
+// The following GUID is for the ID of the typelib if this project is exposed to COM
+[assembly: Guid("7c9d0a33-c2a8-4415-b9ed-9f40337dbec6")]
+
+// Version information for an assembly consists of the following four values:
+//
+//      Major Version
+//      Minor Version
+//      Build Number
+//      Revision
+//
+// You can specify all the values or you can default the Build and Revision Numbers
+// by using the '*' as shown below:
+// [assembly: AssemblyVersion("1.0.*")]
+[assembly: AssemblyVersion("1.0.0.0")]
+[assembly: AssemblyFileVersion("1.0.0.0")]
diff --git a/Engine/Source/Programs/UnrealGameSync/UnrealGameSyncAutomation/UnrealGameSyncAutomation.csproj b/Engine/Source/Programs/UnrealGameSync/UnrealGameSyncAutomation/UnrealGameSyncAutomation.csproj
new file mode 100644
--- /dev/null
+++ b/Engine/Source/Programs/UnrealGameSync/UnrealGameSyncAutomation/UnrealGameSyncAutomation.csproj
@@ -0,0 +1,54 @@
+﻿<?xml version="1.0" encoding="utf-8"?>
+<Project ToolsVersion="15.0" xmlns="http://schemas.microsoft.com/developer/msbuild/2003">
+  <Import Project="$(MSBuildExtensionsPath)\$(MSBuildToolsVersion)\Microsoft.Common.props" Condition="Exists('$(MSBuildExtensionsPath)\$(MSBuildToolsVersion)\Microsoft.Common.props')" />
+  <PropertyGroup>
+    <Configuration Condition=" '$(Configuration)' == '' ">Debug</Configuration>
+    <Platform Condition=" '$(Platform)' == '' ">AnyCPU</Platform>
+    <ProjectGuid>{7C9D0A33-C2A8-4415-B9ED-9F40337DBEC6}</ProjectGuid>
+    <OutputType>Exe</OutputType>
+    <RootNamespace>UnrealGameSyncAutomation</RootNamespace>
+    <AssemblyName>UnrealGameSyncAutomation</AssemblyName>
+    <TargetFrameworkVersion>v4.5</TargetFrameworkVersion>
+    <FileAlignment>512</FileAlignment>
+    <AutoGenerateBindingRedirects>true</AutoGenerateBindingRedirects>
+    <Deterministic>true</Deterministic>
+    <TargetFrameworkProfile />
+  </PropertyGroup>
+  <PropertyGroup Condition=" '$(Configuration)|$(Platform)' == 'Debug|AnyCPU' ">
+    <PlatformTarget>AnyCPU</PlatformTarget>
+    <DebugSymbols>true</DebugSymbols>
+    <DebugType>full</DebugType>
+    <Optimize>false</Optimize>
+    <OutputPath>bin\Debug\</OutputPath>
+    <DefineConstants>DEBUG;TRACE</DefineConstants>
+    <ErrorReport>prompt</ErrorReport>
+    <WarningLevel>4</WarningLevel>
+  </PropertyGroup>
+  <PropertyGroup Condition=" '$(Configuration)|$(Platform)' == 'Release|AnyCPU' ">
+    <PlatformTarget>AnyCPU</PlatformTarget>
+    <DebugType>pdbonly</DebugType>
+    <Optimize>true</Optimize>
+    <OutputPath>bin\Release\</OutputPath>
+    <DefineConstants>TRACE</DefineConstants>
+    <ErrorReport>prompt</ErrorReport>
+    <WarningLevel>4</WarningLevel>
+  </PropertyGroup>
+  <ItemGroup>
+    <Reference Include="System" />
+    <Reference Include="System.Core" />
+    <Reference Include="System.Xml.Linq" />
+    <Reference Include="System.Data.DataSetExtensions" />
+    <Reference Include="Microsoft.CSharp" />
+    <Reference Include="System.Data" />
+    <Reference Include="System.Net.Http" />
+    <Reference Include="System.Xml" />
+  </ItemGroup>
+  <ItemGroup>
+    <Compile Include="Program.cs" />
+    <Compile Include="Properties\AssemblyInfo.cs" />
+  </ItemGroup>
+  <ItemGroup>
+    <None Include="App.config" />
+  </ItemGroup>
+  <Import Project="$(MSBuildToolsPath)\Microsoft.CSharp.targets" />
+</Project>
\ No newline at end of file
diff --git a/Engine/Source/Programs/UnrealGameSync/UnrealGameSyncLauncher/Forms/LogWindow.cs b/Engine/Source/Programs/UnrealGameSync/UnrealGameSyncLauncher/Forms/LogWindow.cs
--- a/Engine/Source/Programs/UnrealGameSync/UnrealGameSyncLauncher/Forms/LogWindow.cs
+++ b/Engine/Source/Programs/UnrealGameSync/UnrealGameSyncLauncher/Forms/LogWindow.cs
@@ -1,4 +1,4 @@
-// Copyright 1998-2018 Epic Games, Inc. All Rights Reserved.
+// Copyright 1998-2019 Epic Games, Inc. All Rights Reserved.
 
 using System;
 using System.Collections.Generic;
diff --git a/Engine/Source/Programs/UnrealGameSync/UnrealGameSyncLauncher/Forms/SettingsWindow.cs b/Engine/Source/Programs/UnrealGameSync/UnrealGameSyncLauncher/Forms/SettingsWindow.cs
--- a/Engine/Source/Programs/UnrealGameSync/UnrealGameSyncLauncher/Forms/SettingsWindow.cs
+++ b/Engine/Source/Programs/UnrealGameSync/UnrealGameSyncLauncher/Forms/SettingsWindow.cs
@@ -1,4 +1,4 @@
-// Copyright 1998-2018 Epic Games, Inc. All Rights Reserved.
+// Copyright 1998-2019 Epic Games, Inc. All Rights Reserved.
 
 using System;
 using System.Collections.Generic;
@@ -82,7 +82,7 @@ namespace UnrealGameSyncLauncher
 				DepotPath = null;
 			}
 
-			Program.SaveSettings(ServerAndPort, UserName, DepotPath);
+			Utility.SaveGlobalPerforceSettings(ServerAndPort, UserName, DepotPath);
 
 			// Create the task for connecting to this server
 			StringWriter Log = new StringWriter();
@@ -93,7 +93,7 @@ namespace UnrealGameSyncLauncher
 			ModalTaskResult Result = PerforceModalTask.Execute(this, null, ServerAndPort, UserName, SyncApplication, "Updating", "Checking for updates, please wait...", Log, out ErrorMessage);
 			if(Result == ModalTaskResult.Succeeded)
 			{
-				Program.SaveSettings(ServerAndPort, UserName, DepotPath);
+				Utility.SaveGlobalPerforceSettings(ServerAndPort, UserName, DepotPath);
 				DialogResult = DialogResult.OK;
 				Close();
 			}
diff --git a/Engine/Source/Programs/UnrealGameSync/UnrealGameSyncLauncher/Program.cs b/Engine/Source/Programs/UnrealGameSync/UnrealGameSyncLauncher/Program.cs
--- a/Engine/Source/Programs/UnrealGameSync/UnrealGameSyncLauncher/Program.cs
+++ b/Engine/Source/Programs/UnrealGameSync/UnrealGameSyncLauncher/Program.cs
@@ -1,4 +1,4 @@
-// Copyright 1998-2018 Epic Games, Inc. All Rights Reserved.
+// Copyright 1998-2019 Epic Games, Inc. All Rights Reserved.
 
 using Microsoft.Win32;
 using System;
@@ -15,14 +15,8 @@ using UnrealGameSync;
 
 namespace UnrealGameSyncLauncher
 {
-	static partial class Program
+	static class Program
 	{
-		/// <summary>
-		/// Specifies the depot path to sync down the stable version of UGS from, without a trailing slash (eg. //depot/UnrealGameSync/bin). This is a site-specific setting. 
-		/// The UnrealGameSync executable should be located at Release/UnrealGameSync.exe under this path, with any dependent DLLs.
-		/// </summary>
-		static readonly string DefaultDepotPath;
-
 		[STAThread]
 		static int Main(string[] Args)
 		{
@@ -70,8 +64,8 @@ namespace UnrealGameSyncLauncher
 				// Read the settings
 				string ServerAndPort = null;
 				string UserName = null;
-				string DepotPath = DefaultDepotPath;
-				ReadSettings(ref ServerAndPort, ref UserName, ref DepotPath);
+				string DepotPath = DeploymentSettings.DefaultDepotPath;
+				Utility.ReadGlobalPerforceSettings(ref ServerAndPort, ref UserName, ref DepotPath);
 
 				// If the shift key is held down, immediately show the settings window
 				if((Control.ModifierKeys & Keys.Shift) != 0)
@@ -104,62 +98,6 @@ namespace UnrealGameSyncLauncher
 			return 1;
 		}
 
-		public static void ReadSettings(ref string ServerAndPort, ref string UserName, ref string DepotPath)
-		{
-			using (RegistryKey Key = Registry.CurrentUser.OpenSubKey("SOFTWARE\\Epic Games\\UnrealGameSync", false))
-			{
-				if(Key != null)
-				{
-					ServerAndPort = Key.GetValue("ServerAndPort", ServerAndPort) as string;
-					UserName = Key.GetValue("UserName", UserName) as string;
-					DepotPath = Key.GetValue("DepotPath", DepotPath) as string;
-				}
-			}
-		}
-
-		public static void SaveSettings(string ServerAndPort, string UserName, string DepotPath)
-		{
-			try
-			{
-				using (RegistryKey Key = Registry.CurrentUser.CreateSubKey("SOFTWARE\\Epic Games\\UnrealGameSync"))
-				{
-					// Delete this legacy setting
-					try { Key.DeleteValue("Server"); } catch(Exception) { }
-
-					if(String.IsNullOrEmpty(ServerAndPort))
-					{
-						try { Key.DeleteValue("ServerAndPort"); } catch(Exception) { }
-					}
-					else
-					{
-						Key.SetValue("ServerAndPort", ServerAndPort);
-					}
-
-					if(String.IsNullOrEmpty(UserName))
-					{
-						try { Key.DeleteValue("UserName"); } catch(Exception) { }
-					}
-					else
-					{
-						Key.SetValue("UserName", UserName);
-					}
-
-					if(String.IsNullOrEmpty(DepotPath) || (DefaultDepotPath != null && String.Equals(DepotPath, DefaultDepotPath, StringComparison.InvariantCultureIgnoreCase)))
-					{
-						try { Key.DeleteValue("DepotPath"); } catch(Exception) { }
-					}
-					else
-					{
-						Key.SetValue("DepotPath", DepotPath);
-					}
-				}
-			}
-			catch(Exception Ex)
-			{
-				MessageBox.Show("Unable to save settings.\n\n" + Ex.ToString());
-			}
-		}
-
 		public static bool SyncAndRun(PerforceConnection Perforce, string BaseDepotPath, bool bUnstable, string[] Args, Mutex InstanceMutex, TextWriter LogWriter)
 		{
 			try
diff --git a/Engine/Source/Programs/UnrealGameSync/UnrealGameSyncLauncher/Properties/AssemblyInfo.cs b/Engine/Source/Programs/UnrealGameSync/UnrealGameSyncLauncher/Properties/AssemblyInfo.cs
--- a/Engine/Source/Programs/UnrealGameSync/UnrealGameSyncLauncher/Properties/AssemblyInfo.cs
+++ b/Engine/Source/Programs/UnrealGameSync/UnrealGameSyncLauncher/Properties/AssemblyInfo.cs
@@ -1,4 +1,4 @@
-// Copyright 1998-2018 Epic Games, Inc. All Rights Reserved.
+// Copyright 1998-2019 Epic Games, Inc. All Rights Reserved.
 
 using System.Reflection;
 using System.Runtime.CompilerServices;
@@ -12,7 +12,7 @@ using System.Runtime.InteropServices;
 [assembly: AssemblyConfiguration("")]
 [assembly: AssemblyCompany("Epic Games, Inc.")]
 [assembly: AssemblyProduct("UnrealGameSync")]
-[assembly: AssemblyCopyright("Copyright 1998-2018 Epic Games, Inc. All Rights Reserved.")]
+[assembly: AssemblyCopyright("Copyright 1998-2019 Epic Games, Inc. All Rights Reserved.")]
 [assembly: AssemblyTrademark("")]
 [assembly: AssemblyCulture("")]
 
diff --git a/Engine/Source/Programs/UnrealGameSync/UnrealGameSyncLauncher/SyncAndRunModalTask.cs b/Engine/Source/Programs/UnrealGameSync/UnrealGameSyncLauncher/SyncAndRunModalTask.cs
--- a/Engine/Source/Programs/UnrealGameSync/UnrealGameSyncLauncher/SyncAndRunModalTask.cs
+++ b/Engine/Source/Programs/UnrealGameSync/UnrealGameSyncLauncher/SyncAndRunModalTask.cs
@@ -1,4 +1,4 @@
-// Copyright 1998-2018 Epic Games, Inc. All Rights Reserved.
+// Copyright 1998-2019 Epic Games, Inc. All Rights Reserved.
 
 using System;
 using System.Collections.Generic;
diff --git a/Engine/Source/Programs/UnrealGameSync/UnrealGameSyncLauncher/SyncAndRunPerforceTask.cs b/Engine/Source/Programs/UnrealGameSync/UnrealGameSyncLauncher/SyncAndRunPerforceTask.cs
--- a/Engine/Source/Programs/UnrealGameSync/UnrealGameSyncLauncher/SyncAndRunPerforceTask.cs
+++ b/Engine/Source/Programs/UnrealGameSync/UnrealGameSyncLauncher/SyncAndRunPerforceTask.cs
@@ -1,4 +1,4 @@
-// Copyright 1998-2018 Epic Games, Inc. All Rights Reserved.
+// Copyright 1998-2019 Epic Games, Inc. All Rights Reserved.
 
 using System;
 using System.Collections.Generic;
diff --git a/Engine/Source/Programs/UnrealGameSync/UnrealGameSyncLauncher/UnrealGameSyncLauncher.csproj b/Engine/Source/Programs/UnrealGameSync/UnrealGameSyncLauncher/UnrealGameSyncLauncher.csproj
--- a/Engine/Source/Programs/UnrealGameSync/UnrealGameSyncLauncher/UnrealGameSyncLauncher.csproj
+++ b/Engine/Source/Programs/UnrealGameSync/UnrealGameSyncLauncher/UnrealGameSyncLauncher.csproj
@@ -89,6 +89,9 @@
     <Compile Include="..\UnrealGameSync\Forms\PasswordWindow.Designer.cs">
       <Link>UnrealGameSync\PasswordWindow.Designer.cs</Link>
     </Compile>
+    <Compile Include="..\UnrealGameSync\DeploymentSettings.cs">
+      <Link>Shared\DeploymentSettings.cs</Link>
+    </Compile>
     <Compile Include="..\UnrealGameSync\ModalTask.cs">
       <Link>Shared\ModalTask.cs</Link>
     </Compile>
@@ -151,9 +154,11 @@
   </ItemGroup>
   <ItemGroup />
   <Choose>
-    <When Condition="Exists('NotForLicensees\ProgramSettings.cs')">
+    <When Condition="Exists('..\UnrealGameSync\NotForLicensees\DeploymentSettingsInternal.cs')">
       <ItemGroup>
-        <Compile Include="NotForLicensees\ProgramSettings.cs" />
+        <Compile Include="..\UnrealGameSync\NotForLicensees\DeploymentSettingsInternal.cs">
+			<Link>Shared\NotForLicensees\DeploymentSettingsInternal.cs</Link>
+		</Compile>
       </ItemGroup>
     </When>
   </Choose>
diff --git a/Engine/Source/Programs/UnrealGameSync/UnrealGameSyncMetadataServer/App_Start/WebApiConfig.cs b/Engine/Source/Programs/UnrealGameSync/UnrealGameSyncMetadataServer/App_Start/WebApiConfig.cs
--- a/Engine/Source/Programs/UnrealGameSync/UnrealGameSyncMetadataServer/App_Start/WebApiConfig.cs
+++ b/Engine/Source/Programs/UnrealGameSync/UnrealGameSyncMetadataServer/App_Start/WebApiConfig.cs
@@ -1,4 +1,4 @@
-﻿// Copyright 1998-2018 Epic Games, Inc. All Rights Reserved.
+﻿// Copyright 1998-2019 Epic Games, Inc. All Rights Reserved.
 
 using System.Web.Http;
 
diff --git a/Engine/Source/Programs/UnrealGameSync/UnrealGameSyncMetadataServer/Connectors/SqlConnector.cs b/Engine/Source/Programs/UnrealGameSync/UnrealGameSyncMetadataServer/Connectors/SqlConnector.cs
--- a/Engine/Source/Programs/UnrealGameSync/UnrealGameSyncMetadataServer/Connectors/SqlConnector.cs
+++ b/Engine/Source/Programs/UnrealGameSync/UnrealGameSyncMetadataServer/Connectors/SqlConnector.cs
@@ -1,4 +1,4 @@
-﻿// Copyright 1998-2018 Epic Games, Inc. All Rights Reserved.
+﻿// Copyright 1998-2019 Epic Games, Inc. All Rights Reserved.
 
 using System;
 using System.Collections.Generic;
diff --git a/Engine/Source/Programs/UnrealGameSync/UnrealGameSyncMetadataServer/Controllers/BuildController.cs b/Engine/Source/Programs/UnrealGameSync/UnrealGameSyncMetadataServer/Controllers/BuildController.cs
--- a/Engine/Source/Programs/UnrealGameSync/UnrealGameSyncMetadataServer/Controllers/BuildController.cs
+++ b/Engine/Source/Programs/UnrealGameSync/UnrealGameSyncMetadataServer/Controllers/BuildController.cs
@@ -1,4 +1,4 @@
-﻿// Copyright 1998-2018 Epic Games, Inc. All Rights Reserved.
+﻿// Copyright 1998-2019 Epic Games, Inc. All Rights Reserved.
 
 using System.Collections.Generic;
 using System.Web.Http;
diff --git a/Engine/Source/Programs/UnrealGameSync/UnrealGameSyncMetadataServer/Controllers/CISController.cs b/Engine/Source/Programs/UnrealGameSync/UnrealGameSyncMetadataServer/Controllers/CISController.cs
--- a/Engine/Source/Programs/UnrealGameSync/UnrealGameSyncMetadataServer/Controllers/CISController.cs
+++ b/Engine/Source/Programs/UnrealGameSync/UnrealGameSyncMetadataServer/Controllers/CISController.cs
@@ -1,4 +1,4 @@
-﻿// Copyright 1998-2018 Epic Games, Inc. All Rights Reserved.
+﻿// Copyright 1998-2019 Epic Games, Inc. All Rights Reserved.
 
 using System;
 using System.Collections.Generic;
diff --git a/Engine/Source/Programs/UnrealGameSync/UnrealGameSyncMetadataServer/Controllers/CommentController.cs b/Engine/Source/Programs/UnrealGameSync/UnrealGameSyncMetadataServer/Controllers/CommentController.cs
--- a/Engine/Source/Programs/UnrealGameSync/UnrealGameSyncMetadataServer/Controllers/CommentController.cs
+++ b/Engine/Source/Programs/UnrealGameSync/UnrealGameSyncMetadataServer/Controllers/CommentController.cs
@@ -1,4 +1,4 @@
-﻿// Copyright 1998-2018 Epic Games, Inc. All Rights Reserved.
+﻿// Copyright 1998-2019 Epic Games, Inc. All Rights Reserved.
 
 using System.Collections.Generic;
 using System.Web.Http;
diff --git a/Engine/Source/Programs/UnrealGameSync/UnrealGameSyncMetadataServer/Controllers/ErrorController.cs b/Engine/Source/Programs/UnrealGameSync/UnrealGameSyncMetadataServer/Controllers/ErrorController.cs
--- a/Engine/Source/Programs/UnrealGameSync/UnrealGameSyncMetadataServer/Controllers/ErrorController.cs
+++ b/Engine/Source/Programs/UnrealGameSync/UnrealGameSyncMetadataServer/Controllers/ErrorController.cs
@@ -1,4 +1,4 @@
-﻿// Copyright 1998-2018 Epic Games, Inc. All Rights Reserved.
+﻿// Copyright 1998-2019 Epic Games, Inc. All Rights Reserved.
 
 using System.Collections.Generic;
 using System.Web.Http;
diff --git a/Engine/Source/Programs/UnrealGameSync/UnrealGameSyncMetadataServer/Controllers/EventController.cs b/Engine/Source/Programs/UnrealGameSync/UnrealGameSyncMetadataServer/Controllers/EventController.cs
--- a/Engine/Source/Programs/UnrealGameSync/UnrealGameSyncMetadataServer/Controllers/EventController.cs
+++ b/Engine/Source/Programs/UnrealGameSync/UnrealGameSyncMetadataServer/Controllers/EventController.cs
@@ -1,4 +1,4 @@
-﻿// Copyright 1998-2018 Epic Games, Inc. All Rights Reserved.
+﻿// Copyright 1998-2019 Epic Games, Inc. All Rights Reserved.
 
 using System.Collections.Generic;
 using System.Web.Http;
diff --git a/Engine/Source/Programs/UnrealGameSync/UnrealGameSyncMetadataServer/Controllers/LatestController.cs b/Engine/Source/Programs/UnrealGameSync/UnrealGameSyncMetadataServer/Controllers/LatestController.cs
--- a/Engine/Source/Programs/UnrealGameSync/UnrealGameSyncMetadataServer/Controllers/LatestController.cs
+++ b/Engine/Source/Programs/UnrealGameSync/UnrealGameSyncMetadataServer/Controllers/LatestController.cs
@@ -1,4 +1,4 @@
-﻿// Copyright 1998-2018 Epic Games, Inc. All Rights Reserved.
+﻿// Copyright 1998-2019 Epic Games, Inc. All Rights Reserved.
 
 using System;
 using System.Collections.Generic;
diff --git a/Engine/Source/Programs/UnrealGameSync/UnrealGameSyncMetadataServer/Controllers/TelemetryController.cs b/Engine/Source/Programs/UnrealGameSync/UnrealGameSyncMetadataServer/Controllers/TelemetryController.cs
--- a/Engine/Source/Programs/UnrealGameSync/UnrealGameSyncMetadataServer/Controllers/TelemetryController.cs
+++ b/Engine/Source/Programs/UnrealGameSync/UnrealGameSyncMetadataServer/Controllers/TelemetryController.cs
@@ -1,4 +1,4 @@
-﻿// Copyright 1998-2018 Epic Games, Inc. All Rights Reserved.
+﻿// Copyright 1998-2019 Epic Games, Inc. All Rights Reserved.
 
 using System.Web.Http;
 using UnrealGameSyncMetadataServer.Connectors;
diff --git a/Engine/Source/Programs/UnrealGameSync/UnrealGameSyncMetadataServer/Global.asax.cs b/Engine/Source/Programs/UnrealGameSync/UnrealGameSyncMetadataServer/Global.asax.cs
--- a/Engine/Source/Programs/UnrealGameSync/UnrealGameSyncMetadataServer/Global.asax.cs
+++ b/Engine/Source/Programs/UnrealGameSync/UnrealGameSyncMetadataServer/Global.asax.cs
@@ -1,4 +1,4 @@
-﻿// Copyright 1998-2018 Epic Games, Inc. All Rights Reserved.
+﻿// Copyright 1998-2019 Epic Games, Inc. All Rights Reserved.
 using System;
 using System.Data.SQLite;
 using System.IO;
@@ -25,60 +25,7 @@ namespace UnrealGameSyncMetadataServer
 				using (SQLiteConnection Connection = new SQLiteConnection(ConnectionString, true))
 				{
 					Connection.Open();
-					string CreateTables = @"CREATE TABLE IF NOT EXISTS
-										[CIS] (
-											[Id]           INTEGER       NOT NULL PRIMARY KEY AUTOINCREMENT,
-											[ChangeNumber] INTEGER       NOT NULL,
-											[BuildType]    NCHAR (32)    NOT NULL,
-											[Result]       NCHAR (10)    NOT NULL,
-											[Url]          VARCHAR (512) NOT NULL,
-											[Project]      VARCHAR (256) NULL,
-											[ArchivePath]  VARCHAR (512) NULL
-										);
-										
-										CREATE TABLE IF NOT EXISTS
-										[Comments] (
-											[Id]           INTEGER    	  NOT NULL PRIMARY KEY AUTOINCREMENT,
-											[ChangeNumber] INTEGER        NOT NULL,
-											[UserName]     VARCHAR (128)  NOT NULL,
-											[Text]         VARCHAR (140)  NOT NULL,
-											[Project]      NVARCHAR (256) NOT NULL
-										);
-
-										CREATE TABLE IF NOT EXISTS
-										[Errors] (
-											[Id]        INTEGER		   NOT NULL PRIMARY KEY AUTOINCREMENT,
-											[Type]      VARCHAR (50)   NOT NULL,
-											[Text]      VARCHAR (1024) NOT NULL,
-											[UserName]  NVARCHAR (128) NOT NULL,
-											[Project]   VARCHAR (128)  NULL,
-											[Timestamp] DATETIME       NOT NULL,
-											[Version]   VARCHAR (64)   NOT NULL,
-											[IpAddress] VARCHAR (64)   NOT NULL
-										);
-										
-										CREATE TABLE IF NOT EXISTS
-										[Telemetry.v2] (
-											[Id]        INTEGER		   NOT NULL PRIMARY KEY AUTOINCREMENT,
-											[Action]    VARCHAR (128)  NOT NULL,
-											[Result]    VARCHAR (128)  NOT NULL,
-											[UserName]  NVARCHAR (128) NOT NULL,
-											[Project]   VARCHAR (128)  NOT NULL,
-											[Timestamp] DATETIME       NOT NULL,
-											[Duration]  REAL           NOT NULL,
-											[Version]   VARCHAR (64)   NULL,
-											[IpAddress] VARCHAR (64)   NULL
-										);
-
-										CREATE TABLE IF NOT EXISTS
-										[UserVotes] (
-											[Id]         INTEGER		NOT NULL PRIMARY KEY AUTOINCREMENT,
-											[Changelist] INTEGER        NOT NULL,
-											[UserName]   NVARCHAR (128) NOT NULL,
-											[Verdict]    NCHAR (32)     NOT NULL,
-											[Project]    NVARCHAR (256) NULL
-										);";
-					using (SQLiteCommand Command = new SQLiteCommand(CreateTables, Connection))
+					using (SQLiteCommand Command = new SQLiteCommand(Properties.Resources.Setup, Connection))
 					{
 						Command.ExecuteNonQuery();
 					}
diff --git a/Engine/Source/Programs/UnrealGameSync/UnrealGameSyncMetadataServer/Models/BuildData.cs b/Engine/Source/Programs/UnrealGameSync/UnrealGameSyncMetadataServer/Models/BuildData.cs
--- a/Engine/Source/Programs/UnrealGameSync/UnrealGameSyncMetadataServer/Models/BuildData.cs
+++ b/Engine/Source/Programs/UnrealGameSync/UnrealGameSyncMetadataServer/Models/BuildData.cs
@@ -1,4 +1,4 @@
-﻿// Copyright 1998-2018 Epic Games, Inc. All Rights Reserved.
+﻿// Copyright 1998-2019 Epic Games, Inc. All Rights Reserved.
 
 using System;
 using System.Collections.Generic;
diff --git a/Engine/Source/Programs/UnrealGameSync/UnrealGameSyncMetadataServer/Models/CommentData.cs b/Engine/Source/Programs/UnrealGameSync/UnrealGameSyncMetadataServer/Models/CommentData.cs
--- a/Engine/Source/Programs/UnrealGameSync/UnrealGameSyncMetadataServer/Models/CommentData.cs
+++ b/Engine/Source/Programs/UnrealGameSync/UnrealGameSyncMetadataServer/Models/CommentData.cs
@@ -1,4 +1,4 @@
-﻿// Copyright 1998-2018 Epic Games, Inc. All Rights Reserved.
+﻿// Copyright 1998-2019 Epic Games, Inc. All Rights Reserved.
 
 using System;
 using System.Collections.Generic;
diff --git a/Engine/Source/Programs/UnrealGameSync/UnrealGameSyncMetadataServer/Models/EventData.cs b/Engine/Source/Programs/UnrealGameSync/UnrealGameSyncMetadataServer/Models/EventData.cs
--- a/Engine/Source/Programs/UnrealGameSync/UnrealGameSyncMetadataServer/Models/EventData.cs
+++ b/Engine/Source/Programs/UnrealGameSync/UnrealGameSyncMetadataServer/Models/EventData.cs
@@ -1,4 +1,4 @@
-﻿// Copyright 1998-2018 Epic Games, Inc. All Rights Reserved.
+﻿// Copyright 1998-2019 Epic Games, Inc. All Rights Reserved.
 
 using System;
 using System.Collections.Generic;
diff --git a/Engine/Source/Programs/UnrealGameSync/UnrealGameSyncMetadataServer/Models/EventSummary.cs b/Engine/Source/Programs/UnrealGameSync/UnrealGameSyncMetadataServer/Models/EventSummary.cs
--- a/Engine/Source/Programs/UnrealGameSync/UnrealGameSyncMetadataServer/Models/EventSummary.cs
+++ b/Engine/Source/Programs/UnrealGameSync/UnrealGameSyncMetadataServer/Models/EventSummary.cs
@@ -1,4 +1,4 @@
-﻿// Copyright 1998-2018 Epic Games, Inc. All Rights Reserved.
+﻿// Copyright 1998-2019 Epic Games, Inc. All Rights Reserved.
 
 using System;
 using System.Collections.Generic;
diff --git a/Engine/Source/Programs/UnrealGameSync/UnrealGameSyncMetadataServer/Models/LatestData.cs b/Engine/Source/Programs/UnrealGameSync/UnrealGameSyncMetadataServer/Models/LatestData.cs
--- a/Engine/Source/Programs/UnrealGameSync/UnrealGameSyncMetadataServer/Models/LatestData.cs
+++ b/Engine/Source/Programs/UnrealGameSync/UnrealGameSyncMetadataServer/Models/LatestData.cs
@@ -1,4 +1,4 @@
-﻿// Copyright 1998-2018 Epic Games, Inc. All Rights Reserved.
+﻿// Copyright 1998-2019 Epic Games, Inc. All Rights Reserved.
 
 using System;
 using System.Collections.Generic;
diff --git a/Engine/Source/Programs/UnrealGameSync/UnrealGameSyncMetadataServer/Models/TelemetryErrorData.cs b/Engine/Source/Programs/UnrealGameSync/UnrealGameSyncMetadataServer/Models/TelemetryErrorData.cs
--- a/Engine/Source/Programs/UnrealGameSync/UnrealGameSyncMetadataServer/Models/TelemetryErrorData.cs
+++ b/Engine/Source/Programs/UnrealGameSync/UnrealGameSyncMetadataServer/Models/TelemetryErrorData.cs
@@ -1,4 +1,4 @@
-﻿// Copyright 1998-2018 Epic Games, Inc. All Rights Reserved.
+﻿// Copyright 1998-2019 Epic Games, Inc. All Rights Reserved.
 
 using System;
 using System.Collections.Generic;
diff --git a/Engine/Source/Programs/UnrealGameSync/UnrealGameSyncMetadataServer/Models/TelemetryTimingData.cs b/Engine/Source/Programs/UnrealGameSync/UnrealGameSyncMetadataServer/Models/TelemetryTimingData.cs
--- a/Engine/Source/Programs/UnrealGameSync/UnrealGameSyncMetadataServer/Models/TelemetryTimingData.cs
+++ b/Engine/Source/Programs/UnrealGameSync/UnrealGameSyncMetadataServer/Models/TelemetryTimingData.cs
@@ -1,4 +1,4 @@
-﻿// Copyright 1998-2018 Epic Games, Inc. All Rights Reserved.
+﻿// Copyright 1998-2019 Epic Games, Inc. All Rights Reserved.
 
 using System;
 using System.Collections.Generic;
diff --git a/Engine/Source/Programs/UnrealGameSync/UnrealGameSyncMetadataServer/Properties/AssemblyInfo.cs b/Engine/Source/Programs/UnrealGameSync/UnrealGameSyncMetadataServer/Properties/AssemblyInfo.cs
--- a/Engine/Source/Programs/UnrealGameSync/UnrealGameSyncMetadataServer/Properties/AssemblyInfo.cs
+++ b/Engine/Source/Programs/UnrealGameSync/UnrealGameSyncMetadataServer/Properties/AssemblyInfo.cs
@@ -1,4 +1,4 @@
-﻿// Copyright 1998-2018 Epic Games, Inc. All Rights Reserved.
+﻿// Copyright 1998-2019 Epic Games, Inc. All Rights Reserved.
 
 using System.Reflection;
 using System.Runtime.CompilerServices;
@@ -12,7 +12,7 @@ using System.Runtime.InteropServices;
 [assembly: AssemblyConfiguration("")]
 [assembly: AssemblyCompany("")]
 [assembly: AssemblyProduct("UnrealGameSync")]
-[assembly: AssemblyCopyright("Copyright 1998-2018 Epic Games, Inc. All Rights Reserved.")]
+[assembly: AssemblyCopyright("Copyright 1998-2019 Epic Games, Inc. All Rights Reserved.")]
 [assembly: AssemblyTrademark("")]
 [assembly: AssemblyCulture("")]
 
diff --git a/Engine/Source/Programs/UnrealGameSync/UnrealGameSyncMetadataServer/UnrealGameSyncMetadataServer.csproj b/Engine/Source/Programs/UnrealGameSync/UnrealGameSyncMetadataServer/UnrealGameSyncMetadataServer.csproj
--- a/Engine/Source/Programs/UnrealGameSync/UnrealGameSyncMetadataServer/UnrealGameSyncMetadataServer.csproj
+++ b/Engine/Source/Programs/UnrealGameSync/UnrealGameSyncMetadataServer/UnrealGameSyncMetadataServer.csproj
@@ -24,6 +24,7 @@
     <NuGetPackageImportStamp>
     </NuGetPackageImportStamp>
     <TargetFrameworkProfile />
+    <Use64BitIISExpress />
   </PropertyGroup>
   <PropertyGroup Condition=" '$(Configuration)|$(Platform)' == 'Debug|AnyCPU' ">
     <TransformWebConfigEnabled>False</TransformWebConfigEnabled>
@@ -89,6 +90,12 @@
   </ItemGroup>
   <ItemGroup>
     <Compile Include="App_Start\WebApiConfig.cs" />
+    <Compile Include="Properties\Resources.Designer.cs">
+      <AutoGen>True</AutoGen>
+      <DesignTime>True</DesignTime>
+      <DependentUpon>Resources.resx</DependentUpon>
+    </Compile>
+    <Content Include="Setup.sql" />
     <Compile Include="Controllers\BuildController.cs" />
     <Compile Include="Controllers\CISController.cs" />
     <Compile Include="Controllers\CommentController.cs" />
@@ -124,12 +131,17 @@
     <Content Include="NotForLicensees\Web.config">
       <SubType>Designer</SubType>
     </Content>
-    <None Include="Properties\PublishProfiles\FilePublish.pubxml" />
-    <None Include="Properties\PublishProfiles\Profile.pubxml" />
   </ItemGroup>
   <ItemGroup>
     <Content Include="packages.config" />
   </ItemGroup>
+  <ItemGroup>
+    <EmbeddedResource Include="Properties\Resources.resx">
+      <Generator>ResXFileCodeGenerator</Generator>
+      <LastGenOutput>Resources.Designer.cs</LastGenOutput>
+      <SubType>Designer</SubType>
+    </EmbeddedResource>
+  </ItemGroup>
   <PropertyGroup>
     <VisualStudioVersion Condition="'$(VisualStudioVersion)' == ''">10.0</VisualStudioVersion>
     <VSToolsPath Condition="'$(VSToolsPath)' == ''">$(MSBuildExtensionsPath32)\Microsoft\VisualStudio\v$(VisualStudioVersion)</VSToolsPath>
