UGS: Fix exception detecting Perforce settings when the project file does not exist.

diff --git a/Engine/Source/Programs/UnrealGameSync/UnrealGameSync/PerforceModalTask.cs b/Engine/Source/Programs/UnrealGameSync/UnrealGameSync/PerforceModalTask.cs
--- a/Engine/Source/Programs/UnrealGameSync/UnrealGameSync/PerforceModalTask.cs
+++ b/Engine/Source/Programs/UnrealGameSync/UnrealGameSync/PerforceModalTask.cs
@@ -43,7 +43,14 @@ namespace UnrealGameSync
 					PerforceConnection Perforce = new PerforceConnection(UserName, null, null);
 					if(ProjectFileName != null)
 					{
-						Directory.SetCurrentDirectory(Path.GetDirectoryName(ProjectFileName));
+						try
+						{
+							Directory.SetCurrentDirectory(Path.GetDirectoryName(ProjectFileName));
+						}
+						catch
+						{
+							// Just ignore
+						}
 					}
 
 					string NewServerAndPort;
