UGS: Only update non-visible tabs every 10 minutes.

diff --git a/Engine/Source/Programs/UnrealGameSync/UnrealGameSync/Controls/WorkspaceControl.cs b/Engine/Source/Programs/UnrealGameSync/UnrealGameSync/Controls/WorkspaceControl.cs
--- a/Engine/Source/Programs/UnrealGameSync/UnrealGameSync/Controls/WorkspaceControl.cs
+++ b/Engine/Source/Programs/UnrealGameSync/UnrealGameSync/Controls/WorkspaceControl.cs
@@ -4105,5 +4105,13 @@ namespace UnrealGameSync
 		{
 			SetBisectStateForSelection(BisectState.Include);
 		}
+
+		private void WorkspaceControl_VisibleChanged(object sender, EventArgs e)
+		{
+			if(PerforceMonitor != null && PerforceMonitor.IsActive != Visible)
+			{
+				PerforceMonitor.IsActive = Visible;
+			}
+		}
 	}
 }
diff --git a/Engine/Source/Programs/UnrealGameSync/UnrealGameSync/PerforceMonitor.cs b/Engine/Source/Programs/UnrealGameSync/UnrealGameSync/PerforceMonitor.cs
--- a/Engine/Source/Programs/UnrealGameSync/UnrealGameSync/PerforceMonitor.cs
+++ b/Engine/Source/Programs/UnrealGameSync/UnrealGameSync/PerforceMonitor.cs
@@ -92,6 +92,12 @@ namespace UnrealGameSync
 			}
 		}
 
+		public bool IsActive
+		{
+			get;
+			set;
+		}
+
 		public string LastStatusMessage
 		{
 			get;
@@ -177,7 +183,7 @@ namespace UnrealGameSync
 				}
 
 				// Wait for another request, or scan for new builds after a timeout
-				RefreshEvent.WaitOne(60 * 1000);
+				RefreshEvent.WaitOne((IsActive? 1 : 10) * 60 * 1000);
 			}
 		}
 
diff --git a/Engine/Source/Programs/UnrealGameSync/UnrealGameSync/Properties/AssemblyInfo.cs b/Engine/Source/Programs/UnrealGameSync/UnrealGameSync/Properties/AssemblyInfo.cs
--- a/Engine/Source/Programs/UnrealGameSync/UnrealGameSync/Properties/AssemblyInfo.cs
+++ b/Engine/Source/Programs/UnrealGameSync/UnrealGameSync/Properties/AssemblyInfo.cs
@@ -34,5 +34,5 @@ using System.Runtime.InteropServices;
 // You can specify all the values or you can default the Build and Revision Numbers 
 // by using the '*' as shown below:
 // [assembly: AssemblyVersion("1.0.*")]
-[assembly: AssemblyVersion("1.134.*")]
-[assembly: AssemblyFileVersion("1.134.0.0")]
+[assembly: AssemblyVersion("1.135.*")]
+[assembly: AssemblyFileVersion("1.135.0.0")]
