UGS: Ignore "files not in client view" error when running a sync preview. This can happen if a sync filter includes a path that doesn't exist.

diff --git a/Engine/Source/Programs/UnrealGameSync/UnrealGameSync/Perforce.cs b/Engine/Source/Programs/UnrealGameSync/UnrealGameSync/Perforce.cs
--- a/Engine/Source/Programs/UnrealGameSync/UnrealGameSync/Perforce.cs
+++ b/Engine/Source/Programs/UnrealGameSync/UnrealGameSync/Perforce.cs
@@ -1137,7 +1137,7 @@ namespace UnrealGameSync
 
 		public bool SyncPreview(string Filter, int ChangeNumber, bool bOnlyFilesInThisChange, out List<PerforceFileRecord> FileRecords, TextWriter Log)
 		{
-			return RunCommand(String.Format("sync -n {0}@{1}{2}", Filter, bOnlyFilesInThisChange? "=" : "", ChangeNumber), out FileRecords, CommandOptions.IgnoreFilesUpToDateError | CommandOptions.IgnoreNoSuchFilesError, Log);
+			return RunCommand(String.Format("sync -n {0}@{1}{2}", Filter, bOnlyFilesInThisChange? "=" : "", ChangeNumber), out FileRecords, CommandOptions.IgnoreFilesUpToDateError | CommandOptions.IgnoreNoSuchFilesError | CommandOptions.IgnoreFilesNotInClientViewError, Log);
 		}
 
 		public bool ForceSync(string Filter, TextWriter Log)
