UGS: Remove redundant entries in the list of sync paths returned by [REDACTED_PERSON]().

diff --git a/Engine/Source/Programs/UnrealGameSync/UnrealGameSync/Workspace.cs b/Engine/Source/Programs/UnrealGameSync/UnrealGameSync/Workspace.cs
--- a/Engine/Source/Programs/UnrealGameSync/UnrealGameSync/Workspace.cs
+++ b/Engine/Source/Programs/UnrealGameSync/UnrealGameSync/Workspace.cs
@@ -984,6 +984,7 @@ namespace UnrealGameSync
 
 		public List<string> GetSyncPaths(bool bSyncAllProjects, string[] SyncFilter)
 		{
+			// Add the default project paths
 			List<string> SyncPaths = new List<string>();
 			if(bSyncAllProjects || !SelectedClientFileName.EndsWith(".uproject", StringComparison.InvariantCultureIgnoreCase))
 			{
@@ -999,6 +1000,8 @@ namespace UnrealGameSync
 				}
 				SyncPaths.Add(PerforceUtils.GetClientOrDepotDirectoryName(SelectedClientFileName) + "/...");
 			}
+
+			// Apply the sync filter to that list. We only want inclusive rules in the output list, but we can manually apply exclusions to previous entries.
 			if(SyncFilter != null)
 			{
 				foreach(string SyncPath in SyncFilter)
@@ -1008,8 +1011,31 @@ namespace UnrealGameSync
 					{
 						SyncPaths.Add(ClientRootPath + TrimSyncPath);
 					}
+					else if(TrimSyncPath.StartsWith("-/") && TrimSyncPath.EndsWith("..."))
+					{
+						SyncPaths.RemoveAll(x => x.StartsWith(ClientRootPath + TrimSyncPath.Substring(1, TrimSyncPath.Length - 4)));
+					}
 				}
 			}
+
+			// Sort the remaining paths by length, and remove any paths which are included twice
+			SyncPaths = SyncPaths.OrderBy(x => x.Length).ToList();
+			for(int Idx = 0; Idx < SyncPaths.Count; Idx++)
+			{
+				string SyncPath = SyncPaths[Idx];
+				if(SyncPath.EndsWith("..."))
+				{
+					string SyncPathPrefix = SyncPath.Substring(0, SyncPath.Length - 3);
+					for(int OtherIdx = SyncPaths.Count - 1; OtherIdx > Idx; OtherIdx--)
+					{
+						if(SyncPaths[OtherIdx].StartsWith(SyncPathPrefix))
+						{
+							SyncPaths.RemoveAt(OtherIdx);
+						}
+					}
+				}
+			}
+
 			return SyncPaths;
 		}
 		
