UGS: Allow for paths being enumerated more than once when searching for files to clean. There may be redundant entries in the list of sync paths.

diff --git a/Engine/Source/Programs/UnrealGameSync/UnrealGameSync/FindFoldersToCleanTask.cs b/Engine/Source/Programs/UnrealGameSync/UnrealGameSync/FindFoldersToCleanTask.cs
--- a/Engine/Source/Programs/UnrealGameSync/UnrealGameSync/FindFoldersToCleanTask.cs
+++ b/Engine/Source/Programs/UnrealGameSync/UnrealGameSync/FindFoldersToCleanTask.cs
@@ -249,7 +249,7 @@ namespace UnrealGameSync
 						}
 						FileFolder = NextFileFolder;
 					}
-					FileFolder.NameToFile.Add(Tokens[Tokens.Length - 1], FileRecord);
+					FileFolder.NameToFile[Tokens[Tokens.Length - 1]] = FileRecord;
 				}
 			}
 
