UGS: Prevent trying to parse a .uprojectdirs file as JSON.

diff --git a/Engine/Source/Programs/UnrealGameSync/UnrealGameSync/DetectProjectSettingsTask.cs b/Engine/Source/Programs/UnrealGameSync/UnrealGameSync/DetectProjectSettingsTask.cs
--- a/Engine/Source/Programs/UnrealGameSync/UnrealGameSync/DetectProjectSettingsTask.cs
+++ b/Engine/Source/Programs/UnrealGameSync/UnrealGameSync/DetectProjectSettingsTask.cs
@@ -310,7 +310,7 @@ namespace UnrealGameSync
 
 			// Figure out if it's an enterprise project
 			List<string> ProjectLines;
-			if(Perforce.Print(NewSelectedClientFileName, out ProjectLines, Log))
+			if(NewSelectedClientFileName.EndsWith(".uproject", StringComparison.InvariantCultureIgnoreCase) && Perforce.Print(NewSelectedClientFileName, out ProjectLines, Log))
 			{
 				string Text = String.Join("\n", ProjectLines);
 				bIsEnterpriseProject = Utility.IsEnterpriseProjectFromText(Text);
