Fix typo.

diff --git a/Engine/Source/Programs/UnrealGameSync/UnrealGameSync/Workspace.cs b/Engine/Source/Programs/UnrealGameSync/UnrealGameSync/Workspace.cs
--- a/Engine/Source/Programs/UnrealGameSync/UnrealGameSync/Workspace.cs
+++ b/Engine/Source/Programs/UnrealGameSync/UnrealGameSync/Workspace.cs
@@ -451,7 +451,7 @@ namespace UnrealGameSync
 					}
 					SyncDepotPaths.RemoveAll(x => !Filter.Matches(x));
 
-					// Findd all the depot paths that will be synced
+					// Find all the depot paths that will be synced
 					HashSet<string> RemainingDepotPaths = new HashSet<string>(StringComparer.InvariantCultureIgnoreCase);
 					RemainingDepotPaths.UnionWith(RemoveDepotPaths);
 					RemainingDepotPaths.UnionWith(SyncDepotPaths);
