UGS: Fix failing to detect enterprise projects correctly when running from installed builds, due to using incorrect client to print project file.

#rnx

diff --git a/Engine/Source/Programs/UnrealGameSync/UnrealGameSync/DetectProjectSettingsTask.cs b/Engine/Source/Programs/UnrealGameSync/UnrealGameSync/DetectProjectSettingsTask.cs
--- a/Engine/Source/Programs/UnrealGameSync/UnrealGameSync/DetectProjectSettingsTask.cs
+++ b/Engine/Source/Programs/UnrealGameSync/UnrealGameSync/DetectProjectSettingsTask.cs
@@ -310,7 +310,7 @@ namespace UnrealGameSync
 
 			// Figure out if it's an enterprise project
 			List<string> ProjectLines;
-			if(NewSelectedClientFileName.EndsWith(".uproject", StringComparison.InvariantCultureIgnoreCase) && Perforce.Print(NewSelectedClientFileName, out ProjectLines, Log))
+			if(NewSelectedClientFileName.EndsWith(".uproject", StringComparison.InvariantCultureIgnoreCase) && PerforceClient.Print(NewSelectedClientFileName, out ProjectLines, Log))
 			{
 				string Text = String.Join("\n", ProjectLines);
 				bIsEnterpriseProject = Utility.IsEnterpriseProjectFromText(Text);
diff --git a/Engine/Source/Programs/UnrealGameSync/UnrealGameSync/Properties/AssemblyInfo.cs b/Engine/Source/Programs/UnrealGameSync/UnrealGameSync/Properties/AssemblyInfo.cs
--- a/Engine/Source/Programs/UnrealGameSync/UnrealGameSync/Properties/AssemblyInfo.cs
+++ b/Engine/Source/Programs/UnrealGameSync/UnrealGameSync/Properties/AssemblyInfo.cs
@@ -34,5 +34,5 @@ using System.Runtime.InteropServices;
 // You can specify all the values or you can default the Build and Revision Numbers 
 // by using the '*' as shown below:
 // [assembly: AssemblyVersion("1.0.*")]
-[assembly: AssemblyVersion("1.136.*")]
-[assembly: AssemblyFileVersion("1.136.0.0")]
+[assembly: AssemblyVersion("1.137.*")]
+[assembly: AssemblyFileVersion("1.137.0.0")]
