UGS: Set the program name as "UGS" when executing P4 commands.

diff --git a/Engine/Source/Programs/UnrealGameSync/UnrealGameSync/Perforce.cs b/Engine/Source/Programs/UnrealGameSync/UnrealGameSync/Perforce.cs
--- a/Engine/Source/Programs/UnrealGameSync/UnrealGameSync/Perforce.cs
+++ b/Engine/Source/Programs/UnrealGameSync/UnrealGameSync/Perforce.cs
@@ -5,8 +5,10 @@ using System.Collections.Generic;
 using System.Diagnostics;
 using System.IO;
 using System.Linq;
+using System.Reflection;
 using System.Text;
 using System.Threading.Tasks;
+using System.Windows.Forms;
 
 namespace UnrealGameSync
 {
@@ -1304,6 +1306,7 @@ namespace UnrealGameSync
 			{
 				FullCommandLine.Append("-s ");
 			}
+			FullCommandLine.AppendFormat("-zprog=UGS -zversion={0} ", Assembly.GetExecutingAssembly().GetName().Version.ToString());
 			FullCommandLine.Append(CommandLine);
 
 			return FullCommandLine.ToString();
diff --git a/Engine/Source/Programs/UnrealGameSync/UnrealGameSync/Properties/AssemblyInfo.cs b/Engine/Source/Programs/UnrealGameSync/UnrealGameSync/Properties/AssemblyInfo.cs
--- a/Engine/Source/Programs/UnrealGameSync/UnrealGameSync/Properties/AssemblyInfo.cs
+++ b/Engine/Source/Programs/UnrealGameSync/UnrealGameSync/Properties/AssemblyInfo.cs
@@ -34,5 +34,5 @@ using System.Runtime.InteropServices;
 // You can specify all the values or you can default the Build and Revision Numbers 
 // by using the '*' as shown below:
 // [assembly: AssemblyVersion("1.0.*")]
-[assembly: AssemblyVersion("1.137.*")]
-[assembly: AssemblyFileVersion("1.137.0.0")]
+[assembly: AssemblyVersion("1.138.*")]
+[assembly: AssemblyFileVersion("1.138.0.0")]
