UGS: Speculative fix for crash when auto-updating, caused by update events firing multiple times.

diff --git a/Engine/Source/Programs/UnrealGameSync/UnrealGameSync/ProgramApplicationContext.cs b/Engine/Source/Programs/UnrealGameSync/UnrealGameSync/ProgramApplicationContext.cs
--- a/Engine/Source/Programs/UnrealGameSync/UnrealGameSync/ProgramApplicationContext.cs
+++ b/Engine/Source/Programs/UnrealGameSync/UnrealGameSync/ProgramApplicationContext.cs
@@ -22,6 +22,7 @@ namespace UnrealGameSync
 		bool bRestoreState;
 		string UpdateSpawn;
 		bool bUnstable;
+		bool bIsClosing;
 
 		BoundedLogWriter Log;
 		UserSettings Settings;
@@ -215,8 +216,9 @@ namespace UnrealGameSync
 
 		private void OnUpdateAvailable()
 		{
-			if(MainWindowInstance != null && MainWindowInstance.CanPerformUpdate())
+			if(MainWindowInstance != null && !bIsClosing && MainWindowInstance.CanPerformUpdate())
 			{
+				bIsClosing = true;
 				MainWindowInstance.ForceClose();
 			}
 		}
