UGS: Improve performance when force-syncing missing or tampered files with the clean workspace dialog, and fix a bug where paths containing @ characters weren't being escaped properly.

diff --git a/Engine/Source/Programs/UnrealGameSync/UnrealGameSync/DeleteFilesTask.cs b/Engine/Source/Programs/UnrealGameSync/UnrealGameSync/DeleteFilesTask.cs
--- a/Engine/Source/Programs/UnrealGameSync/UnrealGameSync/DeleteFilesTask.cs
+++ b/Engine/Source/Programs/UnrealGameSync/UnrealGameSync/DeleteFilesTask.cs
@@ -28,12 +28,18 @@ namespace UnrealGameSync
 		{
 			StringBuilder FailMessage = new StringBuilder();
 
-			foreach(FileInfo FileToSync in FilesToSync)
+			if(FilesToSync.Count > 0)
 			{
+				List<string> RevisionsToSync = new List<string>();
+				foreach(FileInfo FileToSync in FilesToSync)
+				{
+					RevisionsToSync.Add(String.Format("{0}#have", PerforceUtils.EscapePath(FileToSync.FullName)));
+				}
+
 				StringWriter Log = new StringWriter();
-				if(!Perforce.ForceSync(String.Format("{0}#have", FileToSync.FullName), Log))
+				if(!Perforce.Sync(RevisionsToSync, x => { }, new List<string>(), true, null, Log))
 				{
-					FailMessage.AppendFormat("{0} ({1})\r\n", FileToSync.FullName, Log.ToString().Trim().Replace('\n', ' '));
+					FailMessage.Append(Log.ToString());
 				}
 			}
 
diff --git a/Engine/Source/Programs/UnrealGameSync/UnrealGameSync/Perforce.cs b/Engine/Source/Programs/UnrealGameSync/UnrealGameSync/Perforce.cs
--- a/Engine/Source/Programs/UnrealGameSync/UnrealGameSync/Perforce.cs
+++ b/Engine/Source/Programs/UnrealGameSync/UnrealGameSync/Perforce.cs
@@ -1055,7 +1055,7 @@ namespace UnrealGameSync
 
 		public bool Stat(string Filter, out List<PerforceFileRecord> FileRecords, TextWriter Log)
 		{
-			return RunCommand(String.Format("fstat \"{0}\"", Filter), out FileRecords, CommandOptions.None, Log);
+			return RunCommand(String.Format("fstat \"{0}\"", Filter), out FileRecords, CommandOptions.IgnoreFilesNotOnClientError, Log);
 		}
 
 		public bool Sync(string Filter, TextWriter Log)
@@ -1063,7 +1063,7 @@ namespace UnrealGameSync
 			return RunCommand("sync " + Filter, CommandOptions.IgnoreFilesUpToDateError, Log);
 		}
 
-		public bool Sync(List<string> FileRevisions, Action<PerforceFileRecord> SyncOutput, List<string> TamperedFiles, PerforceSyncOptions Options, TextWriter Log)
+		public bool Sync(List<string> FileRevisions, Action<PerforceFileRecord> SyncOutput, List<string> TamperedFiles, bool bForce, PerforceSyncOptions Options, TextWriter Log)
 		{
 			// Write all the files we want to sync to a temp file
 			string TempFileName = Path.GetTempFileName();
@@ -1087,6 +1087,10 @@ namespace UnrealGameSync
 						CommandLine.AppendFormat(" -v net.tcpsize={0}", Options.TcpBufferSize);
 					}
 					CommandLine.Append(" sync");
+					if(bForce)
+					{
+						CommandLine.AppendFormat(" -f");
+					}
 					if(Options != null && Options.NumThreads > 1)
 					{
 						CommandLine.AppendFormat(" --parallel=threads={0}", Options.NumThreads);
@@ -1544,6 +1548,16 @@ namespace UnrealGameSync
 			}
 		}
 
+		static public string EscapePath(string Path)
+		{
+			string NewPath = Path;
+			NewPath = NewPath.Replace("%", "%25");
+			NewPath = NewPath.Replace("*", "%2A");
+			NewPath = NewPath.Replace("#", "%23");
+			NewPath = NewPath.Replace("@", "%40");
+			return NewPath;
+		}
+
 		static public string UnescapePath(string Path)
 		{
 			string NewPath = Path;
diff --git a/Engine/Source/Programs/UnrealGameSync/UnrealGameSync/Properties/AssemblyInfo.cs b/Engine/Source/Programs/UnrealGameSync/UnrealGameSync/Properties/AssemblyInfo.cs
--- a/Engine/Source/Programs/UnrealGameSync/UnrealGameSync/Properties/AssemblyInfo.cs
+++ b/Engine/Source/Programs/UnrealGameSync/UnrealGameSync/Properties/AssemblyInfo.cs
@@ -34,5 +34,5 @@ using System.Runtime.InteropServices;
 // You can specify all the values or you can default the Build and Revision Numbers 
 // by using the '*' as shown below:
 // [assembly: AssemblyVersion("1.0.*")]
-[assembly: AssemblyVersion("1.138.*")]
-[assembly: AssemblyFileVersion("1.138.0.0")]
+[assembly: AssemblyVersion("1.139.*")]
+[assembly: AssemblyFileVersion("1.139.0.0")]
diff --git a/Engine/Source/Programs/UnrealGameSync/UnrealGameSync/Workspace.cs b/Engine/Source/Programs/UnrealGameSync/UnrealGameSync/Workspace.cs
--- a/Engine/Source/Programs/UnrealGameSync/UnrealGameSync/Workspace.cs
+++ b/Engine/Source/Programs/UnrealGameSync/UnrealGameSync/Workspace.cs
@@ -463,7 +463,7 @@ namespace UnrealGameSync
 
 					// Sync them all
 					List<string> TamperedFiles = new List<string>();
-					if(!Perforce.Sync(SyncRevisions, Record => UpdateSyncProgress(Record, RemainingDepotPaths, SyncRevisions.Count), TamperedFiles, Context.PerforceSyncOptions, Log))
+					if(!Perforce.Sync(SyncRevisions, Record => UpdateSyncProgress(Record, RemainingDepotPaths, SyncRevisions.Count), TamperedFiles, false, Context.PerforceSyncOptions, Log))
 					{
 						StatusMessage = "Aborted sync due to errors.";
 						return WorkspaceUpdateResult.FailedToSync;
