UGS: Delete the zip manifest after removing all the files listed in it.

diff --git a/Engine/Source/Programs/UnrealGameSync/UnrealGameSync/Workspace.cs b/Engine/Source/Programs/UnrealGameSync/UnrealGameSync/Workspace.cs
--- a/Engine/Source/Programs/UnrealGameSync/UnrealGameSync/Workspace.cs
+++ b/Engine/Source/Programs/UnrealGameSync/UnrealGameSync/Workspace.cs
@@ -731,6 +731,7 @@ namespace UnrealGameSync
 							Log.WriteLine("Removing {0} binaries...", ArchiveTypeAndDepotPath.Key);
 							Progress.Set(String.Format("Removing {0} binaries...", ArchiveTypeAndDepotPath.Key), 0.0f);
 							ArchiveUtils.RemoveExtractedFiles(LocalRootPath, ManifestFileName, Progress, Log);
+							File.Delete(ManifestFileName);
 							Log.WriteLine();
 						}
 
