UGS: Fix schedule settings being lost for any open projects whenever the settings dialog is opened.

diff --git a/Engine/Source/Programs/UnrealGameSync/UnrealGameSync/Forms/ScheduleWindow.cs b/Engine/Source/Programs/UnrealGameSync/UnrealGameSync/Forms/ScheduleWindow.cs
--- a/Engine/Source/Programs/UnrealGameSync/UnrealGameSync/Forms/ScheduleWindow.cs
+++ b/Engine/Source/Programs/UnrealGameSync/UnrealGameSync/Forms/ScheduleWindow.cs
@@ -33,7 +33,7 @@ namespace UnrealGameSync
 
 			foreach(UserSelectedProjectSettings Project in LocalFileToProject.Values.OrderBy(x => x.ToString()))
 			{
-				bool Enabled = ScheduledProjects.Contains(Project);
+				bool Enabled = ScheduledProjects.Any(x => x.LocalPath == Project.LocalPath);
 				ProjectListBox.Items.Add(Project, Enabled);
 			}
 
