UGS: Fix parsing of progress markup when using parallel executor.

diff --git a/Engine/Source/Programs/UnrealGameSync/UnrealGameSync/OutputAdapters.cs b/Engine/Source/Programs/UnrealGameSync/UnrealGameSync/OutputAdapters.cs
--- a/Engine/Source/Programs/UnrealGameSync/UnrealGameSync/OutputAdapters.cs
+++ b/Engine/Source/Programs/UnrealGameSync/UnrealGameSync/OutputAdapters.cs
@@ -1,4 +1,4 @@
-﻿// Copyright 1998-2018 Epic Games, Inc. All Rights Reserved.
+// Copyright 1998-2018 Epic Games, Inc. All Rights Reserved.
 
 using System;
 using System.Collections.Generic;
@@ -277,11 +277,12 @@ namespace UnrealGameSync
 
 		protected override void FlushLine(string Line)
 		{
-			if(Line.StartsWith(DirectivePrefix))
+			string TrimLine = Line.Trim();
+			if(TrimLine.StartsWith(DirectivePrefix))
 			{
 				// Line that just contains a progress directive
 				bool bSkipLine = false;
-				ProcessInternal(Line.Substring(DirectivePrefix.Length), ref bSkipLine);
+				ProcessInternal(TrimLine.Substring(DirectivePrefix.Length), ref bSkipLine);
 			}
 			else
 			{
@@ -289,7 +290,6 @@ namespace UnrealGameSync
 				string RemainingLine = Line;
 
 				// Look for a progress directive at the end of a line, in square brackets
-				string TrimLine = Line.TrimEnd();
 				if(TrimLine.EndsWith("]"))
 				{
 					for(int LastIdx = TrimLine.Length - 2; LastIdx >= 0 && TrimLine[LastIdx] != ']'; LastIdx--)
