UGS: Improve log output for scheduled syncs and opening new projects.

diff --git a/Engine/Source/Programs/UnrealGameSync/UnrealGameSync/Controls/WorkspaceControl.cs b/Engine/Source/Programs/UnrealGameSync/UnrealGameSync/Controls/WorkspaceControl.cs
--- a/Engine/Source/Programs/UnrealGameSync/UnrealGameSync/Controls/WorkspaceControl.cs
+++ b/Engine/Source/Programs/UnrealGameSync/UnrealGameSync/Controls/WorkspaceControl.cs
@@ -126,7 +126,7 @@ namespace UnrealGameSync
 		IWorkspaceControlOwner Owner;
 		string ApiUrl;
 		string DataFolder;
-		BoundedLogWriter Log;
+		LineBasedTextWriter Log;
 
 		UserSettings Settings;
 		UserWorkspaceSettings WorkspaceSettings;
@@ -216,7 +216,7 @@ namespace UnrealGameSync
 		bool bUpdateBuildMetadataPosted;
 		bool bUpdateReviewsPosted;
 
-		public WorkspaceControl(IWorkspaceControlOwner InOwner, string InApiUrl, string InOriginalExecutableFileName, DetectProjectSettingsTask DetectSettings, BoundedLogWriter InLog, UserSettings InSettings)
+		public WorkspaceControl(IWorkspaceControlOwner InOwner, string InApiUrl, string InOriginalExecutableFileName, DetectProjectSettingsTask DetectSettings, LineBasedTextWriter InLog, UserSettings InSettings)
 		{
 			InitializeComponent();
 
@@ -648,7 +648,6 @@ namespace UnrealGameSync
 			Context.StartTime = DateTime.UtcNow;
 			Context.PerforceSyncOptions = (PerforceSyncOptions)Settings.SyncOptions.Clone();
 
-			Log.WriteLine();
 			Log.WriteLine("Updating workspace at {0}...", Context.StartTime.ToLocalTime().ToString());
 			Log.WriteLine("  ChangeNumber={0}", Context.ChangeNumber);
 			Log.WriteLine("  Options={0}", Context.Options.ToString());
diff --git a/Engine/Source/Programs/UnrealGameSync/UnrealGameSync/Forms/MainWindow.cs b/Engine/Source/Programs/UnrealGameSync/UnrealGameSync/Forms/MainWindow.cs
--- a/Engine/Source/Programs/UnrealGameSync/UnrealGameSync/Forms/MainWindow.cs
+++ b/Engine/Source/Programs/UnrealGameSync/UnrealGameSync/Forms/MainWindow.cs
@@ -57,7 +57,7 @@ namespace UnrealGameSync
 
 		string ApiUrl;
 		string DataFolder;
-		BoundedLogWriter Log;
+		LineBasedTextWriter Log;
 		UserSettings Settings;
 		int TabMenu_TabIdx = -1;
 		int ChangingWorkspacesRefCount;
@@ -73,7 +73,7 @@ namespace UnrealGameSync
 
 		IMainWindowTabPanel CurrentTabPanel;
 
-		public MainWindow(string InApiUrl, string InDataFolder, bool bInRestoreStateOnLoad, string InOriginalExecutableFileName, List<DetectProjectSettingsResult> StartupProjects, BoundedLogWriter InLog, UserSettings InSettings)
+		public MainWindow(string InApiUrl, string InDataFolder, bool bInRestoreStateOnLoad, string InOriginalExecutableFileName, List<DetectProjectSettingsResult> StartupProjects, LineBasedTextWriter InLog, UserSettings InSettings)
 		{
 			InitializeComponent();
 
@@ -444,6 +444,8 @@ namespace UnrealGameSync
 
 				TimeSpan IntervalToFirstTick = NextScheduleTime - CurrentTime;
 				ScheduleTimer = new System.Threading.Timer(x => MainThreadSynchronizationContext.Post((o) => { if(!IsDisposed){ ScheduleTimerElapsed(); } }, null), null, IntervalToFirstTick, TimeSpan.FromDays(1));
+
+				Log.WriteLine("Schedule: Started ScheduleTimer for {0} ({1} remaining)", NextScheduleTime, IntervalToFirstTick);
 			}
 		}
 
@@ -453,16 +455,20 @@ namespace UnrealGameSync
 			{
 				ScheduleTimer.Dispose();
 				ScheduleTimer = null;
+				Log.WriteLine("Schedule: Stopped ScheduleTimer");
 			}
 			StopScheduleSettledTimer();
 		}
 
 		private void ScheduleTimerElapsed()
 		{
+			Log.WriteLine("Schedule: Timer Elapsed");
+
 			// Try to open any missing tabs. 
 			int NumInitialTabs = TabControl.GetTabCount();
 			foreach (UserSelectedProjectSettings ScheduledProject in Settings.ScheduleProjects)
 			{
+				Log.WriteLine("Schedule: Attempting to open {0}", ScheduledProject);
 				TryOpenProject(ScheduledProject, -1, OpenProjectOptions.Quiet);
 			}
 
@@ -481,6 +487,7 @@ namespace UnrealGameSync
 		{
 			StopScheduleSettledTimer();
 			ScheduleSettledTimer = new System.Threading.Timer(x => MainThreadSynchronizationContext.Post((o) => { if(!IsDisposed){ ScheduleSettledTimerElapsed(); } }, null), null, TimeSpan.FromSeconds(20.0), TimeSpan.FromMilliseconds(-1.0));
+			Log.WriteLine("Schedule: Started ScheduleSettledTimer");
 		}
 
 		private void StopScheduleSettledTimer()
@@ -489,17 +496,25 @@ namespace UnrealGameSync
 			{
 				ScheduleSettledTimer.Dispose();
 				ScheduleSettledTimer = null;
+
+				Log.WriteLine("Schedule: Stopped ScheduleSettledTimer");
 			}
 		}
 
 		private void ScheduleSettledTimerElapsed()
 		{
+			Log.WriteLine("Schedule: Starting Sync");
 			for(int Idx = 0; Idx < TabControl.GetTabCount(); Idx++)
 			{
 				WorkspaceControl Workspace = TabControl.GetTabData(Idx) as WorkspaceControl;
-				if(Workspace != null && (Settings.ScheduleAnyOpenProject || Settings.ScheduleProjects.Contains(Workspace.SelectedProject)))
+				if(Workspace != null)
 				{
-					Workspace.ScheduleTimerElapsed();
+					Log.WriteLine("Schedule: Considering {0}", Workspace.SelectedFileName);
+					if(Settings.ScheduleAnyOpenProject || Settings.ScheduleProjects.Contains(Workspace.SelectedProject))
+					{
+						Log.WriteLine("Schedule: Starting Sync");
+						Workspace.ScheduleTimerElapsed();
+					}
 				}
 			}
 		}
@@ -631,7 +646,8 @@ namespace UnrealGameSync
 
 		int TryOpenProject(UserSelectedProjectSettings Project, int ReplaceTabIdx, OpenProjectOptions Options = OpenProjectOptions.None)
 		{
-			using(DetectProjectSettingsTask DetectProjectSettings = new DetectProjectSettingsTask(Project, DataFolder, Log))
+			Log.WriteLine("Detecting settings for {0}", Project);
+			using(DetectProjectSettingsTask DetectProjectSettings = new DetectProjectSettingsTask(Project, DataFolder, new PrefixedTextWriter("  ", Log)))
 			{
 				string ErrorMessage;
 				if(ModalTask.Execute(this, DetectProjectSettings, "Opening Project", "Opening project, please wait...", out ErrorMessage) != ModalTaskResult.Succeeded)
@@ -660,6 +676,7 @@ namespace UnrealGameSync
 					{
 						if(Workspace.SelectedFileName.Equals(ProjectSettings.NewSelectedFileName, StringComparison.InvariantCultureIgnoreCase))
 						{
+							Log.WriteLine("  Already open in tab {0}", TabIdx);
 							if((Options & OpenProjectOptions.Quiet) == 0)
 							{
 								TabControl.SelectTab(TabIdx);
@@ -670,10 +687,12 @@ namespace UnrealGameSync
 						{
 							if((Options & OpenProjectOptions.Quiet) == 0 && MessageBox.Show(String.Format("{0} is already open under {1}.\n\nWould you like to close it?", Path.GetFileNameWithoutExtension(Workspace.SelectedFileName), Workspace.BranchDirectoryName, Path.GetFileNameWithoutExtension(ProjectSettings.NewSelectedFileName)), "Branch already open", MessageBoxButtons.YesNo) == System.Windows.Forms.DialogResult.Yes)
 							{
+								Log.WriteLine("  Another project already open in this workspace, tab {0}. Replacing.", TabIdx);
 								TabControl.RemoveTab(TabIdx);
 							}
 							else
 							{
+								Log.WriteLine("  Another project already open in this workspace, tab {0}. Aborting.", TabIdx);
 								return -1;
 							}
 						}
@@ -709,10 +728,12 @@ namespace UnrealGameSync
 			if(ReplaceTabIdx == -1)
 			{
 				int NewTabIdx = TabControl.InsertTab(-1, NewTabName, NewWorkspace);
+				Log.WriteLine("  Inserted tab {0}", NewTabIdx);
 				return NewTabIdx;
 			}
 			else
 			{
+				Log.WriteLine("  Replacing tab {0}", ReplaceTabIdx);
 				TabControl.InsertTab(ReplaceTabIdx + 1, NewTabName, NewWorkspace);
 				TabControl.RemoveTab(ReplaceTabIdx);
 				return ReplaceTabIdx;
diff --git a/Engine/Source/Programs/UnrealGameSync/UnrealGameSync/OutputAdapters.cs b/Engine/Source/Programs/UnrealGameSync/UnrealGameSync/OutputAdapters.cs
--- a/Engine/Source/Programs/UnrealGameSync/UnrealGameSync/OutputAdapters.cs
+++ b/Engine/Source/Programs/UnrealGameSync/UnrealGameSync/OutputAdapters.cs
@@ -77,6 +77,30 @@ namespace UnrealGameSync
 			Inner.WriteLine(Prefix + Line);
 		}
 	}
+
+	class TimestampLogWriter : LineBasedTextWriter
+	{
+		LineBasedTextWriter Inner;
+
+		public TimestampLogWriter(LineBasedTextWriter Inner)
+		{
+			this.Inner = Inner;
+		}
+
+		protected override void Dispose(bool bDisposing)
+		{
+			if(Inner != null)
+			{
+				Inner.Dispose();
+				Inner = null;
+			}
+		}
+
+		protected override void FlushLine(string Line)
+		{
+			Inner.WriteLine("[{0}] {1}", DateTime.Now, Line);
+		}
+	}
 	
 	class BoundedLogWriter : LineBasedTextWriter
 	{
diff --git a/Engine/Source/Programs/UnrealGameSync/UnrealGameSync/ProgramApplicationContext.cs b/Engine/Source/Programs/UnrealGameSync/UnrealGameSync/ProgramApplicationContext.cs
--- a/Engine/Source/Programs/UnrealGameSync/UnrealGameSync/ProgramApplicationContext.cs
+++ b/Engine/Source/Programs/UnrealGameSync/UnrealGameSync/ProgramApplicationContext.cs
@@ -24,7 +24,7 @@ namespace UnrealGameSync
 		bool bUnstable;
 		bool bIsClosing;
 
-		BoundedLogWriter Log;
+		TimestampLogWriter Log;
 		UserSettings Settings;
 		ActivationListener ActivationListener;
 
@@ -62,7 +62,7 @@ namespace UnrealGameSync
 			MainThreadSynchronizationContext = WindowsFormsSynchronizationContext.Current;
 
 			// Create the log file
-			Log = new BoundedLogWriter(Path.Combine(DataFolder, "UnrealGameSync.log"));
+			Log = new TimestampLogWriter(new BoundedLogWriter(Path.Combine(DataFolder, "UnrealGameSync.log")));
 			Log.WriteLine("Application version: {0}", Assembly.GetExecutingAssembly().GetName().Version);
 			Log.WriteLine("Started at {0}", DateTime.Now.ToString());
 
@@ -145,7 +145,8 @@ namespace UnrealGameSync
 			List<DetectProjectSettingsTask> Tasks = new List<DetectProjectSettingsTask>();
 			foreach(UserSelectedProjectSettings OpenProject in Settings.OpenProjects)
 			{
-				Tasks.Add(new DetectProjectSettingsTask(OpenProject, DataFolder, Log));
+				Log.WriteLine("Detecting settings for {0}", OpenProject);
+				Tasks.Add(new DetectProjectSettingsTask(OpenProject, DataFolder, new PrefixedTextWriter("  ", Log)));
 			}
 
 			// Detect settings for the project we want to open
