Fix directories which are not in Perforce not showing up in the clean workspace dialog.

diff --git a/Engine/Source/Programs/UnrealGameSync/UnrealGameSync/FindFoldersToCleanTask.cs b/Engine/Source/Programs/UnrealGameSync/UnrealGameSync/FindFoldersToCleanTask.cs
--- a/Engine/Source/Programs/UnrealGameSync/UnrealGameSync/FindFoldersToCleanTask.cs
+++ b/Engine/Source/Programs/UnrealGameSync/UnrealGameSync/FindFoldersToCleanTask.cs
@@ -104,44 +104,56 @@ namespace UnrealGameSync
 
 		void MergeTrees(FolderToClean LocalFolder, PerforceHaveFolder PerforceFolder, HashSet<string> OpenClientPaths)
 		{
-			// Loop through all the local sub-folders
-			foreach(FolderToClean LocalSubFolder in LocalFolder.NameToSubFolder.Values)
+			if(PerforceFolder == null)
 			{
-				PerforceHaveFolder PerforceSubFolder;
-				if(PerforceFolder.NameToSubFolder.TryGetValue(LocalSubFolder.Name, out PerforceSubFolder))
+				// Loop through all the local sub-folders
+				foreach(FolderToClean LocalSubFolder in LocalFolder.NameToSubFolder.Values)
 				{
-					MergeTrees(LocalSubFolder, PerforceSubFolder, OpenClientPaths);
+					MergeTrees(LocalSubFolder, null, OpenClientPaths);
 				}
-			}
 
-			// Also merge all the Perforce folders that no longer exist
-			foreach(KeyValuePair<string, PerforceHaveFolder> PerforceSubFolderPair in PerforceFolder.NameToSubFolder)
+				// Delete everything
+				LocalFolder.FilesToDelete.AddRange(LocalFolder.NameToFile.Values);
+			}
+			else
 			{
-				FolderToClean LocalSubFolder;
-				if(!LocalFolder.NameToSubFolder.TryGetValue(PerforceSubFolderPair.Key, out LocalSubFolder))
+				// Loop through all the local sub-folders
+				foreach(FolderToClean LocalSubFolder in LocalFolder.NameToSubFolder.Values)
 				{
-					LocalSubFolder = new FolderToClean(new DirectoryInfo(Path.Combine(LocalFolder.Directory.FullName, PerforceSubFolderPair.Key)));
-					MergeTrees(LocalSubFolder, PerforceSubFolderPair.Value, OpenClientPaths);
-					LocalFolder.NameToSubFolder.Add(LocalSubFolder.Name, LocalSubFolder);
+					PerforceHaveFolder PerforceSubFolder;
+					PerforceFolder.NameToSubFolder.TryGetValue(LocalSubFolder.Name, out PerforceSubFolder);
+					MergeTrees(LocalSubFolder, PerforceSubFolder, OpenClientPaths);
 				}
-			}
 
-			// Find all the files that need to be re-synced
-			foreach(KeyValuePair<string, PerforceFileRecord> FilePair in PerforceFolder.NameToFile)
-			{
-				FileInfo LocalFile;
-				if(!LocalFolder.NameToFile.TryGetValue(FilePair.Key, out LocalFile))
+				// Also merge all the Perforce folders that no longer exist
+				foreach(KeyValuePair<string, PerforceHaveFolder> PerforceSubFolderPair in PerforceFolder.NameToSubFolder)
 				{
-					LocalFolder.FilesToSync.Add(new FileInfo(Path.Combine(LocalFolder.Directory.FullName, FilePair.Key)));
+					FolderToClean LocalSubFolder;
+					if(!LocalFolder.NameToSubFolder.TryGetValue(PerforceSubFolderPair.Key, out LocalSubFolder))
+					{
+						LocalSubFolder = new FolderToClean(new DirectoryInfo(Path.Combine(LocalFolder.Directory.FullName, PerforceSubFolderPair.Key)));
+						MergeTrees(LocalSubFolder, PerforceSubFolderPair.Value, OpenClientPaths);
+						LocalFolder.NameToSubFolder.Add(LocalSubFolder.Name, LocalSubFolder);
+					}
 				}
-				else if((FilePair.Value.Flags & PerforceFileFlags.AlwaysWritable) == 0 && (LocalFile.Attributes & FileAttributes.ReadOnly) == 0 && !OpenClientPaths.Contains(FilePair.Value.ClientPath))
+
+				// Find all the files that need to be re-synced
+				foreach(KeyValuePair<string, PerforceFileRecord> FilePair in PerforceFolder.NameToFile)
 				{
-					LocalFolder.FilesToSync.Add(LocalFile);
+					FileInfo LocalFile;
+					if(!LocalFolder.NameToFile.TryGetValue(FilePair.Key, out LocalFile))
+					{
+						LocalFolder.FilesToSync.Add(new FileInfo(Path.Combine(LocalFolder.Directory.FullName, FilePair.Key)));
+					}
+					else if((FilePair.Value.Flags & PerforceFileFlags.AlwaysWritable) == 0 && (LocalFile.Attributes & FileAttributes.ReadOnly) == 0 && !OpenClientPaths.Contains(FilePair.Value.ClientPath))
+					{
+						LocalFolder.FilesToSync.Add(LocalFile);
+					}
 				}
-			}
 
-			// Find all the files that should be deleted
-			LocalFolder.FilesToDelete.AddRange(LocalFolder.NameToFile.Values.Where(x => !PerforceFolder.NameToFile.ContainsKey(x.Name)));
+				// Find all the files that should be deleted
+				LocalFolder.FilesToDelete.AddRange(LocalFolder.NameToFile.Values.Where(x => !PerforceFolder.NameToFile.ContainsKey(x.Name)));
+			}
 
 			// Figure out if this folder is just an empty directory that needs to be removed
 			LocalFolder.bEmptyLeaf = LocalFolder.NameToFile.Count == 0 && LocalFolder.NameToSubFolder.Count == 0 && LocalFolder.FilesToSync.Count == 0;
