UGS: Do not include executables in diagnostics zip file, and ignore "no such files" error when cleaning workspace.

diff --git a/Engine/Source/Programs/UnrealGameSync/UnrealGameSync/Forms/DiagnosticsWindow.cs b/Engine/Source/Programs/UnrealGameSync/UnrealGameSync/Forms/DiagnosticsWindow.cs
--- a/Engine/Source/Programs/UnrealGameSync/UnrealGameSync/Forms/DiagnosticsWindow.cs
+++ b/Engine/Source/Programs/UnrealGameSync/UnrealGameSync/Forms/DiagnosticsWindow.cs
@@ -1,4 +1,4 @@
-﻿// Copyright 1998-2018 Epic Games, Inc. All Rights Reserved.
+// Copyright 1998-2018 Epic Games, Inc. All Rights Reserved.
 
 using System;
 using System.Collections.Generic;
@@ -56,7 +56,13 @@ namespace UnrealGameSync
 				try
 				{
 					ZipFile Zip = new ZipFile();
-					Zip.AddDirectory(DataFolder);
+					foreach(string FileName in Directory.EnumerateFiles(DataFolder))
+					{
+						if(!FileName.EndsWith(".exe", StringComparison.InvariantCultureIgnoreCase) && !FileName.EndsWith(".dll", StringComparison.InvariantCultureIgnoreCase))
+						{
+							Zip.AddFile(FileName, "");
+						}
+					}
 					Zip.Save(ZipFileName);
 				}
 				catch(Exception Ex)
diff --git a/Engine/Source/Programs/UnrealGameSync/UnrealGameSync/Perforce.cs b/Engine/Source/Programs/UnrealGameSync/UnrealGameSync/Perforce.cs
--- a/Engine/Source/Programs/UnrealGameSync/UnrealGameSync/Perforce.cs
+++ b/Engine/Source/Programs/UnrealGameSync/UnrealGameSync/Perforce.cs
@@ -1055,7 +1055,7 @@ namespace UnrealGameSync
 
 		public bool Stat(string Filter, out List<PerforceFileRecord> FileRecords, TextWriter Log)
 		{
-			return RunCommand(String.Format("fstat \"{0}\"", Filter), out FileRecords, CommandOptions.IgnoreFilesNotOnClientError, Log);
+			return RunCommand(String.Format("fstat \"{0}\"", Filter), out FileRecords, CommandOptions.IgnoreFilesNotOnClientError | CommandOptions.IgnoreNoSuchFilesError, Log);
 		}
 
 		public bool Sync(string Filter, TextWriter Log)
diff --git a/Engine/Source/Programs/UnrealGameSync/UnrealGameSync/Properties/AssemblyInfo.cs b/Engine/Source/Programs/UnrealGameSync/UnrealGameSync/Properties/AssemblyInfo.cs
--- a/Engine/Source/Programs/UnrealGameSync/UnrealGameSync/Properties/AssemblyInfo.cs
+++ b/Engine/Source/Programs/UnrealGameSync/UnrealGameSync/Properties/AssemblyInfo.cs
@@ -34,5 +34,5 @@ using System.Runtime.InteropServices;
 // You can specify all the values or you can default the Build and Revision Numbers 
 // by using the '*' as shown below:
 // [assembly: AssemblyVersion("1.0.*")]
-[assembly: AssemblyVersion("1.140.*")]
-[assembly: AssemblyFileVersion("1.140.0.0")]
+[assembly: AssemblyVersion("1.141.*")]
+[assembly: AssemblyFileVersion("1.141.0.0")]
