UGS: Fix crash writing version files when directory does not already exist.

#rnx

diff --git a/Engine/Source/Programs/UnrealGameSync/UnrealGameSync/Workspace.cs b/Engine/Source/Programs/UnrealGameSync/UnrealGameSync/Workspace.cs
--- a/Engine/Source/Programs/UnrealGameSync/UnrealGameSync/Workspace.cs
+++ b/Engine/Source/Programs/UnrealGameSync/UnrealGameSync/Workspace.cs
@@ -1192,6 +1192,7 @@ namespace UnrealGameSync
 				}
 				else
 				{
+					Directory.CreateDirectory(Path.GetDirectoryName(WhereRecord.LocalPath));
 					Utility.ForceDeleteFile(WhereRecord.LocalPath);
 					if(WhereRecord.DepotPath != null)
 					{
