UGS: Show a dialog box while the launcher is updating executables from Perforce, which allows cancelling the operation if necessary. Allow setting the username on the settings window, and prompt for login credentials if necessary. Should prevent situations where users have to update settings from the command prompt.

Holding down shift during launch now shows the settings dialog rather than an immediate prompt to launch the unstable version (unstable version is shown as a checkbox on this dialog).

diff --git a/Engine/Source/Programs/UnrealGameSync/UnrealGameSync/ModalTask.cs b/Engine/Source/Programs/UnrealGameSync/UnrealGameSync/ModalTask.cs
--- a/Engine/Source/Programs/UnrealGameSync/UnrealGameSync/ModalTask.cs
+++ b/Engine/Source/Programs/UnrealGameSync/UnrealGameSync/ModalTask.cs
@@ -25,7 +25,7 @@ namespace UnrealGameSync
 	{
 		public static ModalTaskResult Execute(IWin32Window Owner, IModalTask Task, string InTitle, string InMessage, out string ErrorMessage)
 		{
-			ModalTaskWindow Window = new ModalTaskWindow(Task, InTitle, InMessage, FormStartPosition.CenterParent);
+			ModalTaskWindow Window = new ModalTaskWindow(Task, InTitle, InMessage, (Owner == null)? FormStartPosition.CenterScreen : FormStartPosition.CenterParent);
 			Window.Complete += () => Window.Close();
 			Window.ShowDialog(Owner);
 			ErrorMessage = Window.ErrorMessage;
diff --git a/Engine/Source/Programs/UnrealGameSync/UnrealGameSync/Perforce.cs b/Engine/Source/Programs/UnrealGameSync/UnrealGameSync/Perforce.cs
--- a/Engine/Source/Programs/UnrealGameSync/UnrealGameSync/Perforce.cs
+++ b/Engine/Source/Programs/UnrealGameSync/UnrealGameSync/Perforce.cs
@@ -292,7 +292,7 @@ namespace UnrealGameSync
 		public string DepotPath;
 	}
 
-	class PerforceSyncOptions
+	public class PerforceSyncOptions
 	{
 		public int NumRetries;
 		public int NumThreads;
diff --git a/Engine/Source/Programs/UnrealGameSync/UnrealGameSync/PerforceModalTask.cs b/Engine/Source/Programs/UnrealGameSync/UnrealGameSync/PerforceModalTask.cs
--- a/Engine/Source/Programs/UnrealGameSync/UnrealGameSync/PerforceModalTask.cs
+++ b/Engine/Source/Programs/UnrealGameSync/UnrealGameSync/PerforceModalTask.cs
@@ -117,6 +117,7 @@ namespace UnrealGameSync
 				LoginResult = Perforce.Login(Password, out PasswordErrorMessage, Log);
 				if(LoginResult != LoginResult.Succeded)
 				{
+					Log.WriteLine(PasswordErrorMessage);
 					ErrorMessage = String.Format("Unable to login: {0}", PasswordErrorMessage);
 					return false;
 				}
@@ -127,6 +128,7 @@ namespace UnrealGameSync
 			{
 				LoginResult = LoginResult.MissingPassword;
 				ErrorMessage = "User is not logged in to Perforce.";
+				Log.WriteLine(ErrorMessage);
 				return false;
 			}
 
diff --git a/Engine/Source/Programs/UnrealGameSync/UnrealGameSync/Program.cs b/Engine/Source/Programs/UnrealGameSync/UnrealGameSync/Program.cs
--- a/Engine/Source/Programs/UnrealGameSync/UnrealGameSync/Program.cs
+++ b/Engine/Source/Programs/UnrealGameSync/UnrealGameSync/Program.cs
@@ -58,6 +58,12 @@ namespace UnrealGameSync
 			string UpdateSpawn;
 			ParseArgument(RemainingArgs, "-updatespawn=", out UpdateSpawn);
 
+			string ServerAndPort;
+			ParseArgument(RemainingArgs, "-p4port=", out ServerAndPort);
+
+			string UserName;
+			ParseArgument(RemainingArgs, "-p4user=", out UserName);
+
 			bool bRestoreState;
 			ParseOption(RemainingArgs, "-restorestate", out bRestoreState);
 
@@ -90,7 +96,7 @@ namespace UnrealGameSync
 			{
 				try
 				{
-					using(UpdateMonitor UpdateMonitor = new UpdateMonitor(new PerforceConnection(null, null, null), UpdatePath))
+					using(UpdateMonitor UpdateMonitor = new UpdateMonitor(new PerforceConnection(UserName, null, ServerAndPort), UpdatePath))
 					{
 						ProgramApplicationContext Context = new ProgramApplicationContext(UpdateMonitor, ApiUrl, DataFolder, ActivateEvent, bRestoreState, UpdateSpawn, ProjectFileName, bUnstable);
 						Application.Run(Context);
diff --git a/Engine/Source/Programs/UnrealGameSync/UnrealGameSyncLauncher/Forms/LogWindow.cs b/Engine/Source/Programs/UnrealGameSync/UnrealGameSyncLauncher/Forms/LogWindow.cs
new file mode 100644
--- /dev/null
+++ b/Engine/Source/Programs/UnrealGameSync/UnrealGameSyncLauncher/Forms/LogWindow.cs
@@ -0,0 +1,25 @@
+// Copyright 1998-2018 Epic Games, Inc. All Rights Reserved.
+
+using System;
+using System.Collections.Generic;
+using System.ComponentModel;
+using System.Data;
+using System.Drawing;
+using System.Linq;
+using System.Text;
+using System.Threading.Tasks;
+using System.Windows.Forms;
+
+namespace UnrealGameSyncLauncher
+{
+	public partial class LogWindow : Form
+	{
+		public LogWindow(string Text)
+		{
+			InitializeComponent();
+
+			LogTextBox.Text = Text;
+			LogTextBox.Select(Text.Length, 0);
+		}
+	}
+}
diff --git a/Engine/Source/Programs/UnrealGameSync/UnrealGameSyncLauncher/Forms/SettingsWindow.cs b/Engine/Source/Programs/UnrealGameSync/UnrealGameSyncLauncher/Forms/SettingsWindow.cs
new file mode 100644
--- /dev/null
+++ b/Engine/Source/Programs/UnrealGameSync/UnrealGameSyncLauncher/Forms/SettingsWindow.cs
@@ -0,0 +1,116 @@
+// Copyright 1998-2018 Epic Games, Inc. All Rights Reserved.
+
+using System;
+using System.Collections.Generic;
+using System.ComponentModel;
+using System.Data;
+using System.Drawing;
+using System.IO;
+using System.Linq;
+using System.Runtime.InteropServices;
+using System.Text;
+using System.Threading;
+using System.Threading.Tasks;
+using System.Windows.Forms;
+using UnrealGameSync;
+
+namespace UnrealGameSyncLauncher
+{
+	partial class SettingsWindow : Form
+	{
+		[DllImport("user32.dll")]
+		private static extern IntPtr SendMessage(IntPtr hWnd, int Msg, int wParam, [MarshalAs(UnmanagedType.LPWStr)] string lParam);
+
+		public delegate bool SyncAndRunDelegate(PerforceConnection Perforce, string DepotPath, bool bUnstable, TextWriter LogWriter);
+
+		const int EM_SETCUEBANNER = 0x1501;
+
+		string LogText;
+		SyncAndRunDelegate SyncAndRun;
+
+		public SettingsWindow(string Prompt, string LogText, string ServerAndPort, string UserName, string DepotPath, bool bUnstable, SyncAndRunDelegate SyncAndRun)
+		{
+			InitializeComponent();
+
+			if(Prompt != null)
+			{
+				this.PromptLabel.Text = Prompt;
+			}
+
+			this.LogText = LogText;
+			this.ServerTextBox.Text = ServerAndPort;
+			this.UserNameTextBox.Text = UserName;
+			this.DepotPathTextBox.Text = DepotPath;
+			this.UseUnstableBuildCheckBox.Checked = bUnstable;
+			this.SyncAndRun = SyncAndRun;
+
+			ViewLogBtn.Visible = LogText != null;
+		}
+
+		protected override void OnLoad(EventArgs e)
+		{
+			base.OnLoad(e);
+
+			SendMessage(ServerTextBox.Handle, EM_SETCUEBANNER, 1, "Default Server");
+			SendMessage(UserNameTextBox.Handle, EM_SETCUEBANNER, 1, "Default User");
+		}
+
+		private void ViewLogBtn_Click(object sender, EventArgs e)
+		{
+			LogWindow Log = new LogWindow(LogText);
+			Log.ShowDialog(this);
+		}
+
+		private void ConnectBtn_Click(object sender, EventArgs e)
+		{
+			// Update the settings
+			string ServerAndPort = ServerTextBox.Text.Trim();
+			if(ServerAndPort.Length == 0)
+			{
+				ServerAndPort = null;
+			}
+
+			string UserName = UserNameTextBox.Text.Trim();
+			if(UserName.Length == 0)
+			{
+				UserName = null;
+			}
+
+			string DepotPath = DepotPathTextBox.Text.Trim();
+			if(DepotPath.Length == 0)
+			{
+				DepotPath = null;
+			}
+
+			Program.SaveSettings(ServerAndPort, UserName, DepotPath);
+
+			// Create the task for connecting to this server
+			StringWriter Log = new StringWriter();
+			SyncAndRunPerforceTask SyncApplication = new SyncAndRunPerforceTask((Perforce, LogWriter) => SyncAndRun(Perforce, DepotPath, UseUnstableBuildCheckBox.Checked, LogWriter));
+
+			// Attempt to sync through a modal dialog
+			string ErrorMessage;
+			ModalTaskResult Result = PerforceModalTask.Execute(this, null, ServerAndPort, UserName, SyncApplication, "Updating", "Checking for updates, please wait...", Log, out ErrorMessage);
+			if(Result == ModalTaskResult.Succeeded)
+			{
+				Program.SaveSettings(ServerAndPort, UserName, DepotPath);
+				DialogResult = DialogResult.OK;
+				Close();
+			}
+
+			if(ErrorMessage != null)
+			{
+				PromptLabel.Text = ErrorMessage;
+			}
+
+			LogText = Log.ToString();
+			ViewLogBtn.Visible = true;
+		}
+
+		private void CancelBtn_Click(object sender, EventArgs e)
+		{
+			DialogResult = DialogResult.Cancel;
+			Close();
+		}
+	}
+}
diff --git a/Engine/Source/Programs/UnrealGameSync/UnrealGameSyncLauncher/Program.cs b/Engine/Source/Programs/UnrealGameSync/UnrealGameSyncLauncher/Program.cs
--- a/Engine/Source/Programs/UnrealGameSync/UnrealGameSyncLauncher/Program.cs
+++ b/Engine/Source/Programs/UnrealGameSync/UnrealGameSyncLauncher/Program.cs
@@ -1,4 +1,4 @@
-﻿// Copyright 1998-2018 Epic Games, Inc. All Rights Reserved.
+// Copyright 1998-2018 Epic Games, Inc. All Rights Reserved.
 
 using Microsoft.Win32;
 using System;
@@ -11,6 +11,7 @@ using System.Text;
 using System.Threading;
 using System.Threading.Tasks;
 using System.Windows.Forms;
+using UnrealGameSync;
 
 namespace UnrealGameSyncLauncher
 {
@@ -63,85 +64,82 @@ namespace UnrealGameSyncLauncher
 					return 1;
 				}
 
-				string Server = null;
-				string DepotPath = DefaultDepotPath;
+				// Figure out if we should sync the unstable build by default
+				bool bUnstable = Args.Contains("-unstable", StringComparer.InvariantCultureIgnoreCase);
 
-				ReadSettings(ref Server, ref DepotPath);
+				// Read the settings
+				string ServerAndPort = null;
+				string UserName = null;
+				string DepotPath = DefaultDepotPath;
+				ReadSettings(ref ServerAndPort, ref UserName, ref DepotPath);
 
-				for(;;)
+				// If the shift key is held down, immediately show the settings window
+				if((Control.ModifierKeys & Keys.Shift) != 0)
 				{
-					StringWriter LogWriter = new StringWriter();
-					LogWriter.WriteLine("Attempting to sync UnrealGameSync...");
-
-					if(!String.IsNullOrEmpty(DepotPath))
+					// Show the settings window immediately
+					SettingsWindow UpdateError = new SettingsWindow(null, null, ServerAndPort, UserName, DepotPath, bUnstable, (Perforce, DepotParam, bUnstableParam, LogWriter) => SyncAndRun(Perforce, DepotParam, bUnstableParam, Args, InstanceMutex, LogWriter));
+					if(UpdateError.ShowDialog() == DialogResult.OK)
 					{
-						try
-						{
-							if(SyncAndRunApplication(Args, InstanceMutex, PerforceFileName, Server, DepotPath, LogWriter))
-							{
-								break;
-							}
-						}
-						catch(Exception Ex)
-						{
-							LogWriter.WriteLine(Ex.ToString());
-						}
+						return 0;
 					}
+				}
+				else
+				{
+					// Try to do a sync with the current settings first
+					SyncAndRunModalTask SyncApplication = new SyncAndRunModalTask(ServerAndPort, UserName, (Perforce, LogWriter) => SyncAndRun(Perforce, DepotPath, bUnstable, Args, InstanceMutex, LogWriter));
 
-					string DefaultServer = null;
-					if(String.IsNullOrEmpty(Server))
+					string ErrorMessage;
+					if(ModalTask.Execute(null, SyncApplication, "Updating", "Checking for updates, please wait...", out ErrorMessage) == ModalTaskResult.Succeeded)
 					{
-						List<string> Lines = new List<string>();
-						RunPerforceCommand(PerforceFileName, "set P4PORT", Lines, new StringWriter());
-
-						if(Lines.Count >= 1 && Lines[0].StartsWith("P4PORT="))
-						{
-							DefaultServer = Lines[0].Substring(7).Split(' ')[0];
-							Server = DefaultServer;
-						}
+						return 0;
 					}
 
-					UpdateErrorWindow ErrorWindow = new UpdateErrorWindow(LogWriter.ToString(), Server, DepotPath);
-					if(ErrorWindow.ShowDialog() != DialogResult.OK)
+					SettingsWindow UpdateError = new SettingsWindow("Unable to update UnrealGameSync from Perforce. Verify that your connection settings are correct.", SyncApplication.LogWriter.ToString(), SyncApplication.ServerAndPort, SyncApplication.UserName, DepotPath, bUnstable, (Perforce, DepotParam, bUnstableParam, LogWriter) => SyncAndRun(Perforce, DepotParam, bUnstableParam, Args, InstanceMutex, LogWriter));
+					if(UpdateError.ShowDialog() == DialogResult.OK)
 					{
-						return 2;
+						return 0;
 					}
-
-					Server = (DefaultServer != null && String.Equals(ErrorWindow.Server, DefaultServer, StringComparison.InvariantCultureIgnoreCase))? null : ErrorWindow.Server;
-					DepotPath = ErrorWindow.DepotPath;
-
-					SaveSettings(Server, DepotPath);
 				}
 			}
-			return 0;
+			return 1;
 		}
 
-		static void ReadSettings(ref string Server, ref string DepotPath)
+		public static void ReadSettings(ref string ServerAndPort, ref string UserName, ref string DepotPath)
 		{
 			using (RegistryKey Key = Registry.CurrentUser.OpenSubKey("SOFTWARE\\Epic Games\\UnrealGameSync", false))
 			{
 				if(Key != null)
 				{
-					Server = Key.GetValue("Server", Server) as string;
+					ServerAndPort = Key.GetValue("Server", ServerAndPort) as string;
+					UserName = Key.GetValue("UserName", UserName) as string;
 					DepotPath = Key.GetValue("DepotPath", DepotPath) as string;
 				}
 			}
 		}
 
-		static void SaveSettings(string Server, string DepotPath)
+		public static void SaveSettings(string ServerAndPort, string UserName, string DepotPath)
 		{
 			using (RegistryKey Key = Registry.CurrentUser.OpenSubKey("SOFTWARE\\Epic Games\\UnrealGameSync", true))
 			{
-				if(String.IsNullOrEmpty(Server))
+				if(String.IsNullOrEmpty(ServerAndPort))
 				{
 					try { Key.DeleteValue("Server"); } catch(Exception) { }
 				}
 				else
 				{
-					Key.SetValue("Server", Server);
+					Key.SetValue("Server", ServerAndPort);
+				}
+
+				if(String.IsNullOrEmpty(UserName))
+				{
+					try { Key.DeleteValue("UserName"); } catch(Exception) { }
+				}
+				else
+				{
+					Key.SetValue("UserName", UserName);
 				}
 
-				if(String.IsNullOrEmpty(DepotPath) || String.Equals(DepotPath, DefaultDepotPath, StringComparison.InvariantCultureIgnoreCase))
+				if(String.IsNullOrEmpty(DepotPath) || (DefaultDepotPath != null && String.Equals(DepotPath, DefaultDepotPath, StringComparison.InvariantCultureIgnoreCase)))
 				{
 					try { Key.DeleteValue("DepotPath"); } catch(Exception) { }
 				}
@@ -152,203 +150,146 @@ namespace UnrealGameSyncLauncher
 			}
 		}
 
-		static bool SyncAndRunApplication(string[] Args, Mutex InstanceMutex, string PerforceFileName, string Server, string BaseDepotPath, TextWriter LogWriter)
+		public static bool SyncAndRun(PerforceConnection Perforce, string BaseDepotPath, bool bUnstable, string[] Args, Mutex InstanceMutex, TextWriter LogWriter)
 		{
-			// Get the path that we're syncing
-			bool bUnstable = Args.Contains("-unstable", StringComparer.InvariantCultureIgnoreCase);
-			if(!bUnstable && (Control.ModifierKeys & Keys.Shift) != 0)
-			{
-				if(MessageBox.Show("Use the latest unstable build of UnrealGameSync?\n\n(This message was triggered by holding down the SHIFT key on startup).", "Use unstable build?", MessageBoxButtons.YesNo) == DialogResult.Yes)
-				{
-					bUnstable = true;
-				}
-			}
-
-			string SyncPath = BaseDepotPath.TrimEnd('/') + (bUnstable? "/UnstableRelease/..." : "/Release/...");
-			LogWriter.WriteLine("Syncing from {0}", SyncPath);
-
-			// Create the target folder
-			string ApplicationFolder = Path.Combine(Environment.GetFolderPath(Environment.SpecialFolder.LocalApplicationData), "UnrealGameSync", "Latest");
-			if(!SafeCreateDirectory(ApplicationFolder))
-			{
-				LogWriter.WriteLine("Couldn't create directory: {0}", ApplicationFolder);
-				return false;
-			}
-
-			// Get the server option
-			string ServerArgument = "";
-			if(!String.IsNullOrEmpty(Server))
-			{
-				ServerArgument = String.Format("-p\"{0}\"", Server);
-			}
-
-			// Find the most recent changelist
-			List<string> SubmittedLines = new List<string>();
-			if(RunPerforceCommand(PerforceFileName, String.Format("{0} changes -s submitted -m 1", ServerArgument), SubmittedLines, LogWriter) != 0)
-			{
-				LogWriter.WriteLine("Couldn't find last changelist");
-				return false;
-			}
-
-			// Split into tokens
-			string[] ChangeTokens = SubmittedLines[0].Split(' ');
-
-			// Parse the changelist number
-			int RequiredChangeNumber;
-			if(ChangeTokens.Length < 2 || ChangeTokens[0] != "Change" || !int.TryParse(ChangeTokens[1], out RequiredChangeNumber))
+			try
 			{
-				LogWriter.WriteLine("Couldn't parse last changelist number");
-				return false;
-			}
-
-			// Read the current version
-			string SyncVersionFile = Path.Combine(ApplicationFolder, "SyncVersion.txt");
-			string RequiredSyncText = String.Format("{0}\n{1}@{2}", Server ?? "", SyncPath, RequiredChangeNumber.ToString());
+				string SyncPath = BaseDepotPath.TrimEnd('/') + (bUnstable? "/UnstableRelease/..." : "/Release/...");
+				LogWriter.WriteLine("Syncing from {0}", SyncPath);
 
-			// Check the application exists
-			string ApplicationExe = Path.Combine(ApplicationFolder, "UnrealGameSync.exe");
-
-			// Check if the version has changed
-			string SyncText;
-			if(!File.Exists(SyncVersionFile) || !File.Exists(ApplicationExe) || !TryReadAllText(SyncVersionFile, out SyncText) || SyncText != RequiredSyncText)
-			{
-				// Try to delete the directory contents. Retry for a while, in case we've been spawned by an application in this folder to do an update.
-				for(int NumRetries = 0; !SafeDeleteDirectoryContents(ApplicationFolder); NumRetries++)
+				// Create the target folder
+				string ApplicationFolder = Path.Combine(Environment.GetFolderPath(Environment.SpecialFolder.LocalApplicationData), "UnrealGameSync", "Latest");
+				if(!SafeCreateDirectory(ApplicationFolder))
 				{
-					if(NumRetries > 20)
-					{
-						LogWriter.WriteLine("Couldn't delete contents of {0} (retried {1} times).", ApplicationFolder, NumRetries);
-						return false;
-					}
-					Thread.Sleep(500);
+					LogWriter.WriteLine("Couldn't create directory: {0}", ApplicationFolder);
+					return false;
 				}
-				
-				// Find all the files in the sync path at this changelist
-				List<string> FileLines = new List<string>();
-				if(RunPerforceCommand(PerforceFileName, String.Format("{0} -z tag fstat \"{1}@{2}\"", ServerArgument, SyncPath, RequiredChangeNumber), FileLines, LogWriter) != 0)
+
+				// Find the most recent changelist
+				List<PerforceChangeSummary> Changes;
+				if(!Perforce.FindChanges(SyncPath, 1, out Changes, LogWriter) || Changes.Count < 1)
 				{
-					LogWriter.WriteLine("Couldn't find matching files.");
+					LogWriter.WriteLine("Couldn't find last changelist");
 					return false;
 				}
 
-				// Sync all the files in this list to the same directory structure under the application folder
-				string DepotPathPrefix = SyncPath.Substring(0, SyncPath.LastIndexOf('/') + 1);
-				foreach(string FileLine in FileLines)
+				// Take the first changelist number
+				int RequiredChangeNumber = Changes[0].Number;
+
+				// Read the current version
+				string SyncVersionFile = Path.Combine(ApplicationFolder, "SyncVersion.txt");
+				string RequiredSyncText = String.Format("{0}\n{1}@{2}", Perforce.ServerAndPort ?? "", SyncPath, RequiredChangeNumber);
+
+				// Check the application exists
+				string ApplicationExe = Path.Combine(ApplicationFolder, "UnrealGameSync.exe");
+
+				// Check if the version has changed
+				string SyncText;
+				if(!File.Exists(SyncVersionFile) || !File.Exists(ApplicationExe) || !TryReadAllText(SyncVersionFile, out SyncText) || SyncText != RequiredSyncText)
 				{
-					const string DepotPathTag = "... depotFile ";
-					if(FileLine.StartsWith(DepotPathTag))
+					// Try to delete the directory contents. Retry for a while, in case we've been spawned by an application in this folder to do an update.
+					for(int NumRetries = 0; !SafeDeleteDirectoryContents(ApplicationFolder); NumRetries++)
 					{
-						string DepotPath = FileLine.Substring(DepotPathTag.Length).Trim();
-						if(!DepotPath.StartsWith(DepotPathPrefix, StringComparison.InvariantCultureIgnoreCase))
+						if(NumRetries > 20)
 						{
-							LogWriter.WriteLine("Found file {0} which did not begin with {1}", DepotPath, DepotPathPrefix);
+							LogWriter.WriteLine("Couldn't delete contents of {0} (retried {1} times).", ApplicationFolder, NumRetries);
 							return false;
 						}
+						Thread.Sleep(500);
+					}
+				
+					// Find all the files in the sync path at this changelist
+					List<PerforceFileRecord> FileRecords;
+					if(!Perforce.Stat(String.Format("{0}@{1}", SyncPath, RequiredChangeNumber), out FileRecords, LogWriter))
+					{
+						LogWriter.WriteLine("Couldn't find matching files.");
+						return false;
+					}
 
-						string LocalPath = Path.Combine(ApplicationFolder, DepotPath.Substring(DepotPathPrefix.Length).Replace('/', Path.DirectorySeparatorChar));
+					// Sync all the files in this list to the same directory structure under the application folder
+					string DepotPathPrefix = SyncPath.Substring(0, SyncPath.LastIndexOf('/') + 1);
+					foreach(PerforceFileRecord FileRecord in FileRecords)
+					{
+						string LocalPath = Path.Combine(ApplicationFolder, FileRecord.DepotPath.Substring(DepotPathPrefix.Length).Replace('/', Path.DirectorySeparatorChar));
 						if(!SafeCreateDirectory(Path.GetDirectoryName(LocalPath)))
 						{
 							LogWriter.WriteLine("Couldn't create folder {0}", Path.GetDirectoryName(LocalPath));
 							return false;
 						}
-						if(RunPerforceCommand(PerforceFileName, String.Format("print -o \"{0}\" \"{1}@{2}\"", LocalPath, DepotPath, RequiredChangeNumber), null, LogWriter) != 0)
+						if(!Perforce.PrintToFile(FileRecord.DepotPath, LocalPath, LogWriter))
 						{
-							LogWriter.WriteLine("Couldn't sync {0} to {1}", DepotPath, LocalPath);
+							LogWriter.WriteLine("Couldn't sync {0} to {1}", FileRecord.DepotPath, LocalPath);
 							return false;
 						}
 					}
-				}
 
-				// Check the application exists
-				if(!File.Exists(ApplicationExe))
-				{
-					LogWriter.WriteLine("Application was not synced from Perforce. Check that UnrealGameSync exists at {0}/UnrealGameSync.exe, and you have access to it.", SyncPath);
-					return false;
+					// Check the application exists
+					if(!File.Exists(ApplicationExe))
+					{
+						LogWriter.WriteLine("Application was not synced from Perforce. Check that UnrealGameSync exists at {0}/UnrealGameSync.exe, and you have access to it.", SyncPath);
+						return false;
+					}
+
+					// Update the version
+					if(!TryWriteAllText(SyncVersionFile, RequiredSyncText))
+					{
+						LogWriter.WriteLine("Couldn't write sync text to {0}", SyncVersionFile);
+						return false;
+					}
 				}
+				LogWriter.WriteLine();
 
-				// Update the version
-				if(!TryWriteAllText(SyncVersionFile, RequiredSyncText))
+				// Build the command line for the synced application, including the sync path to monitor for updates
+				StringBuilder NewCommandLine = new StringBuilder(String.Format("-updatepath=\"{0}@>{1}\" -updatespawn=\"{2}\"{3}", SyncPath, RequiredChangeNumber, Assembly.GetEntryAssembly().Location, bUnstable? " -unstable" : ""));
+				if(!String.IsNullOrEmpty(Perforce.ServerAndPort))
 				{
-					LogWriter.WriteLine("Couldn't write sync text to {0}", SyncVersionFile);
-					return false;
+					NewCommandLine.AppendFormat(" -p4port={0}", QuoteArgument(Perforce.ServerAndPort));
 				}
-			}
-			LogWriter.WriteLine();
-
-			// Build the command line for the synced application, including the sync path to monitor for updates
-			StringBuilder NewCommandLine = new StringBuilder(String.Format("-updatepath=\"{0}@>{1}\" -updatespawn=\"{2}\"{3}", SyncPath, RequiredChangeNumber, Assembly.GetEntryAssembly().Location, bUnstable? " -unstable" : ""));
-			foreach(string Arg in Args)
-			{
-				if(Arg.Contains(' '))
+				if(!String.IsNullOrEmpty(Perforce.UserName))
 				{
-					NewCommandLine.AppendFormat( "\"{0}\"", Arg);
+					NewCommandLine.AppendFormat(" -p4user={0}", QuoteArgument(Perforce.UserName));
 				}
-				else
+				foreach(string Arg in Args)
 				{
-					NewCommandLine.AppendFormat(" {0}", Arg);
+					NewCommandLine.AppendFormat(" {0}", QuoteArgument(Arg));
 				}
-			}
 
-			// Release the mutex now so that the new application can start up
-			InstanceMutex.Close();
+				// Release the mutex now so that the new application can start up
+				InstanceMutex.Close();
 
-			// Spawn the application
-			LogWriter.WriteLine("Spawning {0} with command line: {1}", ApplicationExe, NewCommandLine.ToString());
-			using(Process ChildProcess = new Process())
-			{
-				ChildProcess.StartInfo.FileName = ApplicationExe;
-				ChildProcess.StartInfo.Arguments = NewCommandLine.ToString();
-				ChildProcess.StartInfo.UseShellExecute = false;
-				ChildProcess.StartInfo.CreateNoWindow = false;
-				if(!ChildProcess.Start())
+				// Spawn the application
+				LogWriter.WriteLine("Spawning {0} with command line: {1}", ApplicationExe, NewCommandLine.ToString());
+				using(Process ChildProcess = new Process())
 				{
-					LogWriter.WriteLine("Failed to start process");
-					return false;
+					ChildProcess.StartInfo.FileName = ApplicationExe;
+					ChildProcess.StartInfo.Arguments = NewCommandLine.ToString();
+					ChildProcess.StartInfo.UseShellExecute = false;
+					ChildProcess.StartInfo.CreateNoWindow = false;
+					if(!ChildProcess.Start())
+					{
+						LogWriter.WriteLine("Failed to start process");
+						return false;
+					}
 				}
-			}
-
-			return true;
-		}
 
-		static int RunPerforceCommand(string PerforceFileName, string CommandLine, List<string> Lines, TextWriter LogWriter)
-		{
-			LogWriter.WriteLine();
-			LogWriter.WriteLine("Running p4.exe {0}", CommandLine);
-			using(Process ChildProcess = new Process())
+				return true;
+			}
+			catch(Exception Ex)
 			{
-				DataReceivedEventHandler OutputHandler = (x, y) => HandlePerforceOutput(x, y, Lines, LogWriter);
-				
-				ChildProcess.StartInfo.FileName = PerforceFileName;
-				ChildProcess.StartInfo.Arguments = CommandLine;
-				ChildProcess.StartInfo.UseShellExecute = false;
-				ChildProcess.StartInfo.RedirectStandardOutput = true;
-				ChildProcess.StartInfo.RedirectStandardError = true;
-				ChildProcess.OutputDataReceived += OutputHandler;
-				ChildProcess.ErrorDataReceived += OutputHandler;
-				ChildProcess.StartInfo.CreateNoWindow = true;
-				ChildProcess.StartInfo.StandardOutputEncoding = new System.Text.UTF8Encoding(false, false);
-				ChildProcess.Start();
-				ChildProcess.BeginOutputReadLine();
-				ChildProcess.BeginErrorReadLine();
-				ChildProcess.WaitForExit();
-
-				LogWriter.WriteLine("Finished with exit code {0}", ChildProcess.ExitCode);
-				return ChildProcess.ExitCode;
+				LogWriter.WriteLine(Ex.ToString());
+				return false;
 			}
 		}
 
-		static void HandlePerforceOutput(object Sender, DataReceivedEventArgs Args, List<string> Lines, TextWriter LogWriter)
+		static string QuoteArgument(string Arg)
 		{
-			if(Args.Data != null)
+			if(Arg.IndexOf(' ') != -1 && !Arg.StartsWith("\""))
 			{
-				lock(LogWriter)
-				{
-					if(Lines != null)
-					{
-						Lines.Add(Args.Data);
-					}
-					LogWriter.WriteLine("p4> {0}", Args.Data);
-				}
+				return String.Format("\"{0}\"", Arg);
+			}
+			else
+			{
+				return Arg;
 			}
 		}
 
diff --git a/Engine/Source/Programs/UnrealGameSync/UnrealGameSyncLauncher/SyncAndRunModalTask.cs b/Engine/Source/Programs/UnrealGameSync/UnrealGameSyncLauncher/SyncAndRunModalTask.cs
new file mode 100644
--- /dev/null
+++ b/Engine/Source/Programs/UnrealGameSync/UnrealGameSyncLauncher/SyncAndRunModalTask.cs
@@ -0,0 +1,51 @@
+// Copyright 1998-2018 Epic Games, Inc. All Rights Reserved.
+
+using System;
+using System.Collections.Generic;
+using System.Diagnostics;
+using System.IO;
+using System.Linq;
+using System.Text;
+using System.Threading;
+using System.Threading.Tasks;
+using UnrealGameSync;
+
+namespace UnrealGameSyncLauncher
+{
+	class SyncAndRunModalTask : IModalTask
+	{
+		public delegate bool SyncAndRunDelegate(PerforceConnection Perforce, TextWriter LogWriter);
+
+		public string ServerAndPort;
+		public string UserName;
+		SyncAndRunDelegate SyncAndRun;
+		public StringWriter LogWriter;
+
+		public SyncAndRunModalTask(string ServerAndPort, string UserName, SyncAndRunDelegate SyncAndRun)
+		{
+			this.ServerAndPort = ServerAndPort;
+			this.UserName = UserName;
+			this.SyncAndRun = SyncAndRun;
+			this.LogWriter = new StringWriter();
+		}
+
+		public bool Run(out string ErrorMessage)
+		{
+			if(!PerforceModalTask.TryGetServerSettings(null, ref ServerAndPort, ref UserName, LogWriter))
+			{
+				ErrorMessage = "Unable to get Perforce settings.";
+				return false;
+			}
+
+			PerforceConnection Perforce = new PerforceConnection(UserName, null, ServerAndPort);
+			if(!SyncAndRun(Perforce, LogWriter))
+			{
+				ErrorMessage = "Unable to sync and run application.";
+				return false;
+			}
+
+			ErrorMessage = null;
+			return true;
+		}
+	}
+}
diff --git a/Engine/Source/Programs/UnrealGameSync/UnrealGameSyncLauncher/SyncAndRunPerforceTask.cs b/Engine/Source/Programs/UnrealGameSync/UnrealGameSyncLauncher/SyncAndRunPerforceTask.cs
new file mode 100644
--- /dev/null
+++ b/Engine/Source/Programs/UnrealGameSync/UnrealGameSyncLauncher/SyncAndRunPerforceTask.cs
@@ -0,0 +1,39 @@
+// Copyright 1998-2018 Epic Games, Inc. All Rights Reserved.
+
+using System;
+using System.Collections.Generic;
+using System.IO;
+using System.Linq;
+using System.Text;
+using System.Threading;
+using System.Threading.Tasks;
+using UnrealGameSync;
+
+namespace UnrealGameSyncLauncher
+{
+	class SyncAndRunPerforceTask : IPerforceModalTask
+	{
+		public delegate bool SyncAndRunDelegate(PerforceConnection Perforce, TextWriter LogWriter);
+
+		SyncAndRunDelegate SyncAndRun;
+
+		public SyncAndRunPerforceTask(SyncAndRunDelegate SyncAndRun)
+		{
+			this.SyncAndRun = SyncAndRun;
+		}
+
+		public bool Run(PerforceConnection Perforce, TextWriter LogWriter, out string ErrorMessage)
+		{
+			if(SyncAndRun(Perforce, LogWriter))
+			{
+				ErrorMessage = null;
+				return true;
+			}
+			else
+			{
+				ErrorMessage = "Failed to sync application.";
+				return false;
+			}
+		}
+	}
+}
diff --git a/Engine/Source/Programs/UnrealGameSync/UnrealGameSyncLauncher/UnrealGameSyncLauncher.csproj b/Engine/Source/Programs/UnrealGameSync/UnrealGameSyncLauncher/UnrealGameSyncLauncher.csproj
--- a/Engine/Source/Programs/UnrealGameSync/UnrealGameSyncLauncher/UnrealGameSyncLauncher.csproj
+++ b/Engine/Source/Programs/UnrealGameSync/UnrealGameSyncLauncher/UnrealGameSyncLauncher.csproj
@@ -58,6 +58,7 @@
   <ItemGroup>
     <Reference Include="System" />
     <Reference Include="System.Core" />
+    <Reference Include="System.Web.Extensions" />
     <Reference Include="System.Xml.Linq" />
     <Reference Include="System.Data.DataSetExtensions" />
     <Reference Include="Microsoft.CSharp" />
@@ -68,20 +69,57 @@
     <Reference Include="System.Xml" />
   </ItemGroup>
   <ItemGroup>
-    <Compile Include="UpdateLogWindow.cs">
+    <Compile Include="..\UnrealGameSync\ChildProcess.cs">
+      <Link>Shared\ChildProcess.cs</Link>
+    </Compile>
+    <Compile Include="..\UnrealGameSync\Forms\ModalTaskWindow.cs">
+      <Link>UnrealGameSync\ModalTaskWindow.cs</Link>
       <SubType>Form</SubType>
     </Compile>
-    <Compile Include="UpdateLogWindow.Designer.cs">
-      <DependentUpon>UpdateLogWindow.cs</DependentUpon>
+    <Compile Include="..\UnrealGameSync\Forms\ModalTaskWindow.Designer.cs">
+      <Link>UnrealGameSync\ModalTaskWindow.Designer.cs</Link>
     </Compile>
-    <Compile Include="Program.cs" />
-    <Compile Include="Properties\AssemblyInfo.cs" />
-    <Compile Include="UpdateErrorWindow.cs">
+    <Compile Include="..\UnrealGameSync\Forms\PasswordWindow.cs">
+      <Link>UnrealGameSync\PasswordWindow.cs</Link>
+      <SubType>Form</SubType>
+    </Compile>
+    <Compile Include="..\UnrealGameSync\Forms\PasswordWindow.Designer.cs">
+      <Link>UnrealGameSync\PasswordWindow.Designer.cs</Link>
+    </Compile>
+    <Compile Include="..\UnrealGameSync\ModalTask.cs">
+      <Link>Shared\ModalTask.cs</Link>
+    </Compile>
+    <Compile Include="..\UnrealGameSync\Perforce.cs">
+      <Link>Shared\Perforce.cs</Link>
+    </Compile>
+    <Compile Include="..\UnrealGameSync\PerforceModalTask.cs">
+      <Link>UnrealGameSync\PerforceModalTask.cs</Link>
+    </Compile>
+    <Compile Include="..\UnrealGameSync\Utility.cs">
+      <Link>Shared\Utility.cs</Link>
+    </Compile>
+    <Compile Include="Forms\LogWindow.cs">
       <SubType>Form</SubType>
     </Compile>
-    <Compile Include="UpdateErrorWindow.Designer.cs">
-      <DependentUpon>UpdateErrorWindow.cs</DependentUpon>
+    <Compile Include="Forms\LogWindow.Designer.cs">
+      <DependentUpon>LogWindow.cs</DependentUpon>
     </Compile>
+    <Compile Include="Forms\SettingsWindow.cs">
+      <SubType>Form</SubType>
+    </Compile>
+    <Compile Include="Forms\SettingsWindow.Designer.cs">
+      <DependentUpon>SettingsWindow.cs</DependentUpon>
+    </Compile>
+    <Compile Include="SyncAndRunModalTask.cs" />
+    <Compile Include="SyncAndRunPerforceTask.cs" />
+    <Compile Include="Program.cs" />
+    <Compile Include="Properties\AssemblyInfo.cs" />
+    <EmbeddedResource Include="Forms\LogWindow.resx">
+      <DependentUpon>LogWindow.cs</DependentUpon>
+    </EmbeddedResource>
+    <EmbeddedResource Include="Forms\SettingsWindow.resx">
+      <DependentUpon>SettingsWindow.cs</DependentUpon>
+    </EmbeddedResource>
     <EmbeddedResource Include="Properties\Resources.resx">
       <Generator>ResXFileCodeGenerator</Generator>
       <LastGenOutput>Resources.Designer.cs</LastGenOutput>
@@ -91,12 +129,6 @@
       <AutoGen>True</AutoGen>
       <DependentUpon>Resources.resx</DependentUpon>
     </Compile>
-    <EmbeddedResource Include="UpdateErrorWindow.resx">
-      <DependentUpon>UpdateErrorWindow.cs</DependentUpon>
-    </EmbeddedResource>
-    <EmbeddedResource Include="UpdateLogWindow.resx">
-      <DependentUpon>UpdateLogWindow.cs</DependentUpon>
-    </EmbeddedResource>
     <None Include="Properties\Settings.settings">
       <Generator>SettingsSingleFileGenerator</Generator>
       <LastGenOutput>Settings.Designer.cs</LastGenOutput>
@@ -113,6 +145,7 @@
   <ItemGroup>
     <Content Include="Resources\Icon.ico" />
   </ItemGroup>
+  <ItemGroup />
   <Choose>
     <When Condition="Exists('NotForLicensees\ProgramSettings.cs')">
       <ItemGroup>
diff --git a/Engine/Source/Programs/UnrealGameSync/UnrealGameSyncLauncher/UpdateErrorWindow.cs b/Engine/Source/Programs/UnrealGameSync/UnrealGameSyncLauncher/UpdateErrorWindow.cs
deleted file mode 100644
diff --git a/Engine/Source/Programs/UnrealGameSync/UnrealGameSyncLauncher/UpdateLogWindow.cs b/Engine/Source/Programs/UnrealGameSync/UnrealGameSyncLauncher/UpdateLogWindow.cs
deleted file mode 100644
