Update version number for UGS launcher to 1.13.

diff --git a/Engine/Source/Programs/UnrealGameSync/UnrealGameSyncLauncher/Properties/AssemblyInfo.cs b/Engine/Source/Programs/UnrealGameSync/UnrealGameSyncLauncher/Properties/AssemblyInfo.cs
--- a/Engine/Source/Programs/UnrealGameSync/UnrealGameSyncLauncher/Properties/AssemblyInfo.cs
+++ b/Engine/Source/Programs/UnrealGameSync/UnrealGameSyncLauncher/Properties/AssemblyInfo.cs
@@ -1,4 +1,4 @@
-﻿// Copyright 1998-2018 Epic Games, Inc. All Rights Reserved.
+// Copyright 1998-2018 Epic Games, Inc. All Rights Reserved.
 
 using System.Reflection;
 using System.Runtime.CompilerServices;
@@ -34,5 +34,5 @@ using System.Runtime.InteropServices;
 // You can specify all the values or you can default the Build and Revision Numbers 
 // by using the '*' as shown below:
 // [assembly: AssemblyVersion("1.0.*")]
-[assembly: AssemblyVersion("1.12.0.0")]
-[assembly: AssemblyFileVersion("1.12.0.0")]
+[assembly: AssemblyVersion("1.13.0.0")]
+[assembly: AssemblyFileVersion("1.13.0.0")]
