UGS: Main window is now restored to normal size when activated by clicking on the tray icon.

diff --git a/Engine/Source/Programs/UnrealGameSync/UnrealGameSync/Forms/MainWindow.cs b/Engine/Source/Programs/UnrealGameSync/UnrealGameSync/Forms/MainWindow.cs
--- a/Engine/Source/Programs/UnrealGameSync/UnrealGameSync/Forms/MainWindow.cs
+++ b/Engine/Source/Programs/UnrealGameSync/UnrealGameSync/Forms/MainWindow.cs
@@ -337,6 +337,10 @@ namespace UnrealGameSync
 		public void ShowAndActivate()
 		{
 			Show();
+			if(WindowState == FormWindowState.Minimized)
+			{
+				WindowState = FormWindowState.Normal;
+			}
 			Activate();
 		}
 
