UGS: Speculative fix for unreproduced exception accessing disposed window while shutting down.

#rnx

diff --git a/Engine/Source/Programs/UnrealGameSync/UnrealGameSync/Controls/WorkspaceControl.cs b/Engine/Source/Programs/UnrealGameSync/UnrealGameSync/Controls/WorkspaceControl.cs
--- a/Engine/Source/Programs/UnrealGameSync/UnrealGameSync/Controls/WorkspaceControl.cs
+++ b/Engine/Source/Programs/UnrealGameSync/UnrealGameSync/Controls/WorkspaceControl.cs
@@ -165,6 +165,7 @@ namespace UnrealGameSync
 		}
 
 		SynchronizationContext MainThreadSynchronizationContext;
+		bool bIsDisposing;
 
 		string EditorTargetName;
 		bool bIsEnterpriseProject;
@@ -398,6 +399,8 @@ namespace UnrealGameSync
 		/// <param name="disposing">true if managed resources should be disposed; otherwise, false.</param>
 		protected override void Dispose(bool disposing)
 		{
+			bIsDisposing = true;
+
 			if (disposing && (components != null))
 			{
 				components.Dispose();
@@ -499,7 +502,7 @@ namespace UnrealGameSync
 
 		private void LoginExpired()
 		{
-			if(!IsDisposed)
+			if(!bIsDisposing)
 			{
 				Owner.RequestProjectChange(this, SelectedProject);
 			}
@@ -804,7 +807,7 @@ namespace UnrealGameSync
 			if(!bUpdateBuildListPosted)
 			{
 				bUpdateBuildListPosted = true;
-				MainThreadSynchronizationContext.Post((o) => { bUpdateBuildListPosted = false; if(!IsDisposed) { UpdateBuildList(); } }, null);
+				MainThreadSynchronizationContext.Post((o) => { bUpdateBuildListPosted = false; if(!bIsDisposing) { UpdateBuildList(); } }, null);
 			}
 		}
 
@@ -985,7 +988,7 @@ namespace UnrealGameSync
 			if(!bUpdateBuildMetadataPosted)
 			{
 				bUpdateBuildMetadataPosted = true;
-				MainThreadSynchronizationContext.Post((o) => { bUpdateBuildMetadataPosted = false; if(!IsDisposed) { UpdateBuildMetadata(); } }, null);
+				MainThreadSynchronizationContext.Post((o) => { bUpdateBuildMetadataPosted = false; if(!bIsDisposing) { UpdateBuildMetadata(); } }, null);
 			}
 		}
 
@@ -1025,7 +1028,7 @@ namespace UnrealGameSync
 			if(!bUpdateReviewsPosted)
 			{
 				bUpdateReviewsPosted = true;
-				MainThreadSynchronizationContext.Post((o) => { bUpdateReviewsPosted = false; if(!IsDisposed) { UpdateReviews(); } }, null);
+				MainThreadSynchronizationContext.Post((o) => { bUpdateReviewsPosted = false; if(!bIsDisposing) { UpdateReviews(); } }, null);
 			}
 		}
 
diff --git a/Engine/Source/Programs/UnrealGameSync/UnrealGameSync/Properties/AssemblyInfo.cs b/Engine/Source/Programs/UnrealGameSync/UnrealGameSync/Properties/AssemblyInfo.cs
--- a/Engine/Source/Programs/UnrealGameSync/UnrealGameSync/Properties/AssemblyInfo.cs
+++ b/Engine/Source/Programs/UnrealGameSync/UnrealGameSync/Properties/AssemblyInfo.cs
@@ -34,5 +34,5 @@ using System.Runtime.InteropServices;
 // You can specify all the values or you can default the Build and Revision Numbers 
 // by using the '*' as shown below:
 // [assembly: AssemblyVersion("1.0.*")]
-[assembly: AssemblyVersion("1.141.*")]
-[assembly: AssemblyFileVersion("1.141.0.0")]
+[assembly: AssemblyVersion("1.142.*")]
+[assembly: AssemblyFileVersion("1.142.0.0")]
