UGS: Allow the $(Change) variable to be used in more places than just the context menu.

diff --git a/Engine/Source/Programs/UnrealGameSync/UnrealGameSync/Controls/WorkspaceControl.cs b/Engine/Source/Programs/UnrealGameSync/UnrealGameSync/Controls/WorkspaceControl.cs
--- a/Engine/Source/Programs/UnrealGameSync/UnrealGameSync/Controls/WorkspaceControl.cs
+++ b/Engine/Source/Programs/UnrealGameSync/UnrealGameSync/Controls/WorkspaceControl.cs
@@ -627,7 +627,7 @@ namespace UnrealGameSync
 
 			string[] CombinedSyncFilter = UserSettings.GetCombinedSyncFilter(Workspace.GetSyncCategories(), Settings.SyncView, Settings.SyncExcludedCategories, WorkspaceSettings.SyncView, WorkspaceSettings.SyncIncludedCategories, WorkspaceSettings.SyncExcludedCategories);
 
-			WorkspaceUpdateContext Context = new WorkspaceUpdateContext(ChangeNumber, Options, CombinedSyncFilter, GetDefaultBuildStepObjects(), ProjectSettings.BuildSteps, null, GetWorkspaceVariables());
+			WorkspaceUpdateContext Context = new WorkspaceUpdateContext(ChangeNumber, Options, CombinedSyncFilter, GetDefaultBuildStepObjects(), ProjectSettings.BuildSteps, null, GetWorkspaceVariables(ChangeNumber));
 			if(Options.HasFlag(WorkspaceUpdateOptions.SyncArchives))
 			{
 				string EditorArchivePath = null;
@@ -1984,7 +1984,7 @@ namespace UnrealGameSync
 							Options |= WorkspaceUpdateOptions.UseIncrementalBuilds;
 						}
 		
-						WorkspaceUpdateContext Context = new WorkspaceUpdateContext(Workspace.CurrentChangeNumber, Options, null, GetDefaultBuildStepObjects(), ProjectSettings.BuildSteps, null, GetWorkspaceVariables());
+						WorkspaceUpdateContext Context = new WorkspaceUpdateContext(Workspace.CurrentChangeNumber, Options, null, GetDefaultBuildStepObjects(), ProjectSettings.BuildSteps, null, GetWorkspaceVariables(Workspace.CurrentChangeNumber));
 						StartWorkspaceUpdate(Context);
 					}
 				}
@@ -2726,8 +2726,7 @@ namespace UnrealGameSync
 						ConfigFile ProjectConfigFile = PerforceMonitor.LatestProjectConfigFile;
 						if(ProjectConfigFile != null)
 						{
-							Dictionary<string, string> Variables = GetWorkspaceVariables();
-							Variables.Add("Change", String.Format("{0}", ContextMenuChange.Number));
+							Dictionary<string, string> Variables = GetWorkspaceVariables(ContextMenuChange.Number);
 
 							string[] ChangeContextMenuEntries = ProjectConfigFile.GetValues("Options.ContextMenu", new string[0]);
 							foreach(string ChangeContextMenuEntry in ChangeContextMenuEntries)
@@ -3477,18 +3476,19 @@ namespace UnrealGameSync
 				}
 				else
 				{
-					WorkspaceUpdateContext Context = new WorkspaceUpdateContext(Workspace.CurrentChangeNumber, WorkspaceUpdateOptions.Build, null, GetDefaultBuildStepObjects(), ProjectSettings.BuildSteps, new HashSet<Guid>{ UniqueId }, GetWorkspaceVariables());
+					WorkspaceUpdateContext Context = new WorkspaceUpdateContext(Workspace.CurrentChangeNumber, WorkspaceUpdateOptions.Build, null, GetDefaultBuildStepObjects(), ProjectSettings.BuildSteps, new HashSet<Guid>{ UniqueId }, GetWorkspaceVariables(Workspace.CurrentChangeNumber));
 					StartWorkspaceUpdate(Context);
 				}
 			}
 		}
 
-		private Dictionary<string, string> GetWorkspaceVariables()
+		private Dictionary<string, string> GetWorkspaceVariables(int ChangeNumber)
 		{
 			BuildConfig EditorBuildConfig = GetEditorBuildConfig();
 
 			Dictionary<string, string> Variables = new Dictionary<string,string>();
 			Variables.Add("Stream", StreamName);
+			Variables.Add("Change", ChangeNumber.ToString());
 			Variables.Add("ClientName", ClientName);
 			Variables.Add("BranchDir", BranchDirectoryName);
 			Variables.Add("ProjectDir", Path.GetDirectoryName(SelectedFileName));
@@ -3524,7 +3524,7 @@ namespace UnrealGameSync
 				List<BuildStep> UserSteps = GetUserBuildSteps(ProjectBuildStepObjects);
 
 				// Show the dialog
-				ModifyBuildStepsWindow EditStepsWindow = new ModifyBuildStepsWindow(TargetNames, UserSteps, new HashSet<Guid>(ProjectBuildStepObjects.Keys), BranchDirectoryName, GetWorkspaceVariables());
+				ModifyBuildStepsWindow EditStepsWindow = new ModifyBuildStepsWindow(TargetNames, UserSteps, new HashSet<Guid>(ProjectBuildStepObjects.Keys), BranchDirectoryName, GetWorkspaceVariables(Workspace.CurrentChangeNumber));
 				EditStepsWindow.ShowDialog();
 
 				// Update the user settings
