UGS: Setting for whether or not to use incremental builds is now exposed through the variable "$(UseIncrementalBuilds)" for use by custom build steps.

diff --git a/Engine/Source/Programs/UnrealGameSync/UnrealGameSync/Controls/WorkspaceControl.cs b/Engine/Source/Programs/UnrealGameSync/UnrealGameSync/Controls/WorkspaceControl.cs
--- a/Engine/Source/Programs/UnrealGameSync/UnrealGameSync/Controls/WorkspaceControl.cs
+++ b/Engine/Source/Programs/UnrealGameSync/UnrealGameSync/Controls/WorkspaceControl.cs
@@ -3497,6 +3497,7 @@ namespace UnrealGameSync
 			Variables.Add("UE4EditorCmdExe", GetEditorExePath(EditorBuildConfig).Replace(".exe", "-Cmd.exe"));
 			Variables.Add("UE4EditorConfig", EditorBuildConfig.ToString());
 			Variables.Add("UE4EditorDebugArg", (EditorBuildConfig == BuildConfig.Debug || EditorBuildConfig == BuildConfig.DebugGame)? " -debug" : "");
+			Variables.Add("UseIncrementalBuilds", Settings.bUseIncrementalBuilds? "1" : "0");
 			return Variables;
 		}
 
