UGS: A per-project list of folders and extensions to be deleted by default when running the 'clean workspace' tool can now be specified through the <ProjectDir>/Build/UnrealGameSync.ini file. Settings may be specified for an individual branch (via a category with the depot path to the project) or for wherever the project is currently open (via the [Default] category).

The SafeToDeleteFolders list specifies a substring that will be checked against folder paths. Anything containing this folder will be marked as safe for delete by default.

The SafeToDeleteExtensions list specifies a list of extensions for files that can always be deleted.

Example:

 [Default]
+SafeToDeleteFolders=/MyGame/Test/
+SafeToDeleteFolders=/DataService/
+SafeToDeleteExtensions=.xx1
+SafeToDeleteExtensions=.xx2

diff --git a/Engine/Source/Programs/UnrealGameSync/UnrealGameSync/Controls/WorkspaceControl.cs b/Engine/Source/Programs/UnrealGameSync/UnrealGameSync/Controls/WorkspaceControl.cs
--- a/Engine/Source/Programs/UnrealGameSync/UnrealGameSync/Controls/WorkspaceControl.cs
+++ b/Engine/Source/Programs/UnrealGameSync/UnrealGameSync/Controls/WorkspaceControl.cs
@@ -3423,9 +3423,22 @@ namespace UnrealGameSync
 				return;
 			}
 
+			string ExtraSafeToDeleteFolders;
+			if(!TryGetProjectSetting(Workspace.ProjectConfigFile, "SafeToDeleteFolders", out ExtraSafeToDeleteFolders))
+			{
+				ExtraSafeToDeleteFolders = "";
+			}
+
+			string ExtraSafeToDeleteExtensions;
+			if(!TryGetProjectSetting(Workspace.ProjectConfigFile, "SafeToDeleteExtensions", out ExtraSafeToDeleteExtensions))
+			{
+				ExtraSafeToDeleteExtensions = "";
+			}
+
 			string[] CombinedSyncFilter = UserSettings.GetCombinedSyncFilter(Workspace.GetSyncCategories(), Settings.SyncView, Settings.SyncExcludedCategories, WorkspaceSettings.SyncView, WorkspaceSettings.SyncIncludedCategories, WorkspaceSettings.SyncExcludedCategories);
 			List<string> SyncPaths = Workspace.GetSyncPaths(WorkspaceSettings.bSyncAllProjects ?? Settings.bSyncAllProjects, CombinedSyncFilter);
-			CleanWorkspaceWindow.DoClean(ParentForm, Workspace.Perforce, BranchDirectoryName, Workspace.ClientRootPath, SyncPaths, Log);
+
+			CleanWorkspaceWindow.DoClean(ParentForm, Workspace.Perforce, BranchDirectoryName, Workspace.ClientRootPath, SyncPaths, ExtraSafeToDeleteFolders.Split('\n'), ExtraSafeToDeleteExtensions.Split('\n'), Log);
 		}
 
 		private void UpdateBuildSteps()
diff --git a/Engine/Source/Programs/UnrealGameSync/UnrealGameSync/Forms/CleanWorkspaceWindow.cs b/Engine/Source/Programs/UnrealGameSync/UnrealGameSync/Forms/CleanWorkspaceWindow.cs
--- a/Engine/Source/Programs/UnrealGameSync/UnrealGameSync/Forms/CleanWorkspaceWindow.cs
+++ b/Engine/Source/Programs/UnrealGameSync/UnrealGameSync/Forms/CleanWorkspaceWindow.cs
@@ -98,18 +98,23 @@ namespace UnrealGameSync
 			".csproj.references",
 		};
 
+		IReadOnlyList<string> ExtraSafeToDeleteFolders;
+		IReadOnlyList<string> ExtraSafeToDeleteExtensions;
+
 		[DllImport("Shell32.dll", EntryPoint = "ExtractIconExW", CharSet = CharSet.Unicode, ExactSpelling = true, CallingConvention = CallingConvention.StdCall)]
 		private static extern int ExtractIconEx(string sFile, int iIndex, IntPtr piLargeVersion, out IntPtr piSmallVersion, int amountIcons);
 
-		private CleanWorkspaceWindow(PerforceConnection PerforceClient, FolderToClean RootFolderToClean)
+		private CleanWorkspaceWindow(PerforceConnection PerforceClient, FolderToClean RootFolderToClean, string[] ExtraSafeToDeleteFolders, string[] ExtraSafeToDeleteExtensions)
 		{
 			this.PerforceClient = PerforceClient;
 			this.RootFolderToClean = RootFolderToClean;
+			this.ExtraSafeToDeleteFolders = ExtraSafeToDeleteFolders.Select(x => x.Trim().Replace('\\', '/').Trim('/')).Where(x => x.Length > 0).Select(x => String.Format("/{0}/", x.ToLowerInvariant())).ToArray();
+			this.ExtraSafeToDeleteExtensions = ExtraSafeToDeleteExtensions.Select(x => x.Trim().ToLowerInvariant()).Where(x => x.Length > 0).ToArray();
 
 			InitializeComponent();
 		}
 
-		public static void DoClean(IWin32Window Owner, PerforceConnection PerforceClient, string LocalRootPath, string ClientRootPath, IReadOnlyList<string> SyncPaths, TextWriter Log)
+		public static void DoClean(IWin32Window Owner, PerforceConnection PerforceClient, string LocalRootPath, string ClientRootPath, IReadOnlyList<string> SyncPaths, string[] ExtraSafeToDeleteFolders, string[] ExtraSafeToDeleteExtensions, TextWriter Log)
 		{
 			// Figure out which folders to clean
 			FolderToClean RootFolderToClean = new FolderToClean(new DirectoryInfo(LocalRootPath));
@@ -134,7 +139,7 @@ namespace UnrealGameSync
 			}
 
 			// Populate the tree
-			CleanWorkspaceWindow CleanWorkspace = new CleanWorkspaceWindow(PerforceClient, RootFolderToClean);
+			CleanWorkspaceWindow CleanWorkspace = new CleanWorkspaceWindow(PerforceClient, RootFolderToClean, ExtraSafeToDeleteFolders, ExtraSafeToDeleteExtensions);
 			CleanWorkspace.ShowDialog();
 		}
 
@@ -276,12 +281,12 @@ namespace UnrealGameSync
 
 		private bool IsSafeToDeleteFolder(string FolderPath)
 		{
-			return SafeToDeleteFolders.Any(x => FolderPath.EndsWith(x));
+			return SafeToDeleteFolders.Any(x => FolderPath.EndsWith(x)) || ExtraSafeToDeleteFolders.Any(x => FolderPath.EndsWith(x));
 		}
 
 		private bool IsSafeToDeleteFile(string FolderPath, string Name)
 		{
-			return SafeToDeleteExtensions.Any(x => Name.EndsWith(x));
+			return SafeToDeleteExtensions.Any(x => Name.EndsWith(x)) || ExtraSafeToDeleteExtensions.Any(x => Name.EndsWith(x));
 		}
 
 		private void TreeView_DrawNode(object sender, DrawTreeNodeEventArgs e)
