UGS: Align all badges with the same name, to make it easier to see which CIS steps are being run. Allow overriding the slot taken by a particular badge by calling it "SlotName:LabelName".

diff --git a/Engine/Source/Programs/UnrealGameSync/UnrealGameSync/Controls/WorkspaceControl.cs b/Engine/Source/Programs/UnrealGameSync/UnrealGameSync/Controls/WorkspaceControl.cs
--- a/Engine/Source/Programs/UnrealGameSync/UnrealGameSync/Controls/WorkspaceControl.cs
+++ b/Engine/Source/Programs/UnrealGameSync/UnrealGameSync/Controls/WorkspaceControl.cs
@@ -196,6 +196,8 @@ namespace UnrealGameSync
 		Font BuildFont;
 		Font SelectedBuildFont;
 		Font BadgeFont;
+		List<KeyValuePair<string, string>> BadgeNameAndGroupPairs = new List<KeyValuePair<string, string>>();
+		Dictionary<string, Size> BadgeLabelToSize = new Dictionary<string, Size>();
 
 		string OriginalExecutableFileName;
 
@@ -994,6 +996,48 @@ namespace UnrealGameSync
 
 		void UpdateBuildMetadata()
 		{
+			// Clear the badge size cache
+			BadgeLabelToSize.Clear();
+
+			// Reset the badge groups
+			Dictionary<string, string> BadgeNameToGroup = new Dictionary<string, string>();
+
+			// Read the group mappings from the config file
+			string GroupDefinitions;
+			if(TryGetProjectSetting(PerforceMonitor.LatestProjectConfigFile, "BadgeGroups", out GroupDefinitions))
+			{
+				string[] GroupDefinitionsArray = GroupDefinitions.Split('\n');
+				for(int Idx = 0; Idx < GroupDefinitionsArray.Length; Idx++)
+				{
+					string GroupName = String.Format("{0:0000}", Idx);
+					foreach(string BadgeName in GroupDefinitionsArray[Idx].Split(',').Select(x => x.Trim()))
+					{
+						BadgeNameToGroup[BadgeName] = GroupName;
+					}
+				}
+			}
+
+			// Add a dummy group for any other badges we have
+			foreach(ListViewItem Item in BuildList.Items)
+			{
+				EventSummary Summary = EventMonitor.GetSummaryForChange(((PerforceChangeSummary)Item.Tag).Number);
+				if(Summary != null)
+				{
+					foreach(BuildData Build in Summary.Builds)
+					{
+						string BadgeSlot = Build.BadgeName;
+						if(!BadgeNameToGroup.ContainsKey(BadgeSlot))
+						{
+							BadgeNameToGroup.Add(BadgeSlot, "XXXX");
+						}
+					}
+				}
+			}
+
+			// Sort the list of groups for display
+			BadgeNameAndGroupPairs = BadgeNameToGroup.OrderBy(x => x.Value).ThenBy(x => x.Key).ToList();
+
+			// Update everything else
 			ChangeNumberToArchivePath.Clear();
 			ChangeNumberToLayoutInfo.Clear();
 			BuildList.Invalidate();
@@ -1553,11 +1597,11 @@ namespace UnrealGameSync
 				Rectangle Bounds = new Rectangle(ListLocation.X + Badge.Offset, ListLocation.Y, Badge.Width, Badge.Height);
 				if(Badge.UniqueId != null && Badge.UniqueId == HoverBadgeUniqueId)
 				{
-					DrawBadge(Graphics, Bounds, Badge.Label, Color.FromArgb(Alpha, Badge.HoverBackgroundColor), bMergeLeft, bMergeRight);
+					DrawBadge(Graphics, Bounds, Badge.Label, Color.FromArgb((Alpha * Badge.HoverBackgroundColor.A) / 255, Badge.HoverBackgroundColor), bMergeLeft, bMergeRight);
 				}
 				else
 				{
-					DrawBadge(Graphics, Bounds, Badge.Label, Color.FromArgb(Alpha, Badge.BackgroundColor), bMergeLeft, bMergeRight);
+					DrawBadge(Graphics, Bounds, Badge.Label, Color.FromArgb((Alpha * Badge.BackgroundColor.A) / 255, Badge.BackgroundColor), bMergeLeft, bMergeRight);
 				}
 			}
 		}
@@ -1705,62 +1749,49 @@ namespace UnrealGameSync
 			EventSummary Summary = EventMonitor.GetSummaryForChange(ChangeNumber);
 			if(Summary != null && Summary.Builds.Count > 0)
 			{
-				Dictionary<string, string> BadgeNameToGroup = new Dictionary<string, string>();
-
-				// Read the group mappings
-				string GroupDefinitions;
-				if(TryGetProjectSetting(PerforceMonitor.LatestProjectConfigFile, "BadgeGroups", out GroupDefinitions))
-				{
-					string[] GroupDefinitionsArray = GroupDefinitions.Split('\n');
-					for(int Idx = 0; Idx < GroupDefinitionsArray.Length; Idx++)
-					{
-						string GroupName = String.Format("{0:0000}", Idx);
-						foreach(string BadgeName in GroupDefinitionsArray[Idx].Split(',').Select(x => x.Trim()))
-						{
-							BadgeNameToGroup[BadgeName] = GroupName;
-						}
-					}
-				}
-
-				// Build a list of builds with their group
-				List<Tuple<string, BuildData>> GroupedBuilds = new List<Tuple<string, BuildData>>();
+				// Create a lookup for build data for each badge name
+				Dictionary<string, BuildData> BadgeNameToBuildData = new Dictionary<string, BuildData>();
 				foreach(BuildData Build in Summary.Builds)
 				{
-					string Group;
-					if(!BadgeNameToGroup.TryGetValue(Build.BuildType, out Group))
-					{
-						Group = "XXXX";
-					}
-					GroupedBuilds.Add(Tuple.Create(Group, Build));
+					BadgeNameToBuildData[Build.BadgeName] = Build;
 				}
 
 				// Add all the badges, sorted by group
-				foreach(Tuple<string, BuildData> GroupedBuild in GroupedBuilds.OrderBy(x => x.Item1).ThenBy(x => x.Item2.BuildType))
+				foreach(KeyValuePair<string, string> BadgeNameAndGroup in BadgeNameAndGroupPairs)
 				{
-					BuildData Build = GroupedBuild.Item2;
+					string BadgeLabel = BadgeNameAndGroup.Key;
+					Color BadgeColor = Color.FromArgb(0, Color.White);
 
-					Color BadgeColor = Color.FromArgb(128, 192, 64);
-					if(Build.Result == BuildDataResult.Starting)
-					{
-						BadgeColor = Color.FromArgb(128, 192, 255);
-					}
-					else if(Build.Result == BuildDataResult.Warning)
-					{
-						BadgeColor = Color.FromArgb(255, 192, 0);
-					}
-					else if(Build.Result == BuildDataResult.Failure)
-					{
-						BadgeColor = Color.FromArgb(192, 64, 0);
-					}
-					else if(Build.Result == BuildDataResult.Skipped)
+					BuildData Build;
+					if(BadgeNameToBuildData.TryGetValue(BadgeNameAndGroup.Key, out Build))
 					{
-						BadgeColor = Color.FromArgb(192, 192, 192);
+						BadgeLabel = Build.BadgeLabel;
+						if(Build.Result == BuildDataResult.Starting)
+						{
+							BadgeColor = Color.FromArgb(128, 192, 255);
+						}
+						else if(Build.Result == BuildDataResult.Warning)
+						{
+							BadgeColor = Color.FromArgb(255, 192, 0);
+						}
+						else if(Build.Result == BuildDataResult.Failure)
+						{
+							BadgeColor = Color.FromArgb(192, 64, 0);
+						}
+						else if(Build.Result == BuildDataResult.Skipped)
+						{
+							BadgeColor = Color.FromArgb(192, 192, 192);
+						}
+						else
+						{
+							BadgeColor = Color.FromArgb(128, 192, 64);
+						}
 					}
 
-					Color HoverBadgeColor = Color.FromArgb(Math.Min(BadgeColor.R + 32, 255), Math.Min(BadgeColor.G + 32, 255), Math.Min(BadgeColor.B + 32, 255));
+					Color HoverBadgeColor = Color.FromArgb(BadgeColor.A, Math.Min(BadgeColor.R + 32, 255), Math.Min(BadgeColor.G + 32, 255), Math.Min(BadgeColor.B + 32, 255));
 
-					string UniqueId = String.Format("{0}:{1}", ChangeNumber, Build.BuildType);
-					Badges.Add(new BadgeInfo(Build.BuildType, GroupedBuild.Item1, UniqueId, BadgeColor, HoverBadgeColor, Build));
+					string UniqueId = String.Format("{0}:{1}", ChangeNumber, BadgeNameAndGroup.Key);
+					Badges.Add(new BadgeInfo(BadgeLabel, BadgeNameAndGroup.Value, UniqueId, BadgeColor, HoverBadgeColor, Build));
 				}
 			}
 
@@ -1786,29 +1817,39 @@ namespace UnrealGameSync
 
 		private void DrawBadge(Graphics Graphics, Rectangle BadgeRect, string BadgeText, Color BadgeColor, bool bMergeLeft, bool bMergeRight)
 		{
-			using (GraphicsPath Path = new GraphicsPath())
+			if(BadgeColor.A != 0)
 			{
-				Path.StartFigure();
-				Path.AddLine(BadgeRect.Left + (bMergeLeft? 1 : 0), BadgeRect.Top, BadgeRect.Left - (bMergeLeft? 1 : 0), BadgeRect.Bottom);
-				Path.AddLine(BadgeRect.Left - (bMergeLeft? 1 : 0), BadgeRect.Bottom, BadgeRect.Right - 1 - (bMergeRight? 1 : 0), BadgeRect.Bottom);
-				Path.AddLine(BadgeRect.Right - 1 - (bMergeRight? 1 : 0), BadgeRect.Bottom, BadgeRect.Right - 1 + (bMergeRight? 1 : 0), BadgeRect.Top);
-				Path.AddLine(BadgeRect.Right - 1 + (bMergeRight? 1 : 0), BadgeRect.Top, BadgeRect.Left + (bMergeLeft? 1 : 0), BadgeRect.Top);
-				Path.CloseFigure();
-
-				using(SolidBrush Brush = new SolidBrush(BadgeColor))
+				using (GraphicsPath Path = new GraphicsPath())
 				{
-					Graphics.FillPath(Brush, Path);
+					Path.StartFigure();
+					Path.AddLine(BadgeRect.Left + (bMergeLeft? 1 : 0), BadgeRect.Top, BadgeRect.Left - (bMergeLeft? 1 : 0), BadgeRect.Bottom);
+					Path.AddLine(BadgeRect.Left - (bMergeLeft? 1 : 0), BadgeRect.Bottom, BadgeRect.Right - 1 - (bMergeRight? 1 : 0), BadgeRect.Bottom);
+					Path.AddLine(BadgeRect.Right - 1 - (bMergeRight? 1 : 0), BadgeRect.Bottom, BadgeRect.Right - 1 + (bMergeRight? 1 : 0), BadgeRect.Top);
+					Path.AddLine(BadgeRect.Right - 1 + (bMergeRight? 1 : 0), BadgeRect.Top, BadgeRect.Left + (bMergeLeft? 1 : 0), BadgeRect.Top);
+					Path.CloseFigure();
+
+					using(SolidBrush Brush = new SolidBrush(BadgeColor))
+					{
+						Graphics.FillPath(Brush, Path);
+					}
 				}
-			}
 
-			TextRenderer.DrawText(Graphics, BadgeText, BadgeFont, BadgeRect, Color.White, TextFormatFlags.HorizontalCenter | TextFormatFlags.VerticalCenter | TextFormatFlags.SingleLine | TextFormatFlags.NoPrefix | TextFormatFlags.PreserveGraphicsClipping);
+				TextRenderer.DrawText(Graphics, BadgeText, BadgeFont, BadgeRect, Color.White, TextFormatFlags.HorizontalCenter | TextFormatFlags.VerticalCenter | TextFormatFlags.SingleLine | TextFormatFlags.NoPrefix | TextFormatFlags.PreserveGraphicsClipping);
+			}
 		}
 
 		private Size GetBadgeSize(string BadgeText)
 		{
-			Size LabelSize = TextRenderer.MeasureText(BadgeText, BadgeFont);
-			int BadgeHeight = BadgeFont.Height + 1;
-			return new Size(LabelSize.Width + BadgeHeight - 4, BadgeHeight);
+			Size BadgeSize;
+			if(!BadgeLabelToSize.TryGetValue(BadgeText, out BadgeSize))
+			{
+				Size LabelSize = TextRenderer.MeasureText(BadgeText, BadgeFont);
+				int BadgeHeight = BadgeFont.Height + 1;
+
+				BadgeSize = new Size(LabelSize.Width + BadgeHeight - 4, BadgeHeight);
+				BadgeLabelToSize[BadgeText] = BadgeSize;
+			}
+			return BadgeSize;
 		}
 
 		private static bool GetLastUpdateMessage(WorkspaceUpdateResult Result, string ResultMessage, out string Message)
@@ -2639,7 +2680,10 @@ namespace UnrealGameSync
 								if(BadgeBounds.Contains(Args.Location))
 								{
 									BuildData Build = (BuildData)Badge.UserData;
-									Process.Start(Build.Url);
+									if(Build != null)
+									{
+										Process.Start(Build.Url);
+									}
 									break;
 								}
 							}
diff --git a/Engine/Source/Programs/UnrealGameSync/UnrealGameSync/EventMonitor.cs b/Engine/Source/Programs/UnrealGameSync/UnrealGameSync/EventMonitor.cs
--- a/Engine/Source/Programs/UnrealGameSync/UnrealGameSync/EventMonitor.cs
+++ b/Engine/Source/Programs/UnrealGameSync/UnrealGameSync/EventMonitor.cs
@@ -80,6 +80,38 @@ namespace UnrealGameSync
 		{
 			get { return Result == BuildDataResult.Failure; }
 		}
+
+		public string BadgeName
+		{
+			get
+			{
+				int Idx = BuildType.IndexOf(':');
+				if(Idx == -1)
+				{
+					return BuildType;
+				}
+				else
+				{
+					return BuildType.Substring(0, Idx);
+				}
+			}
+		}
+
+		public string BadgeLabel
+		{
+			get
+			{
+				int Idx = BuildType.IndexOf(':');
+				if(Idx == -1)
+				{
+					return BuildType;
+				}
+				else
+				{
+					return BuildType.Substring(Idx + 1);
+				}
+			}
+		}
 	}
 
 	enum ReviewVerdict
