UGS: Fix display of multiline errors in the status panel.

#rnx

diff --git a/Engine/Source/Programs/UnrealGameSync/UnrealGameSync/Forms/MainWindow.cs b/Engine/Source/Programs/UnrealGameSync/UnrealGameSync/Forms/MainWindow.cs
--- a/Engine/Source/Programs/UnrealGameSync/UnrealGameSync/Forms/MainWindow.cs
+++ b/Engine/Source/Programs/UnrealGameSync/UnrealGameSync/Forms/MainWindow.cs
@@ -292,9 +292,13 @@ namespace UnrealGameSync
 
 			Lines.Add(new StatusLine(){ LineHeight = 0.5f });
 
-			StatusLine ErrorLine = new StatusLine();
-			ErrorLine.AddText(Message);
-			Lines.Add(ErrorLine);
+			foreach(string MessageLine in Message.Split('\n'))
+			{
+				StatusLine ErrorLine = new StatusLine();
+				ErrorLine.AddText(MessageLine);
+				ErrorLine.LineHeight = 0.8f;
+				Lines.Add(ErrorLine);
+			}
 
 			Lines.Add(new StatusLine(){ LineHeight = 0.5f });
 
