UGS: Compress build badges in a column if it shrinks below the size that they would be visible.

#rnx

diff --git a/Engine/Source/Programs/UnrealGameSync/UnrealGameSync/Controls/WorkspaceControl.cs b/Engine/Source/Programs/UnrealGameSync/UnrealGameSync/Controls/WorkspaceControl.cs
--- a/Engine/Source/Programs/UnrealGameSync/UnrealGameSync/Controls/WorkspaceControl.cs
+++ b/Engine/Source/Programs/UnrealGameSync/UnrealGameSync/Controls/WorkspaceControl.cs
@@ -85,6 +85,14 @@ namespace UnrealGameSync
 				this.UserData = UserData;
 			}
 
+			public BadgeInfo(BadgeInfo Other)
+				: this(Other.Label, Other.Group, Other.UniqueId, Other.BackgroundColor, Other.HoverBackgroundColor, Other.UserData)
+			{
+				this.Offset = Other.Offset;
+				this.Width = Other.Width;
+				this.Height = Other.Height;
+			}
+
 			public Rectangle GetBounds(Point ListLocation)
 			{
 				return new Rectangle(ListLocation.X + Offset, ListLocation.Y, Width, Height);
@@ -1444,8 +1452,12 @@ namespace UnrealGameSync
 				{
 				    e.Graphics.IntersectClip(e.Bounds);
 				    e.Graphics.SmoothingMode = SmoothingMode.HighQuality;
+
 					Point BuildsLocation = GetBadgeListLocation(Layout.BuildBadges, e.Bounds, HorizontalAlign.Center, VerticalAlignment.Middle);
-					DrawBadges(e.Graphics, BuildsLocation, Layout.BuildBadges, BadgeAlpha);
+					BuildsLocation.X = Math.Max(BuildsLocation.X, e.Bounds.Left);
+
+					List<BadgeInfo> FinalBadges = FinalLayoutBadges(Layout.BuildBadges, e.Bounds.Width);
+					DrawBadges(e.Graphics, BuildsLocation, FinalBadges, BadgeAlpha);
 				}
 			}
 			else if(e.ColumnIndex == StatusColumn.Index)
@@ -1586,6 +1598,29 @@ namespace UnrealGameSync
 			return new Point(X, Y);
 		}
 
+		private List<BadgeInfo> FinalLayoutBadges(List<BadgeInfo> Badges, int Width)
+		{
+			int TotalWidth = Badges[Badges.Count - 1].Offset + Badges[Badges.Count - 1].Width;
+			if(TotalWidth < Width)
+			{
+				return Badges;
+			}
+			else
+			{
+				List<BadgeInfo> FinalBadges = Badges.Where(x => x.BackgroundColor.A != 0).Select(x => new BadgeInfo(x)).ToList();
+
+				// Just clamp the edges
+				int Offset = Math.Max(Width, FinalBadges.Sum(x => x.Width));
+				for(int Idx = FinalBadges.Count - 1; Idx >= 0; Idx--)
+				{
+					Offset = Math.Min(Offset - FinalBadges[Idx].Width, FinalBadges[Idx].Offset);
+					FinalBadges[Idx].Offset = Offset;
+				}
+
+				return FinalBadges;
+			}
+		}
+
 		private void DrawBadges(Graphics Graphics, Point ListLocation, List<BadgeInfo> Badges, int Alpha)
 		{
 			for(int Idx = 0; Idx < Badges.Count; Idx++)
@@ -1796,6 +1831,7 @@ namespace UnrealGameSync
 			}
 
 			LayoutBadges(Badges);
+
 			return Badges;
 		}
 
@@ -2674,7 +2710,10 @@ namespace UnrealGameSync
 						if(LayoutInfo.BuildBadges.Count > 0)
 						{
 							Point BuildListLocation = GetBadgeListLocation(LayoutInfo.BuildBadges, HitTest.SubItem.Bounds, HorizontalAlign.Center, VerticalAlignment.Middle);
-							foreach(BadgeInfo Badge in LayoutInfo.BuildBadges)
+							BuildListLocation.X = Math.Max(BuildListLocation.X, HitTest.SubItem.Bounds.Left);
+
+							List<BadgeInfo> FinalBadges = FinalLayoutBadges(LayoutInfo.BuildBadges, HitTest.SubItem.Bounds.Width);
+							foreach(BadgeInfo Badge in FinalBadges)
 							{
 								Rectangle BadgeBounds = Badge.GetBounds(BuildListLocation);
 								if(BadgeBounds.Contains(Args.Location))
@@ -2838,8 +2877,11 @@ namespace UnrealGameSync
 					ChangeLayoutInfo LayoutInfo = GetChangeLayoutInfo((PerforceChangeSummary)HitTest.Item.Tag);
 					if(LayoutInfo.BuildBadges.Count > 0)
 					{
-						Point ListLocation = GetBadgeListLocation(LayoutInfo.BuildBadges, HitTest.SubItem.Bounds, HorizontalAlign.Center, VerticalAlignment.Middle);
-						NewHoverBadgeUniqueId = HitTestBadge(e.Location, LayoutInfo.BuildBadges, ListLocation);
+						Point BuildListLocation = GetBadgeListLocation(LayoutInfo.BuildBadges, HitTest.SubItem.Bounds, HorizontalAlign.Center, VerticalAlignment.Middle);
+						BuildListLocation.X = Math.Max(BuildListLocation.X, HitTest.SubItem.Bounds.Left);
+
+						List<BadgeInfo> FinalBadges = FinalLayoutBadges(LayoutInfo.BuildBadges, HitTest.SubItem.Bounds.Width);
+						NewHoverBadgeUniqueId = HitTestBadge(e.Location, FinalBadges, BuildListLocation);
 					}
 				}
 			}
