UGS: Added support for showing the latest completed build from a specific list of badges in the status panel. To declare a badge as one that should appear in the status panel rather than the CIS column, add it to the project's UnrealGameSync.ini in the project or [Default] section like so:

+ServiceBadges=RoboMerge

diff --git a/Engine/Source/Programs/UnrealGameSync/UnrealGameSync/Controls/StatusPanel.cs b/Engine/Source/Programs/UnrealGameSync/UnrealGameSync/Controls/StatusPanel.cs
--- a/Engine/Source/Programs/UnrealGameSync/UnrealGameSync/Controls/StatusPanel.cs
+++ b/Engine/Source/Programs/UnrealGameSync/UnrealGameSync/Controls/StatusPanel.cs
@@ -1,4 +1,4 @@
-﻿// Copyright 1998-2018 Epic Games, Inc. All Rights Reserved.
+// Copyright 1998-2018 Epic Games, Inc. All Rights Reserved.
 
 using System;
 using System.Collections.Generic;
@@ -15,10 +15,13 @@ namespace UnrealGameSync
 	class StatusElementResources
 	{
 		Dictionary<FontStyle, Font> FontCache = new Dictionary<FontStyle, Font>();
+		public readonly Font BadgeFont;
 
 		public StatusElementResources(Font BaseFont)
 		{
 			FontCache.Add(FontStyle.Regular, BaseFont);
+
+			BadgeFont = new Font(BaseFont.FontFamily, BaseFont.Size - 1.25f, FontStyle.Bold);
 		}
 
 		public void Dispose()
@@ -30,6 +33,8 @@ namespace UnrealGameSync
 					FontPair.Value.Dispose();
 				}
 			}
+
+			BadgeFont.Dispose();
 		}
 
 		public Font FindOrAddFont(FontStyle Style)
@@ -171,6 +176,57 @@ namespace UnrealGameSync
 		}
 	}
 
+	class BadgeStatusElement : StatusElement
+	{
+		string Name;
+		Color BackgroundColor;
+		Color HoverBackgroundColor;
+		Action ClickAction;
+
+		public BadgeStatusElement(string InName, Color InBackgroundColor, Action InClickAction)
+		{
+			Name = InName;
+			BackgroundColor = InBackgroundColor;
+			if(ClickAction == null)
+			{
+				HoverBackgroundColor = BackgroundColor;
+			}
+			else
+			{
+				HoverBackgroundColor = Color.FromArgb(Math.Min(BackgroundColor.R + 32, 255), Math.Min(BackgroundColor.G + 32, 255), Math.Min(BackgroundColor.B + 32, 255));
+			}
+			ClickAction = InClickAction;
+			if(ClickAction != null)
+			{
+				Cursor = Cursors.Hand;
+			}
+		}
+
+		public override void OnClick(Point Location)
+		{
+			if(ClickAction != null)
+			{
+				ClickAction();
+			}
+		}
+
+		public override Size Measure(Graphics Graphics, StatusElementResources Resources)
+		{
+			Size LabelSize = TextRenderer.MeasureText(Name, Resources.BadgeFont);
+			int BadgeHeight = Resources.BadgeFont.Height + 1;
+			return new Size(LabelSize.Width + BadgeHeight - 4, BadgeHeight);
+		}
+
+		public override void Draw(Graphics Graphics, StatusElementResources Resources)
+		{
+			using(SolidBrush Brush = new SolidBrush(bMouseOver? HoverBackgroundColor : BackgroundColor))
+			{
+				Graphics.FillRectangle(Brush, Bounds);
+			}
+			TextRenderer.DrawText(Graphics, Name, Resources.BadgeFont, Bounds, Color.White, TextFormatFlags.HorizontalCenter | TextFormatFlags.VerticalCenter | TextFormatFlags.SingleLine | TextFormatFlags.NoPrefix | TextFormatFlags.PreserveGraphicsClipping);
+		}
+	}
+
 	class ProgressBarStatusElement : StatusElement
 	{
 		float Progress;
@@ -247,6 +303,11 @@ namespace UnrealGameSync
 			Elements.Add(new LinkStatusElement(InText, InStyle, InLinkAction));
 		}
 
+		public void AddBadge(string InText, Color InBackgroundColor, Action InClickAction)
+		{
+			Elements.Add(new BadgeStatusElement(InText, InBackgroundColor, InClickAction));
+		}
+
 		public void AddProgressBar(float Progress)
 		{
 			Elements.Add(new ProgressBarStatusElement(Progress));
@@ -502,9 +563,9 @@ namespace UnrealGameSync
 			int LogoY = (BodyHeight - LogoHeight) / 2;
 
 			// Layout the caption. We may move the logo to make room for this.
+			LogoY -= Font.Height / 2;
 			if(Caption != null)
 			{
-				LogoY -= Font.Height / 2;
 				Caption.Layout(Graphics, Point.Empty, Resources);
 				int CaptionWidth = Caption.Bounds.Width;
 				Caption.Layout(Graphics, new Point(Math.Min(LogoX + (LogoWidth / 2) - (CaptionWidth / 2), DividerX - CaptionWidth), LogoY + LogoHeight), Resources);
diff --git a/Engine/Source/Programs/UnrealGameSync/UnrealGameSync/Controls/WorkspaceControl.cs b/Engine/Source/Programs/UnrealGameSync/UnrealGameSync/Controls/WorkspaceControl.cs
--- a/Engine/Source/Programs/UnrealGameSync/UnrealGameSync/Controls/WorkspaceControl.cs
+++ b/Engine/Source/Programs/UnrealGameSync/UnrealGameSync/Controls/WorkspaceControl.cs
@@ -206,6 +206,7 @@ namespace UnrealGameSync
 		Font BadgeFont;
 		List<KeyValuePair<string, string>> BadgeNameAndGroupPairs = new List<KeyValuePair<string, string>>();
 		Dictionary<string, Size> BadgeLabelToSize = new Dictionary<string, Size>();
+		List<KeyValuePair<string, BuildData>> ServiceBadges = new List<KeyValuePair<string, BuildData>>();
 
 		string OriginalExecutableFileName;
 
@@ -369,6 +370,7 @@ namespace UnrealGameSync
 			UpdateBuildSteps();
 			UpdateSyncActionCheckboxes();
 			UpdateStatusPanel();
+			UpdateServiceBadges();
 
 			if(CurrentChangeNumber != -1)
 			{
@@ -379,6 +381,23 @@ namespace UnrealGameSync
 			EventMonitor.Start();
 		}
 
+		private void UpdateServiceBadges()
+		{
+			string[] ServiceBadgeNames;
+			if(!TryGetProjectSetting(Workspace.ProjectConfigFile, "ServiceBadges", out ServiceBadgeNames))
+			{
+				ServiceBadgeNames = new string[0];
+			}
+
+			ServiceBadges.Clear();
+			foreach(string ServiceBadgeName in ServiceBadgeNames)
+			{
+				BuildData LatestBuild;
+				EventMonitor.TryGetLatestBuild(ServiceBadgeName, out LatestBuild);
+				ServiceBadges.Add(new KeyValuePair<string, BuildData>(ServiceBadgeName, LatestBuild));
+			}
+		}
+
 		protected override void OnLoad(EventArgs e)
 		{
 			base.OnLoad(e);
@@ -731,6 +750,7 @@ namespace UnrealGameSync
 			UpdateTimer.Stop();
 
 			UpdateSyncConfig(Workspace.CurrentChangeNumber, Workspace.CurrentSyncFilterHash);
+			UpdateServiceBadges();
 
 			if(Result == WorkspaceUpdateResult.Success && Context.Options.HasFlag(WorkspaceUpdateOptions.SyncSingleChange))
 			{
@@ -1042,6 +1062,12 @@ namespace UnrealGameSync
 				}
 			}
 
+			// Remove anything that's a service badge
+			foreach(KeyValuePair<string, BuildData> ServiceBadge in ServiceBadges)
+			{
+				BadgeNameToGroup.Remove(ServiceBadge.Key);
+			}
+
 			// Sort the list of groups for display
 			BadgeNameAndGroupPairs = BadgeNameToGroup.OrderBy(x => x.Value).ThenBy(x => x.Key).ToList();
 
@@ -1049,6 +1075,7 @@ namespace UnrealGameSync
 			ChangeNumberToArchivePath.Clear();
 			ChangeNumberToLayoutInfo.Clear();
 			BuildList.Invalidate();
+			UpdateServiceBadges();
 			UpdateStatusPanel();
 			UpdateBuildFailureNotification();
 		}
@@ -1757,6 +1784,21 @@ namespace UnrealGameSync
 			return false;
 		}
 
+		private bool TryGetProjectSetting(ConfigFile ProjectConfigFile, string Name, out string[] Values)
+		{
+			string ValueList;
+			if(TryGetProjectSetting(ProjectConfigFile, Name, out ValueList))
+			{
+				Values = ValueList.Split('\n').Select(x => x.Trim()).Where(x => x.Length > 0).ToArray();
+				return true;
+			}
+			else
+			{
+				Values = null;
+				return false;
+			}
+		}
+
 		private bool TryGetProjectSetting(ConfigFile ProjectConfigFile, string Name, string LegacyName, out string Value)
 		{
 			string NewValue;
@@ -1801,26 +1843,7 @@ namespace UnrealGameSync
 					if(BadgeNameToBuildData.TryGetValue(BadgeNameAndGroup.Key, out Build))
 					{
 						BadgeLabel = Build.BadgeLabel;
-						if(Build.Result == BuildDataResult.Starting)
-						{
-							BadgeColor = Color.FromArgb(128, 192, 255);
-						}
-						else if(Build.Result == BuildDataResult.Warning)
-						{
-							BadgeColor = Color.FromArgb(255, 192, 0);
-						}
-						else if(Build.Result == BuildDataResult.Failure)
-						{
-							BadgeColor = Color.FromArgb(192, 64, 0);
-						}
-						else if(Build.Result == BuildDataResult.Skipped)
-						{
-							BadgeColor = Color.FromArgb(192, 192, 192);
-						}
-						else
-						{
-							BadgeColor = Color.FromArgb(128, 192, 64);
-						}
+						BadgeColor = GetBuildBadgeColor(Build.Result);
 					}
 
 					Color HoverBadgeColor = Color.FromArgb(BadgeColor.A, Math.Min(BadgeColor.R + 32, 255), Math.Min(BadgeColor.G + 32, 255), Math.Min(BadgeColor.B + 32, 255));
@@ -1835,6 +1858,30 @@ namespace UnrealGameSync
 			return Badges;
 		}
 
+		private static Color GetBuildBadgeColor(BuildDataResult Result)
+		{
+			if(Result == BuildDataResult.Starting)
+			{
+				return Color.FromArgb(128, 192, 255);
+			}
+			else if(Result == BuildDataResult.Warning)
+			{
+				return Color.FromArgb(255, 192, 0);
+			}
+			else if(Result == BuildDataResult.Failure)
+			{
+				return Color.FromArgb(192, 64, 0);
+			}
+			else if(Result == BuildDataResult.Skipped)
+			{
+				return Color.FromArgb(192, 192, 192);
+			}
+			else
+			{
+				return Color.FromArgb(128, 192, 64);
+			}
+		}
+
 		private Rectangle GetSyncBadgeRectangle(Rectangle Bounds)
 		{
 			Size BadgeSize = GetBadgeSize("Sync");
@@ -2318,6 +2365,18 @@ namespace UnrealGameSync
 				ProgramsLine.AddLink("Visual Studio", FontStyle.Regular, () => { OpenSolution(); });
 				ProgramsLine.AddText("  |  ");
 				ProgramsLine.AddLink("Windows Explorer", FontStyle.Regular, () => { Process.Start("explorer.exe", String.Format("\"{0}\"", Path.GetDirectoryName(SelectedFileName))); });
+				foreach(KeyValuePair<string, BuildData> ServiceBadge in ServiceBadges)
+				{
+					ProgramsLine.AddText("  |  ");
+					if(ServiceBadge.Value == null)
+					{
+						ProgramsLine.AddBadge(ServiceBadge.Key, GetBuildBadgeColor(BuildDataResult.Skipped), null);
+					}
+					else
+					{
+						ProgramsLine.AddBadge(ServiceBadge.Key, GetBuildBadgeColor(ServiceBadge.Value.Result), () => { Process.Start(ServiceBadge.Value.Url); });
+					}
+				}
 				ProgramsLine.AddText("  |  ");
 				ProgramsLine.AddLink("More... \u25BE", FontStyle.Regular, (P, R) => { ShowActionsMenu(R); });
 				Lines.Add(ProgramsLine);
diff --git a/Engine/Source/Programs/UnrealGameSync/UnrealGameSync/EventMonitor.cs b/Engine/Source/Programs/UnrealGameSync/UnrealGameSync/EventMonitor.cs
--- a/Engine/Source/Programs/UnrealGameSync/UnrealGameSync/EventMonitor.cs
+++ b/Engine/Source/Programs/UnrealGameSync/UnrealGameSync/EventMonitor.cs
@@ -148,6 +148,7 @@ namespace UnrealGameSync
 		ConcurrentQueue<BuildData> IncomingBuilds = new ConcurrentQueue<BuildData>();
 		Dictionary<int, EventSummary> ChangeNumberToSummary = new Dictionary<int, EventSummary>();
 		Dictionary<string, EventData> UserNameToLastSyncEvent = new Dictionary<string, EventData>(StringComparer.InvariantCultureIgnoreCase);
+		Dictionary<string, BuildData> BadgeNameToLatestBuild = new Dictionary<string, BuildData>();
 		BoundedLogWriter LogWriter;
 		long LastEventId;
 		long LastCommentId;
@@ -342,6 +343,12 @@ namespace UnrealGameSync
 
 			Summary.Builds.Add(Build);
 			Summary.Verdict = GetVerdict(Summary.Reviews, Summary.Builds);
+
+			BuildData LatestBuild;
+			if(!BadgeNameToLatestBuild.TryGetValue(Build.BadgeName, out LatestBuild) || Build.ChangeNumber > LatestBuild.ChangeNumber || (Build.ChangeNumber == LatestBuild.ChangeNumber && Build.Id > LatestBuild.Id))
+			{
+				BadgeNameToLatestBuild[Build.BadgeName] = Build;
+			}
 		}
 
 		void ApplyCommentUpdate(CommentData Comment)
@@ -657,6 +664,11 @@ namespace UnrealGameSync
 			return Summary;
 		}
 
+		public bool TryGetLatestBuild(string BuildType, out BuildData BuildData)
+		{
+			return BadgeNameToLatestBuild.TryGetValue(BuildType, out BuildData);
+		}
+
 		public static bool IsReview(EventType Type)
 		{
 			return IsPositiveReview(Type) || IsNegativeReview(Type) || Type == EventType.Unknown;
