UGS: Fix tooltips not showing for changes that have description badges.

diff --git a/Engine/Source/Programs/UnrealGameSync/UnrealGameSync/Controls/WorkspaceControl.cs b/Engine/Source/Programs/UnrealGameSync/UnrealGameSync/Controls/WorkspaceControl.cs
--- a/Engine/Source/Programs/UnrealGameSync/UnrealGameSync/Controls/WorkspaceControl.cs
+++ b/Engine/Source/Programs/UnrealGameSync/UnrealGameSync/Controls/WorkspaceControl.cs
@@ -3245,7 +3245,7 @@ namespace UnrealGameSync
 					if(DescriptionBounds.Contains(ClientPoint))
 					{
 						ChangeLayoutInfo LayoutInfo = GetChangeLayoutInfo(Change);
-						if(LayoutInfo.DescriptionBadges.Count == 0 || ClientPoint.X < GetBadgeListLocation(LayoutInfo.BuildBadges, DescriptionBounds, HorizontalAlign.Right, VerticalAlignment.Middle).X - 2)
+						if(LayoutInfo.DescriptionBadges.Count == 0 || ClientPoint.X < GetBadgeListLocation(LayoutInfo.DescriptionBadges, DescriptionBounds, HorizontalAlign.Right, VerticalAlignment.Middle).X - 2)
 						{
 							BuildListToolTip.Show(Change.Description, BuildList, new Point(DescriptionBounds.Left, DescriptionBounds.Bottom + 2));
 							return;
