UGS: Allow adding additional columns to display particular badges by adding entries from the project config file. Example syntax:

+Columns=(Name="Desktop",MinWidth=50,DesiredWidth=100,Weight=3,Badges="Editor")
+Columns=(Name="Mobile",MinWidth=50,DesiredWidth=100,Weight=3,Badges="IOS,Android")

Same form can be used to control how default columns are displayed (though badge settings are ignored). Also allow PerforceMonitor to detect local changes to project config files and update settings automatically.

diff --git a/Engine/Source/Programs/UnrealGameSync/UnrealGameSync/Controls/WorkspaceControl.cs b/Engine/Source/Programs/UnrealGameSync/UnrealGameSync/Controls/WorkspaceControl.cs
--- a/Engine/Source/Programs/UnrealGameSync/UnrealGameSync/Controls/WorkspaceControl.cs
+++ b/Engine/Source/Programs/UnrealGameSync/UnrealGameSync/Controls/WorkspaceControl.cs
@@ -104,6 +104,7 @@ namespace UnrealGameSync
 			public List<BadgeInfo> DescriptionBadges;
 			public List<BadgeInfo> TypeBadges;
 			public List<BadgeInfo> BuildBadges;
+			public Dictionary<string, List<BadgeInfo>> CustomBadges;
 		}
 
 		static Rectangle GoodBuildIcon = new Rectangle(0, 0, 16, 16);
@@ -223,6 +224,8 @@ namespace UnrealGameSync
 		float[] ColumnWeights;
 		int[] MinColumnWidths;
 		int[] DesiredColumnWidths;
+		string LastColumnSettings;
+		List<ColumnHeader> CustomColumns;
 
 		bool bUpdateBuildListPosted;
 		bool bUpdateBuildMetadataPosted;
@@ -304,7 +307,7 @@ namespace UnrealGameSync
 
 			string ProjectLogBaseName = Path.Combine(DataFolder, String.Format("{0}@{1}", PerforceClient.ClientName, DetectSettings.BranchClientPath.Replace("//" + PerforceClient.ClientName + "/", "").Trim('/').Replace("/", "$")));
 
-			PerforceMonitor = new PerforceMonitor(PerforceClient, DetectSettings.BranchClientPath, DetectSettings.NewSelectedClientFileName, DetectSettings.NewSelectedProjectIdentifier, ProjectLogBaseName + ".p4.log", DetectSettings.bIsEnterpriseProject, DetectSettings.LatestProjectConfigFile);
+			PerforceMonitor = new PerforceMonitor(PerforceClient, DetectSettings.BranchClientPath, DetectSettings.NewSelectedClientFileName, DetectSettings.NewSelectedProjectIdentifier, ProjectLogBaseName + ".p4.log", DetectSettings.bIsEnterpriseProject, DetectSettings.LatestProjectConfigFile, DetectSettings.LocalConfigFiles);
 			PerforceMonitor.OnUpdate += UpdateBuildListCallback;
 			PerforceMonitor.OnUpdateMetadata += UpdateBuildMetadataCallback;
 			PerforceMonitor.OnStreamChange += StreamChangedCallback;
@@ -313,52 +316,7 @@ namespace UnrealGameSync
 			EventMonitor = new EventMonitor(ApiUrl, PerforceUtils.GetClientOrDepotDirectoryName(SelectedProjectIdentifier), DetectSettings.PerforceClient.UserName, ProjectLogBaseName + ".review.log");
 			EventMonitor.OnUpdatesReady += UpdateReviewsCallback;
 
-			ColumnWidths = new float[BuildList.Columns.Count];
-			for(int Idx = 0; Idx < BuildList.Columns.Count; Idx++)
-			{
-				ColumnWidths[Idx] = BuildList.Columns[Idx].Width;
-			}
-
-			MinColumnWidths = Enumerable.Repeat(32, BuildList.Columns.Count).ToArray();
-			MinColumnWidths[IconColumn.Index] = 50;
-			MinColumnWidths[TypeColumn.Index] = 100;
-			MinColumnWidths[TimeColumn.Index] = 75;
-			MinColumnWidths[ChangeColumn.Index] = 75;
-			MinColumnWidths[CISColumn.Index] = 200;
-
-			DesiredColumnWidths = Enumerable.Repeat(65536, BuildList.Columns.Count).ToArray();
-			DesiredColumnWidths[IconColumn.Index] = MinColumnWidths[IconColumn.Index];
-			DesiredColumnWidths[TypeColumn.Index] = MinColumnWidths[TypeColumn.Index];
-			DesiredColumnWidths[TimeColumn.Index] = MinColumnWidths[TimeColumn.Index];
-			DesiredColumnWidths[ChangeColumn.Index] = MinColumnWidths[ChangeColumn.Index];
-			DesiredColumnWidths[AuthorColumn.Index] = 120;
-			DesiredColumnWidths[CISColumn.Index] = 200;
-			DesiredColumnWidths[StatusColumn.Index] = 300;
-
-			ColumnWeights = Enumerable.Repeat(1.0f, BuildList.Columns.Count).ToArray();
-			ColumnWeights[IconColumn.Index] = 3.0f;
-			ColumnWeights[TypeColumn.Index] = 3.0f;
-			ColumnWeights[TimeColumn.Index] = 3.0f;
-			ColumnWeights[ChangeColumn.Index] = 3.0f;
-			ColumnWeights[DescriptionColumn.Index] = 1.25f;
-			ColumnWeights[CISColumn.Index] = 1.5f;
-
-			ConfigFile ProjectConfigFile = PerforceMonitor.LatestProjectConfigFile;
-			for(int Idx = 0; Idx < BuildList.Columns.Count; Idx++)
-			{
-				if(!String.IsNullOrEmpty(BuildList.Columns[Idx].Text))
-				{
-					string StringValue;
-					if(TryGetProjectSetting(ProjectConfigFile, String.Format("ColumnWidth_{0}", BuildList.Columns[Idx].Text), out StringValue))
-					{
-						int IntValue;
-						if(Int32.TryParse(StringValue, out IntValue))
-						{
-							DesiredColumnWidths[Idx] = IntValue;
-						}
-					}
-				}
-			}
+			UpdateColumnSettings();
 
 			string LogFileName = Path.Combine(DataFolder, ProjectLogBaseName + ".sync.log");
 			SyncLog.OpenFile(LogFileName);
@@ -381,10 +339,127 @@ namespace UnrealGameSync
 			EventMonitor.Start();
 		}
 
+		private void UpdateColumnSettings()
+		{
+			string NextColumnSettings;
+			TryGetProjectSetting(PerforceMonitor.LatestProjectConfigFile, "Columns", out NextColumnSettings);
+
+			if(CustomColumns == null || NextColumnSettings != LastColumnSettings)
+			{
+				LastColumnSettings = NextColumnSettings;
+
+				if(CustomColumns != null)
+				{
+					foreach(ColumnHeader CustomColumn in CustomColumns)
+					{
+						BuildList.Columns.Remove(CustomColumn);
+					}
+				}
+
+				Dictionary<string, ColumnHeader> NameToColumn = new Dictionary<string, ColumnHeader>();
+				foreach(ColumnHeader Column in BuildList.Columns)
+				{
+					NameToColumn[Column.Text] = Column;
+					Column.Tag = null;
+				}
+
+				CustomColumns = new List<ColumnHeader>();
+				if(NextColumnSettings != null)
+				{
+					foreach(string CustomColumn in NextColumnSettings.Split('\n'))
+					{
+						ConfigObject ColumnConfig = new ConfigObject(CustomColumn);
+
+						string Name = ColumnConfig.GetValue("Name", null);
+						if(Name != null)
+						{
+							ColumnHeader Column;
+							if(NameToColumn.TryGetValue(Name, out Column))
+							{
+								Column.Tag = ColumnConfig;
+							}
+							else
+							{
+								Column = new ColumnHeader();
+								Column.Text = Name;
+								Column.Tag = ColumnConfig;
+								BuildList.Columns.Insert(CISColumn.Index + CustomColumns.Count + 1, Column);
+
+								CustomColumns.Add(Column);
+							}
+						}
+					}
+				}
+
+				ColumnWidths = new float[BuildList.Columns.Count];
+				for(int Idx = 0; Idx < BuildList.Columns.Count; Idx++)
+				{
+					ColumnWidths[Idx] = BuildList.Columns[Idx].Width;
+				}
+
+				MinColumnWidths = Enumerable.Repeat(32, BuildList.Columns.Count).ToArray();
+				MinColumnWidths[IconColumn.Index] = 50;
+				MinColumnWidths[TypeColumn.Index] = 100;
+				MinColumnWidths[TimeColumn.Index] = 75;
+				MinColumnWidths[ChangeColumn.Index] = 75;
+				MinColumnWidths[CISColumn.Index] = 200;
+
+				DesiredColumnWidths = Enumerable.Repeat(65536, BuildList.Columns.Count).ToArray();
+				DesiredColumnWidths[IconColumn.Index] = MinColumnWidths[IconColumn.Index];
+				DesiredColumnWidths[TypeColumn.Index] = MinColumnWidths[TypeColumn.Index];
+				DesiredColumnWidths[TimeColumn.Index] = MinColumnWidths[TimeColumn.Index];
+				DesiredColumnWidths[ChangeColumn.Index] = MinColumnWidths[ChangeColumn.Index];
+				DesiredColumnWidths[AuthorColumn.Index] = 120;
+				DesiredColumnWidths[CISColumn.Index] = 200;
+				DesiredColumnWidths[StatusColumn.Index] = 300;
+
+				ColumnWeights = Enumerable.Repeat(1.0f, BuildList.Columns.Count).ToArray();
+				ColumnWeights[IconColumn.Index] = 3.0f;
+				ColumnWeights[TypeColumn.Index] = 3.0f;
+				ColumnWeights[TimeColumn.Index] = 3.0f;
+				ColumnWeights[ChangeColumn.Index] = 3.0f;
+				ColumnWeights[DescriptionColumn.Index] = 1.25f;
+				ColumnWeights[CISColumn.Index] = 1.5f;
+
+				foreach(ColumnHeader Column in BuildList.Columns)
+				{
+					ConfigObject ColumnConfig = (ConfigObject)Column.Tag;
+					if(ColumnConfig != null)
+					{
+						MinColumnWidths[Column.Index] = ColumnConfig.GetValue("MinWidth", MinColumnWidths[Column.Index]);
+						DesiredColumnWidths[Column.Index] = ColumnConfig.GetValue("DesiredWidth", DesiredColumnWidths[Column.Index]);
+						ColumnWeights[Column.Index] = ColumnConfig.GetValue("Weight", MinColumnWidths[Column.Index]);
+					}
+				}
+
+				ConfigFile ProjectConfigFile = PerforceMonitor.LatestProjectConfigFile;
+				for(int Idx = 0; Idx < BuildList.Columns.Count; Idx++)
+				{
+					if(!String.IsNullOrEmpty(BuildList.Columns[Idx].Text))
+					{
+						string StringValue;
+						if(TryGetProjectSetting(ProjectConfigFile, String.Format("ColumnWidth_{0}", BuildList.Columns[Idx].Text), out StringValue))
+						{
+							int IntValue;
+							if(Int32.TryParse(StringValue, out IntValue))
+							{
+								DesiredColumnWidths[Idx] = IntValue;
+							}
+						}
+					}
+				}
+
+				if(ColumnWidths != null)
+				{
+					ResizeColumns(ColumnWidths.Sum());
+				}
+			}
+		}
+
 		private void UpdateServiceBadges()
 		{
 			string[] ServiceBadgeNames;
-			if(!TryGetProjectSetting(Workspace.ProjectConfigFile, "ServiceBadges", out ServiceBadgeNames))
+			if(!TryGetProjectSetting(PerforceMonitor.LatestProjectConfigFile, "ServiceBadges", out ServiceBadgeNames))
 			{
 				ServiceBadgeNames = new string[0];
 			}
@@ -750,7 +825,6 @@ namespace UnrealGameSync
 			UpdateTimer.Stop();
 
 			UpdateSyncConfig(Workspace.CurrentChangeNumber, Workspace.CurrentSyncFilterHash);
-			UpdateServiceBadges();
 
 			if(Result == WorkspaceUpdateResult.Success && Context.Options.HasFlag(WorkspaceUpdateOptions.SyncSingleChange))
 			{
@@ -1024,6 +1098,9 @@ namespace UnrealGameSync
 
 		void UpdateBuildMetadata()
 		{
+			// Update the column settings first, since this may affect the badges we hide
+			UpdateColumnSettings();
+
 			// Clear the badge size cache
 			BadgeLabelToSize.Clear();
 
@@ -1068,6 +1145,16 @@ namespace UnrealGameSync
 				BadgeNameToGroup.Remove(ServiceBadge.Key);
 			}
 
+			// Remove everything that's in a custom column
+			foreach(ColumnHeader CustomColumn in CustomColumns)
+			{
+				ConfigObject ColumnConfig = (ConfigObject)CustomColumn.Tag;
+				foreach(string BadgeName in ColumnConfig.GetValue("Badges", "").Split(','))
+				{
+					BadgeNameToGroup.Remove(BadgeName);
+				}
+			}
+
 			// Sort the list of groups for display
 			BadgeNameAndGroupPairs = BadgeNameToGroup.OrderBy(x => x.Value).ThenBy(x => x.Key).ToList();
 
@@ -1308,7 +1395,10 @@ namespace UnrealGameSync
 
 				LayoutInfo.DescriptionBadges = CreateDescriptionBadges(Change);
 				LayoutInfo.TypeBadges = CreateTypeBadges(Change.Number);
-				LayoutInfo.BuildBadges = CreateBuildBadges(Change.Number);
+
+				EventSummary Summary = EventMonitor.GetSummaryForChange(Change.Number);
+				LayoutInfo.BuildBadges = CreateBuildBadges(Change.Number, Summary);
+				LayoutInfo.CustomBadges = CreateCustomBadges(Change.Number, Summary);
 
 				ChangeNumberToLayoutInfo.Add(Change.Number, LayoutInfo);
 			}
@@ -1565,7 +1655,21 @@ namespace UnrealGameSync
 			}
 			else
 			{
-				throw new NotImplementedException();
+				ColumnHeader Column = BuildList.Columns[e.ColumnIndex];
+				if(CustomColumns.Contains(Column))
+				{
+				    ChangeLayoutInfo Layout = GetChangeLayoutInfo(Change);
+
+				    List<BadgeInfo> Badges;
+				    if(Layout.CustomBadges.TryGetValue(Column.Text, out Badges) && Badges.Count > 0)
+				    {
+				        e.Graphics.IntersectClip(e.Bounds);
+				        e.Graphics.SmoothingMode = SmoothingMode.HighQuality;
+    
+					    Point BuildsLocation = GetBadgeListLocation(Badges, e.Bounds, HorizontalAlign.Center, VerticalAlignment.Middle);
+					    DrawBadges(e.Graphics, BuildsLocation, Badges, BadgeAlpha);
+				    }
+				}
 			}
 		}
 
@@ -1819,11 +1923,10 @@ namespace UnrealGameSync
 			return false;
 		}
 
-		private List<BadgeInfo> CreateBuildBadges(int ChangeNumber)
+		private List<BadgeInfo> CreateBuildBadges(int ChangeNumber, EventSummary Summary)
 		{
 			List<BadgeInfo> Badges = new List<BadgeInfo>();
 
-			EventSummary Summary = EventMonitor.GetSummaryForChange(ChangeNumber);
 			if(Summary != null && Summary.Builds.Count > 0)
 			{
 				// Create a lookup for build data for each badge name
@@ -1836,20 +1939,11 @@ namespace UnrealGameSync
 				// Add all the badges, sorted by group
 				foreach(KeyValuePair<string, string> BadgeNameAndGroup in BadgeNameAndGroupPairs)
 				{
-					string BadgeLabel = BadgeNameAndGroup.Key;
-					Color BadgeColor = Color.FromArgb(0, Color.White);
-
 					BuildData Build;
-					if(BadgeNameToBuildData.TryGetValue(BadgeNameAndGroup.Key, out Build))
-					{
-						BadgeLabel = Build.BadgeLabel;
-						BadgeColor = GetBuildBadgeColor(Build.Result);
-					}
-
-					Color HoverBadgeColor = Color.FromArgb(BadgeColor.A, Math.Min(BadgeColor.R + 32, 255), Math.Min(BadgeColor.G + 32, 255), Math.Min(BadgeColor.B + 32, 255));
+					BadgeNameToBuildData.TryGetValue(BadgeNameAndGroup.Key, out Build);
 
-					string UniqueId = String.Format("{0}:{1}", ChangeNumber, BadgeNameAndGroup.Key);
-					Badges.Add(new BadgeInfo(BadgeLabel, BadgeNameAndGroup.Value, UniqueId, BadgeColor, HoverBadgeColor, Build));
+					BadgeInfo Badge = CreateBadge(ChangeNumber, BadgeNameAndGroup.Key, BadgeNameAndGroup.Value, Build);
+					Badges.Add(Badge);
 				}
 			}
 
@@ -1858,6 +1952,51 @@ namespace UnrealGameSync
 			return Badges;
 		}
 
+		private BadgeInfo CreateBadge(int ChangeNumber, string BadgeName, string BadgeGroup, BuildData Build)
+		{
+			string BadgeLabel = BadgeName;
+			Color BadgeColor = Color.FromArgb(0, Color.White);
+
+			if(Build != null)
+			{
+				BadgeLabel = Build.BadgeLabel;
+				BadgeColor = GetBuildBadgeColor(Build.Result);
+			}
+
+			Color HoverBadgeColor = Color.FromArgb(BadgeColor.A, Math.Min(BadgeColor.R + 32, 255), Math.Min(BadgeColor.G + 32, 255), Math.Min(BadgeColor.B + 32, 255));
+
+			string UniqueId = String.Format("{0}:{1}", ChangeNumber, BadgeName);
+			return new BadgeInfo(BadgeLabel, BadgeGroup, UniqueId, BadgeColor, HoverBadgeColor, Build);
+		}
+
+		private Dictionary<string, List<BadgeInfo>> CreateCustomBadges(int ChangeNumber, EventSummary Summary)
+		{
+			Dictionary<string, List<BadgeInfo>> ColumnNameToBadges = new Dictionary<string, List<BadgeInfo>>();
+			if(Summary != null && Summary.Builds.Count > 0)
+			{
+				foreach(ColumnHeader CustomColumn in CustomColumns)
+				{
+					ConfigObject Config = (ConfigObject)CustomColumn.Tag;
+					if(Config != null)
+					{
+						List<BadgeInfo> Badges = new List<BadgeInfo>();
+
+						string[] BadgeNames = Config.GetValue("Badges", "").Split(new char[]{ ',' }, StringSplitOptions.RemoveEmptyEntries);
+						foreach(string BadgeName in BadgeNames)
+						{
+							BadgeInfo Badge = CreateBadge(ChangeNumber, BadgeName, "XXXX", Summary.Builds.FirstOrDefault(x => x.BadgeName == BadgeName));
+							Badges.Add(Badge);
+						}
+
+						LayoutBadges(Badges);
+
+						ColumnNameToBadges[CustomColumn.Text] = Badges;
+					}
+				}
+			}
+			return ColumnNameToBadges;
+		}
+
 		private static Color GetBuildBadgeColor(BuildDataResult Result)
 		{
 			if(Result == BuildDataResult.Starting)
@@ -2772,17 +2911,38 @@ namespace UnrealGameSync
 							BuildListLocation.X = Math.Max(BuildListLocation.X, HitTest.SubItem.Bounds.Left);
 
 							List<BadgeInfo> FinalBadges = FinalLayoutBadges(LayoutInfo.BuildBadges, HitTest.SubItem.Bounds.Width);
-							foreach(BadgeInfo Badge in FinalBadges)
+
+							BadgeInfo Badge = HitTestBadge(Args.Location, FinalBadges, BuildListLocation);
+							if(Badge != null)
 							{
-								Rectangle BadgeBounds = Badge.GetBounds(BuildListLocation);
-								if(BadgeBounds.Contains(Args.Location))
+								BuildData Build = (BuildData)Badge.UserData;
+								if(Build != null)
+								{
+									Process.Start(Build.Url);
+								}
+							}
+						}
+					}
+
+					foreach(ColumnHeader CustomColumn in CustomColumns)
+					{
+						if(CustomColumn.Index < HitTest.Item.SubItems.Count && HitTest.Item.SubItems[CustomColumn.Index] == HitTest.SubItem)
+						{
+							ChangeLayoutInfo LayoutInfo = GetChangeLayoutInfo((PerforceChangeSummary)HitTest.Item.Tag);
+
+							List<BadgeInfo> Badges;
+							if(LayoutInfo.CustomBadges.TryGetValue(CustomColumn.Text, out Badges) && Badges.Count > 0)
+							{
+								Point ListLocation = GetBadgeListLocation(Badges, HitTest.SubItem.Bounds, HorizontalAlign.Center, VerticalAlignment.Middle);
+
+								BadgeInfo Badge = HitTestBadge(Args.Location, Badges, ListLocation);
+								if(Badge != null)
 								{
 									BuildData Build = (BuildData)Badge.UserData;
 									if(Build != null)
 									{
 										Process.Start(Build.Url);
 									}
-									break;
 								}
 							}
 						}
@@ -2928,7 +3088,7 @@ namespace UnrealGameSync
 					if(LayoutInfo.DescriptionBadges.Count > 0)
 					{
 						Point ListLocation = GetBadgeListLocation(LayoutInfo.DescriptionBadges, HitTest.SubItem.Bounds, HorizontalAlign.Right, VerticalAlignment.Middle);
-						NewHoverBadgeUniqueId = HitTestBadge(e.Location, LayoutInfo.DescriptionBadges, ListLocation);
+						NewHoverBadgeUniqueId = HitTestBadge(e.Location, LayoutInfo.DescriptionBadges, ListLocation)?.UniqueId;
 					}
 				}
 				else if(ColumnIndex == CISColumn.Index)
@@ -2940,7 +3100,20 @@ namespace UnrealGameSync
 						BuildListLocation.X = Math.Max(BuildListLocation.X, HitTest.SubItem.Bounds.Left);
 
 						List<BadgeInfo> FinalBadges = FinalLayoutBadges(LayoutInfo.BuildBadges, HitTest.SubItem.Bounds.Width);
-						NewHoverBadgeUniqueId = HitTestBadge(e.Location, FinalBadges, BuildListLocation);
+						NewHoverBadgeUniqueId = HitTestBadge(e.Location, FinalBadges, BuildListLocation)?.UniqueId;
+					}
+				}
+				else if(CustomColumns.Contains(BuildList.Columns[ColumnIndex]))
+				{
+					ColumnHeader Column = BuildList.Columns[ColumnIndex];
+
+					ChangeLayoutInfo LayoutInfo = GetChangeLayoutInfo((PerforceChangeSummary)HitTest.Item.Tag);
+
+					List<BadgeInfo> Badges;
+					if(LayoutInfo.CustomBadges.TryGetValue(Column.Text, out Badges) && Badges.Count > 0)
+					{
+						Point ListLocation = GetBadgeListLocation(Badges, HitTest.SubItem.Bounds, HorizontalAlign.Center, VerticalAlignment.Middle);
+						NewHoverBadgeUniqueId = HitTestBadge(e.Location, Badges, ListLocation)?.UniqueId;
 					}
 				}
 			}
@@ -2962,14 +3135,14 @@ namespace UnrealGameSync
 			}
 		}
 
-		private string HitTestBadge(Point Location, List<BadgeInfo> BadgeList, Point ListLocation)
+		private BadgeInfo HitTestBadge(Point Location, List<BadgeInfo> BadgeList, Point ListLocation)
 		{
 			foreach(BadgeInfo Badge in BadgeList)
 			{
 				Rectangle BadgeBounds = Badge.GetBounds(ListLocation);
 				if(BadgeBounds.Contains(Location))
 				{
-					return Badge.UniqueId;
+					return Badge;
 				}
 			}
 			return null;
diff --git a/Engine/Source/Programs/UnrealGameSync/UnrealGameSync/DetectProjectSettingsTask.cs b/Engine/Source/Programs/UnrealGameSync/UnrealGameSync/DetectProjectSettingsTask.cs
--- a/Engine/Source/Programs/UnrealGameSync/UnrealGameSync/DetectProjectSettingsTask.cs
+++ b/Engine/Source/Programs/UnrealGameSync/UnrealGameSync/DetectProjectSettingsTask.cs
@@ -27,6 +27,7 @@ namespace UnrealGameSync
 		public string DataFolder;
 		public bool bIsEnterpriseProject;
 		public ConfigFile LatestProjectConfigFile;
+		public List<KeyValuePair<string, DateTime>> LocalConfigFiles;
 		TextWriter Log;
 
 		public DetectProjectSettingsTask(UserSelectedProjectSettings SelectedProject, string InDataFolder, TextWriter InLog)
@@ -317,7 +318,8 @@ namespace UnrealGameSync
 			}
 
 			// Read the initial config file
-			LatestProjectConfigFile = PerforceMonitor.ReadProjectConfigFile(PerforceClient, BranchClientPath, NewSelectedClientFileName, Log);
+			LocalConfigFiles = new List<KeyValuePair<string, DateTime>>();
+			LatestProjectConfigFile = PerforceMonitor.ReadProjectConfigFile(PerforceClient, BranchClientPath, NewSelectedClientFileName, LocalConfigFiles, Log);
 
 			// Succeed!
 			ErrorMessage = null;
diff --git a/Engine/Source/Programs/UnrealGameSync/UnrealGameSync/PerforceMonitor.cs b/Engine/Source/Programs/UnrealGameSync/UnrealGameSync/PerforceMonitor.cs
--- a/Engine/Source/Programs/UnrealGameSync/UnrealGameSync/PerforceMonitor.cs
+++ b/Engine/Source/Programs/UnrealGameSync/UnrealGameSync/PerforceMonitor.cs
@@ -45,13 +45,14 @@ namespace UnrealGameSync
 		BoundedLogWriter LogWriter;
 		bool bIsEnterpriseProject;
 		bool bDisposing;
+		List<KeyValuePair<string, DateTime>> LocalConfigFiles;
 
 		public event Action OnUpdate;
 		public event Action OnUpdateMetadata;
 		public event Action OnStreamChange;
 		public event Action OnLoginExpired;
 
-		public PerforceMonitor(PerforceConnection InPerforce, string InBranchClientPath, string InSelectedClientFileName, string InSelectedProjectIdentifier, string InLogPath, bool bInIsEnterpriseProject, ConfigFile InProjectConfigFile)
+		public PerforceMonitor(PerforceConnection InPerforce, string InBranchClientPath, string InSelectedClientFileName, string InSelectedProjectIdentifier, string InLogPath, bool bInIsEnterpriseProject, ConfigFile InProjectConfigFile, List<KeyValuePair<string, DateTime>> InLocalConfigFiles)
 		{
 			Perforce = InPerforce;
 			BranchClientPath = InBranchClientPath;
@@ -63,6 +64,7 @@ namespace UnrealGameSync
 			OtherStreamNames = new List<string>();
 			bIsEnterpriseProject = bInIsEnterpriseProject;
 			LatestProjectConfigFile = InProjectConfigFile;
+			LocalConfigFiles = InLocalConfigFiles;
 
 			LogWriter = new BoundedLogWriter(InLogPath);
 		}
@@ -419,6 +421,16 @@ namespace UnrealGameSync
 					OnUpdateMetadata();
 				}
 			}
+
+			if(LocalConfigFiles.Any(x => File.GetLastWriteTimeUtc(x.Key) != x.Value))
+			{
+				UpdateProjectConfigFile();
+				if(OnUpdateMetadata != null)
+				{
+					OnUpdateMetadata();
+				}
+			}
+
 			return true;
 		}
 
@@ -500,10 +512,11 @@ namespace UnrealGameSync
 
 		void UpdateProjectConfigFile()
 		{
-			LatestProjectConfigFile = ReadProjectConfigFile(Perforce, BranchClientPath, SelectedClientFileName, LogWriter);
+			LocalConfigFiles.Clear();
+			LatestProjectConfigFile = ReadProjectConfigFile(Perforce, BranchClientPath, SelectedClientFileName, LocalConfigFiles, LogWriter);
 		}
 
-		public static ConfigFile ReadProjectConfigFile(PerforceConnection Perforce, string BranchClientPath, string SelectedClientFileName, TextWriter Log)
+		public static ConfigFile ReadProjectConfigFile(PerforceConnection Perforce, string BranchClientPath, string SelectedClientFileName, List<KeyValuePair<string, DateTime>> LocalConfigFiles, TextWriter Log)
 		{
 			List<string> ConfigFilePaths = Utility.GetConfigFileLocations(BranchClientPath, SelectedClientFileName, '/');
 
@@ -524,6 +537,8 @@ namespace UnrealGameSync
 							string LocalFileName;
 							if(Perforce.ConvertToLocalPath(ConfigFilePath, out LocalFileName, Log))
 							{
+								DateTime LastModifiedTime = File.GetLastWriteTimeUtc(LocalFileName);
+								LocalConfigFiles.Add(new KeyValuePair<string, DateTime>(LocalFileName, LastModifiedTime));
 								Lines = File.ReadAllLines(LocalFileName).ToList();
 							}
 						}
