UGS: Fix exception when login is discovered to have expired during a workspace update.

diff --git a/Engine/Source/Programs/UnrealGameSync/UnrealGameSync/Controls/WorkspaceControl.cs b/Engine/Source/Programs/UnrealGameSync/UnrealGameSync/Controls/WorkspaceControl.cs
--- a/Engine/Source/Programs/UnrealGameSync/UnrealGameSync/Controls/WorkspaceControl.cs
+++ b/Engine/Source/Programs/UnrealGameSync/UnrealGameSync/Controls/WorkspaceControl.cs
@@ -845,11 +845,7 @@ namespace UnrealGameSync
 			WorkspaceSettings.LastBuiltChangeNumber = Workspace.LastBuiltChangeNumber;
 			Settings.Save();
 
-			if(Result == WorkspaceUpdateResult.FailedToSyncLoginExpired)
-			{
-				LoginExpired();
-			}
-			else if(Result == WorkspaceUpdateResult.FilesToResolve)
+			if(Result == WorkspaceUpdateResult.FilesToResolve)
 			{
 				MessageBox.Show("You have files to resolve after syncing your workspace. Please check P4.");
 			}
@@ -904,6 +900,12 @@ namespace UnrealGameSync
 			Refresh();
 			UpdateStatusPanel();
 			UpdateSyncActionCheckboxes();
+
+			// Do this last because it may result in the control being disposed
+			if(Result == WorkspaceUpdateResult.FailedToSyncLoginExpired)
+			{
+				LoginExpired();
+			}
 		}
 
 		void UpdateBuildListCallback()
