UGS: Fix an exception on shutdown due to the icon being hidden after it's already been disposed.

diff --git a/Engine/Source/Programs/UnrealGameSync/UnrealGameSync/ProgramApplicationContext.cs b/Engine/Source/Programs/UnrealGameSync/UnrealGameSync/ProgramApplicationContext.cs
--- a/Engine/Source/Programs/UnrealGameSync/UnrealGameSync/ProgramApplicationContext.cs
+++ b/Engine/Source/Programs/UnrealGameSync/UnrealGameSync/ProgramApplicationContext.cs
@@ -241,6 +241,12 @@ namespace UnrealGameSync
 				Components = null;
 			}
 
+			if(NotifyIcon != null)
+			{
+				NotifyIcon.Dispose();
+				NotifyIcon = null;
+			}
+
 			if(Log != null)
 			{
 				Log.Dispose();
@@ -348,7 +354,10 @@ namespace UnrealGameSync
 		{
 			base.ExitThreadCore();
 
-			NotifyIcon.Visible = false;
+			if(NotifyIcon != null)
+			{
+				NotifyIcon.Visible = false;
+			}
 		}
 	}
 }
