UGS: Add an unhandled exception filter which sends the exception data to the backend.

diff --git a/Engine/Source/Programs/UnrealGameSync/UnrealGameSync/Program.cs b/Engine/Source/Programs/UnrealGameSync/UnrealGameSync/Program.cs
--- a/Engine/Source/Programs/UnrealGameSync/UnrealGameSync/Program.cs
+++ b/Engine/Source/Programs/UnrealGameSync/UnrealGameSync/Program.cs
@@ -94,31 +94,33 @@ namespace UnrealGameSync
 
 			using(TelemetryWriter Telemetry = new TelemetryWriter(ApiUrl, Path.Combine(DataFolder, "Telemetry.log")))
 			{
-				try
+				AppDomain.CurrentDomain.UnhandledException += CurrentDomain_UnhandledException;
+				using(UpdateMonitor UpdateMonitor = new UpdateMonitor(new PerforceConnection(UserName, null, ServerAndPort), UpdatePath))
 				{
-					using(UpdateMonitor UpdateMonitor = new UpdateMonitor(new PerforceConnection(UserName, null, ServerAndPort), UpdatePath))
+					ProgramApplicationContext Context = new ProgramApplicationContext(UpdateMonitor, ApiUrl, DataFolder, ActivateEvent, bRestoreState, UpdateSpawn, ProjectFileName, bUnstable);
+					Application.Run(Context);
+
+					if(UpdateMonitor.IsUpdateAvailable && UpdateSpawn != null)
 					{
-						ProgramApplicationContext Context = new ProgramApplicationContext(UpdateMonitor, ApiUrl, DataFolder, ActivateEvent, bRestoreState, UpdateSpawn, ProjectFileName, bUnstable);
-						Application.Run(Context);
-
-						if(UpdateMonitor.IsUpdateAvailable && UpdateSpawn != null)
-						{
-							InstanceMutex.Close();
-							Utility.SpawnProcess(UpdateSpawn, "-restorestate" + (bUnstable? " -unstable" : ""));
-						}
+						InstanceMutex.Close();
+						Utility.SpawnProcess(UpdateSpawn, "-restorestate" + (bUnstable? " -unstable" : ""));
 					}
 				}
-				catch(Exception Ex)
+			}
+		}
+
+		private static void CurrentDomain_UnhandledException(object Sender, UnhandledExceptionEventArgs Args)
+		{
+			Exception Ex = Args.ExceptionObject as Exception;
+			if(Ex != null)
+			{
+				StringBuilder ExceptionTrace = new StringBuilder(Ex.ToString());
+				for(Exception InnerEx = Ex.InnerException; InnerEx != null; InnerEx = InnerEx.InnerException)
 				{
-					StringBuilder ExceptionTrace = new StringBuilder(Ex.ToString());
-					for(Exception InnerEx = Ex.InnerException; InnerEx != null; InnerEx = InnerEx.InnerException)
-					{
-						ExceptionTrace.Append("\nInner Exception:\n");
-						ExceptionTrace.Append(InnerEx.ToString());
-					}
-					TelemetryWriter.Enqueue(TelemetryErrorType.Crash, ExceptionTrace.ToString(), null, DateTime.Now);
-					MessageBox.Show(String.Format("UnrealGameSync has crashed.\n\n{0}", ExceptionTrace.ToString()));
+					ExceptionTrace.Append("\nInner Exception:\n");
+					ExceptionTrace.Append(InnerEx.ToString());
 				}
+				TelemetryWriter.Enqueue(TelemetryErrorType.Crash, ExceptionTrace.ToString(), null, DateTime.Now);
 			}
 		}
 
diff --git a/Engine/Source/Programs/UnrealGameSync/UnrealGameSync/Telemetry.cs b/Engine/Source/Programs/UnrealGameSync/UnrealGameSync/Telemetry.cs
--- a/Engine/Source/Programs/UnrealGameSync/UnrealGameSync/Telemetry.cs
+++ b/Engine/Source/Programs/UnrealGameSync/UnrealGameSync/Telemetry.cs
@@ -1,4 +1,4 @@
-﻿// Copyright 1998-2018 Epic Games, Inc. All Rights Reserved.
+// Copyright 1998-2018 Epic Games, Inc. All Rights Reserved.
 
 using System;
 using System.Collections.Concurrent;
@@ -109,7 +109,7 @@ namespace UnrealGameSync
 			if(WorkerThread != null)
 			{
 				RefreshEvent.Set();
-				if(!WorkerThread.Join(100))
+				if(!WorkerThread.Join(10 * 1000))
 				{
 					WorkerThread.Abort();
 					WorkerThread.Join();
