UGS: Reduce the number of characters displayed for build badges based on the available space.

diff --git a/Engine/Source/Programs/UnrealGameSync/UnrealGameSync/Controls/WorkspaceControl.cs b/Engine/Source/Programs/UnrealGameSync/UnrealGameSync/Controls/WorkspaceControl.cs
--- a/Engine/Source/Programs/UnrealGameSync/UnrealGameSync/Controls/WorkspaceControl.cs
+++ b/Engine/Source/Programs/UnrealGameSync/UnrealGameSync/Controls/WorkspaceControl.cs
@@ -69,6 +69,7 @@ namespace UnrealGameSync
 			public Color BackgroundColor;
 			public Color HoverBackgroundColor;
 			public object UserData;
+			public string ToolTip;
 
 			public BadgeInfo(string Label, string Group, Color BadgeColor)
 				: this(Label, Group, null, BadgeColor, BadgeColor)
@@ -226,6 +227,7 @@ namespace UnrealGameSync
 		int[] DesiredColumnWidths;
 		string LastColumnSettings;
 		List<ColumnHeader> CustomColumns;
+		int MaxBuildBadgeChars;
 
 		bool bUpdateBuildListPosted;
 		bool bUpdateBuildMetadataPosted;
@@ -1160,6 +1162,9 @@ namespace UnrealGameSync
 			// Sort the list of groups for display
 			BadgeNameAndGroupPairs = BadgeNameToGroup.OrderBy(x => x.Value).ThenBy(x => x.Key).ToList();
 
+			// Figure out whether to show smaller badges due to limited space
+			UpdateMaxBuildBadgeChars();
+
 			// Update everything else
 			ChangeNumberToArchivePath.Clear();
 			ChangeNumberToLayoutInfo.Clear();
@@ -1169,6 +1174,60 @@ namespace UnrealGameSync
 			UpdateBuildFailureNotification();
 		}
 
+		void UpdateMaxBuildBadgeChars()
+		{
+			int NewMaxBuildBadgeChars;
+			if(GetBuildBadgeStripWidth(-1) < CISColumn.Width)
+			{
+				NewMaxBuildBadgeChars = -1;
+			}
+			else if(GetBuildBadgeStripWidth(3) < CISColumn.Width)
+			{
+				NewMaxBuildBadgeChars = 3;
+			}
+			else if(GetBuildBadgeStripWidth(2) < CISColumn.Width)
+			{
+				NewMaxBuildBadgeChars = 2;
+			}
+			else
+			{
+				NewMaxBuildBadgeChars = 1;
+			}
+
+			if(NewMaxBuildBadgeChars != MaxBuildBadgeChars)
+			{
+				ChangeNumberToLayoutInfo.Clear();
+				BuildList.Invalidate();
+				MaxBuildBadgeChars = NewMaxBuildBadgeChars;
+			}
+		}
+
+		int GetBuildBadgeStripWidth(int MaxNumChars)
+		{
+			// Create dummy badges for each badge name
+			List<BadgeInfo> DummyBadges = new List<BadgeInfo>();
+			foreach(KeyValuePair<string, string> BadgeNameAndGroupPair in BadgeNameAndGroupPairs)
+			{
+				string BadgeName = BadgeNameAndGroupPair.Key;
+				if(MaxNumChars != -1 && BadgeName.Length > MaxNumChars)
+				{
+					BadgeName = BadgeName.Substring(0, MaxNumChars);
+				}
+
+				BadgeInfo DummyBadge = CreateBadge(-1, BadgeName, BadgeNameAndGroupPair.Value, null);
+				DummyBadges.Add(DummyBadge);
+			}
+
+			// Check if they fit within the column
+			int Width = 0;
+			if(DummyBadges.Count > 0)
+			{
+				LayoutBadges(DummyBadges);
+				Width = DummyBadges[DummyBadges.Count - 1].Offset + DummyBadges[DummyBadges.Count - 1].Width;
+			}
+			return Width;
+		}
+
 		bool ShouldIncludeInReviewedList(int ChangeNumber)
 		{
 			if(PromotedChangeNumbers.Contains(ChangeNumber))
@@ -1575,8 +1634,7 @@ namespace UnrealGameSync
 					Point BuildsLocation = GetBadgeListLocation(Layout.BuildBadges, e.Bounds, HorizontalAlign.Center, VerticalAlignment.Middle);
 					BuildsLocation.X = Math.Max(BuildsLocation.X, e.Bounds.Left);
 
-					List<BadgeInfo> FinalBadges = FinalLayoutBadges(Layout.BuildBadges, e.Bounds.Width);
-					DrawBadges(e.Graphics, BuildsLocation, FinalBadges, BadgeAlpha);
+					DrawBadges(e.Graphics, BuildsLocation, Layout.BuildBadges, BadgeAlpha);
 				}
 			}
 			else if(e.ColumnIndex == StatusColumn.Index)
@@ -1731,29 +1789,6 @@ namespace UnrealGameSync
 			return new Point(X, Y);
 		}
 
-		private List<BadgeInfo> FinalLayoutBadges(List<BadgeInfo> Badges, int Width)
-		{
-			int TotalWidth = Badges[Badges.Count - 1].Offset + Badges[Badges.Count - 1].Width;
-			if(TotalWidth < Width)
-			{
-				return Badges;
-			}
-			else
-			{
-				List<BadgeInfo> FinalBadges = Badges.Where(x => x.BackgroundColor.A != 0).Select(x => new BadgeInfo(x)).ToList();
-
-				// Just clamp the edges
-				int Offset = Math.Max(Width, FinalBadges.Sum(x => x.Width));
-				for(int Idx = FinalBadges.Count - 1; Idx >= 0; Idx--)
-				{
-					Offset = Math.Min(Offset - FinalBadges[Idx].Width, FinalBadges[Idx].Offset);
-					FinalBadges[Idx].Offset = Offset;
-				}
-
-				return FinalBadges;
-			}
-		}
-
 		private void DrawBadges(Graphics Graphics, Point ListLocation, List<BadgeInfo> Badges, int Alpha)
 		{
 			for(int Idx = 0; Idx < Badges.Count; Idx++)
@@ -1945,12 +1980,17 @@ namespace UnrealGameSync
 					BadgeNameToBuildData.TryGetValue(BadgeNameAndGroup.Key, out Build);
 
 					BadgeInfo Badge = CreateBadge(ChangeNumber, BadgeNameAndGroup.Key, BadgeNameAndGroup.Value, Build);
+					if(MaxBuildBadgeChars != -1 && Badge.Label.Length > MaxBuildBadgeChars)
+					{
+						Badge.ToolTip = Badge.Label;
+						Badge.Label = Badge.Label.Substring(0, MaxBuildBadgeChars);
+					}
 					Badges.Add(Badge);
 				}
 			}
 
+			// Layout the badges
 			LayoutBadges(Badges);
-
 			return Badges;
 		}
 
@@ -2912,9 +2952,7 @@ namespace UnrealGameSync
 							Point BuildListLocation = GetBadgeListLocation(LayoutInfo.BuildBadges, HitTest.SubItem.Bounds, HorizontalAlign.Center, VerticalAlignment.Middle);
 							BuildListLocation.X = Math.Max(BuildListLocation.X, HitTest.SubItem.Bounds.Left);
 
-							List<BadgeInfo> FinalBadges = FinalLayoutBadges(LayoutInfo.BuildBadges, HitTest.SubItem.Bounds.Width);
-
-							BadgeInfo Badge = HitTestBadge(Args.Location, FinalBadges, BuildListLocation);
+							BadgeInfo Badge = HitTestBadge(Args.Location, LayoutInfo.BuildBadges, BuildListLocation);
 							if(Badge != null)
 							{
 								BuildData Build = (BuildData)Badge.UserData;
@@ -3101,8 +3139,13 @@ namespace UnrealGameSync
 						Point BuildListLocation = GetBadgeListLocation(LayoutInfo.BuildBadges, HitTest.SubItem.Bounds, HorizontalAlign.Center, VerticalAlignment.Middle);
 						BuildListLocation.X = Math.Max(BuildListLocation.X, HitTest.SubItem.Bounds.Left);
 
-						List<BadgeInfo> FinalBadges = FinalLayoutBadges(LayoutInfo.BuildBadges, HitTest.SubItem.Bounds.Width);
-						NewHoverBadgeUniqueId = HitTestBadge(e.Location, FinalBadges, BuildListLocation)?.UniqueId;
+						BadgeInfo Badge = HitTestBadge(e.Location, LayoutInfo.BuildBadges, BuildListLocation);
+						NewHoverBadgeUniqueId = (Badge != null)? Badge.UniqueId : null;
+
+						if(Badge != null && Badge.ToolTip != null && HoverBadgeUniqueId != NewHoverBadgeUniqueId)
+						{
+							BuildListToolTip.Show(Badge.ToolTip, BuildList, new Point(BuildListLocation.X + Badge.Offset, HitTest.Item.Bounds.Bottom + 2));
+						}
 					}
 				}
 				else if(CustomColumns.Contains(BuildList.Columns[ColumnIndex]))
@@ -3416,6 +3459,12 @@ namespace UnrealGameSync
 				{
 					PerforceChangeSummary Change = (PerforceChangeSummary)Args.Item.Tag;
 
+					Rectangle CISBounds = Args.Item.SubItems[CISColumn.Index].Bounds;
+					if(CISBounds.Contains(ClientPoint))
+					{
+						return;
+					}
+
 					Rectangle DescriptionBounds = Args.Item.SubItems[DescriptionColumn.Index].Bounds;
 					if(DescriptionBounds.Contains(ClientPoint))
 					{
@@ -4162,6 +4211,11 @@ namespace UnrealGameSync
 			UpdateNumRequestedBuilds(true);
 		}
 
+		private void BuildList_ColumnWidthChanging(object sender, ColumnWidthChangingEventArgs e)
+		{
+			UpdateMaxBuildBadgeChars();
+		}
+
 		private void BuildList_ColumnWidthChanged(object sender, ColumnWidthChangedEventArgs e)
 		{
 			if(ColumnWidths != null && MinColumnWidths != null)
@@ -4174,6 +4228,7 @@ namespace UnrealGameSync
 				}
 				ColumnWidths[e.ColumnIndex] = NewWidth;
 			}
+			UpdateMaxBuildBadgeChars();
 		}
 
 		private void BuildList_Resize(object sender, EventArgs e)
