UGS: Allow syncing projects where source code is not available (and various version files don't exist).

diff --git a/Engine/Source/Programs/UnrealGameSync/UnrealGameSync/Perforce.cs b/Engine/Source/Programs/UnrealGameSync/UnrealGameSync/Perforce.cs
--- a/Engine/Source/Programs/UnrealGameSync/UnrealGameSync/Perforce.cs
+++ b/Engine/Source/Programs/UnrealGameSync/UnrealGameSync/Perforce.cs
@@ -284,7 +284,7 @@ namespace UnrealGameSync
 		}
 	}
 
-	[DebuggerDisplay("{Path}")]
+	[DebuggerDisplay("{DepotPath}")]
 	class PerforceWhereRecord
 	{
 		public string LocalPath;
diff --git a/Engine/Source/Programs/UnrealGameSync/UnrealGameSync/Workspace.cs b/Engine/Source/Programs/UnrealGameSync/UnrealGameSync/Workspace.cs
--- a/Engine/Source/Programs/UnrealGameSync/UnrealGameSync/Workspace.cs
+++ b/Engine/Source/Programs/UnrealGameSync/UnrealGameSync/Workspace.cs
@@ -1149,19 +1149,27 @@ namespace UnrealGameSync
 			Log.WriteLine("p4>   {0} {1}", Record.Action, Record.ClientPath);
 		}
 
-		bool UpdateVersionFile(string LocalPath, Dictionary<string, string> VersionStrings, int ChangeNumber)
+		bool UpdateVersionFile(string ClientPath, Dictionary<string, string> VersionStrings, int ChangeNumber)
 		{
-			PerforceWhereRecord WhereRecord;
-			if(!Perforce.Where(LocalPath, out WhereRecord, Log))
+			List<PerforceFileRecord> Records;
+			if(!Perforce.Stat(ClientPath, out Records, Log))
 			{
-				Log.WriteLine("P4 where failed for {0}", LocalPath);
+				Log.WriteLine("Failed to query records for {0}", ClientPath);
 				return false;
 			}
+			if(Records.Count == 0)
+			{
+				Log.WriteLine("Ignoring {0}; not found on server.", ClientPath);
+				return true;
+			}
+
+			string LocalPath = Records[0].ClientPath; // Actually a filesystem path
+			string DepotPath = Records[0].DepotPath;
 
 			List<string> Lines;
-			if(!Perforce.Print(String.Format("{0}@{1}", WhereRecord.DepotPath, ChangeNumber), out Lines, Log))
+			if(!Perforce.Print(String.Format("{0}@{1}", DepotPath, ChangeNumber), out Lines, Log))
 			{
-				Log.WriteLine("Couldn't get default contents of {0}", WhereRecord.DepotPath);
+				Log.WriteLine("Couldn't get default contents of {0}", DepotPath);
 				return false;
 			}
 
@@ -1179,33 +1187,33 @@ namespace UnrealGameSync
 				Writer.WriteLine(NewLine);
 			}
 
-			return WriteVersionFile(WhereRecord, Writer.ToString());
+			return WriteVersionFile(LocalPath, DepotPath, Writer.ToString());
 		}
 
-		bool WriteVersionFile(PerforceWhereRecord WhereRecord, string NewText)
+		bool WriteVersionFile(string LocalPath, string DepotPath, string NewText)
 		{
 			try
 			{
-				if(File.Exists(WhereRecord.LocalPath) && File.ReadAllText(WhereRecord.LocalPath) == NewText)
+				if(File.Exists(LocalPath) && File.ReadAllText(LocalPath) == NewText)
 				{
-					Log.WriteLine("Ignored {0}; contents haven't changed", WhereRecord.LocalPath);
+					Log.WriteLine("Ignored {0}; contents haven't changed", LocalPath);
 				}
 				else
 				{
-					Directory.CreateDirectory(Path.GetDirectoryName(WhereRecord.LocalPath));
-					Utility.ForceDeleteFile(WhereRecord.LocalPath);
-					if(WhereRecord.DepotPath != null)
+					Directory.CreateDirectory(Path.GetDirectoryName(LocalPath));
+					Utility.ForceDeleteFile(LocalPath);
+					if(DepotPath != null)
 					{
-						Perforce.Sync(WhereRecord.DepotPath + "#0", Log);
+						Perforce.Sync(DepotPath + "#0", Log);
 					}
-					File.WriteAllText(WhereRecord.LocalPath, NewText);
-					Log.WriteLine("Written {0}", WhereRecord.LocalPath);
+					File.WriteAllText(LocalPath, NewText);
+					Log.WriteLine("Written {0}", LocalPath);
 				}
 				return true;
 			}
 			catch(Exception Ex)
 			{
-				Log.WriteException(Ex, "Failed to write to {0}.", WhereRecord.LocalPath);
+				Log.WriteException(Ex, "Failed to write to {0}.", LocalPath);
 				return false;
 			}
 		}
