UGS: Remove requirement for UE4Editor.target.cs to be visible in the depot in order to open a project.

diff --git a/Engine/Source/Programs/UnrealGameSync/UnrealGameSync/Controls/WorkspaceControl.cs b/Engine/Source/Programs/UnrealGameSync/UnrealGameSync/Controls/WorkspaceControl.cs
--- a/Engine/Source/Programs/UnrealGameSync/UnrealGameSync/Controls/WorkspaceControl.cs
+++ b/Engine/Source/Programs/UnrealGameSync/UnrealGameSync/Controls/WorkspaceControl.cs
@@ -291,7 +291,7 @@ namespace UnrealGameSync
 			ServerTimeZone = DetectSettings.ServerTimeZone;
 
 			// Update the branch directory
-			BranchDirectoryName = Path.GetFullPath(Path.Combine(Path.GetDirectoryName(DetectSettings.BaseEditorTargetPath), "..", ".."));
+			BranchDirectoryName = DetectSettings.BranchDirectoryName;
 
 			// Check if we've the project we've got open in this workspace is the one we're actually synced to
 			int CurrentChangeNumber = -1;
diff --git a/Engine/Source/Programs/UnrealGameSync/UnrealGameSync/DetectProjectSettingsTask.cs b/Engine/Source/Programs/UnrealGameSync/UnrealGameSync/DetectProjectSettingsTask.cs
--- a/Engine/Source/Programs/UnrealGameSync/UnrealGameSync/DetectProjectSettingsTask.cs
+++ b/Engine/Source/Programs/UnrealGameSync/UnrealGameSync/DetectProjectSettingsTask.cs
@@ -19,7 +19,7 @@ namespace UnrealGameSync
 		public string NewSelectedProjectIdentifier;
 		public string NewProjectEditorTarget;
 		public string BranchClientPath;
-		public string BaseEditorTargetPath;
+		public string BranchDirectoryName;
 		public string NewSelectedClientFileName;
 		public string StreamName;
 		public Image ProjectLogo;
@@ -226,7 +226,7 @@ namespace UnrealGameSync
 				if(NewSelectedClientFileName[EndIdx] == '/')
 				{
 					bool bFileExists;
-					if(PerforceClient.FileExists(NewSelectedClientFileName.Substring(0, EndIdx) + "/Engine/Source/UE4Editor.target.cs", out bFileExists, Log) && bFileExists)
+					if(PerforceClient.FileExists(NewSelectedClientFileName.Substring(0, EndIdx) + "/Engine/Build/Build.version", out bFileExists, Log) && bFileExists)
 					{
 						BranchClientPath = NewSelectedClientFileName.Substring(0, EndIdx);
 						break;
@@ -236,11 +236,13 @@ namespace UnrealGameSync
 			Log.WriteLine("Found branch root at {0}", BranchClientPath);
 
 			// Get the local branch root
-			if(!PerforceClient.ConvertToLocalPath(BranchClientPath + "/Engine/Source/UE4Editor.target.cs", out BaseEditorTargetPath, Log))
+			string BuildVersionPath;
+			if(!PerforceClient.ConvertToLocalPath(BranchClientPath + "/Engine/Build/Build.version", out BuildVersionPath, Log))
 			{
-				ErrorMessage = String.Format("Couldn't get local path for editor target file");
+				ErrorMessage = String.Format("Couldn't get local path for Engine/Build/Build.version");
 				return false;
 			}
+			BranchDirectoryName = Path.GetFullPath(Path.Combine(Path.GetDirectoryName(BuildVersionPath), "..", ".."));
 
 			// Find the editor target for this project
 			if(NewSelectedFileName.EndsWith(".uproject", StringComparison.InvariantCultureIgnoreCase))
