UGS: Allow customizing the position of custom columns, via the Index=N attribute.

diff --git a/Engine/Source/Programs/UnrealGameSync/UnrealGameSync/Controls/WorkspaceControl.cs b/Engine/Source/Programs/UnrealGameSync/UnrealGameSync/Controls/WorkspaceControl.cs
--- a/Engine/Source/Programs/UnrealGameSync/UnrealGameSync/Controls/WorkspaceControl.cs
+++ b/Engine/Source/Programs/UnrealGameSync/UnrealGameSync/Controls/WorkspaceControl.cs
@@ -382,10 +382,16 @@ namespace UnrealGameSync
 							}
 							else
 							{
+								int Index = ColumnConfig.GetValue("Index", -1);
+								if(Index < 0 || Index > BuildList.Columns.Count)
+								{
+									Index = ((CustomColumns.Count > 0)? CustomColumns[CustomColumns.Count - 1].Index : CISColumn.Index) + 1;
+								}
+
 								Column = new ColumnHeader();
 								Column.Text = Name;
 								Column.Tag = ColumnConfig;
-								BuildList.Columns.Insert(CISColumn.Index + CustomColumns.Count + 1, Column);
+								BuildList.Columns.Insert(Index, Column);
 
 								CustomColumns.Add(Column);
 							}
