UGS: Include *.inl when looking for code changes.

diff --git a/Engine/Source/Programs/UnrealGameSync/UnrealGameSync/PerforceMonitor.cs b/Engine/Source/Programs/UnrealGameSync/UnrealGameSync/PerforceMonitor.cs
--- a/Engine/Source/Programs/UnrealGameSync/UnrealGameSync/PerforceMonitor.cs
+++ b/Engine/Source/Programs/UnrealGameSync/UnrealGameSync/PerforceMonitor.cs
@@ -361,7 +361,7 @@ namespace UnrealGameSync
 			// Update them in batches
 			foreach(int QueryChangeNumber in QueryChangeNumbers)
 			{
-				string[] CodeExtensions = { ".cs", ".h", ".cpp", ".usf", ".ush", ".uproject", ".uplugin" };
+				string[] CodeExtensions = { ".cs", ".h", ".cpp", ".inl", ".usf", ".ush", ".uproject", ".uplugin" };
 
 				// If there's something to check for, find all the content changes after this changelist
 				PerforceDescribeRecord DescribeRecord;
