UGS: Add a more useful error if a file to be synced exceeds the max allowed path length.

diff --git a/Engine/Source/Programs/UnrealGameSync/UnrealGameSync/Workspace.cs b/Engine/Source/Programs/UnrealGameSync/UnrealGameSync/Workspace.cs
--- a/Engine/Source/Programs/UnrealGameSync/UnrealGameSync/Workspace.cs
+++ b/Engine/Source/Programs/UnrealGameSync/UnrealGameSync/Workspace.cs
@@ -368,11 +368,20 @@ namespace UnrealGameSync
 						// Remove all the files that are not included by the filter
 						foreach(PerforceFileRecord HaveFile in HaveFiles)
 						{
-							string FullPath = Path.GetFullPath(HaveFile.Path);
-							if(MatchFilter(FullPath, SyncPathsFilter) && !MatchFilter(FullPath, UserFilter))
+							try
+							{
+								string FullPath = Path.GetFullPath(HaveFile.Path);
+								if(MatchFilter(FullPath, SyncPathsFilter) && !MatchFilter(FullPath, UserFilter))
+								{
+									Log.WriteLine("  {0}", HaveFile.DepotPath);
+									RemoveDepotPaths.Add(HaveFile.DepotPath);
+								}
+							}
+							catch(PathTooLongException)
 							{
-								Log.WriteLine("  {0}", HaveFile.DepotPath);
-								RemoveDepotPaths.Add(HaveFile.DepotPath);
+								Log.WriteLine("The local path for {0} exceeds the maximum allowed by Windows. Re-sync your workspace to a directory with a shorter name, or delete the file from the server.", HaveFile.Path);
+								StatusMessage = "File exceeds maximum path length allowed by Windows.";
+								return WorkspaceUpdateResult.FailedToSync;
 							}
 						}
 
@@ -426,6 +435,20 @@ namespace UnrealGameSync
 							SyncRecords.RemoveAll(x => !String.IsNullOrEmpty(x.ClientPath) && !MatchFilter(Path.GetFullPath(x.ClientPath), UserFilter));
 						}
 
+						foreach(PerforceFileRecord SyncRecord in SyncRecords)
+						{
+							try
+							{
+								Path.GetFullPath(SyncRecord.ClientPath);
+							}
+							catch(PathTooLongException)
+							{
+								Log.WriteLine("The local path for {0} exceeds the maximum allowed by Windows. Re-sync your workspace to a directory with a shorter name, or delete the file from the server.", SyncRecord.ClientPath);
+								StatusMessage = "File exceeds maximum path length allowed by Windows.";
+								return WorkspaceUpdateResult.FailedToSync;
+							}
+						}
+
 						SyncDepotPaths.AddRange(SyncRecords.Select(x => x.DepotPath));
 
 						List<PerforceFileRecord> OpenRecords;
