UGS: Move the check for overlong paths earlier.

diff --git a/Engine/Source/Programs/UnrealGameSync/UnrealGameSync/Properties/AssemblyInfo.cs b/Engine/Source/Programs/UnrealGameSync/UnrealGameSync/Properties/AssemblyInfo.cs
--- a/Engine/Source/Programs/UnrealGameSync/UnrealGameSync/Properties/AssemblyInfo.cs
+++ b/Engine/Source/Programs/UnrealGameSync/UnrealGameSync/Properties/AssemblyInfo.cs
@@ -34,5 +34,5 @@ using System.Runtime.InteropServices;
 // You can specify all the values or you can default the Build and Revision Numbers 
 // by using the '*' as shown below:
 // [assembly: AssemblyVersion("1.0.*")]
-[assembly: AssemblyVersion("1.146.*")]
-[assembly: AssemblyFileVersion("1.146.0.0")]
+[assembly: AssemblyVersion("1.147.*")]
+[assembly: AssemblyFileVersion("1.147.0.0")]
diff --git a/Engine/Source/Programs/UnrealGameSync/UnrealGameSync/Workspace.cs b/Engine/Source/Programs/UnrealGameSync/UnrealGameSync/Workspace.cs
--- a/Engine/Source/Programs/UnrealGameSync/UnrealGameSync/Workspace.cs
+++ b/Engine/Source/Programs/UnrealGameSync/UnrealGameSync/Workspace.cs
@@ -430,16 +430,14 @@ namespace UnrealGameSync
 							return WorkspaceUpdateResult.FailedToSync;
 						}
 
-						if(UserFilter != null)
-						{
-							SyncRecords.RemoveAll(x => !String.IsNullOrEmpty(x.ClientPath) && !MatchFilter(Path.GetFullPath(x.ClientPath), UserFilter));
-						}
-
 						foreach(PerforceFileRecord SyncRecord in SyncRecords)
 						{
 							try
 							{
-								Path.GetFullPath(SyncRecord.ClientPath);
+								if(!String.IsNullOrEmpty(SyncRecord.ClientPath))
+								{
+									Path.GetFullPath(SyncRecord.ClientPath);
+								}
 							}
 							catch(PathTooLongException)
 							{
@@ -449,6 +447,11 @@ namespace UnrealGameSync
 							}
 						}
 
+						if(UserFilter != null)
+						{
+							SyncRecords.RemoveAll(x => !String.IsNullOrEmpty(x.ClientPath) && !MatchFilter(Path.GetFullPath(x.ClientPath), UserFilter));
+						}
+
 						SyncDepotPaths.AddRange(SyncRecords.Select(x => x.DepotPath));
 
 						List<PerforceFileRecord> OpenRecords;
