UGS: Added a dialog which shows all the required platform SDKs for a branch, linked from the status panel in UGS.

The llist is configured via the UGS config file submitted to Engine/Programs/UnrealGameSync/UnrealGameSync.ini (and may be overridden by the project config file if necessary):

    [Default]
    ; Set this to a network share which contains the SDK installers for your site
    SdkInstallerDir=

    ; All the required SDKs for the current version of the engine
    +SdkInfo=(Category="Android", Description="NDK r21", Browse="$(SdkInstallerDir)\\Android")
    +SdkInfo=(Category="Windows", Description="Visual Studio 2017")
    +SdkInfo=(Category="Windows", Description="Visual C++ Toolchain 14.13.26128")
    +SdkInfo=(Category="Windows", Description="Windows SDK 10.0.16299.0")

Similar entries for console platforms are added in console subdirectories. Each entry may contain an Install="Foo.exe" and/or Browse="C:\Foo" style attribute, specifying the path to an installer to run or directory to open in explorer respectively.

The SdkInstallerDir setting is used as a base directory for the default installers, seen above for Android. Licensees may override this with a network path specific to the site that UGS is being deployed to (either in this file, in a project specific config file, or in a Engine/Programs/UnrealGameSync/NotForLicensees/UnrealGameSync.ini file).

diff --git a/Engine/Programs/UnrealGameSync/UnrealGameSync.ini b/Engine/Programs/UnrealGameSync/UnrealGameSync.ini
--- a/Engine/Programs/UnrealGameSync/UnrealGameSync.ini
+++ b/Engine/Programs/UnrealGameSync/UnrealGameSync.ini
@@ -9,4 +9,14 @@ UseFastModularVersioning=1
 +ExcludeChanges=^Changed the dynamic query of collection '.*' to '.*'
 +ExcludeChanges=^Changed the parent of collection '.*'
 +ExcludeChanges=^Upgraded collection '.*' \(was version \d+, now version \d+\)
-+ExcludeChanges=^Collection '.*' not modified
\ No newline at end of file
++ExcludeChanges=^Collection '.*' not modified
+
+[Default]
+; Set this to a network share which contains the SDK installers for your site
+SdkInstallerDir= 
+
+; All the required SDKs for the current version of the engine
++SdkInfo=(Category="Android", Description="NDK r21", Browse="$(SdkInstallerDir)\\Android")
++SdkInfo=(Category="Windows", Description="Visual Studio 2017")
++SdkInfo=(Category="Windows", Description="Visual C++ Toolchain 14.13.26128")
++SdkInfo=(Category="Windows", Description="Windows SDK 10.0.16299.0")
diff --git a/Engine/Source/Programs/UnrealGameSync/UnrealGameSync/Controls/BuildListControl.cs b/Engine/Source/Programs/UnrealGameSync/UnrealGameSync/Controls/BuildListControl.cs
--- a/Engine/Source/Programs/UnrealGameSync/UnrealGameSync/Controls/BuildListControl.cs
+++ b/Engine/Source/Programs/UnrealGameSync/UnrealGameSync/Controls/BuildListControl.cs
@@ -1,4 +1,4 @@
-﻿// Copyright 1998-2018 Epic Games, Inc. All Rights Reserved.
+// Copyright 1998-2018 Epic Games, Inc. All Rights Reserved.
 
 using System;
 using System.Collections.Generic;
@@ -12,7 +12,7 @@ using System.Windows.Forms;
 
 namespace UnrealGameSync
 {
-    public class BuildListControl : ListView 
+    public class BuildListControl : CustomListViewControl
 	{
 		[StructLayout(LayoutKind.Sequential)]
 		public class SCROLLINFO
diff --git a/Engine/Source/Programs/UnrealGameSync/UnrealGameSync/Controls/CustomListViewControl.cs b/Engine/Source/Programs/UnrealGameSync/UnrealGameSync/Controls/CustomListViewControl.cs
new file mode 100644
--- /dev/null
+++ b/Engine/Source/Programs/UnrealGameSync/UnrealGameSync/Controls/CustomListViewControl.cs
@@ -0,0 +1,106 @@
+using System;
+using System.Collections.Generic;
+using System.ComponentModel;
+using System.Drawing;
+using System.Data;
+using System.Linq;
+using System.Text;
+using System.Threading.Tasks;
+using System.Windows.Forms;
+using System.Windows.Forms.VisualStyles;
+
+namespace UnrealGameSync
+{
+	public partial class CustomListViewControl : ListView
+	{
+		VisualStyleRenderer SelectedItemRenderer;
+		VisualStyleRenderer TrackedItemRenderer;
+		public int HoverItem = -1;
+
+		public CustomListViewControl()
+		{
+            if (Application.RenderWithVisualStyles) 
+            { 
+				SelectedItemRenderer = new VisualStyleRenderer("Explorer::ListView", 1, 3);
+				TrackedItemRenderer = new VisualStyleRenderer("Explorer::ListView", 1, 2); 
+			}
+		}
+
+		protected override void OnMouseLeave(EventArgs e)
+		{
+			base.OnMouseLeave(e);
+
+			if(HoverItem != -1)
+			{
+				HoverItem = -1;
+				Invalidate();
+			}
+		}
+
+		protected override void OnMouseMove(MouseEventArgs e)
+		{
+			ListViewHitTestInfo HitTest = this.HitTest(e.Location);
+
+			int PrevHoverItem = HoverItem;
+			HoverItem = (HitTest.Item == null)? -1 : HitTest.Item.Index;
+			if(HoverItem != PrevHoverItem)
+			{
+				if(HoverItem != -1)
+				{
+					RedrawItems(HoverItem, HoverItem, true);
+				}
+				if(PrevHoverItem != -1 && PrevHoverItem < Items.Count)
+				{
+					RedrawItems(PrevHoverItem, PrevHoverItem, true);
+				}
+			}
+
+			base.OnMouseMove(e);
+		}
+
+		public void DrawBackground(Graphics Graphics, ListViewItem Item)
+		{
+			if(Item.Selected)
+			{
+				DrawSelectedBackground(Graphics, Item.Bounds);
+			}
+			else if(Item.Index == HoverItem)
+			{
+				DrawTrackedBackground(Graphics, Item.Bounds);
+			}
+			else
+			{
+				DrawDefaultBackground(Graphics, Item.Bounds);
+			}
+		}
+
+		public void DrawDefaultBackground(Graphics Graphics, Rectangle Bounds)
+		{
+			Graphics.FillRectangle(SystemBrushes.Window, Bounds);
+		}
+
+		public void DrawSelectedBackground(Graphics Graphics, Rectangle Bounds)
+		{
+			if(Application.RenderWithVisualStyles)
+			{
+				SelectedItemRenderer.DrawBackground(Graphics, Bounds);
+			}
+			else
+			{
+				Graphics.FillRectangle(SystemBrushes.ButtonFace, Bounds);
+			}
+		}
+
+		public void DrawTrackedBackground(Graphics Graphics, Rectangle Bounds)
+		{
+			if(Application.RenderWithVisualStyles)
+			{
+				TrackedItemRenderer.DrawBackground(Graphics, Bounds);
+			}
+			else
+			{
+				Graphics.FillRectangle(SystemBrushes.Window, Bounds);
+			}
+		}
+	}
+}
diff --git a/Engine/Source/Programs/UnrealGameSync/UnrealGameSync/Controls/WorkspaceControl.cs b/Engine/Source/Programs/UnrealGameSync/UnrealGameSync/Controls/WorkspaceControl.cs
--- a/Engine/Source/Programs/UnrealGameSync/UnrealGameSync/Controls/WorkspaceControl.cs
+++ b/Engine/Source/Programs/UnrealGameSync/UnrealGameSync/Controls/WorkspaceControl.cs
@@ -197,12 +197,9 @@ namespace UnrealGameSync
 
 		TimeSpan ServerTimeZone;
 
-		int HoverItem = -1;
 		string HoverBadgeUniqueId = null;
 		bool bHoverSync;
 		PerforceChangeSummary ContextMenuChange;
-		VisualStyleRenderer SelectedItemRenderer;
-		VisualStyleRenderer TrackedItemRenderer;
 		Font BuildFont;
 		Font SelectedBuildFont;
 		Font BadgeFont;
@@ -239,12 +236,6 @@ namespace UnrealGameSync
 
 			MainThreadSynchronizationContext = SynchronizationContext.Current;
 
-            if (Application.RenderWithVisualStyles) 
-            { 
-				SelectedItemRenderer = new VisualStyleRenderer("Explorer::ListView", 1, 3);
-				TrackedItemRenderer = new VisualStyleRenderer("Explorer::ListView", 1, 2); 
-			}
-
 			Owner = InOwner;
 			ApiUrl = InApiUrl;
 			DataFolder = DetectSettings.DataFolder;
@@ -1062,9 +1053,9 @@ namespace UnrealGameSync
 				}
 			}
 
-			if(HoverItem > BuildList.Items.Count)
+			if(BuildList.HoverItem > BuildList.Items.Count)
 			{
-				HoverItem = -1;
+				BuildList.HoverItem = -1;
 			}
 
 			UpdateBuildFailureNotification();
@@ -1391,31 +1382,21 @@ namespace UnrealGameSync
 
 		private void BuildList_DrawItem(object sender, DrawListViewItemEventArgs e)
 		{
-			if(Application.RenderWithVisualStyles)
+			if(e.State.HasFlag(ListViewItemStates.Selected))
 			{
-				if(e.State.HasFlag(ListViewItemStates.Selected))
-				{
-					SelectedItemRenderer.DrawBackground(e.Graphics, e.Bounds);
-				}
-				else if(e.ItemIndex == HoverItem)
-				{
-					TrackedItemRenderer.DrawBackground(e.Graphics, e.Bounds);
-				}
-				else if(((PerforceChangeSummary)e.Item.Tag).Number == Workspace.PendingChangeNumber)
-				{
-					TrackedItemRenderer.DrawBackground(e.Graphics, e.Bounds);
-				}
+				BuildList.DrawSelectedBackground(e.Graphics, e.Bounds);
+			}
+			else if(e.ItemIndex == BuildList.HoverItem)
+			{
+				BuildList.DrawTrackedBackground(e.Graphics, e.Bounds);
+			}
+			else if(((PerforceChangeSummary)e.Item.Tag).Number == Workspace.PendingChangeNumber)
+			{
+				BuildList.DrawTrackedBackground(e.Graphics, e.Bounds);
 			}
 			else
 			{
-				if(e.State.HasFlag(ListViewItemStates.Selected))
-				{
-					e.Graphics.FillRectangle(SystemBrushes.ButtonFace, e.Bounds);
-				}
-				else
-				{
-					e.Graphics.FillRectangle(SystemBrushes.Window, e.Bounds);
-				}
+				BuildList.DrawDefaultBackground(e.Graphics, e.Bounds);
 			}
 		}
 
@@ -1679,7 +1660,7 @@ namespace UnrealGameSync
 						MaxX -= HappyIcon.Width;
 						e.Graphics.DrawImage(Properties.Resources.Icons, MaxX, IconY, (Review == null || !EventMonitor.IsNegativeReview(Review.Type))? HappyIcon : DisabledHappyIcon, GraphicsUnit.Pixel);
 					}
-					else if(e.ItemIndex == HoverItem && bAllowSync)
+					else if(e.ItemIndex == BuildList.HoverItem && bAllowSync)
 					{
 						Rectangle SyncBadgeRectangle = GetSyncBadgeRectangle(e.SubItem.Bounds);
 						DrawBadge(e.Graphics, SyncBadgeRectangle, "Sync", bHoverSync? Color.FromArgb(140, 180, 230) : Color.FromArgb(112, 146, 190), true, true);
@@ -2547,11 +2528,20 @@ namespace UnrealGameSync
 					ProgramsLine.AddLink("Unreal Editor", FontStyle.Regular, () => { LaunchEditor(); });
 					ProgramsLine.AddText("  |  ");
 				}
+
+				string[] SdkInfoEntries;
+				if(TryGetProjectSetting(PerforceMonitor.LatestProjectConfigFile, "SdkInfo", out SdkInfoEntries))
+				{
+					ProgramsLine.AddLink("SDK Info", FontStyle.Regular, () => { ShowRequiredSdkInfo(); });
+					ProgramsLine.AddText("  |  ");
+				}
+
 				ProgramsLine.AddLink("Perforce", FontStyle.Regular, () => { OpenPerforce(); });
 				ProgramsLine.AddText("  |  ");
 				ProgramsLine.AddLink("Visual Studio", FontStyle.Regular, () => { OpenSolution(); });
 				ProgramsLine.AddText("  |  ");
 				ProgramsLine.AddLink("Windows Explorer", FontStyle.Regular, () => { Process.Start("explorer.exe", String.Format("\"{0}\"", Path.GetDirectoryName(SelectedFileName))); });
+
 				foreach(KeyValuePair<string, BuildData> ServiceBadge in ServiceBadges)
 				{
 					ProgramsLine.AddText("  |  ");
@@ -2650,6 +2640,17 @@ namespace UnrealGameSync
 			StatusPanel.Set(Lines, Caption, Alert, TintColor);
 		}
 
+		private void ShowRequiredSdkInfo()
+		{
+			string[] SdkInfoEntries;
+			if(TryGetProjectSetting(PerforceMonitor.LatestProjectConfigFile, "SdkInfo", out SdkInfoEntries))
+			{
+				Dictionary<string, string> Variables = GetWorkspaceVariables(-1);
+				SdkInfoWindow Window = new SdkInfoWindow(SdkInfoEntries, Variables, BadgeFont);
+				Window.ShowDialog();
+			}
+		}
+
 		private StatusLine CreateStatusLineFromMarkdown(string Text)
 		{
 			StatusLine Line = new StatusLine();
@@ -3110,19 +3111,6 @@ namespace UnrealGameSync
 		private void BuildList_MouseMove(object sender, MouseEventArgs e)
 		{
 			ListViewHitTestInfo HitTest = BuildList.HitTest(e.Location);
-			int HitTestIndex = (HitTest.Item == null)? -1 : HitTest.Item.Index;
-			if(HitTestIndex != HoverItem)
-			{
-				if(HoverItem != -1)
-				{
-					BuildList.RedrawItems(HoverItem, HoverItem, true);
-				}
-				HoverItem = HitTestIndex;
-				if(HoverItem != -1)
-				{
-					BuildList.RedrawItems(HoverItem, HoverItem, true);
-				}
-			}
 
 			string NewHoverBadgeUniqueId = null;
 			if(HitTest.Item != null)
@@ -3525,12 +3513,6 @@ namespace UnrealGameSync
 		{
 			BuildListToolTip.Hide(BuildList);
 
-			if(HoverItem != -1)
-			{
-				HoverItem = -1;
-				BuildList.Invalidate();
-			}
-
 			if(HoverBadgeUniqueId != null)
 			{
 				HoverBadgeUniqueId = null;
@@ -3874,6 +3856,9 @@ namespace UnrealGameSync
 		{
 			BuildConfig EditorBuildConfig = GetEditorBuildConfig();
 
+			string SdkInstallerDir;
+			TryGetProjectSetting(PerforceMonitor.LatestProjectConfigFile, "SdkInstallerDir", out SdkInstallerDir);
+
 			Dictionary<string, string> Variables = new Dictionary<string,string>();
 			Variables.Add("Stream", StreamName);
 			Variables.Add("Change", ChangeNumber.ToString());
@@ -3886,6 +3871,10 @@ namespace UnrealGameSync
 			Variables.Add("UE4EditorConfig", EditorBuildConfig.ToString());
 			Variables.Add("UE4EditorDebugArg", (EditorBuildConfig == BuildConfig.Debug || EditorBuildConfig == BuildConfig.DebugGame)? " -debug" : "");
 			Variables.Add("UseIncrementalBuilds", Settings.bUseIncrementalBuilds? "1" : "0");
+			if(!String.IsNullOrEmpty(SdkInstallerDir))
+			{
+				Variables.Add("SdkInstallerDir", SdkInstallerDir);
+			}
 			return Variables;
 		}
 
diff --git a/Engine/Source/Programs/UnrealGameSync/UnrealGameSync/Forms/SdkInfoWindow.cs b/Engine/Source/Programs/UnrealGameSync/UnrealGameSync/Forms/SdkInfoWindow.cs
new file mode 100644
--- /dev/null
+++ b/Engine/Source/Programs/UnrealGameSync/UnrealGameSync/Forms/SdkInfoWindow.cs
@@ -0,0 +1,282 @@
+using System;
+using System.Collections.Generic;
+using System.ComponentModel;
+using System.Data;
+using System.Diagnostics;
+using System.Drawing;
+using System.Drawing.Drawing2D;
+using System.IO;
+using System.Linq;
+using System.Text;
+using System.Threading.Tasks;
+using System.Windows.Forms;
+
+namespace UnrealGameSync
+{
+	public partial class SdkInfoWindow : Form
+	{
+		class SdkItem
+		{
+			public string Category;
+			public string Description;
+			public string Install;
+			public string Browse;
+		}
+
+		class BadgeInfo
+		{
+			public string UniqueId;
+			public string Label;
+			public Rectangle Rectangle;
+			public Action OnClick;
+		}
+
+		Font BadgeFont;
+		string HoverBadgeUniqueId;
+
+		public SdkInfoWindow(string[] SdkInfoEntries, Dictionary<string, string> Variables, Font BadgeFont)
+		{
+			InitializeComponent();
+
+			this.BadgeFont = BadgeFont;
+
+			Dictionary<string, ConfigObject> UniqueIdToObject = new Dictionary<string, ConfigObject>(StringComparer.InvariantCultureIgnoreCase);
+			foreach(string SdkInfoEntry in SdkInfoEntries)
+			{
+				ConfigObject Object = new ConfigObject(SdkInfoEntry);
+
+				string UniqueId = Object.GetValue("UniqueId", Guid.NewGuid().ToString());
+
+				ConfigObject ExistingObject;
+				if(UniqueIdToObject.TryGetValue(UniqueId, out ExistingObject))
+				{
+					ExistingObject.AddOverrides(Object, null);
+				}
+				else
+				{
+					UniqueIdToObject.Add(UniqueId, Object);
+				}
+			}
+
+			List<SdkItem> Items = new List<SdkItem>();
+			foreach(ConfigObject Object in UniqueIdToObject.Values)
+			{
+				SdkItem Item = new SdkItem();
+				Item.Category = Object.GetValue("Category", "Other");
+				Item.Description = Object.GetValue("Description", "");
+
+				Item.Install = Utility.ExpandVariables(Object.GetValue("Install", ""), Variables);
+				if(Item.Install.Contains("$("))
+				{
+					Item.Install = null;
+				}
+
+				Item.Browse = Utility.ExpandVariables(Object.GetValue("Browse", ""), Variables);
+				if(Item.Browse.Contains("$("))
+				{
+					Item.Browse = null;
+				}
+
+				if(!String.IsNullOrEmpty(Item.Install) && String.IsNullOrEmpty(Item.Browse))
+				{
+					try
+					{
+						Item.Browse = Path.GetDirectoryName(Item.Install);
+					}
+					catch
+					{
+						Item.Browse = null;
+					}
+				}
+
+				Items.Add(Item);
+			}
+
+			foreach(IGrouping<string, SdkItem> ItemGroup in Items.GroupBy(x => x.Category).OrderBy(x => x.Key))
+			{
+				ListViewGroup Group = new ListViewGroup(ItemGroup.Key);
+				SdkListView.Groups.Add(Group);
+
+				foreach(SdkItem Item in ItemGroup)
+				{
+					ListViewItem NewItem = new ListViewItem(Group);
+					NewItem.SubItems.Add(Item.Description);
+					NewItem.SubItems.Add(new ListViewItem.ListViewSubItem(){ Tag = Item });
+					SdkListView.Items.Add(NewItem);
+				}
+			}
+
+			System.Reflection.PropertyInfo DoubleBufferedProperty = typeof(Control).GetProperty("DoubleBuffered", System.Reflection.BindingFlags.NonPublic | System.Reflection.BindingFlags.Instance);
+			DoubleBufferedProperty.SetValue(SdkListView, true, null);
+
+			if(SdkListView.Items.Count > 0)
+			{
+				int ItemsHeight = SdkListView.Items[SdkListView.Items.Count - 1].Bounds.Bottom + 20;
+				Height = SdkListView.Top + ItemsHeight + (Height - SdkListView.Bottom);
+			}
+		}
+
+		private void OkBtn_Click(object sender, EventArgs e)
+		{
+			DialogResult = DialogResult.OK;
+			Close();
+		}
+
+		private void SdkListView_DrawItem(object sender, DrawListViewItemEventArgs e)
+		{
+			SdkListView.DrawBackground(e.Graphics, e.Item);
+		}
+
+		private void SdkListView_DrawSubItem(object sender, DrawListViewSubItemEventArgs e)
+		{
+			if(e.ColumnIndex != columnHeader3.Index)
+			{
+				TextRenderer.DrawText(e.Graphics, e.SubItem.Text, SdkListView.Font, e.Bounds, SdkListView.ForeColor, TextFormatFlags.EndEllipsis | TextFormatFlags.SingleLine | TextFormatFlags.Left | TextFormatFlags.VerticalCenter | TextFormatFlags.NoPrefix);
+			}
+			else
+			{
+				e.Graphics.SmoothingMode = SmoothingMode.AntiAlias;
+
+				List<BadgeInfo> Badges = GetBadges(e.Item, e.SubItem);
+				for(int Idx = 0; Idx < Badges.Count; Idx++)
+				{
+					Color BadgeColor = (HoverBadgeUniqueId == Badges[Idx].UniqueId)? Color.FromArgb(140, 180, 230) : Color.FromArgb(112, 146, 190);
+					if(Badges[Idx].OnClick != null)
+					{
+						DrawBadge(e.Graphics, Badges[Idx].Label, Badges[Idx].Rectangle, (Idx > 0), (Idx < Badges.Count - 1), BadgeColor);
+					}
+				}
+			}
+		}
+
+		private void SdkListView_MouseMove(object sender, MouseEventArgs e)
+		{
+			string NewHoverUniqueId = null;
+
+			ListViewHitTestInfo HitTest = SdkListView.HitTest(e.Location);
+			if(HitTest.Item != null && HitTest.SubItem == HitTest.Item.SubItems[2])
+			{
+				List<BadgeInfo> Badges = GetBadges(HitTest.Item, HitTest.SubItem);
+				foreach(BadgeInfo Badge in Badges)
+				{
+					if(Badge.Rectangle.Contains(e.Location))
+					{
+						NewHoverUniqueId = Badge.UniqueId;
+					}
+				}
+			}
+
+			if(NewHoverUniqueId != HoverBadgeUniqueId)
+			{
+				HoverBadgeUniqueId = NewHoverUniqueId;
+				SdkListView.Invalidate();
+			}
+		}
+
+		private void SdkListView_MouseLeave(object sender, EventArgs e)
+		{
+			HoverBadgeUniqueId = null;
+		}
+
+		private void SdkListView_MouseDown(object sender, MouseEventArgs e)
+		{
+			ListViewHitTestInfo HitTest = SdkListView.HitTest(e.Location);
+			if(HitTest.Item != null && HitTest.SubItem == HitTest.Item.SubItems[2])
+			{
+				List<BadgeInfo> Badges = GetBadges(HitTest.Item, HitTest.SubItem);
+				foreach(BadgeInfo Badge in Badges)
+				{
+					if(Badge.Rectangle.Contains(e.Location))
+					{
+						Badge.OnClick();
+					}
+				}
+			}
+		}
+
+		private List<BadgeInfo> GetBadges(ListViewItem Item, ListViewItem.ListViewSubItem SubItem)
+		{
+			string UniqueIdPrefix = String.Format("{0}_", Item.Index);
+
+			List<BadgeInfo> Badges = new List<BadgeInfo>();
+
+			SdkItem Sdk = (SdkItem)SubItem.Tag;
+
+			Action InstallAction = null;
+			if(!String.IsNullOrEmpty(Sdk.Install))
+			{
+				InstallAction = () => { Install(Sdk.Install); };
+			}
+			Badges.Add(new BadgeInfo(){ UniqueId = UniqueIdPrefix + "_Install", Label = "Install", OnClick = InstallAction });
+
+			Action BrowseAction = null;
+			if(!String.IsNullOrEmpty(Sdk.Browse))
+			{
+				BrowseAction = () => { Browse(Sdk.Browse); };
+			}
+			Badges.Add(new BadgeInfo(){ UniqueId = UniqueIdPrefix + "_Browse", Label = "Browse", OnClick = BrowseAction });
+
+			int Right = SubItem.Bounds.Right - 10;
+			for(int Idx = Badges.Count - 1; Idx >= 0; Idx--)
+			{
+				Size BadgeSize = GetBadgeSize(Badges[Idx].Label);
+				Right -= BadgeSize.Width;
+				Badges[Idx].Rectangle = new Rectangle(Right, SubItem.Bounds.Y + (SubItem.Bounds.Height - BadgeSize.Height) / 2, BadgeSize.Width, BadgeSize.Height);
+			}
+
+			return Badges;
+		}
+
+		private void Browse(string DirectoryName)
+		{
+			try
+			{
+				Process.Start("explorer.exe", String.Format("\"{0}\"", DirectoryName));
+			}
+			catch(Exception Ex)
+			{
+				MessageBox.Show(String.Format("Unable to open explorer to {0}: {1}", DirectoryName, Ex.Message));
+			}
+		}
+
+		private void Install(string FileName)
+		{
+			try
+			{
+				Process.Start(FileName);
+			}
+			catch(Exception Ex)
+			{
+				MessageBox.Show(String.Format("Unable to run {0}: {1}", FileName, Ex.Message));
+			}
+		}
+
+		private Size GetBadgeSize(string BadgeText)
+		{
+			Size LabelSize = TextRenderer.MeasureText(BadgeText, BadgeFont);
+			int BadgeHeight = BadgeFont.Height + 1;
+
+			return new Size(LabelSize.Width + BadgeHeight - 4, BadgeHeight);
+		}
+
+		private void DrawBadge(Graphics Graphics, string BadgeText, Rectangle BadgeRect, bool bMergeLeft, bool bMergeRight, Color BadgeColor)
+		{
+			using (GraphicsPath Path = new GraphicsPath())
+			{
+				Path.StartFigure();
+				Path.AddLine(BadgeRect.Left + (bMergeLeft? 1 : 0), BadgeRect.Top, BadgeRect.Left - (bMergeLeft? 1 : 0), BadgeRect.Bottom);
+				Path.AddLine(BadgeRect.Left - (bMergeLeft? 1 : 0), BadgeRect.Bottom, BadgeRect.Right - 1 - (bMergeRight? 1 : 0), BadgeRect.Bottom);
+				Path.AddLine(BadgeRect.Right - 1 - (bMergeRight? 1 : 0), BadgeRect.Bottom, BadgeRect.Right - 1 + (bMergeRight? 1 : 0), BadgeRect.Top);
+				Path.AddLine(BadgeRect.Right - 1 + (bMergeRight? 1 : 0), BadgeRect.Top, BadgeRect.Left + (bMergeLeft? 1 : 0), BadgeRect.Top);
+				Path.CloseFigure();
+
+				using(SolidBrush Brush = new SolidBrush(BadgeColor))
+				{
+					Graphics.FillPath(Brush, Path);
+				}
+			}
+
+			TextRenderer.DrawText(Graphics, BadgeText, BadgeFont, BadgeRect, Color.White, TextFormatFlags.HorizontalCenter | TextFormatFlags.VerticalCenter | TextFormatFlags.SingleLine | TextFormatFlags.NoPrefix | TextFormatFlags.PreserveGraphicsClipping);
+		}
+	}
+}
diff --git a/Engine/Source/Programs/UnrealGameSync/UnrealGameSync/PerforceMonitor.cs b/Engine/Source/Programs/UnrealGameSync/UnrealGameSync/PerforceMonitor.cs
--- a/Engine/Source/Programs/UnrealGameSync/UnrealGameSync/PerforceMonitor.cs
+++ b/Engine/Source/Programs/UnrealGameSync/UnrealGameSync/PerforceMonitor.cs
@@ -526,29 +526,36 @@ namespace UnrealGameSync
 			ConfigFile ProjectConfig = new ConfigFile();
 			foreach(string ConfigFilePath in ConfigFilePaths)
 			{
-				List<string> Lines;
-				if(Perforce.Print(ConfigFilePath, out Lines, Log))
+				List<string> Lines = null;
+
+				// If this file is open for edit, read the local version
+				if(OpenFiles != null && OpenFiles.Any(x => x.ClientPath.Equals(ConfigFilePath, StringComparison.InvariantCultureIgnoreCase)))
 				{
-					// If this file is open for edit, read the local version instead
-					if(OpenFiles != null && OpenFiles.Any(x => x.ClientPath.Equals(ConfigFilePath, StringComparison.InvariantCultureIgnoreCase)))
+					try
 					{
-						try
+						string LocalFileName;
+						if(Perforce.ConvertToLocalPath(ConfigFilePath, out LocalFileName, Log))
 						{
-							string LocalFileName;
-							if(Perforce.ConvertToLocalPath(ConfigFilePath, out LocalFileName, Log))
-							{
-								DateTime LastModifiedTime = File.GetLastWriteTimeUtc(LocalFileName);
-								LocalConfigFiles.Add(new KeyValuePair<string, DateTime>(LocalFileName, LastModifiedTime));
-								Lines = File.ReadAllLines(LocalFileName).ToList();
-							}
-						}
-						catch(Exception Ex)
-						{
-							Log.WriteLine("Failed to read local config file for {0}: {1}", ConfigFilePath, Ex.ToString());
+							DateTime LastModifiedTime = File.GetLastWriteTimeUtc(LocalFileName);
+							LocalConfigFiles.Add(new KeyValuePair<string, DateTime>(LocalFileName, LastModifiedTime));
+							Lines = File.ReadAllLines(LocalFileName).ToList();
 						}
 					}
+					catch(Exception Ex)
+					{
+						Log.WriteLine("Failed to read local config file for {0}: {1}", ConfigFilePath, Ex.ToString());
+					}
+				}
 
-					// Merge the text with the config file
+				// Otherwise try to get it from perforce
+				if(Lines == null)
+				{
+					Perforce.Print(ConfigFilePath, out Lines, Log);
+				}
+
+				// Merge the text with the config file
+				if(Lines != null)
+				{
 					try
 					{
 						ProjectConfig.Parse(Lines.ToArray());
diff --git a/Engine/Source/Programs/UnrealGameSync/UnrealGameSync/UnrealGameSync.csproj b/Engine/Source/Programs/UnrealGameSync/UnrealGameSync/UnrealGameSync.csproj
--- a/Engine/Source/Programs/UnrealGameSync/UnrealGameSync/UnrealGameSync.csproj
+++ b/Engine/Source/Programs/UnrealGameSync/UnrealGameSync/UnrealGameSync.csproj
@@ -114,6 +114,9 @@
     <Compile Include="ArchiveManifest.cs" />
     <Compile Include="BuildStep.cs" />
     <Compile Include="ChildProcess.cs" />
+    <Compile Include="Controls\CustomListViewControl.cs">
+      <SubType>Component</SubType>
+    </Compile>
     <Compile Include="Controls\ErrorPanel.cs">
       <SubType>Component</SubType>
     </Compile>
@@ -215,6 +218,12 @@
     <Compile Include="Forms\ProgramsRunningWindow.Designer.cs">
       <DependentUpon>ProgramsRunningWindow.cs</DependentUpon>
     </Compile>
+    <Compile Include="Forms\SdkInfoWindow.cs">
+      <SubType>Form</SubType>
+    </Compile>
+    <Compile Include="Forms\SdkInfoWindow.Designer.cs">
+      <DependentUpon>SdkInfoWindow.cs</DependentUpon>
+    </Compile>
     <Compile Include="Forms\SelectProjectFromWorkspaceWindow.cs">
       <SubType>Form</SubType>
     </Compile>
@@ -366,6 +375,9 @@
     <EmbeddedResource Include="Forms\ProgramsRunningWindow.resx">
       <DependentUpon>ProgramsRunningWindow.cs</DependentUpon>
     </EmbeddedResource>
+    <EmbeddedResource Include="Forms\SdkInfoWindow.resx">
+      <DependentUpon>SdkInfoWindow.cs</DependentUpon>
+    </EmbeddedResource>
     <EmbeddedResource Include="Forms\SelectProjectFromWorkspaceWindow.resx">
       <DependentUpon>SelectProjectFromWorkspaceWindow.cs</DependentUpon>
     </EmbeddedResource>
diff --git a/Engine/Source/Programs/UnrealGameSync/UnrealGameSync/Utility.cs b/Engine/Source/Programs/UnrealGameSync/UnrealGameSync/Utility.cs
--- a/Engine/Source/Programs/UnrealGameSync/UnrealGameSync/Utility.cs
+++ b/Engine/Source/Programs/UnrealGameSync/UnrealGameSync/Utility.cs
@@ -267,15 +267,24 @@ namespace UnrealGameSync
 		{
 			List<string> ProjectConfigFileNames = new List<string>();
 			ProjectConfigFileNames.Add(String.Format("{0}{1}Engine{1}Programs{1}UnrealGameSync{1}UnrealGameSync.ini", BaseWorkspacePath, Separator));
+			ProjectConfigFileNames.Add(String.Format("{0}{1}Engine{1}Programs{1}UnrealGameSync{1}PS4{1}UnrealGameSync.ini", BaseWorkspacePath, Separator));
+			ProjectConfigFileNames.Add(String.Format("{0}{1}Engine{1}Programs{1}UnrealGameSync{1}XboxOne{1}UnrealGameSync.ini", BaseWorkspacePath, Separator));
+			ProjectConfigFileNames.Add(String.Format("{0}{1}Engine{1}Programs{1}UnrealGameSync{1}Switch{1}UnrealGameSync.ini", BaseWorkspacePath, Separator));
 			ProjectConfigFileNames.Add(String.Format("{0}{1}Engine{1}Programs{1}UnrealGameSync{1}NotForLicensees{1}UnrealGameSync.ini", BaseWorkspacePath, Separator));
 			if(ProjectPath.EndsWith(".uproject", StringComparison.InvariantCultureIgnoreCase))
 			{
 				ProjectConfigFileNames.Add(String.Format("{0}{1}Build{1}UnrealGameSync.ini", ProjectPath.Substring(0, ProjectPath.LastIndexOf(Separator)), Separator));
+			    ProjectConfigFileNames.Add(String.Format("{0}{1}Build{1}PS4{1}UnrealGameSync.ini", ProjectPath.Substring(0, ProjectPath.LastIndexOf(Separator)), Separator));
+			    ProjectConfigFileNames.Add(String.Format("{0}{1}Build{1}XboxOne{1}UnrealGameSync.ini", ProjectPath.Substring(0, ProjectPath.LastIndexOf(Separator)), Separator));
+			    ProjectConfigFileNames.Add(String.Format("{0}{1}Build{1}Switch{1}UnrealGameSync.ini", ProjectPath.Substring(0, ProjectPath.LastIndexOf(Separator)), Separator));
 				ProjectConfigFileNames.Add(String.Format("{0}{1}Build{1}NotForLicensees{1}UnrealGameSync.ini", ProjectPath.Substring(0, ProjectPath.LastIndexOf(Separator)), Separator));
 			}
 			else
 			{
 				ProjectConfigFileNames.Add(String.Format("{0}{1}Engine{1}Programs{1}UnrealGameSync{1}DefaultProject.ini", BaseWorkspacePath, Separator));
+			    ProjectConfigFileNames.Add(String.Format("{0}{1}Engine{1}Programs{1}UnrealGameSync{1}PS4{1}DefaultProject.ini", BaseWorkspacePath, Separator));
+			    ProjectConfigFileNames.Add(String.Format("{0}{1}Engine{1}Programs{1}UnrealGameSync{1}XboxOne{1}DefaultProject.ini", BaseWorkspacePath, Separator));
+			    ProjectConfigFileNames.Add(String.Format("{0}{1}Engine{1}Programs{1}UnrealGameSync{1}Switch{1}DefaultProject.ini", BaseWorkspacePath, Separator));
 				ProjectConfigFileNames.Add(String.Format("{0}{1}Engine{1}Programs{1}UnrealGameSync{1}NotForLicensees{1}DefaultProject.ini", BaseWorkspacePath, Separator));
 			}
 			return ProjectConfigFileNames;
