UGS: Add support for DebugGame editors that have an executable with a DebugGame suffix.

diff --git a/Engine/Source/Programs/UnrealGameSync/UnrealGameSync/Controls/WorkspaceControl.cs b/Engine/Source/Programs/UnrealGameSync/UnrealGameSync/Controls/WorkspaceControl.cs
--- a/Engine/Source/Programs/UnrealGameSync/UnrealGameSync/Controls/WorkspaceControl.cs
+++ b/Engine/Source/Programs/UnrealGameSync/UnrealGameSync/Controls/WorkspaceControl.cs
@@ -2286,7 +2286,7 @@ namespace UnrealGameSync
 		private string GetEditorExePath(BuildConfig Config)
 		{
 			string ExeFileName = "UE4Editor.exe";
-			if(Config != BuildConfig.DebugGame && Config != BuildConfig.Development)
+			if((Config != BuildConfig.DebugGame || PerforceMonitor.LatestProjectConfigFile.GetValue("Options.DebugGameHasSeparateExecutable", false)) && Config != BuildConfig.Development)
 			{
 				ExeFileName	= String.Format("UE4Editor-Win64-{0}.exe", Config.ToString());
 			}
