UGS: Prevent attempts to automatically reopen projects while a modal dialog is up, or the workspace is syncing.

diff --git a/Engine/Source/Programs/UnrealGameSync/UnrealGameSync/Forms/MainWindow.cs b/Engine/Source/Programs/UnrealGameSync/UnrealGameSync/Forms/MainWindow.cs
--- a/Engine/Source/Programs/UnrealGameSync/UnrealGameSync/Forms/MainWindow.cs
+++ b/Engine/Source/Programs/UnrealGameSync/UnrealGameSync/Forms/MainWindow.cs
@@ -576,7 +576,7 @@ namespace UnrealGameSync
 		public void RequestProjectChange(WorkspaceControl Workspace, UserSelectedProjectSettings Project)
 		{
 			int TabIdx = TabControl.FindTabIndex(Workspace);
-			if(TabIdx != -1)
+			if(TabIdx != -1 && !Workspace.IsBusy() && CanFocus)
 			{
 				TryOpenProject(Project, TabIdx);
 			}
