UGS: Prevent modal dialogs when login expires in P4, and prompt for password when hitting "retry".

diff --git a/Engine/Source/Programs/UnrealGameSync/UnrealGameSync/Controls/WorkspaceControl.cs b/Engine/Source/Programs/UnrealGameSync/UnrealGameSync/Controls/WorkspaceControl.cs
--- a/Engine/Source/Programs/UnrealGameSync/UnrealGameSync/Controls/WorkspaceControl.cs
+++ b/Engine/Source/Programs/UnrealGameSync/UnrealGameSync/Controls/WorkspaceControl.cs
@@ -34,7 +34,7 @@ namespace UnrealGameSync
 	interface IWorkspaceControlOwner
 	{
 		void EditSelectedProject(WorkspaceControl Workspace);
-		void RequestProjectChange(WorkspaceControl Workspace, UserSelectedProjectSettings Project);
+		void RequestProjectChange(WorkspaceControl Workspace, UserSelectedProjectSettings Project, bool bModal);
 		void ShowAndActivate();
 		void StreamChanged(WorkspaceControl Workspace);
 		void SetTabNames(TabLabels TabNames);
@@ -607,7 +607,7 @@ namespace UnrealGameSync
 		{
 			if(!bIsDisposing)
 			{
-				Owner.RequestProjectChange(this, SelectedProject);
+				Owner.RequestProjectChange(this, SelectedProject, false);
 			}
 		}
 
@@ -4169,7 +4169,7 @@ namespace UnrealGameSync
 
 			foreach(UserSelectedProjectSettings RecentProject in Settings.RecentProjects)
 			{
-				ToolStripMenuItem Item = new ToolStripMenuItem(RecentProject.ToString(), null, new EventHandler((o, e) => Owner.RequestProjectChange(this, RecentProject)));
+				ToolStripMenuItem Item = new ToolStripMenuItem(RecentProject.ToString(), null, new EventHandler((o, e) => Owner.RequestProjectChange(this, RecentProject, true)));
 				RecentMenu.Items.Insert(RecentMenu.Items.Count - 2, Item);
 			}
 
diff --git a/Engine/Source/Programs/UnrealGameSync/UnrealGameSync/Forms/MainWindow.cs b/Engine/Source/Programs/UnrealGameSync/UnrealGameSync/Forms/MainWindow.cs
--- a/Engine/Source/Programs/UnrealGameSync/UnrealGameSync/Forms/MainWindow.cs
+++ b/Engine/Source/Programs/UnrealGameSync/UnrealGameSync/Forms/MainWindow.cs
@@ -268,7 +268,7 @@ namespace UnrealGameSync
 
 		private void CreateErrorPanel(int ReplaceTabIdx, UserSelectedProjectSettings Project, string Message)
 		{
-			Log.WriteLine(Message);
+			Log.WriteLine(Message ?? "Unknown error");
 
 			ErrorPanel ErrorPanel = new ErrorPanel(Project);
 			ErrorPanel.Parent = TabPanel;
@@ -281,7 +281,12 @@ namespace UnrealGameSync
 
 			string SummaryText = String.Format("Unable to open '{0}'.", Project.ToString());
 
-			int NewContentWidth = Math.Max(Math.Max(TextRenderer.MeasureText(SummaryText, ErrorPanel.Font).Width, TextRenderer.MeasureText(Message, ErrorPanel.Font).Width), 400);
+			int NewContentWidth = Math.Max(TextRenderer.MeasureText(SummaryText, ErrorPanel.Font).Width, 400);
+			if(!String.IsNullOrEmpty(Message))
+			{
+				NewContentWidth = Math.Max(NewContentWidth, TextRenderer.MeasureText(Message, ErrorPanel.Font).Width);
+			}
+
 			ErrorPanel.SetContentWidth(NewContentWidth);
 
 			List<StatusLine> Lines = new List<StatusLine>();
@@ -290,23 +295,26 @@ namespace UnrealGameSync
 			SummaryLine.AddText(SummaryText);
 			Lines.Add(SummaryLine);
 
-			Lines.Add(new StatusLine(){ LineHeight = 0.5f });
-
-			foreach(string MessageLine in Message.Split('\n'))
+			if(!String.IsNullOrEmpty(Message))
 			{
-				StatusLine ErrorLine = new StatusLine();
-				ErrorLine.AddText(MessageLine);
-				ErrorLine.LineHeight = 0.8f;
-				Lines.Add(ErrorLine);
+				Lines.Add(new StatusLine(){ LineHeight = 0.5f });
+
+				foreach(string MessageLine in Message.Split('\n'))
+				{
+					StatusLine ErrorLine = new StatusLine();
+					ErrorLine.AddText(MessageLine);
+					ErrorLine.LineHeight = 0.8f;
+					Lines.Add(ErrorLine);
+				}
 			}
 
 			Lines.Add(new StatusLine(){ LineHeight = 0.5f });
 
 			StatusLine ActionLine = new StatusLine();
-			ActionLine.AddLink("Open...", FontStyle.Bold | FontStyle.Underline, () => { BeginInvoke(new MethodInvoker(() => { EditSelectedProject(ErrorPanel); })); });
-			ActionLine.AddText(" | ");
 			ActionLine.AddLink("Retry", FontStyle.Bold | FontStyle.Underline, () => { BeginInvoke(new MethodInvoker(() => { TryOpenProject(Project, TabControl.FindTabIndex(ErrorPanel)); })); });
 			ActionLine.AddText(" | ");
+			ActionLine.AddLink("Settings", FontStyle.Bold | FontStyle.Underline, () => { BeginInvoke(new MethodInvoker(() => { EditSelectedProject(ErrorPanel); })); });
+			ActionLine.AddText(" | ");
 			ActionLine.AddLink("Close", FontStyle.Bold | FontStyle.Underline, () => { BeginInvoke(new MethodInvoker(() => { TabControl.RemoveTab(TabControl.FindTabIndex(ErrorPanel)); })); });
 			Lines.Add(ActionLine);
 
@@ -573,12 +581,19 @@ namespace UnrealGameSync
 			Refresh();
 		}
 
-		public void RequestProjectChange(WorkspaceControl Workspace, UserSelectedProjectSettings Project)
+		public void RequestProjectChange(WorkspaceControl Workspace, UserSelectedProjectSettings Project, bool bModal)
 		{
 			int TabIdx = TabControl.FindTabIndex(Workspace);
 			if(TabIdx != -1 && !Workspace.IsBusy() && CanFocus)
 			{
-				TryOpenProject(Project, TabIdx);
+				if(bModal)
+				{
+					TryOpenProject(Project, TabIdx);
+				}
+				else
+				{
+					TryOpenProject(Project, TabIdx, OpenProjectOptions.Quiet);
+				}
 			}
 		}
 
@@ -658,12 +673,20 @@ namespace UnrealGameSync
 			using(DetectProjectSettingsTask DetectProjectSettings = new DetectProjectSettingsTask(Project, DataFolder, new PrefixedTextWriter("  ", Log)))
 			{
 				string ErrorMessage;
-				if(ModalTask.Execute(this, DetectProjectSettings, "Opening Project", "Opening project, please wait...", out ErrorMessage) != ModalTaskResult.Succeeded)
+
+				ModalTaskResult Result;
+				if((Options & OpenProjectOptions.Quiet) != 0)
 				{
-					if(!String.IsNullOrEmpty(ErrorMessage) && (Options & OpenProjectOptions.Quiet) == 0)
-					{
-						CreateErrorPanel(ReplaceTabIdx, Project, ErrorMessage);
-					}
+					Result = ModalTask.Execute(this, DetectProjectSettings, "Opening Project", "Opening project, please wait...", out ErrorMessage);
+				}
+				else
+				{
+					Result = PerforceModalTask.Execute(this, Project.LocalPath, Project.ServerAndPort, Project.UserName, DetectProjectSettings, "Opening Project", "Opening project, please wait...", Log, out ErrorMessage);
+				}
+
+				if(Result != ModalTaskResult.Succeeded)
+				{
+					CreateErrorPanel(ReplaceTabIdx, Project, ErrorMessage);
 					return -1;
 				}
 				return TryOpenProject(DetectProjectSettings, ReplaceTabIdx, Options);
