UGS: Always show the main window when launched manually, and run with -RestoreState when launched at startup. Also add a couple more places that save the visibility state, since logging off seems like it can terminate the process abrubtly.

diff --git a/Engine/Source/Programs/UnrealGameSync/UnrealGameSync/Controls/WorkspaceControl.cs b/Engine/Source/Programs/UnrealGameSync/UnrealGameSync/Controls/WorkspaceControl.cs
--- a/Engine/Source/Programs/UnrealGameSync/UnrealGameSync/Controls/WorkspaceControl.cs
+++ b/Engine/Source/Programs/UnrealGameSync/UnrealGameSync/Controls/WorkspaceControl.cs
@@ -4016,7 +4016,7 @@ namespace UnrealGameSync
 			}
 			else
 			{
-				Key.SetValue("UnrealGameSync", OriginalExecutableFileName);
+				Key.SetValue("UnrealGameSync", String.Format("\"{0}\" -RestoreState", OriginalExecutableFileName));
 			}
 		}
 
diff --git a/Engine/Source/Programs/UnrealGameSync/UnrealGameSync/Forms/MainWindow.cs b/Engine/Source/Programs/UnrealGameSync/UnrealGameSync/Forms/MainWindow.cs
--- a/Engine/Source/Programs/UnrealGameSync/UnrealGameSync/Forms/MainWindow.cs
+++ b/Engine/Source/Programs/UnrealGameSync/UnrealGameSync/Forms/MainWindow.cs
@@ -230,7 +230,10 @@ namespace UnrealGameSync
 			if(!bAllowClose && Settings.bKeepInTray)
 			{
 				Hide();
-				EventArgs.Cancel = true; 
+				EventArgs.Cancel = true;
+
+				Settings.bWindowVisible = Visible;
+				Settings.Save();
 			}
 			else
 			{
@@ -354,6 +357,9 @@ namespace UnrealGameSync
 				WindowState = FormWindowState.Normal;
 			}
 			Activate();
+
+			Settings.bWindowVisible = Visible;
+			Settings.Save();
 		}
 
 		public bool CanPerformUpdate()
diff --git a/Engine/Source/Programs/UnrealGameSync/UnrealGameSync/ProgramApplicationContext.cs b/Engine/Source/Programs/UnrealGameSync/UnrealGameSync/ProgramApplicationContext.cs
--- a/Engine/Source/Programs/UnrealGameSync/UnrealGameSync/ProgramApplicationContext.cs
+++ b/Engine/Source/Programs/UnrealGameSync/UnrealGameSync/ProgramApplicationContext.cs
@@ -155,14 +155,18 @@ namespace UnrealGameSync
 			DetectStartupProjectSettingsTask = new DetectMultipleProjectSettingsTask(Tasks);
 
 			DetectStartupProjectSettingsWindow = new ModalTaskWindow(DetectStartupProjectSettingsTask, "Opening Projects", "Opening projects, please wait...", FormStartPosition.CenterScreen);
-			if(Settings.bWindowVisible)
+			if(bRestoreState)
 			{
-				DetectStartupProjectSettingsWindow.Show();
-				if(!bRestoreState)
+				if(Settings.bWindowVisible)
 				{
-					DetectStartupProjectSettingsWindow.Activate();
+					DetectStartupProjectSettingsWindow.Show();
 				}
 			}
+			else
+			{
+				DetectStartupProjectSettingsWindow.Show();
+				DetectStartupProjectSettingsWindow.Activate();
+			}
 			DetectStartupProjectSettingsWindow.Complete += OnDetectStartupProjectsComplete;
 		}
 
