- Deprecate CISController
- Add BuildController to replace CIS GET/POST for builds
- Add LatestController, GET does what CIS/GET used to do
- Change Latest/GET to return the last 25 builds filtered by project, rather than the last 5000 individual Ids
- Latest/GET now returns "LatestData" object instead of array of longs
- Updated EventMonitor to match all API changes
- Fixed bug where IDs were getting reset to initial startup values every update loop

diff --git a/Engine/Source/Programs/UnrealGameSync/PostBadgeStatus/Program.cs b/Engine/Source/Programs/UnrealGameSync/PostBadgeStatus/Program.cs
--- a/Engine/Source/Programs/UnrealGameSync/PostBadgeStatus/Program.cs
+++ b/Engine/Source/Programs/UnrealGameSync/PostBadgeStatus/Program.cs
@@ -53,7 +53,7 @@ namespace WriteBadgeStatus
 				}
 				try
 				{
-					SendRequest(RestUrl, "CIS", "POST", new JavaScriptSerializer().Serialize(Build));
+					SendRequest(RestUrl, "Build", "POST", new JavaScriptSerializer().Serialize(Build));
 				}
 				catch(Exception ex)
 				{
diff --git a/Engine/Source/Programs/UnrealGameSync/UnrealGameSync/EventMonitor.cs b/Engine/Source/Programs/UnrealGameSync/UnrealGameSync/EventMonitor.cs
--- a/Engine/Source/Programs/UnrealGameSync/UnrealGameSync/EventMonitor.cs
+++ b/Engine/Source/Programs/UnrealGameSync/UnrealGameSync/EventMonitor.cs
@@ -34,7 +34,12 @@ namespace UnrealGameSync
 		Investigating,
 		Resolved,
 	}
-
+	class LatestData
+	{
+		public long LastEventId { get; set; }
+		public long LastCommentId { get; set; }
+		public long LastBuildId { get; set; }
+	}
 	class EventData
 	{
 		public long Id { get; set; }
@@ -150,10 +155,8 @@ namespace UnrealGameSync
 		Dictionary<string, EventData> UserNameToLastSyncEvent = new Dictionary<string, EventData>(StringComparer.InvariantCultureIgnoreCase);
 		Dictionary<string, BuildData> BadgeNameToLatestBuild = new Dictionary<string, BuildData>();
 		BoundedLogWriter LogWriter;
-		long LastEventId;
-		long LastCommentId;
-		long LastBuildId;
 		bool bDisposing;
+		LatestData LatestIds;
 		HashSet<int> FilterChangeNumbers = new HashSet<int>();
 		List<EventData> InvestigationEvents = new List<EventData>();
 		List<EventData> ActiveInvestigations = new List<EventData>();
@@ -165,6 +168,7 @@ namespace UnrealGameSync
 			ApiUrl = InApiUrl;
 			Project = InProject;
 			CurrentUserName = InCurrentUserName;
+			LatestIds = new LatestData { LastBuildId = 0, LastCommentId = 0, LastEventId = 0 };
 
 			LogWriter = new BoundedLogWriter(InLogFileName);
 			if(ApiUrl == null)
@@ -531,41 +535,44 @@ namespace UnrealGameSync
 				LogWriter.WriteLine();
 				LogWriter.WriteLine("Polling for reviews at {0}...", DateTime.Now.ToString());
 				//////////////
-				/// LastIds 
+				/// Initial Ids 
 				//////////////
-				long[] LastIds= RESTApi.GET<long[]>(ApiUrl, "cis", "GET");
-				LastEventId = LastIds[0];
-				LastCommentId = LastIds[1];
-				LastBuildId = LastIds[2];
+				if (LatestIds.LastBuildId == 0 && LatestIds.LastCommentId == 0 && LatestIds.LastEventId == 0)
+				{
+					LatestData InitialIds = RESTApi.GET<LatestData>(ApiUrl, "latest", string.Format("project={0}", Project));
+					LatestIds.LastBuildId = InitialIds.LastBuildId;
+					LatestIds.LastCommentId = InitialIds.LastCommentId;
+					LatestIds.LastEventId = InitialIds.LastEventId;
+				}
 
 				//////////////
 				/// Reviews 
 				//////////////
-				List<EventData> Events = RESTApi.GET<List<EventData>>(ApiUrl, "event", string.Format("project={0}", Project), string.Format("lasteventid={0}", LastEventId));
+				List<EventData> Events = RESTApi.GET<List<EventData>>(ApiUrl, "event", string.Format("project={0}", Project), string.Format("lasteventid={0}", LatestIds.LastEventId));
 				foreach(EventData Review in Events)
 				{
 					IncomingEvents.Enqueue(Review);
-					LastEventId = Math.Max(LastEventId, Review.Id);
+					LatestIds.LastEventId = Math.Max(LatestIds.LastEventId, Review.Id);
 				}
 
 				//////////////
 				/// Comments 
 				//////////////
-				List<CommentData> Comments = RESTApi.GET<List<CommentData>>(ApiUrl, "comment", string.Format("project={0}", Project), string.Format("lastcommentid={0}", LastCommentId));
+				List<CommentData> Comments = RESTApi.GET<List<CommentData>>(ApiUrl, "comment", string.Format("project={0}", Project), string.Format("lastcommentid={0}", LatestIds.LastCommentId));
 				foreach (CommentData Comment in Comments)
 				{
 					IncomingComments.Enqueue(Comment);
-					LastCommentId = Math.Max(LastCommentId, Comment.Id);
+					LatestIds.LastCommentId = Math.Max(LatestIds.LastCommentId, Comment.Id);
 				}
 
 				//////////////
-				/// CIS 
+				/// Bulids
 				//////////////
-				List<BuildData> Builds = RESTApi.GET<List<BuildData>>(ApiUrl, "cis", string.Format("project={0}", Project), string.Format("lastbuildid={0}", LastBuildId));
+				List<BuildData> Builds = RESTApi.GET<List<BuildData>>(ApiUrl, "build", string.Format("project={0}", Project), string.Format("lastbuildid={0}", LatestIds.LastBuildId));
 				foreach (BuildData Build in Builds)
 				{
 					IncomingBuilds.Enqueue(Build);
-					LastBuildId = Math.Max(LastBuildId, Build.Id);
+					LatestIds.LastBuildId = Math.Max(LatestIds.LastBuildId, Build.Id);
 				}
 
 				LastStatusMessage = String.Format("Last update took {0}ms", Timer.ElapsedMilliseconds);
diff --git a/Engine/Source/Programs/UnrealGameSync/UnrealGameSyncMetadataServer/Connectors/SqlConnector.cs b/Engine/Source/Programs/UnrealGameSync/UnrealGameSyncMetadataServer/Connectors/SqlConnector.cs
--- a/Engine/Source/Programs/UnrealGameSync/UnrealGameSyncMetadataServer/Connectors/SqlConnector.cs
+++ b/Engine/Source/Programs/UnrealGameSync/UnrealGameSyncMetadataServer/Connectors/SqlConnector.cs
@@ -13,63 +13,67 @@ namespace UnrealGameSyncMetadataServer.Connectors
 	public static class SqlConnector
 	{
 		private static string ConnectionString = System.Configuration.ConfigurationManager.AppSettings["ConnectionString"];
-		public static long[] GetLastIds()
+		public static LatestData GetLastIds(string Project = null)
 		{
-			// Get Last IDs
+			// Get ids going back 25 builds for the project being asked for
+			// Do this by grouping by ChangeNumber to get unique entries, then take the 25th id
 			long LastEventId = 0;
 			long LastCommentId = 0;
-			long LastBuildID = 0;
+			long LastBuildId = 0;
+			string ProjectLikeString = "%" + (Project == null ? String.Empty : Project) + "%";
 			using (SQLiteConnection Connection = new SQLiteConnection(ConnectionString))
 			{
 				Connection.Open();
-				using (SQLiteCommand Command = new SQLiteCommand("SELECT MAX(ID) FROM [UserVotes];", Connection))
+				using (SQLiteCommand Command = new SQLiteCommand("SELECT Id FROM [UserVotes] WHERE Project LIKE @param1 GROUP BY Changelist ORDER BY Changelist DESC LIMIT 1 OFFSET 24;", Connection))
 				{
+					Command.Parameters.AddWithValue("@param1", ProjectLikeString);
 					using (SQLiteDataReader Reader = Command.ExecuteReader())
 					{
 						while (Reader.Read())
 						{
 							LastEventId = Reader.GetInt64(0);
-							LastEventId = Math.Max(LastEventId - 5000, 0);
 							break;
 						}
 					}
 				}
-				using (SQLiteCommand Command = new SQLiteCommand("SELECT MAX(ID) FROM [Comments];", Connection))
+				using (SQLiteCommand Command = new SQLiteCommand("SELECT Id FROM [Comments] WHERE Project LIKE @param1 GROUP BY ChangeNumber ORDER BY ChangeNumber DESC LIMIT 1 OFFSET 24;", Connection))
 				{
+					Command.Parameters.AddWithValue("@param1", ProjectLikeString);
 					using (SQLiteDataReader Reader = Command.ExecuteReader())
 					{
 						while (Reader.Read())
 						{
 							LastCommentId = Reader.GetInt32(0);
-							LastCommentId = Math.Max(LastCommentId - 5000, 0);
 							break;
 						}
 					}
 				}
-				using (SQLiteCommand Command = new SQLiteCommand("SELECT MAX(ID) FROM [CIS];", Connection))
+				using (SQLiteCommand Command = new SQLiteCommand("SELECT Id FROM [CIS] WHERE Project LIKE @param1 GROUP BY ChangeNumber ORDER BY ChangeNumber DESC LIMIT 1 OFFSET 24", Connection))
 				{
+					Command.Parameters.AddWithValue("@param1", ProjectLikeString);
 					using (SQLiteDataReader Reader = Command.ExecuteReader())
 					{
 						while (Reader.Read())
 						{
-							LastBuildID = Reader.GetInt32(0);
-							LastBuildID = Math.Max(LastBuildID - 5000, 0);
+							LastBuildId = Reader.GetInt32(0);
 							break;
 						}
 					}
 				}
 			}
-			return new long[] { LastEventId, LastCommentId, LastBuildID };
+			return new LatestData { LastBuildId = LastBuildId, LastCommmentId = LastCommentId, LastEventId = LastEventId };
 		}
 		public static List<EventData> GetUserVotes(string Project, long LastEventId)
 		{
 			List<EventData> ReturnedEvents = new List<EventData>();
+			string ProjectLikeString = "%" + (Project == null ? String.Empty : Project) + "%";
 			using (SQLiteConnection Connection = new SQLiteConnection(ConnectionString))
 			{
 				Connection.Open();
-				using (SQLiteCommand Command = new SQLiteCommand("SELECT Id, Changelist, UserName, Verdict, Project FROM [UserVotes] WHERE Id > @param1 ORDER BY Id", Connection))
+				using (SQLiteCommand Command = new SQLiteCommand("SELECT Id, Changelist, UserName, Verdict, Project FROM [UserVotes] WHERE Id > @param1 AND Project LIKE @param2 ORDER BY Id", Connection))
 				{
 					Command.Parameters.AddWithValue("@param1", LastEventId);
+					Command.Parameters.AddWithValue("@param2", ProjectLikeString);
 					using (SQLiteDataReader Reader = Command.ExecuteReader())
 					{
 						while (Reader.Read())
@@ -95,12 +99,14 @@ namespace UnrealGameSyncMetadataServer.Connectors
 		public static List<CommentData> GetComments(string Project, long LastCommentId)
 		{
 			List<CommentData> ReturnedComments = new List<CommentData>();
+			string ProjectLikeString = "%" + (Project == null ? String.Empty : Project) + "%";
 			using (SQLiteConnection Connection = new SQLiteConnection(ConnectionString))
 			{
 				Connection.Open();
-				using (SQLiteCommand Command = new SQLiteCommand("SELECT Id, ChangeNumber, UserName, Text, Project FROM [Comments] WHERE Id > @param1 ORDER BY Id", Connection))
+				using (SQLiteCommand Command = new SQLiteCommand("SELECT Id, ChangeNumber, UserName, Text, Project FROM [Comments] WHERE Id > @param1 AND Project LIKE @param2 ORDER BY Id", Connection))
 				{
 					Command.Parameters.AddWithValue("@param1", LastCommentId);
+					Command.Parameters.AddWithValue("@param2", ProjectLikeString);
 					using (SQLiteDataReader Reader = Command.ExecuteReader())
 					{
 						while (Reader.Read())
@@ -122,15 +128,17 @@ namespace UnrealGameSyncMetadataServer.Connectors
 			return ReturnedComments;
 		}
 
-		public static List<BuildData> GetCIS(string Project, long LastBuildId)
+		public static List<BuildData> GetBuilds(string Project, long LastBuildId)
 		{
 			List<BuildData> ReturnedBuilds = new List<BuildData>();
+			string ProjectLikeString = "%" + (Project == null ? String.Empty : Project) + "%";
 			using (SQLiteConnection Connection = new SQLiteConnection(ConnectionString))
 			{
 				Connection.Open();
-				using (SQLiteCommand Command = new SQLiteCommand("SELECT Id, ChangeNumber, BuildType, Result, Url, Project FROM [CIS] WHERE Id > @param1 ORDER BY Id", Connection))
+				using (SQLiteCommand Command = new SQLiteCommand("SELECT Id, ChangeNumber, BuildType, Result, Url, Project FROM [CIS] WHERE Id > @param1 AND Project LIKE @param2 ORDER BY Id", Connection))
 				{
 					Command.Parameters.AddWithValue("@param1", LastBuildId);
+					Command.Parameters.AddWithValue("@param2", ProjectLikeString);
 					using (SQLiteDataReader Reader = Command.ExecuteReader())
 					{
 						while (Reader.Read())
@@ -187,7 +195,7 @@ namespace UnrealGameSyncMetadataServer.Connectors
 		}
 
 
-		public static void PostCIS(BuildData Build)
+		public static void PostBuild(BuildData Build)
 		{
 			using (SQLiteConnection Connection = new SQLiteConnection(ConnectionString))
 			{
diff --git a/Engine/Source/Programs/UnrealGameSync/UnrealGameSyncMetadataServer/Controllers/BuildController.cs b/Engine/Source/Programs/UnrealGameSync/UnrealGameSyncMetadataServer/Controllers/BuildController.cs
new file mode 100644
--- /dev/null
+++ b/Engine/Source/Programs/UnrealGameSync/UnrealGameSyncMetadataServer/Controllers/BuildController.cs
@@ -0,0 +1,21 @@
+﻿// Copyright 1998-2018 Epic Games, Inc. All Rights Reserved.
+
+using System.Collections.Generic;
+using System.Web.Http;
+using UnrealGameSyncMetadataServer.Connectors;
+using UnrealGameSyncMetadataServer.Models;
+
+namespace UnrealGameSyncMetadataServer.Controllers
+{
+    public class BuildController : ApiController
+    {
+		public List<BuildData> Get(string Project, long LastBuildId)
+		{
+			return SqlConnector.GetBuilds(Project, LastBuildId);
+		}
+		public void Post([FromBody]BuildData Build)
+		{
+			SqlConnector.PostBuild(Build);
+		}
+	}
+}
diff --git a/Engine/Source/Programs/UnrealGameSync/UnrealGameSyncMetadataServer/Controllers/CISController.cs b/Engine/Source/Programs/UnrealGameSync/UnrealGameSyncMetadataServer/Controllers/CISController.cs
--- a/Engine/Source/Programs/UnrealGameSync/UnrealGameSyncMetadataServer/Controllers/CISController.cs
+++ b/Engine/Source/Programs/UnrealGameSync/UnrealGameSyncMetadataServer/Controllers/CISController.cs
@@ -1,5 +1,6 @@
 ﻿// Copyright 1998-2018 Epic Games, Inc. All Rights Reserved.
 
+using System;
 using System.Collections.Generic;
 using System.Web.Http;
 using UnrealGameSyncMetadataServer.Connectors;
@@ -7,19 +8,20 @@ using UnrealGameSyncMetadataServer.Models;
 
 namespace UnrealGameSyncMetadataServer.Controllers
 {
+	// Deprecated, use Latest or Build controllers instead
     public class CISController : ApiController
     {
-		public long[] Get()
+		public LatestData Get(string Project = null)
 		{
-			return SqlConnector.GetLastIds();
+			return SqlConnector.GetLastIds(Project);
 		}
 		public List<BuildData> Get(string Project, long LastBuildId)
 		{
-			return SqlConnector.GetCIS(Project, LastBuildId);
+			return SqlConnector.GetBuilds(Project, LastBuildId);
 		}
 		public void Post([FromBody]BuildData Build)
 		{
-			SqlConnector.PostCIS(Build);
+			SqlConnector.PostBuild(Build);
 		}
 	}
 }
diff --git a/Engine/Source/Programs/UnrealGameSync/UnrealGameSyncMetadataServer/Controllers/LatestController.cs b/Engine/Source/Programs/UnrealGameSync/UnrealGameSyncMetadataServer/Controllers/LatestController.cs
new file mode 100644
--- /dev/null
+++ b/Engine/Source/Programs/UnrealGameSync/UnrealGameSyncMetadataServer/Controllers/LatestController.cs
@@ -0,0 +1,21 @@
+﻿// Copyright 1998-2018 Epic Games, Inc. All Rights Reserved.
+
+using System;
+using System.Collections.Generic;
+using System.Linq;
+using System.Net;
+using System.Net.Http;
+using System.Web.Http;
+using UnrealGameSyncMetadataServer.Connectors;
+using UnrealGameSyncMetadataServer.Models;
+
+namespace UnrealGameSyncMetadataServer.Controllers
+{
+    public class LatestController : ApiController
+    {
+		public LatestData Get(string Project = null)
+		{
+			return SqlConnector.GetLastIds(Project);
+		}
+	}
+}
diff --git a/Engine/Source/Programs/UnrealGameSync/UnrealGameSyncMetadataServer/Models/LatestData.cs b/Engine/Source/Programs/UnrealGameSync/UnrealGameSyncMetadataServer/Models/LatestData.cs
new file mode 100644
--- /dev/null
+++ b/Engine/Source/Programs/UnrealGameSync/UnrealGameSyncMetadataServer/Models/LatestData.cs
@@ -0,0 +1,16 @@
+﻿// Copyright 1998-2018 Epic Games, Inc. All Rights Reserved.
+
+using System;
+using System.Collections.Generic;
+using System.Linq;
+using System.Web;
+
+namespace UnrealGameSyncMetadataServer.Models
+{
+	public class LatestData
+	{
+		public long LastEventId { get; set; }
+		public long LastCommmentId { get; set; }
+		public long LastBuildId { get; set; }
+	}
+}
\ No newline at end of file
diff --git a/Engine/Source/Programs/UnrealGameSync/UnrealGameSyncMetadataServer/UnrealGameSyncMetadataServer.csproj b/Engine/Source/Programs/UnrealGameSync/UnrealGameSyncMetadataServer/UnrealGameSyncMetadataServer.csproj
--- a/Engine/Source/Programs/UnrealGameSync/UnrealGameSyncMetadataServer/UnrealGameSyncMetadataServer.csproj
+++ b/Engine/Source/Programs/UnrealGameSync/UnrealGameSyncMetadataServer/UnrealGameSyncMetadataServer.csproj
@@ -89,10 +89,12 @@
   </ItemGroup>
   <ItemGroup>
     <Compile Include="App_Start\WebApiConfig.cs" />
+    <Compile Include="Controllers\BuildController.cs" />
     <Compile Include="Controllers\CISController.cs" />
     <Compile Include="Controllers\CommentController.cs" />
     <Compile Include="Controllers\ErrorController.cs" />
     <Compile Include="Controllers\EventController.cs" />
+    <Compile Include="Controllers\LatestController.cs" />
     <Compile Include="Controllers\TelemetryController.cs" />
     <Compile Include="Global.asax.cs">
       <DependentUpon>Global.asax</DependentUpon>
@@ -101,6 +103,7 @@
     <Compile Include="Models\CommentData.cs" />
     <Compile Include="Models\EventData.cs" />
     <Compile Include="Models\EventSummary.cs" />
+    <Compile Include="Models\LatestData.cs" />
     <Compile Include="Models\TelemetryErrorData.cs" />
     <Compile Include="Models\TelemetryTimingData.cs" />
     <Compile Include="Properties\AssemblyInfo.cs" />
