CIS controller still needs to return an array of longs

diff --git a/Engine/Source/Programs/UnrealGameSync/UnrealGameSyncMetadataServer/Controllers/CISController.cs b/Engine/Source/Programs/UnrealGameSync/UnrealGameSyncMetadataServer/Controllers/CISController.cs
--- a/Engine/Source/Programs/UnrealGameSync/UnrealGameSyncMetadataServer/Controllers/CISController.cs
+++ b/Engine/Source/Programs/UnrealGameSync/UnrealGameSyncMetadataServer/Controllers/CISController.cs
@@ -11,9 +11,10 @@ namespace UnrealGameSyncMetadataServer.Controllers
 	// Deprecated, use Latest or Build controllers instead
     public class CISController : ApiController
     {
-		public LatestData Get(string Project = null)
+		public long[] Get(string Project = null)
 		{
-			return SqlConnector.GetLastIds(Project);
+			LatestData LatestIDs = SqlConnector.GetLastIds(Project);
+			return new long[] { LatestIDs.LastEventId, LatestIDs.LastCommmentId, LatestIDs.LastBuildId };
 		}
 		public List<BuildData> Get(string Project, long LastBuildId)
 		{
