UGS: Always show the "Sync Precompiled Editor" option, but disable it and show a tooltip explaining why if it is not available.

diff --git a/Engine/Source/Programs/UnrealGameSync/UnrealGameSync/Controls/WorkspaceControl.cs b/Engine/Source/Programs/UnrealGameSync/UnrealGameSync/Controls/WorkspaceControl.cs
--- a/Engine/Source/Programs/UnrealGameSync/UnrealGameSync/Controls/WorkspaceControl.cs
+++ b/Engine/Source/Programs/UnrealGameSync/UnrealGameSync/Controls/WorkspaceControl.cs
@@ -3190,8 +3190,9 @@ namespace UnrealGameSync
 		private void OptionsButton_Click(object sender, EventArgs e)
 		{
 			OptionsContextMenu_AutoResolveConflicts.Checked = Settings.bAutoResolveConflicts;
-			OptionsContextMenu_SyncPrecompiledEditor.Visible = PerforceMonitor != null && PerforceMonitor.HasZippedBinaries;
+			OptionsContextMenu_SyncPrecompiledEditor.Enabled = PerforceMonitor != null && PerforceMonitor.HasZippedBinaries;
 			OptionsContextMenu_SyncPrecompiledEditor.Checked = Settings.bSyncPrecompiledEditor;
+			OptionsContextMenu_SyncPrecompiledEditor.ToolTipText = PerforceMonitor.ZippedBinariesStatus;
 			OptionsContextMenu_EditorBuildConfiguration.Enabled = !ShouldSyncPrecompiledEditor;
 			UpdateCheckedBuildConfig();
 			OptionsContextMenu_UseIncrementalBuilds.Enabled = !ShouldSyncPrecompiledEditor;
diff --git a/Engine/Source/Programs/UnrealGameSync/UnrealGameSync/PerforceMonitor.cs b/Engine/Source/Programs/UnrealGameSync/UnrealGameSync/PerforceMonitor.cs
--- a/Engine/Source/Programs/UnrealGameSync/UnrealGameSync/PerforceMonitor.cs
+++ b/Engine/Source/Programs/UnrealGameSync/UnrealGameSync/PerforceMonitor.cs
@@ -38,8 +38,6 @@ namespace UnrealGameSync
 		SortedSet<PerforceChangeSummary> Changes = new SortedSet<PerforceChangeSummary>(new PerforceChangeSorter());
 		SortedDictionary<int, PerforceChangeDetails> ChangeDetails = new SortedDictionary<int,PerforceChangeDetails>();
 		SortedSet<int> PromotedChangeNumbers = new SortedSet<int>();
-		int ZippedBinariesConfigChangeNumber;
-		string ZippedBinariesPath;
 		SortedList<int, string> ChangeNumberToZippedBinaries = new SortedList<int,string>();
 		AutoResetEvent RefreshEvent = new AutoResetEvent(false);
 		BoundedLogWriter LogWriter;
@@ -436,70 +434,72 @@ namespace UnrealGameSync
 
 		bool UpdateZippedBinaries()
 		{
-			// Get the path to the config file
-			string ClientConfigFileName = PerforceUtils.GetClientOrDepotDirectoryName(SelectedClientFileName) + "/Build/UnrealGameSync.ini";
+			string ZippedBinariesPath = null;
 
-			// Find the most recent change to that file (if the file doesn't exist, we succeed and just get 0 changes).
-			List<PerforceChangeSummary> ConfigFileChanges;
-			if(!Perforce.FindChanges(ClientConfigFileName, 1, out ConfigFileChanges, LogWriter))
+			// Find all the zipped binaries under this stream
+			ConfigSection ProjectConfigSection = LatestProjectConfigFile.FindSection(SelectedProjectIdentifier);
+			if(ProjectConfigSection != null)
 			{
-				return false;
-			}
-
-			// Update the zipped binaries path if it's changed
-			int NewZippedBinariesConfigChangeNumber = (ConfigFileChanges.Count > 0)? ConfigFileChanges[0].Number : 0;
-			if(NewZippedBinariesConfigChangeNumber != ZippedBinariesConfigChangeNumber)
-			{
-				string NewZippedBinariesPath = null;
-				if(NewZippedBinariesConfigChangeNumber != 0)
-				{
-					List<string> Lines;
-					if(!Perforce.Print(ClientConfigFileName, out Lines, LogWriter))
-					{
-						return false;
-					}
-
-					ConfigFile NewConfigFile = new ConfigFile();
-					NewConfigFile.Parse(Lines.ToArray());
-
-					ConfigSection ProjectSection = NewConfigFile.FindSection(SelectedProjectIdentifier);
-					if(ProjectSection != null)
-					{
-						NewZippedBinariesPath = ProjectSection.GetValue("ZippedBinariesPath", null);
-					}
-				}
-				ZippedBinariesPath = NewZippedBinariesPath;
-				ZippedBinariesConfigChangeNumber = NewZippedBinariesConfigChangeNumber;
+				ZippedBinariesPath = ProjectConfigSection.GetValue("ZippedBinariesPath", null);
 			}
 
+			// Build a new list of zipped binaries
 			SortedList<int, string> NewChangeNumberToZippedBinaries = new SortedList<int,string>();
 			if(ZippedBinariesPath != null)
 			{
-				List<PerforceFileChangeSummary> Changes;
-				if(!Perforce.FindFileChanges(ZippedBinariesPath, 100, out Changes, LogWriter))
+				// Make sure the zipped binaries path exists
+				bool bExists;
+				if(!Perforce.FileExists(ZippedBinariesPath, out bExists, LogWriter))
 				{
 					return false;
 				}
-
-				foreach(PerforceFileChangeSummary Change in Changes)
+				if(bExists)
 				{
-					if(Change.Action != "purge")
+					// Query all the changes to this file
+					List<PerforceFileChangeSummary> Changes;
+					if(!Perforce.FindFileChanges(ZippedBinariesPath, 100, out Changes, LogWriter))
 					{
-						string[] Tokens = Change.Description.Split(' ');
-						if(Tokens[0].StartsWith("[CL") && Tokens[1].EndsWith("]"))
+						return false;
+					}
+
+					// Build a new list of zipped binaries
+					foreach(PerforceFileChangeSummary Change in Changes)
+					{
+						if(Change.Action != "purge")
 						{
-							int OriginalChangeNumber;
-							if(int.TryParse(Tokens[1].Substring(0, Tokens[1].Length - 1), out OriginalChangeNumber) && !NewChangeNumberToZippedBinaries.ContainsKey(OriginalChangeNumber))
+							string[] Tokens = Change.Description.Split(' ');
+							if(Tokens[0].StartsWith("[CL") && Tokens[1].EndsWith("]"))
 							{
-								NewChangeNumberToZippedBinaries[OriginalChangeNumber] = String.Format("{0}#{1}", ZippedBinariesPath, Change.Revision);
+								int OriginalChangeNumber;
+								if(int.TryParse(Tokens[1].Substring(0, Tokens[1].Length - 1), out OriginalChangeNumber) && !NewChangeNumberToZippedBinaries.ContainsKey(OriginalChangeNumber))
+								{
+									NewChangeNumberToZippedBinaries[OriginalChangeNumber] = String.Format("{0}#{1}", ZippedBinariesPath, Change.Revision);
+								}
 							}
 						}
 					}
 				}
 			}
 
-			if(!ChangeNumberToZippedBinaries.SequenceEqual(NewChangeNumberToZippedBinaries))
+			// Get the new status message
+			string NewZippedBinariesStatus;
+			if(ZippedBinariesPath == null)
+			{
+				NewZippedBinariesStatus = String.Format("Precompiled binaries are not available for {0}", SelectedProjectIdentifier);
+			}
+			else if(NewChangeNumberToZippedBinaries.Count == 0)
+			{
+				NewZippedBinariesStatus = String.Format("No valid archives found at {0}", ZippedBinariesPath);
+			}
+			else
 			{
+				NewZippedBinariesStatus = null;
+			}
+
+			// Update the new list of zipped binaries
+			if(!ChangeNumberToZippedBinaries.SequenceEqual(NewChangeNumberToZippedBinaries) || ZippedBinariesStatus != NewZippedBinariesStatus)
+			{
+				ZippedBinariesStatus = NewZippedBinariesStatus;
 				ChangeNumberToZippedBinaries = NewChangeNumberToZippedBinaries;
 				if(OnUpdateMetadata != null && Changes.Count > 0)
 				{
@@ -594,25 +594,25 @@ namespace UnrealGameSync
 			}
 		}
 
-		public string CurrentStream
+		public int LastChangeByCurrentUser
 		{
 			get;
 			private set;
 		}
 
-		public int LastChangeByCurrentUser
+		public int LastCodeChangeByCurrentUser
 		{
 			get;
 			private set;
 		}
 
-		public int LastCodeChangeByCurrentUser
+		public ConfigFile LatestProjectConfigFile
 		{
 			get;
 			private set;
 		}
 
-		public ConfigFile LatestProjectConfigFile
+		public string ZippedBinariesStatus
 		{
 			get;
 			private set;
