UGS: Add additional logging whenever a P4 command fails, and when the user is logged out.

diff --git a/Engine/Source/Programs/UnrealGameSync/UnrealGameSync/Controls/WorkspaceControl.cs b/Engine/Source/Programs/UnrealGameSync/UnrealGameSync/Controls/WorkspaceControl.cs
--- a/Engine/Source/Programs/UnrealGameSync/UnrealGameSync/Controls/WorkspaceControl.cs
+++ b/Engine/Source/Programs/UnrealGameSync/UnrealGameSync/Controls/WorkspaceControl.cs
@@ -607,6 +607,7 @@ namespace UnrealGameSync
 		{
 			if(!bIsDisposing)
 			{
+				Log.WriteLine("Login has expired. Requesting project to be closed.");
 				Owner.RequestProjectChange(this, SelectedProject, false);
 			}
 		}
diff --git a/Engine/Source/Programs/UnrealGameSync/UnrealGameSync/Perforce.cs b/Engine/Source/Programs/UnrealGameSync/UnrealGameSync/Perforce.cs
--- a/Engine/Source/Programs/UnrealGameSync/UnrealGameSync/Perforce.cs
+++ b/Engine/Source/Programs/UnrealGameSync/UnrealGameSync/Perforce.cs
@@ -1241,7 +1241,15 @@ namespace UnrealGameSync
 		private bool RunCommand(string CommandLine, CommandOptions Options, TextWriter Log)
 		{
 			List<string> Lines;
-			return RunCommand(CommandLine, out Lines, Options, Log);
+			bool bResult = RunCommand(CommandLine, out Lines, Options, Log);
+			if(Lines != null)
+			{
+				foreach(string Line in Lines)
+				{
+					Log.WriteLine("p4>   {0}", Line);
+				}
+			}
+			return bResult;
 		}
 
 		private bool RunCommand(string CommandLine, out List<string> Lines, CommandOptions Options, TextWriter Log)
@@ -1258,6 +1266,10 @@ namespace UnrealGameSync
 			if(Utility.ExecuteProcess("p4.exe", null, FullCommandLine, null, out RawOutputLines) != 0 && !Options.HasFlag(CommandOptions.IgnoreExitCode))
 			{
 				Lines = null;
+				foreach(string RawOutputLine in RawOutputLines)
+				{
+					Log.WriteLine("p4>   {0}", RawOutputLine);
+				}
 				return false;
 			}
 
@@ -1271,13 +1283,27 @@ namespace UnrealGameSync
 				List<string> LocalLines = new List<string>();
 				foreach(string RawOutputLine in RawOutputLines)
 				{
-					bResult &= ParseCommandOutput(RawOutputLine, Line => { if(Line.Channel == Channel){ LocalLines.Add(Line.Text); return true; } else { Log.WriteLine(Line.Text); return Line.Channel != PerforceOutputChannel.Error; } }, Options);
+					bResult &= ParseCommandOutput(RawOutputLine, Line => FilterOutput(Line, Channel, LocalLines, Log), Options);
 				}
 				Lines = LocalLines;
 			}
 			return bResult;
 		}
 
+		private bool FilterOutput(PerforceOutputLine Line, PerforceOutputChannel FilterChannel, List<string> FilterLines, TextWriter Log)
+		{
+			if(Line.Channel == FilterChannel)
+			{
+				FilterLines.Add(Line.Text);
+				return true;
+			}
+			else
+			{
+				Log.WriteLine(Line.Text);
+				return Line.Channel != PerforceOutputChannel.Error;
+			}
+		}
+
 		private bool RunCommand(string CommandLine, string Input, HandleOutputDelegate HandleOutput, CommandOptions Options, TextWriter Log)
 		{
 			string FullCommandLine = GetFullCommandLine(CommandLine, Options);
diff --git a/Engine/Source/Programs/UnrealGameSync/UnrealGameSync/Properties/AssemblyInfo.cs b/Engine/Source/Programs/UnrealGameSync/UnrealGameSync/Properties/AssemblyInfo.cs
--- a/Engine/Source/Programs/UnrealGameSync/UnrealGameSync/Properties/AssemblyInfo.cs
+++ b/Engine/Source/Programs/UnrealGameSync/UnrealGameSync/Properties/AssemblyInfo.cs
@@ -34,5 +34,5 @@ using System.Runtime.InteropServices;
 // You can specify all the values or you can default the Build and Revision Numbers 
 // by using the '*' as shown below:
 // [assembly: AssemblyVersion("1.0.*")]
-[assembly: AssemblyVersion("1.150.*")]
-[assembly: AssemblyFileVersion("1.150.0.0")]
+[assembly: AssemblyVersion("1.151.*")]
+[assembly: AssemblyFileVersion("1.151.0.0")]
