UGS: Fix custom build steps dialog inadvertantly modifying config file settings in-place.

diff --git a/Engine/Source/Programs/UnrealGameSync/UnrealGameSync/Controls/WorkspaceControl.cs b/Engine/Source/Programs/UnrealGameSync/UnrealGameSync/Controls/WorkspaceControl.cs
--- a/Engine/Source/Programs/UnrealGameSync/UnrealGameSync/Controls/WorkspaceControl.cs
+++ b/Engine/Source/Programs/UnrealGameSync/UnrealGameSync/Controls/WorkspaceControl.cs
@@ -3936,12 +3936,16 @@ namespace UnrealGameSync
 			Guid UniqueId;
 			if(Guid.TryParse(Object.GetValue(BuildStep.UniqueIdKey, ""), out UniqueId))
 			{
+				// Add or apply Object to the list of steps in Steps. Do not modify Object; make a copy first.
+				ConfigObject NewObject = new ConfigObject(Object);
+
 				ConfigObject DefaultObject;
 				if(Steps.TryGetValue(UniqueId, out DefaultObject))
 				{
-					Object.SetDefaults(DefaultObject);
+					NewObject.SetDefaults(DefaultObject);
 				}
-				Steps[UniqueId] = Object;
+
+				Steps[UniqueId] = NewObject;
 			}
 		}
 
