UGS: Allow for p4 login commands to fail, even though the user is logged in (due to a bad connection, etc...)

diff --git a/Engine/Source/Programs/UnrealGameSync/UnrealGameSync/DetectProjectSettingsTask.cs b/Engine/Source/Programs/UnrealGameSync/UnrealGameSync/DetectProjectSettingsTask.cs
--- a/Engine/Source/Programs/UnrealGameSync/UnrealGameSync/DetectProjectSettingsTask.cs
+++ b/Engine/Source/Programs/UnrealGameSync/UnrealGameSync/DetectProjectSettingsTask.cs
@@ -64,9 +64,17 @@ namespace UnrealGameSync
 				return false;
 			}
 
-			// Create the connection and make sure we're logged in
+			// Create the connection
 			PerforceConnection Perforce = new PerforceConnection(UserName, null, ServerAndPort);
-			if(!Perforce.IsLoggedIn(Log))
+
+			// Make sure we're logged in
+			bool bLoggedIn;
+			if(!Perforce.GetLoggedInState(out bLoggedIn, Log))
+			{
+				ErrorMessage = "Error while checking login status.";
+				return false;
+			}
+			if(!bLoggedIn)
 			{
 				ErrorMessage = "User is not logged in to Perforce.";
 				return false;
diff --git a/Engine/Source/Programs/UnrealGameSync/UnrealGameSync/Perforce.cs b/Engine/Source/Programs/UnrealGameSync/UnrealGameSync/Perforce.cs
--- a/Engine/Source/Programs/UnrealGameSync/UnrealGameSync/Perforce.cs
+++ b/Engine/Source/Programs/UnrealGameSync/UnrealGameSync/Perforce.cs
@@ -491,9 +491,31 @@ namespace UnrealGameSync
 			return new PerforceConnection(UserName, NewClientName, ServerAndPort);
 		}
 
-		public bool IsLoggedIn(TextWriter Log)
+		public bool GetLoggedInState(out bool bIsLoggedIn, TextWriter Log)
 		{
-			return RunCommand("login -s", CommandOptions.None, Log);
+			List<PerforceOutputLine> Lines = new List<PerforceOutputLine>();
+			bool bResult = RunCommand("login -s", null, Line => { Lines.Add(Line); return true; }, CommandOptions.None, Log);
+
+			foreach(PerforceOutputLine Line in Lines)
+			{
+				Log.WriteLine("p4>   {0}", Line.Text);
+			}
+
+			if(bResult)
+			{
+				bIsLoggedIn = true;
+				return true;
+			}
+			else if(Lines[0].Channel == PerforceOutputChannel.Error && Lines[0].Text.Contains("P4PASSWD"))
+			{
+				bIsLoggedIn = false;
+				return true;
+			}
+			else
+			{
+				bIsLoggedIn = false;
+				return false;
+			}
 		}
 
 		public LoginResult Login(string Password, out string ErrorMessage, TextWriter Log)
diff --git a/Engine/Source/Programs/UnrealGameSync/UnrealGameSync/PerforceModalTask.cs b/Engine/Source/Programs/UnrealGameSync/UnrealGameSync/PerforceModalTask.cs
--- a/Engine/Source/Programs/UnrealGameSync/UnrealGameSync/PerforceModalTask.cs
+++ b/Engine/Source/Programs/UnrealGameSync/UnrealGameSync/PerforceModalTask.cs
@@ -124,7 +124,13 @@ namespace UnrealGameSync
 			}
 
 			// Check that we're logged in
-			if(!Perforce.IsLoggedIn(Log))
+			bool bLoggedIn;
+			if(!Perforce.GetLoggedInState(out bLoggedIn, Log))
+			{
+				ErrorMessage = "Unable to get login status.";
+				return false;
+			}
+			if(!bLoggedIn)
 			{
 				LoginResult = LoginResult.MissingPassword;
 				ErrorMessage = "User is not logged in to Perforce.";
diff --git a/Engine/Source/Programs/UnrealGameSync/UnrealGameSync/PerforceMonitor.cs b/Engine/Source/Programs/UnrealGameSync/UnrealGameSync/PerforceMonitor.cs
--- a/Engine/Source/Programs/UnrealGameSync/UnrealGameSync/PerforceMonitor.cs
+++ b/Engine/Source/Programs/UnrealGameSync/UnrealGameSync/PerforceMonitor.cs
@@ -138,47 +138,51 @@ namespace UnrealGameSync
 				Stopwatch Timer = Stopwatch.StartNew();
 
 				// Check we still have a valid login ticket
-				if(!Perforce.IsLoggedIn(LogWriter))
+				bool bLoggedIn;
+				if(Perforce.GetLoggedInState(out bLoggedIn, LogWriter))
 				{
-					LastStatusMessage = "User is not logged in";
-					OnLoginExpired();
-				}
-				else
-				{
-					// Check we haven't switched streams
-					string NewStreamName;
-					if(Perforce.GetActiveStream(out NewStreamName, LogWriter) && NewStreamName != StreamName)
+					if(!bLoggedIn)
 					{
-						OnStreamChange();
+						LastStatusMessage = "User is not logged in";
+						OnLoginExpired();
 					}
-
-					// Update the stream list
-					if(StreamName != null)
+					else
 					{
-						List<string> NewOtherStreamNames;
-						if(!Perforce.FindStreams(PerforceUtils.GetClientOrDepotDirectoryName(StreamName) + "/*", out NewOtherStreamNames, LogWriter))
+						// Check we haven't switched streams
+						string NewStreamName;
+						if(Perforce.GetActiveStream(out NewStreamName, LogWriter) && NewStreamName != StreamName)
 						{
-							NewOtherStreamNames = new List<string>();
+							OnStreamChange();
 						}
-						OtherStreamNames = NewOtherStreamNames;
-					}
 
-					// Check for any p4 changes
-					if(!UpdateChanges())
-					{
-						LastStatusMessage = "Failed to update changes";
-					}
-					else if(!UpdateChangeTypes())
-					{
-						LastStatusMessage = "Failed to update change types";
-					}
-					else if(!UpdateZippedBinaries())
-					{
-						LastStatusMessage = "Failed to update zipped binaries list";
-					}
-					else
-					{
-						LastStatusMessage = String.Format("Last update took {0}ms", Timer.ElapsedMilliseconds);
+						// Update the stream list
+						if(StreamName != null)
+						{
+							List<string> NewOtherStreamNames;
+							if(!Perforce.FindStreams(PerforceUtils.GetClientOrDepotDirectoryName(StreamName) + "/*", out NewOtherStreamNames, LogWriter))
+							{
+								NewOtherStreamNames = new List<string>();
+							}
+							OtherStreamNames = NewOtherStreamNames;
+						}
+
+						// Check for any p4 changes
+						if(!UpdateChanges())
+						{
+							LastStatusMessage = "Failed to update changes";
+						}
+						else if(!UpdateChangeTypes())
+						{
+							LastStatusMessage = "Failed to update change types";
+						}
+						else if(!UpdateZippedBinaries())
+						{
+							LastStatusMessage = "Failed to update zipped binaries list";
+						}
+						else
+						{
+							LastStatusMessage = String.Format("Last update took {0}ms", Timer.ElapsedMilliseconds);
+						}
 					}
 				}
 
diff --git a/Engine/Source/Programs/UnrealGameSync/UnrealGameSync/Workspace.cs b/Engine/Source/Programs/UnrealGameSync/UnrealGameSync/Workspace.cs
--- a/Engine/Source/Programs/UnrealGameSync/UnrealGameSync/Workspace.cs
+++ b/Engine/Source/Programs/UnrealGameSync/UnrealGameSync/Workspace.cs
@@ -305,7 +305,13 @@ namespace UnrealGameSync
 					Log.WriteLine("Syncing to {0}...", PendingChangeNumber);
 
 					// Make sure we're logged in
-					if(!Perforce.IsLoggedIn(Log))
+					bool bLoggedIn;
+					if(!Perforce.GetLoggedInState(out bLoggedIn, Log))
+					{
+						StatusMessage = "Unable to get login status.";
+						return WorkspaceUpdateResult.FailedToSync;
+					}
+					if(!bLoggedIn)
 					{
 						StatusMessage = "User is not logged in.";
 						return WorkspaceUpdateResult.FailedToSyncLoginExpired;
