- Fixing Wildcard match issue, the change to ugsapi sends projects as //Depot/Stream instead of //Depot/Stream/
  Wildcard match was only substringing to 3 chars.
- Checking in the change a while back that increases the number of queried jobs up to 432 based on some maths from Bob about how many builds we want to grab
Published to ugsapi server 8/8/17

diff --git a/Engine/Source/Programs/UnrealGameSync/UnrealGameSyncMetadataServer/Connectors/SqlConnector.cs b/Engine/Source/Programs/UnrealGameSync/UnrealGameSyncMetadataServer/Connectors/SqlConnector.cs
--- a/Engine/Source/Programs/UnrealGameSync/UnrealGameSyncMetadataServer/Connectors/SqlConnector.cs
+++ b/Engine/Source/Programs/UnrealGameSync/UnrealGameSyncMetadataServer/Connectors/SqlConnector.cs
@@ -24,7 +24,7 @@ namespace UnrealGameSyncMetadataServer.Connectors
 			using (SQLiteConnection Connection = new SQLiteConnection(ConnectionString))
 			{
 				Connection.Open();
-				using (SQLiteCommand Command = new SQLiteCommand("SELECT Id FROM [UserVotes] WHERE Project LIKE @param1 GROUP BY Changelist ORDER BY Changelist DESC LIMIT 1 OFFSET 24;", Connection))
+				using (SQLiteCommand Command = new SQLiteCommand("SELECT Id FROM [UserVotes] WHERE Project LIKE @param1 GROUP BY Changelist ORDER BY Changelist DESC LIMIT 1 OFFSET 432;", Connection))
 				{
 					Command.Parameters.AddWithValue("@param1", ProjectLikeString);
 					using (SQLiteDataReader Reader = Command.ExecuteReader())
@@ -36,7 +36,7 @@ namespace UnrealGameSyncMetadataServer.Connectors
 						}
 					}
 				}
-				using (SQLiteCommand Command = new SQLiteCommand("SELECT Id FROM [Comments] WHERE Project LIKE @param1 GROUP BY ChangeNumber ORDER BY ChangeNumber DESC LIMIT 1 OFFSET 24;", Connection))
+				using (SQLiteCommand Command = new SQLiteCommand("SELECT Id FROM [Comments] WHERE Project LIKE @param1 GROUP BY ChangeNumber ORDER BY ChangeNumber DESC LIMIT 1 OFFSET 432;", Connection))
 				{
 					Command.Parameters.AddWithValue("@param1", ProjectLikeString);
 					using (SQLiteDataReader Reader = Command.ExecuteReader())
@@ -48,7 +48,7 @@ namespace UnrealGameSyncMetadataServer.Connectors
 						}
 					}
 				}
-				using (SQLiteCommand Command = new SQLiteCommand("SELECT Id FROM [CIS] WHERE Project LIKE @param1 GROUP BY ChangeNumber ORDER BY ChangeNumber DESC LIMIT 1 OFFSET 24", Connection))
+				using (SQLiteCommand Command = new SQLiteCommand("SELECT Id FROM [CIS] WHERE Project LIKE @param1 GROUP BY ChangeNumber ORDER BY ChangeNumber DESC LIMIT 1 OFFSET 432", Connection))
 				{
 					Command.Parameters.AddWithValue("@param1", ProjectLikeString);
 					using (SQLiteDataReader Reader = Command.ExecuteReader())
@@ -292,7 +292,7 @@ namespace UnrealGameSyncMetadataServer.Connectors
 		}
 		private static bool MatchesWildcard(string Wildcard, string Project)
 		{
-			return Wildcard.EndsWith("...") && Project.StartsWith(Wildcard.Substring(0, Wildcard.Length - 3), StringComparison.InvariantCultureIgnoreCase);
+			return Wildcard.EndsWith("...") && Project.StartsWith(Wildcard.Substring(0, Wildcard.Length - 4), StringComparison.InvariantCultureIgnoreCase);
 		}
 	}
 }
\ No newline at end of file
