Merging latest UGS code from Dev-Build.

diff --git a/Engine/Source/Programs/UnrealGameSync/UnrealGameSync/Controls/LogControl.cs b/Engine/Source/Programs/UnrealGameSync/UnrealGameSync/Controls/LogControl.cs
--- a/Engine/Source/Programs/UnrealGameSync/UnrealGameSync/Controls/LogControl.cs
+++ b/Engine/Source/Programs/UnrealGameSync/UnrealGameSync/Controls/LogControl.cs
@@ -1,4 +1,4 @@
-﻿// Copyright 1998-2019 Epic Games, Inc. All Rights Reserved.
+// Copyright 1998-2019 Epic Games, Inc. All Rights Reserved.
 
 using System;
 using System.Collections.Generic;
@@ -12,6 +12,7 @@ using System.Windows.Forms;
 using System.Runtime.InteropServices;
 using System.Collections.Concurrent;
 using System.IO;
+using System.Text.RegularExpressions;
 
 namespace UnrealGameSync
 {
@@ -749,10 +750,24 @@ namespace UnrealGameSync
 				int SelectMaxIdx;
 				ClipSelectionToLine(Idx, out SelectMinIdx, out SelectMaxIdx);
 
+				Color TextColor;
+				if(Regex.IsMatch(Lines[Idx], "(?<!\\w)error[: ]", RegexOptions.IgnoreCase))
+				{
+					TextColor = Color.FromArgb(189, 54, 47);
+				}
+				else if(Regex.IsMatch(Lines[Idx], "(?<!\\w)warning[: ]", RegexOptions.IgnoreCase))
+				{
+					TextColor = Color.FromArgb(128, 128, 0);
+				}
+				else
+				{
+					TextColor = Color.Black;
+				}
+
 				int TextY = (Idx - ScrollLine) * FontSize.Height;
 				if(SelectMinIdx > 0)
 				{
-					TextRenderer.DrawText(e.Graphics, Lines[Idx].Substring(0, SelectMinIdx), Font, new Point(TextX, TextY), Color.Black, TextFormatFlags.NoPadding);
+					TextRenderer.DrawText(e.Graphics, Lines[Idx].Substring(0, SelectMinIdx), Font, new Point(TextX, TextY), TextColor, TextFormatFlags.NoPadding);
 				}
 				if(SelectMinIdx < SelectMaxIdx)
 				{
@@ -761,7 +776,7 @@ namespace UnrealGameSync
 				}
 				if(SelectMaxIdx < Lines[Idx].Length)
 				{
-					TextRenderer.DrawText(e.Graphics, Lines[Idx].Substring(SelectMaxIdx), Font, new Point(TextX + (SelectMaxIdx * FontSize.Width), TextY), Color.Black, TextFormatFlags.NoPadding);
+					TextRenderer.DrawText(e.Graphics, Lines[Idx].Substring(SelectMaxIdx), Font, new Point(TextX + (SelectMaxIdx * FontSize.Width), TextY), TextColor, TextFormatFlags.NoPadding);
 				}
 			}
 		}
diff --git a/Engine/Source/Programs/UnrealGameSync/UnrealGameSync/Controls/TabControl.cs b/Engine/Source/Programs/UnrealGameSync/UnrealGameSync/Controls/TabControl.cs
--- a/Engine/Source/Programs/UnrealGameSync/UnrealGameSync/Controls/TabControl.cs
+++ b/Engine/Source/Programs/UnrealGameSync/UnrealGameSync/Controls/TabControl.cs
@@ -22,6 +22,7 @@ namespace UnrealGameSync
 			public object Data;
 			public int MinX;
 			public int Width;
+			public int CloseButtonWidth;
 			public Size TextSize;
 			public Tuple<Color, float> Highlight;
 		}
@@ -253,7 +254,8 @@ namespace UnrealGameSync
 					Tab.Width = TabPadding + Tab.TextSize.Width + (int)(TabPadding * DpiScaleX);
 					if(Idx == SelectedTabIdx)
 					{
-						Tab.Width += (int)(TabCloseButtonWidth * DpiScaleX);
+						Tab.CloseButtonWidth = (int)(TabCloseButtonWidth * DpiScaleX);
+						Tab.Width += Tab.CloseButtonWidth;
 					}
 				}
 
@@ -324,7 +326,7 @@ namespace UnrealGameSync
 				{
 					if(HoverTabIdx == SelectedTabIdx)
 					{
-						if(e.Location.X > Tabs[HoverTabIdx].MinX + Tabs[HoverTabIdx].Width - TabCloseButtonWidth)
+						if(e.Location.X > Tabs[HoverTabIdx].MinX + Tabs[HoverTabIdx].Width - Tabs[HoverTabIdx].CloseButtonWidth)
 						{
 							RemoveTab(HoverTabIdx);
 						}
diff --git a/Engine/Source/Programs/UnrealGameSync/UnrealGameSync/Controls/WorkspaceControl.cs b/Engine/Source/Programs/UnrealGameSync/UnrealGameSync/Controls/WorkspaceControl.cs
--- a/Engine/Source/Programs/UnrealGameSync/UnrealGameSync/Controls/WorkspaceControl.cs
+++ b/Engine/Source/Programs/UnrealGameSync/UnrealGameSync/Controls/WorkspaceControl.cs
@@ -1535,23 +1535,19 @@ namespace UnrealGameSync
 		{
 			ExpandRowLayout Layout = new ExpandRowLayout();
 
-			int CurrentMaxChanges = PerforceMonitor.CurrentMaxChanges;
-
-			string ShowingChanges;
-			if(ListIndexToChangeIndex.Count == CurrentMaxChanges)
-			{
-				ShowingChanges = String.Format("Showing {0} changes.", ListIndexToChangeIndex.Count);
-			}
-			else
-			{
-				ShowingChanges = String.Format("Showing {0}/{1} changes.", ListIndexToChangeIndex.Count, CurrentMaxChanges);
-			}
+			string ShowingChanges = String.Format("Showing {0}/{1} changes.", ListIndexToChangeIndex.Count, NumChanges);
 
+			int CurrentMaxChanges = PerforceMonitor.CurrentMaxChanges;
 			if(PerforceMonitor.PendingMaxChanges > CurrentMaxChanges)
 			{
 				Layout.MainText = String.Format("{0}. Fetching {1} more from server...  ", ShowingChanges, PerforceMonitor.PendingMaxChanges - CurrentMaxChanges);
 				Layout.LinkText = "Cancel";
 			}
+			else if(PerforceMonitor.CurrentMaxChanges > NumChanges)
+			{
+				Layout.MainText = ShowingChanges;
+				Layout.LinkText = "";
+			}
 			else
 			{
 				Layout.MainText = String.Format("{0}  ", ShowingChanges);
@@ -1728,7 +1724,7 @@ namespace UnrealGameSync
 			Color TextColor = (bAllowSync || Change.Number == Workspace.PendingChangeNumber || Change.Number == Workspace.CurrentChangeNumber || (WorkspaceSettings != null && WorkspaceSettings.AdditionalChangeNumbers.Contains(Change.Number)))? SystemColors.WindowText : Blend(SystemColors.Window, SystemColors.WindowText, 0.25f);
 
 			const int FadeRange = 6;
-			if(e.ItemIndex >= BuildList.Items.Count - FadeRange)
+			if(e.ItemIndex >= BuildList.Items.Count - FadeRange && NumChanges >= PerforceMonitor.CurrentMaxChanges)
 			{
 				float Opacity = (float)(BuildList.Items.Count - e.ItemIndex - 0.9f) / FadeRange;
 				BadgeAlpha = (int)(BadgeAlpha * Opacity);
diff --git a/Engine/Source/Programs/UnrealGameSync/UnrealGameSync/DeploymentSettings.cs b/Engine/Source/Programs/UnrealGameSync/UnrealGameSync/DeploymentSettings.cs
--- a/Engine/Source/Programs/UnrealGameSync/UnrealGameSync/DeploymentSettings.cs
+++ b/Engine/Source/Programs/UnrealGameSync/UnrealGameSync/DeploymentSettings.cs
@@ -24,5 +24,10 @@ namespace UnrealGameSync
 		/// The UnrealGameSync executable should be located at Release/UnrealGameSync.exe under this path, with any dependent DLLs.
 		/// </summary>
 		public static readonly string DefaultDepotPath = null;
+
+		/// <summary>
+		/// Whether to send telemetry data by default. Can be useful for finding users in need of a hardware upgrade, but can bloat the DB for large teams.
+		/// </summary>
+		public static bool bSendTelemetry = false;
 	}
 }
diff --git a/Engine/Source/Programs/UnrealGameSync/UnrealGameSync/DetectProjectSettingsTask.cs b/Engine/Source/Programs/UnrealGameSync/UnrealGameSync/DetectProjectSettingsTask.cs
--- a/Engine/Source/Programs/UnrealGameSync/UnrealGameSync/DetectProjectSettingsTask.cs
+++ b/Engine/Source/Programs/UnrealGameSync/UnrealGameSync/DetectProjectSettingsTask.cs
@@ -113,7 +113,7 @@ namespace UnrealGameSync
 				{
 					if(!PerforceClient.ConvertToLocalPath(NewSelectedClientFileName, out NewSelectedFileName, Log))
 					{
-						ErrorMessage = String.Format("Couldn't get client path for {0}", NewSelectedFileName);
+						ErrorMessage = String.Format("Couldn't get client path for {0}", NewSelectedClientFileName);
 						return false;
 					}
 				}
@@ -303,7 +303,7 @@ namespace UnrealGameSync
 				else
 				{
 					List<string> ProjectLines;
-					if(PerforceClient.Print(NewSelectedClientFileName, out ProjectLines, Log))
+					if(!PerforceClient.Print(NewSelectedClientFileName, out ProjectLines, Log))
 					{
 						ErrorMessage = String.Format("Unable to get contents of {0}", NewSelectedClientFileName);
 						return false;
diff --git a/Engine/Source/Programs/UnrealGameSync/UnrealGameSync/FindFoldersToCleanTask.cs b/Engine/Source/Programs/UnrealGameSync/UnrealGameSync/FindFoldersToCleanTask.cs
--- a/Engine/Source/Programs/UnrealGameSync/UnrealGameSync/FindFoldersToCleanTask.cs
+++ b/Engine/Source/Programs/UnrealGameSync/UnrealGameSync/FindFoldersToCleanTask.cs
@@ -102,14 +102,14 @@ namespace UnrealGameSync
 			}
 		}
 
-		void MergeTrees(FolderToClean LocalFolder, PerforceHaveFolder PerforceFolder, HashSet<string> OpenClientPaths)
+		void MergeTrees(FolderToClean LocalFolder, PerforceHaveFolder PerforceFolder, HashSet<string> OpenClientPaths, string PerforceConfigFile)
 		{
 			if(PerforceFolder == null)
 			{
 				// Loop through all the local sub-folders
 				foreach(FolderToClean LocalSubFolder in LocalFolder.NameToSubFolder.Values)
 				{
-					MergeTrees(LocalSubFolder, null, OpenClientPaths);
+					MergeTrees(LocalSubFolder, null, OpenClientPaths, PerforceConfigFile);
 				}
 
 				// Delete everything
@@ -122,7 +122,7 @@ namespace UnrealGameSync
 				{
 					PerforceHaveFolder PerforceSubFolder;
 					PerforceFolder.NameToSubFolder.TryGetValue(LocalSubFolder.Name, out PerforceSubFolder);
-					MergeTrees(LocalSubFolder, PerforceSubFolder, OpenClientPaths);
+					MergeTrees(LocalSubFolder, PerforceSubFolder, OpenClientPaths, PerforceConfigFile);
 				}
 
 				// Also merge all the Perforce folders that no longer exist
@@ -132,7 +132,7 @@ namespace UnrealGameSync
 					if(!LocalFolder.NameToSubFolder.TryGetValue(PerforceSubFolderPair.Key, out LocalSubFolder))
 					{
 						LocalSubFolder = new FolderToClean(new DirectoryInfo(Path.Combine(LocalFolder.Directory.FullName, PerforceSubFolderPair.Key)));
-						MergeTrees(LocalSubFolder, PerforceSubFolderPair.Value, OpenClientPaths);
+						MergeTrees(LocalSubFolder, PerforceSubFolderPair.Value, OpenClientPaths, PerforceConfigFile);
 						LocalFolder.NameToSubFolder.Add(LocalSubFolder.Name, LocalSubFolder);
 					}
 				}
@@ -161,6 +161,12 @@ namespace UnrealGameSync
 				}
 			}
 
+			// Remove any config files
+			if(PerforceConfigFile != null)
+			{
+				LocalFolder.FilesToDelete.RemoveAll(x => String.Compare(x.Name, PerforceConfigFile, StringComparison.OrdinalIgnoreCase) == 0);
+			}
+
 			// Figure out if this folder is just an empty directory that needs to be removed
 			LocalFolder.bEmptyLeaf = LocalFolder.NameToFile.Count == 0 && LocalFolder.NameToSubFolder.Count == 0 && LocalFolder.FilesToSync.Count == 0;
 
@@ -294,8 +300,12 @@ namespace UnrealGameSync
 			// Wait to finish scanning the directory
 			FinishedScan.WaitOne();
 
+			// Find the value of the P4CONFIG variable
+			string PerforceConfigFile;
+			PerforceClient.GetSetting("P4CONFIG", out PerforceConfigFile, Log);
+
 			// Merge the trees
-			MergeTrees(RootFolderToClean, RootHaveFolder, OpenLocalFiles);
+			MergeTrees(RootFolderToClean, RootHaveFolder, OpenLocalFiles, PerforceConfigFile);
 
 			// Remove all the empty folders
 			RemoveEmptyFolders(RootFolderToClean);
diff --git a/Engine/Source/Programs/UnrealGameSync/UnrealGameSync/Forms/MainWindow.cs b/Engine/Source/Programs/UnrealGameSync/UnrealGameSync/Forms/MainWindow.cs
--- a/Engine/Source/Programs/UnrealGameSync/UnrealGameSync/Forms/MainWindow.cs
+++ b/Engine/Source/Programs/UnrealGameSync/UnrealGameSync/Forms/MainWindow.cs
@@ -1169,8 +1169,19 @@ namespace UnrealGameSync
 		{
 			if(Settings.WindowBounds != null)
 			{
-				Location = Settings.WindowBounds.Value.Location;
-				Size = Settings.WindowBounds.Value.Size;
+				Rectangle WindowBounds = Settings.WindowBounds.Value;
+				if(WindowBounds.Width > MinimumSize.Width && WindowBounds.Height > MinimumSize.Height)
+				{
+					foreach (Screen Screen in Screen.AllScreens)
+					{
+						if(WindowBounds.IntersectsWith(Screen.Bounds))
+						{
+							Location = Settings.WindowBounds.Value.Location;
+							Size = Settings.WindowBounds.Value.Size;
+							break;
+						}
+					}
+				}
 			}
 			WindowState = Settings.WindowState;
 		}
diff --git a/Engine/Source/Programs/UnrealGameSync/UnrealGameSync/Properties/AssemblyInfo.cs b/Engine/Source/Programs/UnrealGameSync/UnrealGameSync/Properties/AssemblyInfo.cs
--- a/Engine/Source/Programs/UnrealGameSync/UnrealGameSync/Properties/AssemblyInfo.cs
+++ b/Engine/Source/Programs/UnrealGameSync/UnrealGameSync/Properties/AssemblyInfo.cs
@@ -34,5 +34,5 @@ using System.Runtime.InteropServices;
 // You can specify all the values or you can default the Build and Revision Numbers 
 // by using the '*' as shown below:
 // [assembly: AssemblyVersion("1.0.*")]
-[assembly: AssemblyVersion("1.167.*")]
-[assembly: AssemblyFileVersion("1.167.0.0")]
+[assembly: AssemblyVersion("1.171.*")]
+[assembly: AssemblyFileVersion("1.171.0.0")]
diff --git a/Engine/Source/Programs/UnrealGameSync/UnrealGameSync/Telemetry.cs b/Engine/Source/Programs/UnrealGameSync/UnrealGameSync/Telemetry.cs
--- a/Engine/Source/Programs/UnrealGameSync/UnrealGameSync/Telemetry.cs
+++ b/Engine/Source/Programs/UnrealGameSync/UnrealGameSync/Telemetry.cs
@@ -176,6 +176,11 @@ namespace UnrealGameSync
 
 		bool SendTimingData(TelemetryTimingData Data, string Version, string IpAddress)
 		{
+			if(!DeploymentSettings.bSendTelemetry)
+			{
+				return true;
+			}
+
 			try
 			{
 				Stopwatch Timer = Stopwatch.StartNew();
diff --git a/Engine/Source/Programs/UnrealGameSync/UnrealGameSync/Workspace.cs b/Engine/Source/Programs/UnrealGameSync/UnrealGameSync/Workspace.cs
--- a/Engine/Source/Programs/UnrealGameSync/UnrealGameSync/Workspace.cs
+++ b/Engine/Source/Programs/UnrealGameSync/UnrealGameSync/Workspace.cs
@@ -104,11 +104,11 @@ namespace UnrealGameSync
 		{
 			new WorkspaceSyncCategory(new Guid("{6703E989-D912-451D-93AD-B48DE748D282}"), "Content", "*.uasset", "*.umap"),
 			new WorkspaceSyncCategory(new Guid("{6507C2FB-19DD-403A-AFA3-BBF898248D5A}"), "Documentation", "/Engine/Documentation/..."),
-			new WorkspaceSyncCategory(new Guid("{FD7C716E-4BAD-43AE-8FAE-8748EF9EE44D}"), "Platform Support: Android", "/Engine/Source/ThirdParty/.../Android/..."),
+			new WorkspaceSyncCategory(new Guid("{FD7C716E-4BAD-43AE-8FAE-8748EF9EE44D}"), "Platform Support: Android", "/Engine/Source/ThirdParty/.../Android/...", ".../Build/Android/PipelineCaches/..."),
 			new WorkspaceSyncCategory(new Guid("{3299A73D-2176-4C0F-BC99-C1C6631AF6C4}"), "Platform Support: HTML5", "/Engine/Source/ThirdParty/.../HTML5/...", "/Engine/Extras/ThirdPartyNotUE/emsdk/..."),
-			new WorkspaceSyncCategory(new Guid("{176B2EB2-35F7-4E8E-B131-5F1C5F0959AF}"), "Platform Support: iOS", "/Engine/Source/ThirdParty/.../IOS/..."),
+			new WorkspaceSyncCategory(new Guid("{176B2EB2-35F7-4E8E-B131-5F1C5F0959AF}"), "Platform Support: iOS", "/Engine/Source/ThirdParty/.../IOS/...", ".../Build/IOS/PipelineCaches/..."),
 			new WorkspaceSyncCategory(new Guid("{F44B2D25-CBC0-4A8F-B6B3-E4A8125533DD}"), "Platform Support: Linux", "/Engine/Source/ThirdParty/.../Linux/..."),
-			new WorkspaceSyncCategory(new Guid("{2AF45231-0D75-463B-BF9F-ABB3231091BB}"), "Platform Support: Mac", "/Engine/Source/ThirdParty/.../Mac/..."),
+			new WorkspaceSyncCategory(new Guid("{2AF45231-0D75-463B-BF9F-ABB3231091BB}"), "Platform Support: Mac", "/Engine/Source/ThirdParty/.../Mac/...", ".../Build/Mac/PipelineCaches/..."),
 			new WorkspaceSyncCategory(new Guid("{C8CB4934-ADE9-46C9-B6E3-61A659E1FAF5}"), "Platform Support: PS4", ".../PS4/..."),
 			new WorkspaceSyncCategory(new Guid("{F8AE5AC3-DA2D-4719-BABF-8A90D878379E}"), "Platform Support: Switch", ".../Switch/..."),
 			new WorkspaceSyncCategory(new Guid("{3788A0BC-188C-4A0D-950A-D68175F0D110}"), "Platform Support: tvOS", "/Engine/Source/ThirdParty/.../TVOS/..."),
