UGS: Fix compilation errors when DeploymentSettingsInternal.cs is not present.

#rnx

diff --git a/Engine/Source/Programs/UnrealGameSync/UnrealGameSync/UnrealGameSync.csproj b/Engine/Source/Programs/UnrealGameSync/UnrealGameSync/UnrealGameSync.csproj
--- a/Engine/Source/Programs/UnrealGameSync/UnrealGameSync/UnrealGameSync.csproj
+++ b/Engine/Source/Programs/UnrealGameSync/UnrealGameSync/UnrealGameSync.csproj
@@ -331,7 +331,6 @@
       <DependentUpon>SyncFilter.cs</DependentUpon>
     </Compile>
     <Compile Include="ModalTask.cs" />
-    <Compile Include="NotForLicensees\DeploymentSettingsInternal.cs" />
     <Compile Include="PerforceModalTask.cs" />
     <Compile Include="PerforceMonitor.cs" />
     <Compile Include="ProgramApplicationContext.cs" />
@@ -462,7 +461,11 @@
     <None Include="Resources\App.config" />
   </ItemGroup>
   <Choose>
-    <When Condition="Exists('NotForLicensees\ProgramSettings.cs')" />
+    <When Condition="Exists('NotForLicensees\DeploymentSettingsInternal.cs')">
+      <ItemGroup>
+        <Compile Include="NotForLicensees\DeploymentSettingsInternal.cs" />
+      </ItemGroup>
+	</When>
   </Choose>
   <ItemGroup>
     <EmbeddedResource Include="Resources\Icons.png" />
