UGS: Add a matcher for copyright notice warnings/errors.

#rnx

diff --git a/Engine/Source/Programs/UnrealGameSync/MetadataTool/BuildHealth/Matchers/CopyrightNoticeMatcher.cs b/Engine/Source/Programs/UnrealGameSync/MetadataTool/BuildHealth/Matchers/CopyrightNoticeMatcher.cs
new file mode 100644
--- /dev/null
+++ b/Engine/Source/Programs/UnrealGameSync/MetadataTool/BuildHealth/Matchers/CopyrightNoticeMatcher.cs
@@ -0,0 +1,58 @@
+using System;
+using System.Collections.Generic;
+using System.Linq;
+using System.Text;
+using System.Text.RegularExpressions;
+using System.Threading.Tasks;
+
+namespace MetadataTool
+{
+	class CopyrightNoticeMatcher : PatternMatcher
+	{
+		public override string Category => "Copyright";
+
+		public override bool TryMatch(InputJob Job, InputJobStep JobStep, InputDiagnostic Diagnostic, List<BuildHealthIssue> Issues)
+		{
+			// Make sure we're running a step that this applies to
+			if(JobStep.Name.IndexOf("Copyright", StringComparison.OrdinalIgnoreCase) == -1)
+			{
+				return false;
+			}
+
+			// Find any files in compiler output format
+			HashSet<string> SourceFileNames = new HashSet<string>(StringComparer.OrdinalIgnoreCase);
+			foreach (Match FileMatch in Regex.Matches(Diagnostic.Message, @"^\s*(?:WARNING|ERROR):\s*([^ ]+\.[a-zA-Z]+):", RegexOptions.Multiline))
+			{
+				if (FileMatch.Success)
+				{
+					SourceFileNames.Add(FileMatch.Groups[1].Value);
+				}
+			}
+
+			// If we found any source files, create a diagnostic category for them
+			if (SourceFileNames.Count > 0)
+			{
+				BuildHealthIssue Issue = new BuildHealthIssue(Category, Job.Url, new BuildHealthDiagnostic(JobStep.Name, Diagnostic.Message, Diagnostic.Url));
+				Issue.FileNames.UnionWith(SourceFileNames);
+				Issues.Add(Issue);
+				return true;
+			}
+
+			// Otherwise pass
+			return false;
+		}
+
+		public override string GetSummary(BuildHealthIssue Issue)
+		{
+			SortedSet<string> ShortFileNames = GetSourceFileNames(Issue.FileNames);
+			if (ShortFileNames.Count == 0)
+			{
+				return "Missing copyright notices";
+			}
+			else
+			{
+				return String.Format("Missing copyright in {0}", String.Join(", ", ShortFileNames));
+			}
+		}
+	}
+}
diff --git a/Engine/Source/Programs/UnrealGameSync/MetadataTool/CommandHandler_BuildHealth.cs b/Engine/Source/Programs/UnrealGameSync/MetadataTool/CommandHandler_BuildHealth.cs
--- a/Engine/Source/Programs/UnrealGameSync/MetadataTool/CommandHandler_BuildHealth.cs
+++ b/Engine/Source/Programs/UnrealGameSync/MetadataTool/CommandHandler_BuildHealth.cs
@@ -23,7 +23,8 @@ namespace MetadataTool
 		static readonly List<PatternMatcher> Matchers = new List<PatternMatcher>()
 		{ 
 			new CompilePatternMatcher(),
-			new UndefinedSymbolPatternMatcher()
+			new UndefinedSymbolPatternMatcher(),
+			new CopyrightNoticeMatcher()
 		};
 
 		static readonly Dictionary<string, PatternMatcher> CategoryToMatcher = Matchers.ToDictionary(x => x.Category, x => x);
diff --git a/Engine/Source/Programs/UnrealGameSync/MetadataTool/MetadataTool.csproj b/Engine/Source/Programs/UnrealGameSync/MetadataTool/MetadataTool.csproj
--- a/Engine/Source/Programs/UnrealGameSync/MetadataTool/MetadataTool.csproj
+++ b/Engine/Source/Programs/UnrealGameSync/MetadataTool/MetadataTool.csproj
@@ -53,6 +53,7 @@
   <ItemGroup>
     <Compile Include="BuildHealth\ChangeInfo.cs" />
     <Compile Include="BuildHealth\Matchers\CompilePatternMatcher.cs" />
+    <Compile Include="BuildHealth\Matchers\CopyrightNoticeMatcher.cs" />
     <Compile Include="BuildHealth\Matchers\UndefinedSymbolPatternMatcher.cs" />
     <Compile Include="BuildHealth\PatternMatcher.cs" />
     <Compile Include="CommandHandler.cs" />
