UGS: Remove Platforms folder from sync paths for now.

#rnx

diff --git a/Engine/Source/Programs/UnrealGameSync/UnrealGameSync/Workspace.cs b/Engine/Source/Programs/UnrealGameSync/UnrealGameSync/Workspace.cs
--- a/Engine/Source/Programs/UnrealGameSync/UnrealGameSync/Workspace.cs
+++ b/Engine/Source/Programs/UnrealGameSync/UnrealGameSync/Workspace.cs
@@ -1091,7 +1091,7 @@ namespace UnrealGameSync
 				{
 					SyncPaths.Add(ClientRootPath + "/Enterprise/...");
 				}
-				SyncPaths.Add(ClientRootPath + "/Platforms/...");
+//				SyncPaths.Add(ClientRootPath + "/Platforms/...");
 				SyncPaths.Add(PerforceUtils.GetClientOrDepotDirectoryName(SelectedClientFileName) + "/...");
 			}
 
