UGS: Remove code that excludes files to sync if using an older branch. Since this logic relies on the contents of the locally synced config file, it may default to false (and not sync the version header) for the first sync in a stream. Existing logic in Workspace.SyncFileRevisions() should ensure these files are clobbered when necessary, if syncing an earlier branch.

#rnx

diff --git a/Engine/Source/Programs/UnrealGameSync/UnrealGameSync/Workspace.cs b/Engine/Source/Programs/UnrealGameSync/UnrealGameSync/Workspace.cs
--- a/Engine/Source/Programs/UnrealGameSync/UnrealGameSync/Workspace.cs
+++ b/Engine/Source/Programs/UnrealGameSync/UnrealGameSync/Workspace.cs
@@ -558,11 +558,6 @@ namespace UnrealGameSync
 					// Create a filter for all the files we don't want
 					FileFilter Filter = new FileFilter(UserFilter);
 					Filter.Exclude(BuildVersionFileName);
-					if (!ProjectConfigFile.GetValue("Options.UseFastModularVersioningV2", false))
-					{
-						Filter.Exclude(VersionHeaderFileName);
-						Filter.Exclude(ObjectVersionFileName);
-					}
 					if (Context.Options.HasFlag(WorkspaceUpdateOptions.ContentOnly))
 					{
 						Filter.Exclude("*.usf");
