Fix compile errors in UGS.

#rnx

diff --git a/Engine/Source/Programs/UnrealGameSync/UnrealGameSyncLauncher/Forms/SettingsWindow.cs b/Engine/Source/Programs/UnrealGameSync/UnrealGameSyncLauncher/Forms/SettingsWindow.cs
--- a/Engine/Source/Programs/UnrealGameSync/UnrealGameSyncLauncher/Forms/SettingsWindow.cs
+++ b/Engine/Source/Programs/UnrealGameSync/UnrealGameSyncLauncher/Forms/SettingsWindow.cs
@@ -90,7 +90,7 @@ namespace UnrealGameSyncLauncher
 
 			// Attempt to sync through a modal dialog
 			string ErrorMessage;
-			ModalTaskResult Result = PerforceModalTask.Execute(this, null, ServerAndPort, UserName, SyncApplication, "Updating", "Checking for updates, please wait...", Log, out ErrorMessage);
+			ModalTaskResult Result = PerforceModalTask.Execute(this, new PerforceConnection(UserName, null, ServerAndPort), SyncApplication, "Updating", "Checking for updates, please wait...", Log, out ErrorMessage);
 			if(Result == ModalTaskResult.Succeeded)
 			{
 				Utility.SaveGlobalPerforceSettings(ServerAndPort, UserName, DepotPath);
diff --git a/Engine/Source/Programs/UnrealGameSync/UnrealGameSyncLauncher/Program.cs b/Engine/Source/Programs/UnrealGameSync/UnrealGameSyncLauncher/Program.cs
--- a/Engine/Source/Programs/UnrealGameSync/UnrealGameSyncLauncher/Program.cs
+++ b/Engine/Source/Programs/UnrealGameSync/UnrealGameSyncLauncher/Program.cs
@@ -84,7 +84,7 @@ namespace UnrealGameSyncLauncher
 					SyncAndRunPerforceTask SyncApplication = new SyncAndRunPerforceTask((Perforce, LogWriter) => SyncAndRun(Perforce, DepotPath, bUnstable, Args, InstanceMutex, LogWriter));
 
 					string ErrorMessage;
-					if(PerforceModalTask.Execute(null, null, ServerAndPort, UserName, SyncApplication, "Updating", "Checking for updates, please wait...", Log, out ErrorMessage) == ModalTaskResult.Succeeded)
+					if(PerforceModalTask.Execute(null, new PerforceConnection(UserName, null, ServerAndPort), SyncApplication, "Updating", "Checking for updates, please wait...", Log, out ErrorMessage) == ModalTaskResult.Succeeded)
 					{
 						return 0;
 					}
diff --git a/Engine/Source/Programs/UnrealGameSync/UnrealGameSyncLauncher/SyncAndRunModalTask.cs b/Engine/Source/Programs/UnrealGameSync/UnrealGameSyncLauncher/SyncAndRunModalTask.cs
deleted file mode 100644
diff --git a/Engine/Source/Programs/UnrealGameSync/UnrealGameSyncLauncher/UnrealGameSyncLauncher.csproj b/Engine/Source/Programs/UnrealGameSync/UnrealGameSyncLauncher/UnrealGameSyncLauncher.csproj
--- a/Engine/Source/Programs/UnrealGameSync/UnrealGameSyncLauncher/UnrealGameSyncLauncher.csproj
+++ b/Engine/Source/Programs/UnrealGameSync/UnrealGameSyncLauncher/UnrealGameSyncLauncher.csproj
@@ -118,7 +118,6 @@
     <Compile Include="Forms\SettingsWindow.Designer.cs">
       <DependentUpon>SettingsWindow.cs</DependentUpon>
     </Compile>
-    <Compile Include="SyncAndRunModalTask.cs" />
     <Compile Include="SyncAndRunPerforceTask.cs" />
     <Compile Include="Program.cs" />
     <Compile Include="Properties\AssemblyInfo.cs" />
@@ -159,8 +158,8 @@
     <When Condition="Exists('..\UnrealGameSync\NotForLicensees\DeploymentSettingsInternal.cs')">
       <ItemGroup>
         <Compile Include="..\UnrealGameSync\NotForLicensees\DeploymentSettingsInternal.cs">
-			<Link>Shared\NotForLicensees\DeploymentSettingsInternal.cs</Link>
-		</Compile>
+          <Link>Shared\NotForLicensees\DeploymentSettingsInternal.cs</Link>
+        </Compile>
       </ItemGroup>
     </When>
   </Choose>
