UGS: Add DebugGame editor executable names to the list of processes that should be killed before performing a sync. (#6081)

diff --git a/Engine/Source/Programs/UnrealGameSync/UnrealGameSync/Controls/WorkspaceControl.cs b/Engine/Source/Programs/UnrealGameSync/UnrealGameSync/Controls/WorkspaceControl.cs
--- a/Engine/Source/Programs/UnrealGameSync/UnrealGameSync/Controls/WorkspaceControl.cs
+++ b/Engine/Source/Programs/UnrealGameSync/UnrealGameSync/Controls/WorkspaceControl.cs
@@ -2736,6 +2736,8 @@ namespace UnrealGameSync
 			ProcessNames.Add("UE4Editor-Cmd");
 			ProcessNames.Add("UE4Editor-Win64-Debug");
 			ProcessNames.Add("UE4Editor-Win64-Debug-Cmd");
+			ProcessNames.Add("UE4Editor-Win64-DebugGame");
+			ProcessNames.Add("UE4Editor-Win64-DebugGame-Cmd");
 			ProcessNames.Add("CrashReportClient");
 			ProcessNames.Add("CrashReportClient-Win64-Development");
 			ProcessNames.Add("CrashReportClient-Win64-Debug");
