UGS: Add guards around a few places that can throw exceptions due to external factors.

#rnx

diff --git a/Engine/Source/Programs/UnrealGameSync/UnrealGameSync/Forms/IssueBrowserWindow.cs b/Engine/Source/Programs/UnrealGameSync/UnrealGameSync/Forms/IssueBrowserWindow.cs
--- a/Engine/Source/Programs/UnrealGameSync/UnrealGameSync/Forms/IssueBrowserWindow.cs
+++ b/Engine/Source/Programs/UnrealGameSync/UnrealGameSync/Forms/IssueBrowserWindow.cs
@@ -284,7 +284,15 @@ namespace UnrealGameSync
 
 		private void ShowIssue(IssueData Issue)
 		{
-			Issue.Builds = RESTApi.GET<List<IssueBuildData>>(IssueMonitor.ApiUrl, String.Format("issues/{0}/builds", Issue.Id));
+			try
+			{
+				Issue.Builds = RESTApi.GET<List<IssueBuildData>>(IssueMonitor.ApiUrl, String.Format("issues/{0}/builds", Issue.Id));
+			}
+			catch(Exception Ex)
+			{
+				MessageBox.Show(Owner, Ex.ToString(), "Error querying builds", MessageBoxButtons.OK);
+				return;
+			}
 			IssueDetailsWindow.Show(Owner, IssueMonitor, ServerAndPort, UserName, ServerTimeOffset, Issue, Log, CurrentStream);
 		}
 
diff --git a/Engine/Source/Programs/UnrealGameSync/UnrealGameSync/Forms/ModalTaskWindow.cs b/Engine/Source/Programs/UnrealGameSync/UnrealGameSync/Forms/ModalTaskWindow.cs
--- a/Engine/Source/Programs/UnrealGameSync/UnrealGameSync/Forms/ModalTaskWindow.cs
+++ b/Engine/Source/Programs/UnrealGameSync/UnrealGameSync/Forms/ModalTaskWindow.cs
@@ -57,14 +57,22 @@ namespace UnrealGameSync
 
 		private void ThreadProc()
 		{
-			string ErrorMessageValue;
-			if(Task.Run(out ErrorMessageValue))
+			try
 			{
-				PendingResult = ModalTaskResult.Succeeded;
+				string ErrorMessageValue;
+				if (Task.Run(out ErrorMessageValue))
+				{
+					PendingResult = ModalTaskResult.Succeeded;
+				}
+				else
+				{
+					ErrorMessage = ErrorMessageValue;
+					PendingResult = ModalTaskResult.Failed;
+				}
 			}
-			else
+			catch(Exception Ex)
 			{
-				ErrorMessage =  ErrorMessageValue;
+				ErrorMessage = Ex.ToString();
 				PendingResult = ModalTaskResult.Failed;
 			}
 			SyncContext.Post((o) => ThreadCompleteCallback(), null);
