Merging latest version of UGS.

#rnx

diff --git a/Engine/Source/Programs/UnrealGameSync/UnrealGameSync/Forms/CleanWorkspaceWindow.cs b/Engine/Source/Programs/UnrealGameSync/UnrealGameSync/Forms/CleanWorkspaceWindow.cs
--- a/Engine/Source/Programs/UnrealGameSync/UnrealGameSync/Forms/CleanWorkspaceWindow.cs
+++ b/Engine/Source/Programs/UnrealGameSync/UnrealGameSync/Forms/CleanWorkspaceWindow.cs
@@ -489,7 +489,15 @@ namespace UnrealGameSync
 		{
 			TreeNode Node = (TreeNode)FolderContextMenu.Tag;
 			TreeNodeData NodeData = (TreeNodeData)Node.Tag;
-			Process.Start("explorer.exe", String.Format("\"{0}\"", NodeData.Folder.Directory.FullName));
+
+			if (NodeData.Folder != null)
+			{
+				Process.Start("explorer.exe", String.Format("\"{0}\"", NodeData.Folder.Directory.FullName));
+			}
+			else if (NodeData.File != null)
+			{
+				Process.Start("explorer.exe", String.Format("\"{0}\"", NodeData.File.Directory.FullName));
+			}
 		}
 	}
 }
diff --git a/Engine/Source/Programs/UnrealGameSync/UnrealGameSync/Properties/AssemblyInfo.cs b/Engine/Source/Programs/UnrealGameSync/UnrealGameSync/Properties/AssemblyInfo.cs
--- a/Engine/Source/Programs/UnrealGameSync/UnrealGameSync/Properties/AssemblyInfo.cs
+++ b/Engine/Source/Programs/UnrealGameSync/UnrealGameSync/Properties/AssemblyInfo.cs
@@ -34,5 +34,5 @@ using System.Runtime.InteropServices;
 // You can specify all the values or you can default the Build and Revision Numbers 
 // by using the '*' as shown below:
 // [assembly: AssemblyVersion("1.0.*")]
-[assembly: AssemblyVersion("1.177.*")]
-[assembly: AssemblyFileVersion("1.177.0.0")]
+[assembly: AssemblyVersion("1.179.*")]
+[assembly: AssemblyFileVersion("1.179.0.0")]
diff --git a/Engine/Source/Programs/UnrealGameSync/UnrealGameSync/Workspace.cs b/Engine/Source/Programs/UnrealGameSync/UnrealGameSync/Workspace.cs
--- a/Engine/Source/Programs/UnrealGameSync/UnrealGameSync/Workspace.cs
+++ b/Engine/Source/Programs/UnrealGameSync/UnrealGameSync/Workspace.cs
@@ -558,11 +558,6 @@ namespace UnrealGameSync
 					// Create a filter for all the files we don't want
 					FileFilter Filter = new FileFilter(UserFilter);
 					Filter.Exclude(BuildVersionFileName);
-					if (!ProjectConfigFile.GetValue("Options.UseFastModularVersioningV2", false))
-					{
-						Filter.Exclude(VersionHeaderFileName);
-						Filter.Exclude(ObjectVersionFileName);
-					}
 					if (Context.Options.HasFlag(WorkspaceUpdateOptions.ContentOnly))
 					{
 						Filter.Exclude("*.usf");
@@ -1412,7 +1407,19 @@ namespace UnrealGameSync
 				Log.WriteLine("Failed to query records for {0}", ClientPath);
 				return false;
 			}
-			if(Records.Count == 0)
+			if (Records.Count > 1)
+			{
+				// Attempt to remove any existing file which is synced
+				Perforce.ForceSync(String.Format("{0}#0", ClientPath), Log);
+
+				// Try to get the mapped files again
+				if (!Perforce.Stat(ClientPath, out Records, Log))
+				{
+					Log.WriteLine("Failed to query records for {0}", ClientPath);
+					return false;
+				}
+			}
+			if (Records.Count == 0)
 			{
 				Log.WriteLine("Ignoring {0}; not found on server.", ClientPath);
 				return true;
