Merging CL 7936110 from 4.23 stream.

#rnx

diff --git a/Engine/Source/Programs/UnrealGameSync/UnrealGameSync/Perforce.cs b/Engine/Source/Programs/UnrealGameSync/UnrealGameSync/Perforce.cs
--- a/Engine/Source/Programs/UnrealGameSync/UnrealGameSync/Perforce.cs
+++ b/Engine/Source/Programs/UnrealGameSync/UnrealGameSync/Perforce.cs
@@ -1225,10 +1225,10 @@ namespace UnrealGameSync
 					{
 						CommandLine.AppendFormat(" -f");
 					}
-					if(Options != null && Options.NumThreads > 1)
-					{
-						CommandLine.AppendFormat(" --parallel=threads={0}", Options.NumThreads);
-					}
+//					if(Options != null && Options.NumThreads > 1)
+//					{
+//						CommandLine.AppendFormat(" --parallel=threads={0}", Options.NumThreads);
+//					}
 					bResult = RunCommand(CommandLine.ToString(), null, Line => FilterSyncOutput(Line, Parser, TamperedFiles, Log), CommandOptions.NoFailOnErrors | CommandOptions.IgnoreFilesUpToDateError | CommandOptions.IgnoreExitCode, Log);
 				}
 				return bResult;
