-Change to look for NoContent rather than OK, since the server's returning a 204

diff --git a/Engine/Source/Programs/UnrealGameSync/PostBadgeStatus/Program.cs b/Engine/Source/Programs/UnrealGameSync/PostBadgeStatus/Program.cs
--- a/Engine/Source/Programs/UnrealGameSync/PostBadgeStatus/Program.cs
+++ b/Engine/Source/Programs/UnrealGameSync/PostBadgeStatus/Program.cs
@@ -1,4 +1,4 @@
-﻿// Copyright 1998-2019 Epic Games, Inc. All Rights Reserved.
+// Copyright 1998-2019 Epic Games, Inc. All Rights Reserved.
 
 using System;
 using System.Collections.Generic;
@@ -161,7 +161,7 @@ namespace WriteBadgeStatus
 					{
 						ResponseContent = ResponseReader.ReadToEnd();
 						Console.WriteLine(ResponseContent);
-						return Response.StatusCode == HttpStatusCode.OK ? 0 : 1;
+						return Response.StatusCode == HttpStatusCode.NoContent ? 0 : 1;
 					}
 				}
 				
