UGS: Add a wrapper for Process.Start() calls which prints an error if an exception is thrown, rather than allowing it to show an exception dialog.

#rnx

diff --git a/Engine/Source/Programs/UnrealGameSync/UnrealGameSync/Controls/WorkspaceControl.cs b/Engine/Source/Programs/UnrealGameSync/UnrealGameSync/Controls/WorkspaceControl.cs
--- a/Engine/Source/Programs/UnrealGameSync/UnrealGameSync/Controls/WorkspaceControl.cs
+++ b/Engine/Source/Programs/UnrealGameSync/UnrealGameSync/Controls/WorkspaceControl.cs
@@ -2215,11 +2215,11 @@ namespace UnrealGameSync
 								}
 								else if(String.IsNullOrEmpty(ExpandedArguments))
 								{
-									ClickHandler = () => Process.Start(ExpandedUrl);
+									ClickHandler = () => SafeProcessStart(ExpandedUrl);
 								}
 								else
 								{
-									ClickHandler = () => Process.Start(ExpandedUrl, ExpandedArguments);
+									ClickHandler = () => SafeProcessStart(ExpandedUrl, ExpandedArguments);
 								}
 
 								Badges.Add(new BadgeInfo(ReplaceRegexMatches(Name, MatchResult), Group, UniqueId, BadgeColor, HoverBadgeColor, ClickHandler));
@@ -2237,6 +2237,42 @@ namespace UnrealGameSync
 			return Badges;
 		}
 
+		private static void SafeProcessStart(ProcessStartInfo StartInfo)
+		{
+			try
+			{
+				Process.Start(StartInfo);
+			}
+			catch
+			{
+				MessageBox.Show(String.Format("Unable to open {0}", StartInfo.FileName));
+			}
+		}
+
+		private static void SafeProcessStart(string Url)
+		{
+			try
+			{
+				Process.Start(Url);
+			}
+			catch
+			{
+				MessageBox.Show(String.Format("Unable to open {0}", Url));
+			}
+		}
+
+		private static void SafeProcessStart(string Url, string Arguments)
+		{
+			try
+			{
+				Process.Start(Url, Arguments);
+			}
+			catch
+			{
+				MessageBox.Show(String.Format("Unable to open {0} {1}", Url, Arguments));
+			}
+		}
+
 		private string ReplaceRegexMatches(string Text, Match MatchResult)
 		{
 			if(Text != null)
@@ -2392,7 +2428,7 @@ namespace UnrealGameSync
 			}
 			else
 			{
-				ClickHandler = () => Process.Start(BadgeData.Url);
+				ClickHandler = () => SafeProcessStart(BadgeData.Url);
 			}
 
 			string UniqueId = String.Format("{0}:{1}", ChangeNumber, BadgeName);
@@ -2832,7 +2868,7 @@ namespace UnrealGameSync
 			{
 				CommandLine.AppendFormat("-p \"{0}\" -c \"{1}\" -u \"{2}\"", Workspace.Perforce.ServerAndPort ?? "perforce:1666", Workspace.Perforce.ClientName, Workspace.Perforce.UserName);
 			}
-			Process.Start("p4v.exe", CommandLine.ToString());
+			SafeProcessStart("p4v.exe", CommandLine.ToString());
 		}
 
 		private void UpdateStatusPanel()
@@ -2964,7 +3000,7 @@ namespace UnrealGameSync
 				ProgramsLine.AddText("  |  ");
 				ProgramsLine.AddLink("Visual Studio", FontStyle.Regular, () => { OpenSolution(); });
 				ProgramsLine.AddText("  |  ");
-				ProgramsLine.AddLink("Windows Explorer", FontStyle.Regular, () => { Process.Start("explorer.exe", String.Format("\"{0}\"", Path.GetDirectoryName(SelectedFileName))); });
+				ProgramsLine.AddLink("Windows Explorer", FontStyle.Regular, () => { SafeProcessStart("explorer.exe", String.Format("\"{0}\"", Path.GetDirectoryName(SelectedFileName))); });
 
 				string BuildHealthProject;
 				if(TryGetProjectSetting(PerforceMonitor.LatestProjectConfigFile, "BuildHealthProject", out BuildHealthProject))
@@ -2989,7 +3025,7 @@ namespace UnrealGameSync
 					}
 					else
 					{
-						ProgramsLine.AddBadge(ServiceBadge.Key, GetBuildBadgeColor(ServiceBadge.Value.Result), (P, R) => { Process.Start(ServiceBadge.Value.Url); });
+						ProgramsLine.AddBadge(ServiceBadge.Key, GetBuildBadgeColor(ServiceBadge.Value.Result), (P, R) => { SafeProcessStart(ServiceBadge.Value.Url); });
 					}
 				}
 				ProgramsLine.AddText("  |  ");
@@ -3260,7 +3296,7 @@ namespace UnrealGameSync
 							FlushMarkdownText(Line, ElementText, Style);
 							string LinkText = Text.Substring(Idx + 1, EndIdx - (Idx + 1));
 							string LinkUrl = Text.Substring(EndIdx + 2, UrlEndIdx - (EndIdx + 2));
-							Line.AddLink(LinkText, Style, () => Process.Start(LinkUrl));
+							Line.AddLink(LinkText, Style, () => SafeProcessStart(LinkUrl));
 							Idx = UrlEndIdx + 1;
 							continue;
 						}
@@ -3658,7 +3694,7 @@ namespace UnrealGameSync
 									Execute = Utility.ExpandVariables(Execute, Variables);
 									Arguments = Utility.ExpandVariables(Arguments ?? "", Variables);
 
-									ToolStripMenuItem Item = new ToolStripMenuItem(Label, null, new EventHandler((o, a) => Process.Start(Execute, Arguments)));
+									ToolStripMenuItem Item = new ToolStripMenuItem(Label, null, new EventHandler((o, a) => SafeProcessStart(Execute, Arguments)));
 	
 									BuildListContextMenu.Items.Insert(CustomToolEnd, Item);
 									CustomToolEnd++;
@@ -4190,7 +4226,7 @@ namespace UnrealGameSync
 			{
 				ProcessStartInfo StartInfo = new ProcessStartInfo(SolutionFileName);
 				StartInfo.WorkingDirectory = BranchDirectoryName;
-				Process.Start(StartInfo);
+				SafeProcessStart(StartInfo);
 			}
 		}
 
