UGS: Fix starred builds and reviews not updating unless UGS is restarted.

#rnx

diff --git a/Engine/Source/Programs/UnrealGameSync/UnrealGameSync/EventMonitor.cs b/Engine/Source/Programs/UnrealGameSync/UnrealGameSync/EventMonitor.cs
--- a/Engine/Source/Programs/UnrealGameSync/UnrealGameSync/EventMonitor.cs
+++ b/Engine/Source/Programs/UnrealGameSync/UnrealGameSync/EventMonitor.cs
@@ -487,7 +487,7 @@ namespace UnrealGameSync
 					if(Timer.Elapsed > TimeSpan.FromSeconds(RequestThrottle))
 					{
 						bUpdateThrottledRequests = true;
-						Timer.Reset();
+						Timer.Restart();
 					}
 
 					// Read all the new reviews, pass whether or not to fire the throttled requests
