UnrealGameSync: Prevent trying to build (and consequently generate project files) when syncing precompiled binaries. Can happen if the build checkbox is checked before the "Sync Precompiled Binaries" option is enabled.

diff --git a/Engine/Source/Programs/UnrealGameSync/UnrealGameSync/Controls/WorkspaceControl.cs b/Engine/Source/Programs/UnrealGameSync/UnrealGameSync/Controls/WorkspaceControl.cs
--- a/Engine/Source/Programs/UnrealGameSync/UnrealGameSync/Controls/WorkspaceControl.cs
+++ b/Engine/Source/Programs/UnrealGameSync/UnrealGameSync/Controls/WorkspaceControl.cs
@@ -870,7 +870,7 @@ namespace UnrealGameSync
 						MessageBox.Show("There are no compiled editor binaries for this change. To sync it, you must disable syncing of precompiled editor binaries.");
 						return;
 					}
-					Context.Options &= ~WorkspaceUpdateOptions.GenerateProjectFiles;
+					Context.Options &= ~(WorkspaceUpdateOptions.Build | WorkspaceUpdateOptions.GenerateProjectFiles);
 				}
 				Context.ArchiveTypeToDepotPath.Add(EditorArchiveType, EditorArchivePath);
 			}
