UGS: Allow users to run the editor after syncing precompiled binaries.

#rnx

diff --git a/Engine/Source/Programs/UnrealGameSync/UnrealGameSync/Controls/WorkspaceControl.cs b/Engine/Source/Programs/UnrealGameSync/UnrealGameSync/Controls/WorkspaceControl.cs
--- a/Engine/Source/Programs/UnrealGameSync/UnrealGameSync/Controls/WorkspaceControl.cs
+++ b/Engine/Source/Programs/UnrealGameSync/UnrealGameSync/Controls/WorkspaceControl.cs
@@ -3869,7 +3869,7 @@ namespace UnrealGameSync
 			{
 				AfterSyncingLabel.Enabled = true;
 				BuildAfterSyncCheckBox.Enabled = bHasBuildSteps;
-				RunAfterSyncCheckBox.Enabled = BuildAfterSyncCheckBox.Checked;
+				RunAfterSyncCheckBox.Enabled = BuildAfterSyncCheckBox.Checked || ShouldSyncPrecompiledEditor;
 				OpenSolutionAfterSyncCheckBox.Enabled = true;
 			}
 		}
