UGS: Also allow the "Sync Precompiled Editor" option to set the appropriate flag to launch after syncing.

#rnx

diff --git a/Engine/Source/Programs/UnrealGameSync/UnrealGameSync/Controls/WorkspaceControl.cs b/Engine/Source/Programs/UnrealGameSync/UnrealGameSync/Controls/WorkspaceControl.cs
--- a/Engine/Source/Programs/UnrealGameSync/UnrealGameSync/Controls/WorkspaceControl.cs
+++ b/Engine/Source/Programs/UnrealGameSync/UnrealGameSync/Controls/WorkspaceControl.cs
@@ -821,7 +821,7 @@ namespace UnrealGameSync
 			{
 				Options |= WorkspaceUpdateOptions.Build;
 			}
-			if(Settings.bBuildAfterSync && Settings.bRunAfterSync)
+			if((Settings.bBuildAfterSync || ShouldSyncPrecompiledEditor) && Settings.bRunAfterSync)
 			{
 				Options |= WorkspaceUpdateOptions.RunAfterSync;
 			}
