UGS: Add an option to build CrashReportClientEditor rather than CrashReportClient.

diff --git a/Engine/Source/Programs/UnrealGameSync/UnrealGameSync/Controls/WorkspaceControl.cs b/Engine/Source/Programs/UnrealGameSync/UnrealGameSync/Controls/WorkspaceControl.cs
--- a/Engine/Source/Programs/UnrealGameSync/UnrealGameSync/Controls/WorkspaceControl.cs
+++ b/Engine/Source/Programs/UnrealGameSync/UnrealGameSync/Controls/WorkspaceControl.cs
@@ -2777,6 +2777,9 @@ namespace UnrealGameSync
 			ProcessNames.Add("CrashReportClient");
 			ProcessNames.Add("CrashReportClient-Win64-Development");
 			ProcessNames.Add("CrashReportClient-Win64-Debug");
+			ProcessNames.Add("CrashReportClientEditor");
+			ProcessNames.Add("CrashReportClientEditor-Win64-Development");
+			ProcessNames.Add("CrashReportClientEditor-Win64-Debug");
 			ProcessNames.Add("UnrealBuildTool");
 			ProcessNames.Add("AutomationTool");
 
@@ -4619,12 +4622,16 @@ namespace UnrealGameSync
 				ActualEditorTargetName = "UE4Editor";
 			}
 
+			bool bUseCrashReportClientEditor = PerforceMonitor.LatestProjectConfigFile.GetValue("Options.UseCrashReportClientEditor", false);
+
 			List<BuildStep> DefaultBuildSteps = new List<BuildStep>();
 			DefaultBuildSteps.Add(new BuildStep(new Guid("{01F66060-73FA-4CC8-9CB3-E217FBBA954E}"), 0, "Compile UnrealHeaderTool", "Compiling UnrealHeaderTool...", 1, "UnrealHeaderTool", "Win64", "Development", "", !ShouldSyncPrecompiledEditor));
 			DefaultBuildSteps.Add(new BuildStep(new Guid("{F097FF61-C916-4058-8391-35B46C3173D5}"), 1, String.Format("Compile {0}", ActualEditorTargetName), String.Format("Compiling {0}...", ActualEditorTargetName), 10, ActualEditorTargetName, "Win64", Settings.CompiledEditorBuildConfig.ToString(), ProjectArgument, !ShouldSyncPrecompiledEditor));
 			DefaultBuildSteps.Add(new BuildStep(new Guid("{C6E633A1-956F-4AD3-BC95-6D06D131E7B4}"), 2, "Compile ShaderCompileWorker", "Compiling ShaderCompileWorker...", 1, "ShaderCompileWorker", "Win64", "Development", "", !ShouldSyncPrecompiledEditor));
 			DefaultBuildSteps.Add(new BuildStep(new Guid("{24FFD88C-7901-4899-9696-AE1066B4B6E8}"), 3, "Compile UnrealLightmass", "Compiling UnrealLightmass...", 1, "UnrealLightmass", "Win64", "Development", "", !ShouldSyncPrecompiledEditor));
-			DefaultBuildSteps.Add(new BuildStep(new Guid("{FFF20379-06BF-4205-8A3E-C53427736688}"), 4, "Compile CrashReportClient", "Compiling CrashReportClient...", 1, "CrashReportClient", "Win64", "Shipping", "", !ShouldSyncPrecompiledEditor));
+			DefaultBuildSteps.Add(new BuildStep(new Guid("{FFF20379-06BF-4205-8A3E-C53427736688}"), 4, "Compile CrashReportClient", "Compiling CrashReportClient...", 1, "CrashReportClient", "Win64", "Shipping", "", !ShouldSyncPrecompiledEditor && !bUseCrashReportClientEditor));
+			DefaultBuildSteps.Add(new BuildStep(new Guid("{7143D861-58D3-4F83-BADC-BC5DCB2079F6}"), 5, "Compile CrashReportClientEditor", "Compiling CrashReportClientEditor...", 1, "CrashReportClientEditor", "Win64", "Shipping", "", !ShouldSyncPrecompiledEditor && bUseCrashReportClientEditor));
+
 			return DefaultBuildSteps.ToDictionary(x => x.UniqueId, x => x.ToConfigObject());
 		}
 
