UGS: Fix a few minor display issues.

* Links were treating ampersands as accelerator shortcuts.
* Some dialogs were not rendering correctly in high DPI modes.

#rnx

diff --git a/Engine/Source/Programs/UnrealGameSync/UnrealGameSync/Controls/StatusPanel.cs b/Engine/Source/Programs/UnrealGameSync/UnrealGameSync/Controls/StatusPanel.cs
--- a/Engine/Source/Programs/UnrealGameSync/UnrealGameSync/Controls/StatusPanel.cs
+++ b/Engine/Source/Programs/UnrealGameSync/UnrealGameSync/Controls/StatusPanel.cs
@@ -175,7 +175,7 @@ namespace UnrealGameSync
 
 		public override Size Measure(Graphics Graphics, StatusElementResources Resources)
 		{
-			return TextRenderer.MeasureText(Graphics, Text, Resources.FindOrAddFont(Style), new Size(int.MaxValue, int.MaxValue), TextFormatFlags.NoPadding);
+			return TextRenderer.MeasureText(Graphics, Text, Resources.FindOrAddFont(Style), new Size(int.MaxValue, int.MaxValue), TextFormatFlags.NoPadding | TextFormatFlags.NoPrefix);
 		}
 
 		public override void Draw(Graphics Graphics, StatusElementResources Resources)
@@ -189,7 +189,7 @@ namespace UnrealGameSync
 			{
 				TextColor = Color.FromArgb(TextColor.B, TextColor.G, TextColor.R);
 			}
-			TextRenderer.DrawText(Graphics, Text, Resources.FindOrAddFont(Style), Bounds.Location, TextColor, TextFormatFlags.NoPadding);
+			TextRenderer.DrawText(Graphics, Text, Resources.FindOrAddFont(Style), Bounds.Location, TextColor, TextFormatFlags.NoPadding | TextFormatFlags.NoPrefix);
 		}
 	}
 
diff --git a/Engine/Source/Programs/UnrealGameSync/UnrealGameSync/Forms/IssueBrowserWindow.cs b/Engine/Source/Programs/UnrealGameSync/UnrealGameSync/Forms/IssueBrowserWindow.cs
--- a/Engine/Source/Programs/UnrealGameSync/UnrealGameSync/Forms/IssueBrowserWindow.cs
+++ b/Engine/Source/Programs/UnrealGameSync/UnrealGameSync/Forms/IssueBrowserWindow.cs
@@ -46,6 +46,15 @@ namespace UnrealGameSync
 
 			InitializeComponent();
 
+			using (Graphics Graphics = Graphics.FromHwnd(IntPtr.Zero))
+			{
+				float DpiScaleX = Graphics.DpiX / 96.0f;
+				foreach (ColumnHeader Column in IssueListView.Columns)
+				{
+					Column.Width = (int)(Column.Width * DpiScaleX);
+				}
+			}
+
 			System.Reflection.PropertyInfo DoubleBufferedProperty = typeof(Control).GetProperty("DoubleBuffered", System.Reflection.BindingFlags.NonPublic | System.Reflection.BindingFlags.Instance);
 			DoubleBufferedProperty.SetValue(IssueListView, true, null);
 		}
