UGS: Fix exception when right-clicking on "+" tab and selecting "Open Project...".

#rnx

diff --git a/Engine/Source/Programs/UnrealGameSync/UnrealGameSync/Forms/MainWindow.cs b/Engine/Source/Programs/UnrealGameSync/UnrealGameSync/Forms/MainWindow.cs
--- a/Engine/Source/Programs/UnrealGameSync/UnrealGameSync/Forms/MainWindow.cs
+++ b/Engine/Source/Programs/UnrealGameSync/UnrealGameSync/Forms/MainWindow.cs
@@ -1127,7 +1127,14 @@ namespace UnrealGameSync
 
 		private void TabMenu_OpenProject_Click(object sender, EventArgs e)
 		{
-			EditSelectedProject(TabMenu_TabIdx);
+			if (TabMenu_TabIdx == -1)
+			{
+				OpenNewProject();
+			}
+			else
+			{
+				EditSelectedProject(TabMenu_TabIdx);
+			}
 		}
 
 		private void TabMenu_TabNames_Stream_Click(object sender, EventArgs e)
