UGS: Fix changelists that are unstarred being selected as the target changelist to be sync to.

diff --git a/Engine/Source/Programs/UnrealGameSync/UnrealGameSync/Controls/WorkspaceControl.cs b/Engine/Source/Programs/UnrealGameSync/UnrealGameSync/Controls/WorkspaceControl.cs
--- a/Engine/Source/Programs/UnrealGameSync/UnrealGameSync/Controls/WorkspaceControl.cs
+++ b/Engine/Source/Programs/UnrealGameSync/UnrealGameSync/Controls/WorkspaceControl.cs
@@ -4335,7 +4335,7 @@ namespace UnrealGameSync
 					else if(ChangeType == LatestChangeType.Starred)
 					{
 						EventSummary Summary = EventMonitor.GetSummaryForChange(Change.Number);
-						if(((Summary != null && Summary.LastStarReview != null) || PromotedChangeNumbers.Contains(Change.Number)) && CanSyncChange(Change.Number))
+						if(((Summary != null && Summary.LastStarReview != null && Summary.LastStarReview.Type == EventType.Starred) || PromotedChangeNumbers.Contains(Change.Number)) && CanSyncChange(Change.Number))
 						{
 							ChangeNumber = FindNewestGoodContentChange(Change.Number);
 							return true;
