Copying latest UGS changes to 4.24

- Fix for service badges not updating
- Fix exception when using context menu to open new tab, if no other tabs are open

#rnx

diff --git a/Engine/Source/Programs/UnrealGameSync/UnrealGameSync/Controls/WorkspaceControl.cs b/Engine/Source/Programs/UnrealGameSync/UnrealGameSync/Controls/WorkspaceControl.cs
--- a/Engine/Source/Programs/UnrealGameSync/UnrealGameSync/Controls/WorkspaceControl.cs
+++ b/Engine/Source/Programs/UnrealGameSync/UnrealGameSync/Controls/WorkspaceControl.cs
@@ -106,6 +106,35 @@ namespace UnrealGameSync
 			}
 		}
 
+		class ServiceBadgeInfo : IEquatable<ServiceBadgeInfo>
+		{
+			public string Name;
+			public string Url;
+			public BadgeResult Result;
+
+			public ServiceBadgeInfo(string Name, string Url, BadgeResult Result)
+			{
+				this.Name = Name;
+				this.Url = Url;
+				this.Result = Result;
+			}
+
+			public bool Equals(ServiceBadgeInfo OtherBadge)
+			{
+				return OtherBadge != null && Name == OtherBadge.Name && Url == OtherBadge.Url && Result == OtherBadge.Result;
+			}
+
+			public override bool Equals(object Other)
+			{
+				return Equals(Other as ServiceBadgeInfo);
+			}
+
+			public override int GetHashCode()
+			{
+				return Name.GetHashCode() + (Url ?? "").GetHashCode() * 3 + Result.GetHashCode() * 5;
+			}
+		}
+
 		class ChangeLayoutInfo
 		{
 			public List<BadgeInfo> DescriptionBadges;
@@ -120,11 +149,11 @@ namespace UnrealGameSync
 			{
 				PerforceChangeSummary ChangeA = ((ListViewItem)A).Tag as PerforceChangeSummary;
 				PerforceChangeSummary ChangeB = ((ListViewItem)B).Tag as PerforceChangeSummary;
-				if(ChangeA == null)
+				if (ChangeA == null)
 				{
 					return +1;
 				}
-				else if(ChangeB == null)
+				else if (ChangeB == null)
 				{
 					return -1;
 				}
@@ -219,14 +248,14 @@ namespace UnrealGameSync
 		HashSet<int> PromotedChangeNumbers = new HashSet<int>();
 		List<int> ListIndexToChangeIndex = new List<int>();
 		List<int> SortedChangeNumbers = new List<int>();
-		Dictionary<int, string> ChangeNumberToArchivePath = new Dictionary<int,string>();
+		Dictionary<int, string> ChangeNumberToArchivePath = new Dictionary<int, string>();
 		Dictionary<int, ChangeLayoutInfo> ChangeNumberToLayoutInfo = new Dictionary<int, ChangeLayoutInfo>();
 		List<ToolStripMenuItem> CustomToolMenuItems = new List<ToolStripMenuItem>();
 		int NumChanges;
 		int PendingSelectedChangeNumber = -1;
 		bool bHasBuildSteps = false;
 
-		Dictionary<string, int> NotifiedBuildTypeToChangeNumber = new Dictionary<string,int>();
+		Dictionary<string, int> NotifiedBuildTypeToChangeNumber = new Dictionary<string, int>();
 
 		bool bMouseOverExpandLink;
 
@@ -238,7 +267,7 @@ namespace UnrealGameSync
 		Font BadgeFont;
 		List<KeyValuePair<string, string>> BadgeNameAndGroupPairs = new List<KeyValuePair<string, string>>();
 		Dictionary<string, Size> BadgeLabelToSize = new Dictionary<string, Size>();
-		List<KeyValuePair<string, BadgeData>> ServiceBadges = new List<KeyValuePair<string, BadgeData>>();
+		List<ServiceBadgeInfo> ServiceBadges = new List<ServiceBadgeInfo>();
 
 		string OriginalExecutableFileName;
 
@@ -287,12 +316,12 @@ namespace UnrealGameSync
 			ProjectSettings = InSettings.FindOrAddProject(DetectSettings.NewSelectedClientFileName);
 
 			DesiredTaskbarState = Tuple.Create(TaskbarState.NoProgress, 0.0f);
-			
+
 			System.Reflection.PropertyInfo DoubleBufferedProperty = typeof(Control).GetProperty("DoubleBuffered", System.Reflection.BindingFlags.NonPublic | System.Reflection.BindingFlags.Instance);
-			DoubleBufferedProperty.SetValue(BuildList, true, null); 
+			DoubleBufferedProperty.SetValue(BuildList, true, null);
 
 			// force the height of the rows
-			BuildList.SmallImageList = new ImageList(){ ImageSize = new Size(1, 20) };
+			BuildList.SmallImageList = new ImageList() { ImageSize = new Size(1, 20) };
 			BuildList_FontChanged(null, null);
 			BuildList.OnScroll += BuildList_OnScroll;
 			BuildList.ListViewItemSorter = new BuildListItemComparer();
@@ -328,7 +357,7 @@ namespace UnrealGameSync
 			// Check if we've the project we've got open in this workspace is the one we're actually synced to
 			int CurrentChangeNumber = -1;
 			string CurrentSyncFilterHash = null;
-			if(String.Compare(WorkspaceSettings.CurrentProjectIdentifier, SelectedProjectIdentifier, true) == 0)
+			if (String.Compare(WorkspaceSettings.CurrentProjectIdentifier, SelectedProjectIdentifier, true) == 0)
 			{
 				CurrentChangeNumber = WorkspaceSettings.CurrentChangeNumber;
 				CurrentSyncFilterHash = WorkspaceSettings.CurrentSyncFilterHash;
@@ -367,7 +396,7 @@ namespace UnrealGameSync
 			PerforceMonitor.Start();
 			EventMonitor.Start();
 
-			StartupTimer = new System.Threading.Timer(x => MainThreadSynchronizationContext.Post((o) => { if(!IsDisposed){ StartupTimerElapsed(false); } }, null), null, TimeSpan.FromSeconds(20.0), TimeSpan.FromMilliseconds(-1.0));
+			StartupTimer = new System.Threading.Timer(x => MainThreadSynchronizationContext.Post((o) => { if (!IsDisposed) { StartupTimerElapsed(false); } }, null), null, TimeSpan.FromSeconds(20.0), TimeSpan.FromMilliseconds(-1.0));
 			StartupCallbacks = new List<WorkspaceStartupCallback>();
 
 			IssueMonitor.OnIssuesChanged += IssueMonitor_OnIssuesChangedAsync;
@@ -375,7 +404,7 @@ namespace UnrealGameSync
 
 		public void IssueMonitor_OnIssuesChangedAsync()
 		{
-			MainThreadSynchronizationContext.Post((o) => { if(IssueMonitor != null){ IssueMonitor_OnIssuesChanged(); } }, null);
+			MainThreadSynchronizationContext.Post((o) => { if (IssueMonitor != null) { IssueMonitor_OnIssuesChanged(); } }, null);
 		}
 
 		public List<IssueData> GetOpenIssuesForUser()
@@ -387,7 +416,7 @@ namespace UnrealGameSync
 		{
 			bool bPrevUserHasOpenIssues = bUserHasOpenIssues;
 			bUserHasOpenIssues = GetOpenIssuesForUser().Count > 0;
-			if(bUserHasOpenIssues != bPrevUserHasOpenIssues)
+			if (bUserHasOpenIssues != bPrevUserHasOpenIssues)
 			{
 				UpdateStatusPanel();
 				StatusPanel.Invalidate();
@@ -401,10 +430,10 @@ namespace UnrealGameSync
 
 		private void CheckForStartupComplete()
 		{
-			if(StartupTimer != null)
+			if (StartupTimer != null)
 			{
 				int LatestChangeNumber;
-				if(FindChangeToSync(Settings.SyncType, out LatestChangeNumber))
+				if (FindChangeToSync(Settings.SyncType, out LatestChangeNumber))
 				{
 					StartupTimerElapsed(false);
 				}
@@ -413,15 +442,15 @@ namespace UnrealGameSync
 
 		private void StartupTimerElapsed(bool bCancel)
 		{
-			if(StartupTimer != null)
+			if (StartupTimer != null)
 			{
 				StartupTimer.Dispose();
 				StartupTimer = null;
 			}
 
-			if(StartupCallbacks != null)
+			if (StartupCallbacks != null)
 			{
-				foreach(WorkspaceStartupCallback StartupCallback in StartupCallbacks)
+				foreach (WorkspaceStartupCallback StartupCallback in StartupCallbacks)
 				{
 					StartupCallback(this, bCancel);
 				}
@@ -431,7 +460,7 @@ namespace UnrealGameSync
 
 		public void AddStartupCallback(WorkspaceStartupCallback StartupCallback)
 		{
-			if(StartupTimer == null)
+			if (StartupTimer == null)
 			{
 				StartupCallback(this, false);
 			}
@@ -446,46 +475,46 @@ namespace UnrealGameSync
 			string NextColumnSettings;
 			TryGetProjectSetting(PerforceMonitor.LatestProjectConfigFile, "Columns", out NextColumnSettings);
 
-			if(CustomColumns == null || NextColumnSettings != LastColumnSettings)
+			if (CustomColumns == null || NextColumnSettings != LastColumnSettings)
 			{
 				LastColumnSettings = NextColumnSettings;
 
-				if(CustomColumns != null)
+				if (CustomColumns != null)
 				{
-					foreach(ColumnHeader CustomColumn in CustomColumns)
+					foreach (ColumnHeader CustomColumn in CustomColumns)
 					{
 						BuildList.Columns.Remove(CustomColumn);
 					}
 				}
 
 				Dictionary<string, ColumnHeader> NameToColumn = new Dictionary<string, ColumnHeader>();
-				foreach(ColumnHeader Column in BuildList.Columns)
+				foreach (ColumnHeader Column in BuildList.Columns)
 				{
 					NameToColumn[Column.Text] = Column;
 					Column.Tag = null;
 				}
 
 				CustomColumns = new List<ColumnHeader>();
-				if(NextColumnSettings != null)
+				if (NextColumnSettings != null)
 				{
-					foreach(string CustomColumn in NextColumnSettings.Split('\n'))
+					foreach (string CustomColumn in NextColumnSettings.Split('\n'))
 					{
 						ConfigObject ColumnConfig = new ConfigObject(CustomColumn);
 
 						string Name = ColumnConfig.GetValue("Name", null);
-						if(Name != null)
+						if (Name != null)
 						{
 							ColumnHeader Column;
-							if(NameToColumn.TryGetValue(Name, out Column))
+							if (NameToColumn.TryGetValue(Name, out Column))
 							{
 								Column.Tag = ColumnConfig;
 							}
 							else
 							{
 								int Index = ColumnConfig.GetValue("Index", -1);
-								if(Index < 0 || Index > BuildList.Columns.Count)
+								if (Index < 0 || Index > BuildList.Columns.Count)
 								{
-									Index = ((CustomColumns.Count > 0)? CustomColumns[CustomColumns.Count - 1].Index : CISColumn.Index) + 1;
+									Index = ((CustomColumns.Count > 0) ? CustomColumns[CustomColumns.Count - 1].Index : CISColumn.Index) + 1;
 								}
 
 								Column = new ColumnHeader();
@@ -500,12 +529,12 @@ namespace UnrealGameSync
 				}
 
 				ColumnWidths = new float[BuildList.Columns.Count];
-				for(int Idx = 0; Idx < BuildList.Columns.Count; Idx++)
+				for (int Idx = 0; Idx < BuildList.Columns.Count; Idx++)
 				{
 					ColumnWidths[Idx] = BuildList.Columns[Idx].Width;
 				}
 
-				using(Graphics Graphics = Graphics.FromHwnd(IntPtr.Zero))
+				using (Graphics Graphics = Graphics.FromHwnd(IntPtr.Zero))
 				{
 					float DpiScaleX = Graphics.DpiX / 96.0f;
 
@@ -534,10 +563,10 @@ namespace UnrealGameSync
 				ColumnWeights[DescriptionColumn.Index] = 1.25f;
 				ColumnWeights[CISColumn.Index] = 1.5f;
 
-				foreach(ColumnHeader Column in BuildList.Columns)
+				foreach (ColumnHeader Column in BuildList.Columns)
 				{
 					ConfigObject ColumnConfig = (ConfigObject)Column.Tag;
-					if(ColumnConfig != null)
+					if (ColumnConfig != null)
 					{
 						MinColumnWidths[Column.Index] = ColumnConfig.GetValue("MinWidth", MinColumnWidths[Column.Index]);
 						DesiredColumnWidths[Column.Index] = ColumnConfig.GetValue("DesiredWidth", DesiredColumnWidths[Column.Index]);
@@ -546,15 +575,15 @@ namespace UnrealGameSync
 				}
 
 				ConfigFile ProjectConfigFile = PerforceMonitor.LatestProjectConfigFile;
-				for(int Idx = 0; Idx < BuildList.Columns.Count; Idx++)
+				for (int Idx = 0; Idx < BuildList.Columns.Count; Idx++)
 				{
-					if(!String.IsNullOrEmpty(BuildList.Columns[Idx].Text))
+					if (!String.IsNullOrEmpty(BuildList.Columns[Idx].Text))
 					{
 						string StringValue;
-						if(TryGetProjectSetting(ProjectConfigFile, String.Format("ColumnWidth_{0}", BuildList.Columns[Idx].Text), out StringValue))
+						if (TryGetProjectSetting(ProjectConfigFile, String.Format("ColumnWidth_{0}", BuildList.Columns[Idx].Text), out StringValue))
 						{
 							int IntValue;
-							if(Int32.TryParse(StringValue, out IntValue))
+							if (Int32.TryParse(StringValue, out IntValue))
 							{
 								DesiredColumnWidths[Idx] = IntValue;
 							}
@@ -562,28 +591,38 @@ namespace UnrealGameSync
 					}
 				}
 
-				if(ColumnWidths != null)
+				if (ColumnWidths != null)
 				{
 					ResizeColumns(ColumnWidths.Sum());
 				}
 			}
 		}
 
-		private void UpdateServiceBadges()
+		private bool UpdateServiceBadges()
 		{
 			string[] ServiceBadgeNames;
-			if(!TryGetProjectSetting(PerforceMonitor.LatestProjectConfigFile, "ServiceBadges", out ServiceBadgeNames))
+			if (!TryGetProjectSetting(PerforceMonitor.LatestProjectConfigFile, "ServiceBadges", out ServiceBadgeNames))
 			{
 				ServiceBadgeNames = new string[0];
 			}
 
-			ServiceBadges.Clear();
-			foreach(string ServiceBadgeName in ServiceBadgeNames)
+			List<ServiceBadgeInfo> PrevServiceBadges = ServiceBadges;
+
+			ServiceBadges = new List<ServiceBadgeInfo>();
+			foreach (string ServiceBadgeName in ServiceBadgeNames)
 			{
 				BadgeData LatestBuild;
-				EventMonitor.TryGetLatestBadge(ServiceBadgeName, out LatestBuild);
-				ServiceBadges.Add(new KeyValuePair<string, BadgeData>(ServiceBadgeName, LatestBuild));
+				if (EventMonitor.TryGetLatestBadge(ServiceBadgeName, out LatestBuild))
+				{
+					ServiceBadges.Add(new ServiceBadgeInfo(ServiceBadgeName, LatestBuild.Url, LatestBuild.Result));
+				}
+				else
+				{
+					ServiceBadges.Add(new ServiceBadgeInfo(ServiceBadgeName, null, BadgeResult.Skipped));
+				}
 			}
+
+			return !Enumerable.SequenceEqual(ServiceBadges, PrevServiceBadges);
 		}
 
 		protected override void OnLoad(EventArgs e)
@@ -594,9 +633,9 @@ namespace UnrealGameSync
 
 			// Find the default widths 
 			ColumnWidths = new float[BuildList.Columns.Count];
-			for(int Idx = 0; Idx < BuildList.Columns.Count; Idx++)
+			for (int Idx = 0; Idx < BuildList.Columns.Count; Idx++)
 			{
-				if(DesiredColumnWidths[Idx] > 0)
+				if (DesiredColumnWidths[Idx] > 0)
 				{
 					ColumnWidths[Idx] = DesiredColumnWidths[Idx];
 				}
@@ -625,56 +664,56 @@ namespace UnrealGameSync
 
 			UpdateTimer.Stop();
 
-			if(StartupCallbacks != null)
+			if (StartupCallbacks != null)
 			{
-				foreach(WorkspaceStartupCallback StartupCallback in StartupCallbacks)
+				foreach (WorkspaceStartupCallback StartupCallback in StartupCallbacks)
 				{
 					StartupCallback(this, true);
 				}
 				StartupCallbacks = null;
 			}
 
-			if(IssueMonitor != null)
+			if (IssueMonitor != null)
 			{
 				IssueMonitor.OnIssuesChanged -= IssueMonitor_OnIssuesChangedAsync;
 				IssueMonitor = null;
 			}
-			if(StartupTimer != null)
+			if (StartupTimer != null)
 			{
 				StartupTimer.Dispose();
 				StartupTimer = null;
 			}
-			if(NotificationWindow != null)
+			if (NotificationWindow != null)
 			{
 				NotificationWindow.Dispose();
 				NotificationWindow = null;
 			}
-			if(PerforceMonitor != null)
+			if (PerforceMonitor != null)
 			{
 				PerforceMonitor.Dispose();
 				PerforceMonitor = null;
 			}
-			if(Workspace != null)
+			if (Workspace != null)
 			{
 				Workspace.Dispose();
 				Workspace = null;
 			}
-			if(EventMonitor != null)
+			if (EventMonitor != null)
 			{
 				EventMonitor.Dispose();
 				EventMonitor = null;
 			}
-			if(BuildFont != null)
+			if (BuildFont != null)
 			{
 				BuildFont.Dispose();
 				BuildFont = null;
 			}
-			if(SelectedBuildFont != null)
+			if (SelectedBuildFont != null)
 			{
 				SelectedBuildFont.Dispose();
 				SelectedBuildFont = null;
 			}
-			if(BadgeFont != null)
+			if (BadgeFont != null)
 			{
 				BadgeFont.Dispose();
 				BadgeFont = null;
@@ -719,16 +758,16 @@ namespace UnrealGameSync
 		private void StreamChanged()
 		{
 			Owner.StreamChanged(this);
-/*
-			StatusPanel.SuspendDisplay();
+			/*
+						StatusPanel.SuspendDisplay();
 
-			string PrevSelectedFileName = SelectedFileName;
-			if(TryCloseProject())
-			{
-				OpenProject(PrevSelectedFileName);
-			}
+						string PrevSelectedFileName = SelectedFileName;
+						if(TryCloseProject())
+						{
+							OpenProject(PrevSelectedFileName);
+						}
 
-			StatusPanel.ResumeDisplay();*/
+						StatusPanel.ResumeDisplay();*/
 		}
 
 		private void StreamChangedCallback()
@@ -738,7 +777,7 @@ namespace UnrealGameSync
 
 		private void LoginExpired()
 		{
-			if(!bIsDisposing)
+			if (!bIsDisposing)
 			{
 				Log.WriteLine("Login has expired. Requesting project to be closed.");
 				Owner.RequestProjectChange(this, SelectedProject, false);
@@ -755,8 +794,8 @@ namespace UnrealGameSync
 			WorkspaceSettings.CurrentProjectIdentifier = SelectedProjectIdentifier;
 			WorkspaceSettings.CurrentChangeNumber = ChangeNumber;
 			WorkspaceSettings.CurrentSyncFilterHash = SyncFilterHash;
-			if(ChangeNumber == -1 || ChangeNumber != WorkspaceSettings.CurrentChangeNumber)
-			{ 
+			if (ChangeNumber == -1 || ChangeNumber != WorkspaceSettings.CurrentChangeNumber)
+			{
 				WorkspaceSettings.AdditionalChangeNumbers.Clear();
 			}
 			Settings.Save();
@@ -769,14 +808,14 @@ namespace UnrealGameSync
 
 		void ShrinkNumRequestedBuilds()
 		{
-			if(PerforceMonitor != null && BuildList.Items.Count > 0 && PendingSelectedChangeNumber == -1)
+			if (PerforceMonitor != null && BuildList.Items.Count > 0 && PendingSelectedChangeNumber == -1)
 			{
 				// Find the number of visible items using a (slightly wasteful) binary search
 				int VisibleItemCount = 1;
-				for(int StepSize = BuildList.Items.Count / 2; StepSize >= 1; )
+				for (int StepSize = BuildList.Items.Count / 2; StepSize >= 1;)
 				{
 					int TestIndex = VisibleItemCount + StepSize;
-					if(TestIndex < BuildList.Items.Count && BuildList.GetItemRect(TestIndex).Top < BuildList.Height)
+					if (TestIndex < BuildList.Items.Count && BuildList.GetItemRect(TestIndex).Top < BuildList.Height)
 					{
 						VisibleItemCount += StepSize;
 					}
@@ -788,21 +827,21 @@ namespace UnrealGameSync
 
 				// Figure out the last index to ensure is visible
 				int LastVisibleIndex = VisibleItemCount;
-				if(LastVisibleIndex >= ListIndexToChangeIndex.Count)
+				if (LastVisibleIndex >= ListIndexToChangeIndex.Count)
 				{
 					LastVisibleIndex = ListIndexToChangeIndex.Count - 1;
 				}
 
 				// Get the max number of changes to ensure this
 				int NewPendingMaxChanges = 0;
-				if(LastVisibleIndex >= 0)
+				if (LastVisibleIndex >= 0)
 				{
 					NewPendingMaxChanges = ListIndexToChangeIndex[LastVisibleIndex];
 				}
 				NewPendingMaxChanges = PerforceMonitor.InitialMaxChangesValue + ((Math.Max(NewPendingMaxChanges - PerforceMonitor.InitialMaxChangesValue, 0) + BuildListExpandCount - 1) / BuildListExpandCount) * BuildListExpandCount;
 
 				// Shrink the number of changes retained by the PerforceMonitor class
-				if(PerforceMonitor.PendingMaxChanges > NewPendingMaxChanges)
+				if (PerforceMonitor.PendingMaxChanges > NewPendingMaxChanges)
 				{
 					PerforceMonitor.PendingMaxChanges = NewPendingMaxChanges;
 				}
@@ -817,15 +856,15 @@ namespace UnrealGameSync
 		void StartSync(int ChangeNumber, WorkspaceUpdateCallback Callback)
 		{
 			WorkspaceUpdateOptions Options = WorkspaceUpdateOptions.Sync | WorkspaceUpdateOptions.SyncArchives | WorkspaceUpdateOptions.GenerateProjectFiles;
-			if(Settings.bBuildAfterSync)
+			if (Settings.bBuildAfterSync)
 			{
 				Options |= WorkspaceUpdateOptions.Build;
 			}
-			if((Settings.bBuildAfterSync || ShouldSyncPrecompiledEditor) && Settings.bRunAfterSync)
+			if ((Settings.bBuildAfterSync || ShouldSyncPrecompiledEditor) && Settings.bRunAfterSync)
 			{
 				Options |= WorkspaceUpdateOptions.RunAfterSync;
 			}
-			if(Settings.bOpenSolutionAfterSync)
+			if (Settings.bOpenSolutionAfterSync)
 			{
 				Options |= WorkspaceUpdateOptions.OpenSolutionAfterSync;
 			}
@@ -839,9 +878,9 @@ namespace UnrealGameSync
 
 		void StartWorkspaceUpdate(int ChangeNumber, WorkspaceUpdateOptions Options, WorkspaceUpdateCallback Callback)
 		{
-			if((Options & (WorkspaceUpdateOptions.Sync | WorkspaceUpdateOptions.Build)) != 0 && GetProcessesRunningInWorkspace().Length > 0)
+			if ((Options & (WorkspaceUpdateOptions.Sync | WorkspaceUpdateOptions.Build)) != 0 && GetProcessesRunningInWorkspace().Length > 0)
 			{
-				if((Options & WorkspaceUpdateOptions.ScheduledBuild) != 0)
+				if ((Options & WorkspaceUpdateOptions.ScheduledBuild) != 0)
 				{
 					SyncLog.Clear();
 					SyncLog.AppendLine("Editor is open; scheduled sync has been aborted.");
@@ -849,7 +888,7 @@ namespace UnrealGameSync
 				}
 				else
 				{
-					if(!WaitForProgramsToFinish())
+					if (!WaitForProgramsToFinish())
 					{
 						return;
 					}
@@ -859,18 +898,24 @@ namespace UnrealGameSync
 			string[] CombinedSyncFilter = UserSettings.GetCombinedSyncFilter(Workspace.GetSyncCategories(), Settings.SyncView, Settings.SyncExcludedCategories, WorkspaceSettings.SyncView, WorkspaceSettings.SyncIncludedCategories, WorkspaceSettings.SyncExcludedCategories);
 
 			WorkspaceUpdateContext Context = new WorkspaceUpdateContext(ChangeNumber, Options, CombinedSyncFilter, GetDefaultBuildStepObjects(), ProjectSettings.BuildSteps, null, GetWorkspaceVariables(ChangeNumber));
-			if(Options.HasFlag(WorkspaceUpdateOptions.SyncArchives))
+			if (Options.HasFlag(WorkspaceUpdateOptions.SyncArchives))
 			{
 				string EditorArchivePath = null;
-				if(ShouldSyncPrecompiledEditor)
+				if (ShouldSyncPrecompiledEditor)
 				{
 					EditorArchivePath = GetArchivePathForChangeNumber(ChangeNumber);
-					if(EditorArchivePath == null)
+					if (EditorArchivePath == null)
 					{
 						MessageBox.Show("There are no compiled editor binaries for this change. To sync it, you must disable syncing of precompiled editor binaries.");
 						return;
 					}
 					Context.Options &= ~(WorkspaceUpdateOptions.Build | WorkspaceUpdateOptions.GenerateProjectFiles);
+
+					string[] ZippedBinariesSyncFilter;
+					if (TryGetProjectSetting(PerforceMonitor.LatestProjectConfigFile, "ZippedBinariesSyncFilter", out ZippedBinariesSyncFilter) && ZippedBinariesSyncFilter.Length > 0)
+					{
+						Context.SyncFilter = Enumerable.Concat(Context.SyncFilter, ZippedBinariesSyncFilter).ToArray();
+					}
 				}
 				Context.ArchiveTypeToDepotPath.Add(EditorArchiveType, EditorArchivePath);
 			}
@@ -879,19 +924,19 @@ namespace UnrealGameSync
 
 		void StartWorkspaceUpdate(WorkspaceUpdateContext Context, WorkspaceUpdateCallback Callback)
 		{
-			if(Settings.bAutoResolveConflicts)
+			if (Settings.bAutoResolveConflicts)
 			{
 				Context.Options |= WorkspaceUpdateOptions.AutoResolveChanges;
 			}
-			if(Settings.bUseIncrementalBuilds)
+			if (Settings.bUseIncrementalBuilds)
 			{
 				Context.Options |= WorkspaceUpdateOptions.UseIncrementalBuilds;
 			}
-			if(WorkspaceSettings.bSyncAllProjects ?? Settings.bSyncAllProjects)
+			if (WorkspaceSettings.bSyncAllProjects ?? Settings.bSyncAllProjects)
 			{
 				Context.Options |= WorkspaceUpdateOptions.SyncAllProjects | WorkspaceUpdateOptions.IncludeAllProjectsInSolution;
 			}
-			if(WorkspaceSettings.bIncludeAllProjectsInSolution ?? Settings.bIncludeAllProjectsInSolution)
+			if (WorkspaceSettings.bIncludeAllProjectsInSolution ?? Settings.bIncludeAllProjectsInSolution)
 			{
 				Context.Options |= WorkspaceUpdateOptions.IncludeAllProjectsInSolution;
 			}
@@ -906,24 +951,24 @@ namespace UnrealGameSync
 			Log.WriteLine("  Options={0}", Context.Options.ToString());
 			Log.WriteLine("  Clobbering {0} files", Context.ClobberFiles.Count);
 
-			if(Context.Options.HasFlag(WorkspaceUpdateOptions.Sync))
+			if (Context.Options.HasFlag(WorkspaceUpdateOptions.Sync))
 			{
 				UpdateSyncConfig(-1, Workspace.CurrentSyncFilterHash);
 				EventMonitor.PostEvent(Context.ChangeNumber, EventType.Syncing);
 			}
 
-			if(Context.Options.HasFlag(WorkspaceUpdateOptions.Sync) || Context.Options.HasFlag(WorkspaceUpdateOptions.Build))
+			if (Context.Options.HasFlag(WorkspaceUpdateOptions.Sync) || Context.Options.HasFlag(WorkspaceUpdateOptions.Build))
 			{
-				if(!Context.Options.HasFlag(WorkspaceUpdateOptions.ContentOnly) && (Context.CustomBuildSteps == null || Context.CustomBuildSteps.Count == 0))
+				if (!Context.Options.HasFlag(WorkspaceUpdateOptions.ContentOnly) && (Context.CustomBuildSteps == null || Context.CustomBuildSteps.Count == 0))
 				{
-					foreach(BuildConfig Config in Enum.GetValues(typeof(BuildConfig)))
+					foreach (BuildConfig Config in Enum.GetValues(typeof(BuildConfig)))
 					{
 						List<string> EditorReceiptPaths = GetEditorReceiptPaths(Config);
-						foreach(string EditorReceiptPath in EditorReceiptPaths)
+						foreach (string EditorReceiptPath in EditorReceiptPaths)
 						{
-							if(File.Exists(EditorReceiptPath))
+							if (File.Exists(EditorReceiptPath))
 							{
-								try { File.Delete(EditorReceiptPath); } catch(Exception){ }
+								try { File.Delete(EditorReceiptPath); } catch (Exception) { }
 							}
 						}
 					}
@@ -939,7 +984,7 @@ namespace UnrealGameSync
 
 		void CancelWorkspaceUpdate()
 		{
-			if(Workspace.IsBusy() && MessageBox.Show("Are you sure you want to cancel the current operation?", "Cancel operation", MessageBoxButtons.YesNo) == DialogResult.Yes)
+			if (Workspace.IsBusy() && MessageBox.Show("Are you sure you want to cancel the current operation?", "Cancel operation", MessageBoxButtons.YesNo) == DialogResult.Yes)
 			{
 				WorkspaceSettings.LastSyncChangeNumber = Workspace.PendingChangeNumber;
 				WorkspaceSettings.LastSyncResult = WorkspaceUpdateResult.Canceled;
@@ -950,7 +995,7 @@ namespace UnrealGameSync
 
 				Workspace.CancelUpdate();
 
-				if(UpdateCallback != null)
+				if (UpdateCallback != null)
 				{
 					UpdateCallback(WorkspaceUpdateResult.Canceled);
 					UpdateCallback = null;
@@ -978,13 +1023,13 @@ namespace UnrealGameSync
 
 			UpdateSyncConfig(Workspace.CurrentChangeNumber, Workspace.CurrentSyncFilterHash);
 
-			if(Result == WorkspaceUpdateResult.Success && Context.Options.HasFlag(WorkspaceUpdateOptions.SyncSingleChange))
+			if (Result == WorkspaceUpdateResult.Success && Context.Options.HasFlag(WorkspaceUpdateOptions.SyncSingleChange))
 			{
 				WorkspaceSettings.AdditionalChangeNumbers.Add(Context.ChangeNumber);
 				Settings.Save();
 			}
 
-			if(Result == WorkspaceUpdateResult.Success && WorkspaceSettings.ExpandedArchiveTypes != null)
+			if (Result == WorkspaceUpdateResult.Success && WorkspaceSettings.ExpandedArchiveTypes != null)
 			{
 				WorkspaceSettings.ExpandedArchiveTypes = WorkspaceSettings.ExpandedArchiveTypes.Except(Context.ArchiveTypeToDepotPath.Where(x => x.Value == null).Select(x => x.Key)).ToArray();
 			}
@@ -997,61 +1042,61 @@ namespace UnrealGameSync
 			WorkspaceSettings.LastBuiltChangeNumber = Workspace.LastBuiltChangeNumber;
 			Settings.Save();
 
-			if(Result == WorkspaceUpdateResult.FilesToResolve)
+			if (Result == WorkspaceUpdateResult.FilesToResolve)
 			{
 				MessageBox.Show("You have files to resolve after syncing your workspace. Please check P4.");
 			}
-			else if(Result == WorkspaceUpdateResult.FilesToDelete)
+			else if (Result == WorkspaceUpdateResult.FilesToDelete)
 			{
 				DesiredTaskbarState = Tuple.Create(TaskbarState.Paused, 0.0f);
 				Owner.UpdateProgress();
 
 				DeleteWindow Window = new DeleteWindow(Context.DeleteFiles);
-				if(Window.ShowDialog(this) == DialogResult.OK)
+				if (Window.ShowDialog(this) == DialogResult.OK)
 				{
 					StartWorkspaceUpdate(Context, UpdateCallback);
 					return;
 				}
 			}
-			else if(Result == WorkspaceUpdateResult.FilesToClobber)
+			else if (Result == WorkspaceUpdateResult.FilesToClobber)
 			{
 				DesiredTaskbarState = Tuple.Create(TaskbarState.Paused, 0.0f);
 				Owner.UpdateProgress();
 
 				ClobberWindow Window = new ClobberWindow(Context.ClobberFiles);
-				if(Window.ShowDialog(this) == DialogResult.OK)
+				if (Window.ShowDialog(this) == DialogResult.OK)
 				{
 					StartWorkspaceUpdate(Context, UpdateCallback);
 					return;
 				}
 			}
-			else if(Result == WorkspaceUpdateResult.FailedToCompileWithCleanWorkspace)
+			else if (Result == WorkspaceUpdateResult.FailedToCompileWithCleanWorkspace)
 			{
 				EventMonitor.PostEvent(Context.ChangeNumber, EventType.DoesNotCompile);
 			}
-			else if(Result == WorkspaceUpdateResult.Success)
+			else if (Result == WorkspaceUpdateResult.Success)
 			{
-				if(Context.Options.HasFlag(WorkspaceUpdateOptions.Build))
+				if (Context.Options.HasFlag(WorkspaceUpdateOptions.Build))
 				{
 					EventMonitor.PostEvent(Context.ChangeNumber, EventType.Compiles);
 				}
-				if(Context.Options.HasFlag(WorkspaceUpdateOptions.RunAfterSync))
+				if (Context.Options.HasFlag(WorkspaceUpdateOptions.RunAfterSync))
 				{
 					LaunchEditor();
 				}
-				if(Context.Options.HasFlag(WorkspaceUpdateOptions.OpenSolutionAfterSync))
+				if (Context.Options.HasFlag(WorkspaceUpdateOptions.OpenSolutionAfterSync))
 				{
 					OpenSolution();
 				}
 			}
 
-			if(UpdateCallback != null)
+			if (UpdateCallback != null)
 			{
 				UpdateCallback(Result);
 				UpdateCallback = null;
 			}
 
-			DesiredTaskbarState = Tuple.Create((Result == WorkspaceUpdateResult.Success)? TaskbarState.NoProgress : TaskbarState.Error, 0.0f);
+			DesiredTaskbarState = Tuple.Create((Result == WorkspaceUpdateResult.Success) ? TaskbarState.NoProgress : TaskbarState.Error, 0.0f);
 			Owner.UpdateProgress();
 
 			BuildList.Invalidate();
@@ -1060,7 +1105,7 @@ namespace UnrealGameSync
 			UpdateSyncActionCheckboxes();
 
 			// Do this last because it may result in the control being disposed
-			if(Result == WorkspaceUpdateResult.FailedToSyncLoginExpired)
+			if (Result == WorkspaceUpdateResult.FailedToSyncLoginExpired)
 			{
 				LoginExpired();
 			}
@@ -1068,16 +1113,16 @@ namespace UnrealGameSync
 
 		void UpdateBuildListCallback()
 		{
-			if(!bUpdateBuildListPosted)
+			if (!bUpdateBuildListPosted)
 			{
 				bUpdateBuildListPosted = true;
-				MainThreadSynchronizationContext.Post((o) => { bUpdateBuildListPosted = false; if(!bIsDisposing) { UpdateBuildList(); } }, null);
+				MainThreadSynchronizationContext.Post((o) => { bUpdateBuildListPosted = false; if (!bIsDisposing) { UpdateBuildList(); } }, null);
 			}
 		}
 
 		void UpdateBuildList()
 		{
-			if(SelectedFileName != null)
+			if (SelectedFileName != null)
 			{
 				ChangeNumberToArchivePath.Clear();
 				ChangeNumberToLayoutInfo.Clear();
@@ -1088,10 +1133,10 @@ namespace UnrealGameSync
 				PromotedChangeNumbers = PerforceMonitor.GetPromotedChangeNumbers();
 
 				string[] ExcludeChanges = new string[0];
-				if(Workspace != null)
+				if (Workspace != null)
 				{
 					ConfigFile ProjectConfigFile = PerforceMonitor.LatestProjectConfigFile;
-					if(ProjectConfigFile != null)
+					if (ProjectConfigFile != null)
 					{
 						ExcludeChanges = ProjectConfigFile.GetValues("Options.ExcludeChanges", ExcludeChanges);
 					}
@@ -1103,16 +1148,16 @@ namespace UnrealGameSync
 				NumChanges = Changes.Count;
 				ListIndexToChangeIndex = new List<int>();
 				SortedChangeNumbers = new List<int>();
-				
+
 				List<PerforceChangeSummary> FilteredChanges = new List<PerforceChangeSummary>();
-				for(int ChangeIdx = 0; ChangeIdx < Changes.Count; ChangeIdx++)
+				for (int ChangeIdx = 0; ChangeIdx < Changes.Count; ChangeIdx++)
 				{
 					PerforceChangeSummary Change = Changes[ChangeIdx];
-					if(ShouldShowChange(Change, ExcludeChanges) || PromotedChangeNumbers.Contains(Change.Number))
+					if (ShouldShowChange(Change, ExcludeChanges) || PromotedChangeNumbers.Contains(Change.Number))
 					{
 						SortedChangeNumbers.Add(Change.Number);
 
-						if(!bHideUnreviewed || (!EventMonitor.IsUnderInvestigation(Change.Number) && (ShouldIncludeInReviewedList(Change.Number) || bFirstChange)))
+						if (!bHideUnreviewed || (!EventMonitor.IsUnderInvestigation(Change.Number) && (ShouldIncludeInReviewedList(Change.Number) || bFirstChange)))
 						{
 							bFirstChange = false;
 
@@ -1152,13 +1197,13 @@ namespace UnrealGameSync
 
 			// Remove any changes that no longer exist, and update the rest
 			Dictionary<int, PerforceChangeSummary> ChangeNumberToSummary = Changes.ToDictionary(x => x.Number, x => x);
-			for(int Idx = BuildList.Items.Count - 1; Idx >= 0; Idx--)
+			for (int Idx = BuildList.Items.Count - 1; Idx >= 0; Idx--)
 			{
 				PerforceChangeSummary Change = BuildList.Items[Idx].Tag as PerforceChangeSummary;
 				if (Change != null)
 				{
 					PerforceChangeSummary Summary;
-					if(ChangeNumberToSummary.TryGetValue(Change.Number, out Summary))
+					if (ChangeNumberToSummary.TryGetValue(Change.Number, out Summary))
 					{
 						// Update
 						BuildList.Items[Idx].SubItems[DescriptionColumn.Index].Text = Change.Description.Replace('\n', ' ');
@@ -1173,7 +1218,7 @@ namespace UnrealGameSync
 			}
 
 			// Add everything left over
-			foreach(PerforceChangeSummary Change in ChangeNumberToSummary.Values)
+			foreach (PerforceChangeSummary Change in ChangeNumberToSummary.Values)
 			{
 				BuildList_AddItem(Change);
 			}
@@ -1183,7 +1228,7 @@ namespace UnrealGameSync
 			for (int Idx = BuildList.Items.Count - 1; Idx >= 0; Idx--)
 			{
 				ListViewItem Item = BuildList.Items[Idx];
-				if(Item != ExpandItem)
+				if (Item != ExpandItem)
 				{
 					NewExpandItemGroup = Item.Group;
 					break;
@@ -1193,7 +1238,7 @@ namespace UnrealGameSync
 			// Remove the expand row if it's in the wrong place
 			if (ExpandItem != null)
 			{
-				if(!bShowExpandItem || ExpandItem.Group != NewExpandItemGroup)
+				if (!bShowExpandItem || ExpandItem.Group != NewExpandItemGroup)
 				{
 					BuildList.Items.Remove(ExpandItem);
 					ExpandItem = null;
@@ -1239,7 +1284,7 @@ namespace UnrealGameSync
 				}
 
 				PerforceChangeSummary LastChange = null;
-				for(int Idx = NextGroup.Items.Count - 1; Idx >= 0 && LastChange == null; Idx--)
+				for (int Idx = NextGroup.Items.Count - 1; Idx >= 0 && LastChange == null; Idx--)
 				{
 					LastChange = NextGroup.Items[Idx].Tag as PerforceChangeSummary;
 				}
@@ -1286,10 +1331,10 @@ namespace UnrealGameSync
 
 			// Insert it at the right position within the group
 			int GroupInsertIdx = 0;
-			for(; GroupInsertIdx < Group.Items.Count; GroupInsertIdx++)
+			for (; GroupInsertIdx < Group.Items.Count; GroupInsertIdx++)
 			{
 				PerforceChangeSummary OtherChange = Group.Items[GroupInsertIdx].Tag as PerforceChangeSummary;
-				if(OtherChange == null || Change.Number >= OtherChange.Number)
+				if (OtherChange == null || Change.Number >= OtherChange.Number)
 				{
 					break;
 				}
@@ -1302,46 +1347,46 @@ namespace UnrealGameSync
 
 		bool ShouldShowChange(PerforceChangeSummary Change, string[] ExcludeChanges)
 		{
-			if(ProjectSettings.FilterType != FilterType.None)
+			if (ProjectSettings.FilterType != FilterType.None)
 			{
 				PerforceChangeDetails Details;
-				if(!PerforceMonitor.TryGetChangeDetails(Change.Number, out Details))
+				if (!PerforceMonitor.TryGetChangeDetails(Change.Number, out Details))
 				{
 					return false;
 				}
-				if(ProjectSettings.FilterType == FilterType.Code && !Details.bContainsCode)
+				if (ProjectSettings.FilterType == FilterType.Code && !Details.bContainsCode)
 				{
 					return false;
 				}
-				if(ProjectSettings.FilterType == FilterType.Content && !Details.bContainsContent)
+				if (ProjectSettings.FilterType == FilterType.Content && !Details.bContainsContent)
 				{
 					return false;
 				}
 			}
-			if(ProjectSettings.FilterBadges.Count > 0)
+			if (ProjectSettings.FilterBadges.Count > 0)
 			{
 				EventSummary Summary = EventMonitor.GetSummaryForChange(Change.Number);
-				if(Summary == null || !Summary.Badges.Any(x => ProjectSettings.FilterBadges.Contains(x.BadgeName)))
+				if (Summary == null || !Summary.Badges.Any(x => ProjectSettings.FilterBadges.Contains(x.BadgeName)))
 				{
 					return false;
 				}
 			}
-			if(!Settings.bShowAutomatedChanges)
+			if (!Settings.bShowAutomatedChanges)
 			{
-				foreach(string ExcludeChange in ExcludeChanges)
+				foreach (string ExcludeChange in ExcludeChanges)
 				{
-					if(Regex.IsMatch(Change.Description, ExcludeChange, RegexOptions.IgnoreCase))
+					if (Regex.IsMatch(Change.Description, ExcludeChange, RegexOptions.IgnoreCase))
 					{
 						return false;
 					}
 				}
 
-				if(String.Compare(Change.User, "buildmachine", true) == 0 && Change.Description.IndexOf("lightmaps", StringComparison.InvariantCultureIgnoreCase) == -1)
+				if (String.Compare(Change.User, "buildmachine", true) == 0 && Change.Description.IndexOf("lightmaps", StringComparison.InvariantCultureIgnoreCase) == -1)
 				{
 					return false;
 				}
 			}
-			if(IsBisectModeEnabled() && !WorkspaceSettings.ChangeNumberToBisectState.ContainsKey(Change.Number))
+			if (IsBisectModeEnabled() && !WorkspaceSettings.ChangeNumberToBisectState.ContainsKey(Change.Number))
 			{
 				return false;
 			}
@@ -1350,10 +1395,10 @@ namespace UnrealGameSync
 
 		void UpdateBuildMetadataCallback()
 		{
-			if(!bUpdateBuildMetadataPosted)
+			if (!bUpdateBuildMetadataPosted)
 			{
 				bUpdateBuildMetadataPosted = true;
-				MainThreadSynchronizationContext.Post((o) => { bUpdateBuildMetadataPosted = false; if(!bIsDisposing) { UpdateBuildMetadata(); } }, null);
+				MainThreadSynchronizationContext.Post((o) => { bUpdateBuildMetadataPosted = false; if (!bIsDisposing) { UpdateBuildMetadata(); } }, null);
 			}
 		}
 
@@ -1370,13 +1415,13 @@ namespace UnrealGameSync
 
 			// Read the group mappings from the config file
 			string GroupDefinitions;
-			if(TryGetProjectSetting(PerforceMonitor.LatestProjectConfigFile, "BadgeGroups", out GroupDefinitions))
+			if (TryGetProjectSetting(PerforceMonitor.LatestProjectConfigFile, "BadgeGroups", out GroupDefinitions))
 			{
 				string[] GroupDefinitionsArray = GroupDefinitions.Split('\n');
-				for(int Idx = 0; Idx < GroupDefinitionsArray.Length; Idx++)
+				for (int Idx = 0; Idx < GroupDefinitionsArray.Length; Idx++)
 				{
 					string GroupName = String.Format("{0:0000}", Idx);
-					foreach(string BadgeName in GroupDefinitionsArray[Idx].Split(',').Select(x => x.Trim()))
+					foreach (string BadgeName in GroupDefinitionsArray[Idx].Split(',').Select(x => x.Trim()))
 					{
 						BadgeNameToGroup[BadgeName] = GroupName;
 					}
@@ -1384,17 +1429,17 @@ namespace UnrealGameSync
 			}
 
 			// Add a dummy group for any other badges we have
-			foreach(ListViewItem Item in BuildList.Items)
+			foreach (ListViewItem Item in BuildList.Items)
 			{
-				if(Item.Tag != null)
+				if (Item.Tag != null)
 				{
 					EventSummary Summary = EventMonitor.GetSummaryForChange(((PerforceChangeSummary)Item.Tag).Number);
-					if(Summary != null)
+					if (Summary != null)
 					{
-						foreach(BadgeData Badge in Summary.Badges)
+						foreach (BadgeData Badge in Summary.Badges)
 						{
 							string BadgeSlot = Badge.BadgeName;
-							if(!BadgeNameToGroup.ContainsKey(BadgeSlot))
+							if (!BadgeNameToGroup.ContainsKey(BadgeSlot))
 							{
 								BadgeNameToGroup.Add(BadgeSlot, "XXXX");
 							}
@@ -1404,16 +1449,16 @@ namespace UnrealGameSync
 			}
 
 			// Remove anything that's a service badge
-			foreach(KeyValuePair<string, BadgeData> ServiceBadge in ServiceBadges)
+			foreach (ServiceBadgeInfo ServiceBadge in ServiceBadges)
 			{
-				BadgeNameToGroup.Remove(ServiceBadge.Key);
+				BadgeNameToGroup.Remove(ServiceBadge.Name);
 			}
 
 			// Remove everything that's in a custom column
-			foreach(ColumnHeader CustomColumn in CustomColumns)
+			foreach (ColumnHeader CustomColumn in CustomColumns)
 			{
 				ConfigObject ColumnConfig = (ConfigObject)CustomColumn.Tag;
-				foreach(string BadgeName in ColumnConfig.GetValue("Badges", "").Split(','))
+				foreach (string BadgeName in ColumnConfig.GetValue("Badges", "").Split(','))
 				{
 					BadgeNameToGroup.Remove(BadgeName);
 				}
@@ -1435,7 +1480,7 @@ namespace UnrealGameSync
 			CheckForStartupComplete();
 
 			// If we are filtering by badges, we may also need to update the build list
-			if(ProjectSettings.FilterType != FilterType.None || ProjectSettings.FilterBadges.Count > 0)
+			if (ProjectSettings.FilterType != FilterType.None || ProjectSettings.FilterBadges.Count > 0)
 			{
 				UpdateBuildList();
 			}
@@ -1444,15 +1489,15 @@ namespace UnrealGameSync
 		void UpdateMaxBuildBadgeChars()
 		{
 			int NewMaxBuildBadgeChars;
-			if(GetBuildBadgeStripWidth(-1) < CISColumn.Width)
+			if (GetBuildBadgeStripWidth(-1) < CISColumn.Width)
 			{
 				NewMaxBuildBadgeChars = -1;
 			}
-			else if(GetBuildBadgeStripWidth(3) < CISColumn.Width)
+			else if (GetBuildBadgeStripWidth(3) < CISColumn.Width)
 			{
 				NewMaxBuildBadgeChars = 3;
 			}
-			else if(GetBuildBadgeStripWidth(2) < CISColumn.Width)
+			else if (GetBuildBadgeStripWidth(2) < CISColumn.Width)
 			{
 				NewMaxBuildBadgeChars = 2;
 			}
@@ -1461,7 +1506,7 @@ namespace UnrealGameSync
 				NewMaxBuildBadgeChars = 1;
 			}
 
-			if(NewMaxBuildBadgeChars != MaxBuildBadgeChars)
+			if (NewMaxBuildBadgeChars != MaxBuildBadgeChars)
 			{
 				ChangeNumberToLayoutInfo.Clear();
 				BuildList.Invalidate();
@@ -1473,10 +1518,10 @@ namespace UnrealGameSync
 		{
 			// Create dummy badges for each badge name
 			List<BadgeInfo> DummyBadges = new List<BadgeInfo>();
-			foreach(KeyValuePair<string, string> BadgeNameAndGroupPair in BadgeNameAndGroupPairs)
+			foreach (KeyValuePair<string, string> BadgeNameAndGroupPair in BadgeNameAndGroupPairs)
 			{
 				string BadgeName = BadgeNameAndGroupPair.Key;
-				if(MaxNumChars != -1 && BadgeName.Length > MaxNumChars)
+				if (MaxNumChars != -1 && BadgeName.Length > MaxNumChars)
 				{
 					BadgeName = BadgeName.Substring(0, MaxNumChars);
 				}
@@ -1487,7 +1532,7 @@ namespace UnrealGameSync
 
 			// Check if they fit within the column
 			int Width = 0;
-			if(DummyBadges.Count > 0)
+			if (DummyBadges.Count > 0)
 			{
 				LayoutBadges(DummyBadges);
 				Width = DummyBadges[DummyBadges.Count - 1].Offset + DummyBadges[DummyBadges.Count - 1].Width;
@@ -1497,19 +1542,19 @@ namespace UnrealGameSync
 
 		bool ShouldIncludeInReviewedList(int ChangeNumber)
 		{
-			if(PromotedChangeNumbers.Contains(ChangeNumber))
+			if (PromotedChangeNumbers.Contains(ChangeNumber))
 			{
 				return true;
 			}
 
 			EventSummary Review = EventMonitor.GetSummaryForChange(ChangeNumber);
-			if(Review != null)
+			if (Review != null)
 			{
-				if(Review.LastStarReview != null && Review.LastStarReview.Type == EventType.Starred)
+				if (Review.LastStarReview != null && Review.LastStarReview.Type == EventType.Starred)
 				{
 					return true;
 				}
-				if(Review.Verdict == ReviewVerdict.Good || Review.Verdict == ReviewVerdict.Mixed)
+				if (Review.Verdict == ReviewVerdict.Good || Review.Verdict == ReviewVerdict.Mixed)
 				{
 					return true;
 				}
@@ -1519,10 +1564,10 @@ namespace UnrealGameSync
 
 		void UpdateReviewsCallback()
 		{
-			if(!bUpdateReviewsPosted)
+			if (!bUpdateReviewsPosted)
 			{
 				bUpdateReviewsPosted = true;
-				MainThreadSynchronizationContext.Post((o) => { bUpdateReviewsPosted = false; if(!bIsDisposing) { UpdateReviews(); } }, null);
+				MainThreadSynchronizationContext.Post((o) => { bUpdateReviewsPosted = false; if (!bIsDisposing) { UpdateReviews(); } }, null);
 			}
 		}
 
@@ -1531,6 +1576,12 @@ namespace UnrealGameSync
 			ChangeNumberToArchivePath.Clear();
 			ChangeNumberToLayoutInfo.Clear();
 			EventMonitor.ApplyUpdates();
+
+			if (UpdateServiceBadges())
+			{
+				UpdateStatusPanel();
+			}
+
 			Refresh();
 			UpdateBuildFailureNotification();
 			CheckForStartupComplete();
@@ -3019,16 +3070,16 @@ namespace UnrealGameSync
 					}
 				}
 
-				foreach(KeyValuePair<string, BadgeData> ServiceBadge in ServiceBadges)
+				foreach(ServiceBadgeInfo ServiceBadge in ServiceBadges)
 				{
 					ProgramsLine.AddText("  |  ");
-					if(ServiceBadge.Value == null)
+					if(ServiceBadge.Url == null)
 					{
-						ProgramsLine.AddBadge(ServiceBadge.Key, GetBuildBadgeColor(BadgeResult.Skipped), null);
+						ProgramsLine.AddBadge(ServiceBadge.Name, GetBuildBadgeColor(ServiceBadge.Result), null);
 					}
 					else
 					{
-						ProgramsLine.AddBadge(ServiceBadge.Key, GetBuildBadgeColor(ServiceBadge.Value.Result), (P, R) => { SafeProcessStart(ServiceBadge.Value.Url); });
+						ProgramsLine.AddBadge(ServiceBadge.Name, GetBuildBadgeColor(ServiceBadge.Result), (P, R) => { SafeProcessStart(ServiceBadge.Url); });
 					}
 				}
 				ProgramsLine.AddText("  |  ");
diff --git a/Engine/Source/Programs/UnrealGameSync/UnrealGameSync/Forms/MainWindow.cs b/Engine/Source/Programs/UnrealGameSync/UnrealGameSync/Forms/MainWindow.cs
--- a/Engine/Source/Programs/UnrealGameSync/UnrealGameSync/Forms/MainWindow.cs
+++ b/Engine/Source/Programs/UnrealGameSync/UnrealGameSync/Forms/MainWindow.cs
@@ -1127,14 +1127,7 @@ namespace UnrealGameSync
 
 		private void TabMenu_OpenProject_Click(object sender, EventArgs e)
 		{
-			if (TabMenu_TabIdx == -1)
-			{
-				OpenNewProject();
-			}
-			else
-			{
-				EditSelectedProject(TabMenu_TabIdx);
-			}
+			EditSelectedProject(TabMenu_TabIdx);
 		}
 
 		private void TabMenu_TabNames_Stream_Click(object sender, EventArgs e)
diff --git a/Engine/Source/Programs/UnrealGameSync/UnrealGameSync/Properties/AssemblyInfo.cs b/Engine/Source/Programs/UnrealGameSync/UnrealGameSync/Properties/AssemblyInfo.cs
--- a/Engine/Source/Programs/UnrealGameSync/UnrealGameSync/Properties/AssemblyInfo.cs
+++ b/Engine/Source/Programs/UnrealGameSync/UnrealGameSync/Properties/AssemblyInfo.cs
@@ -34,5 +34,5 @@ using System.Runtime.InteropServices;
 // You can specify all the values or you can default the Build and Revision Numbers 
 // by using the '*' as shown below:
 // [assembly: AssemblyVersion("1.0.*")]
-[assembly: AssemblyVersion("1.181.*")]
-[assembly: AssemblyFileVersion("1.181.0.0")]
+[assembly: AssemblyVersion("1.182.*")]
+[assembly: AssemblyFileVersion("1.182.0.0")]
