Merging latest UGS fixes (incorrect compatible changelist in hotfix releases, option in issue settings window not being enabled)

#rnx

diff --git a/Engine/Source/Programs/UnrealGameSync/UnrealGameSync/FastJSON/Formatter.cs b/Engine/Source/Programs/UnrealGameSync/UnrealGameSync/FastJSON/Formatter.cs
new file mode 100644
--- /dev/null
+++ b/Engine/Source/Programs/UnrealGameSync/UnrealGameSync/FastJSON/Formatter.cs
@@ -0,0 +1,97 @@
+﻿using System.Collections.Generic;
+using System.Text;
+
+namespace fastJSON
+{
+    internal static class Formatter
+    {
+        public static string Indent = "   ";
+
+        public static void AppendIndent(StringBuilder sb, int count)
+        {
+            for (; count > 0; --count) sb.Append(Indent);
+        }
+
+        public static bool IsEscaped(StringBuilder sb, int index)
+        {
+            bool escaped = false;
+            while (index > 0 && sb[--index] == '\\') escaped = !escaped;
+            return escaped;
+        }
+
+        public static string PrettyPrint(string input)
+        {
+            var output = new StringBuilder(input.Length * 2);
+            char? quote = null;
+            int depth = 0;
+
+            for (int i = 0; i < input.Length; ++i)
+            {
+                char ch = input[i];
+
+                if (ch == '\"') // found string span
+                {
+                    bool str = true;
+                    while (str)
+                    {
+                        output.Append(ch);
+                        ch = input[++i];
+                        if (ch == '\\')
+                        {
+                            if (input[i + 1] == '\"')
+                            {
+                                output.Append(ch);
+                                ch = input[++i];
+                            }
+                        }
+                        else if (ch == '\"')
+                            str = false;
+                    }
+                }
+
+                switch (ch)
+                {
+                    case '{':
+                    case '[':
+                        output.Append(ch);
+                        if (!quote.HasValue)
+                        {
+                            output.AppendLine();
+                            AppendIndent(output, ++depth);
+                        }
+                        break;
+                    case '}':
+                    case ']':
+                        if (quote.HasValue)
+                            output.Append(ch);
+                        else
+                        {
+                            output.AppendLine();
+                            AppendIndent(output, --depth);
+                            output.Append(ch);
+                        }
+                        break;
+                    case ',':
+                        output.Append(ch);
+                        if (!quote.HasValue)
+                        {
+                            output.AppendLine();
+                            AppendIndent(output, depth);
+                        }
+                        break;
+                    case ':':
+                        if (quote.HasValue) output.Append(ch);
+                        else output.Append(" : ");
+                        break;
+                    default:
+                        if (quote.HasValue || !char.IsWhiteSpace(ch))
+                            output.Append(ch);
+                        break;
+                }
+
+            }
+
+            return output.ToString();
+        }
+    }
+}
\ No newline at end of file
diff --git a/Engine/Source/Programs/UnrealGameSync/UnrealGameSync/FastJSON/Getters.cs b/Engine/Source/Programs/UnrealGameSync/UnrealGameSync/FastJSON/Getters.cs
new file mode 100644
--- /dev/null
+++ b/Engine/Source/Programs/UnrealGameSync/UnrealGameSync/FastJSON/Getters.cs
@@ -0,0 +1,17 @@
+﻿using System;
+using System.Collections.Generic;
+
+#if !__MonoCS__
+#pragma warning disable CS1591
+#endif
+namespace fastJSON
+{
+    public sealed class DatasetSchema
+    {
+        public List<string> Info { get; set; }
+        public string Name { get; set; }
+    }
+}
+#if !__MonoCS__
+#pragma warning restore CS1591
+#endif
diff --git a/Engine/Source/Programs/UnrealGameSync/UnrealGameSync/FastJSON/JSON.cs b/Engine/Source/Programs/UnrealGameSync/UnrealGameSync/FastJSON/JSON.cs
new file mode 100644
--- /dev/null
+++ b/Engine/Source/Programs/UnrealGameSync/UnrealGameSync/FastJSON/JSON.cs
@@ -0,0 +1,904 @@
+﻿using System;
+using System.Collections;
+using System.Collections.Generic;
+#if !SILVERLIGHT
+using System.Data;
+#endif
+using System.Globalization;
+using System.IO;
+using System.Reflection;
+using System.Reflection.Emit;
+
+#if !__MonoCS__
+#pragma warning disable CS1591
+#endif
+namespace fastJSON
+{
+    public delegate string Serialize(object data);
+    public delegate object Deserialize(string data);
+
+    public sealed class JSONParameters
+    {
+        /// <summary>
+        /// Use the optimized fast Dataset Schema format (default = True)
+        /// </summary>
+        public bool UseOptimizedDatasetSchema = true;
+        /// <summary>
+        /// Use the fast GUID format (default = True)
+        /// </summary>
+        public bool UseFastGuid = true;
+        /// <summary>
+        /// Serialize null values to the output (default = True)
+        /// </summary>
+        public bool SerializeNullValues = true;
+        /// <summary>
+        /// Use the UTC date format (default = True)
+        /// </summary>
+        public bool UseUTCDateTime = true;
+        /// <summary>
+        /// Show the readonly properties of types in the output (default = False)
+        /// </summary>
+        public bool ShowReadOnlyProperties = false;
+        /// <summary>
+        /// Use the $types extension to optimise the output json (default = True)
+        /// </summary>
+        public bool UsingGlobalTypes = true;
+        /// <summary>
+        /// ** work in progress
+        /// </summary>
+        public bool IgnoreCaseOnDeserialize = false;
+        /// <summary>
+        /// Anonymous types have read only properties 
+        /// </summary>
+        public bool EnableAnonymousTypes = false;
+        /// <summary>
+        /// Enable fastJSON extensions $types, $type, $map (default = True)
+        /// </summary>
+        public bool UseExtensions = true;
+        /// <summary>
+        /// Use escaped unicode i.e. \uXXXX format for non ASCII characters (default = True)
+        /// </summary>
+        public bool UseEscapedUnicode = true;
+
+        public void FixValues()
+        {
+            if (UseExtensions == false) // disable conflicting params
+            {
+                UsingGlobalTypes = false;
+            }
+        }
+    }
+
+    public sealed class JSON
+    {
+        //public readonly static JSON Instance = new JSON();
+        [ThreadStatic]
+        private static JSON _instance;
+
+        public static JSON Instance
+        {
+            get { return _instance ?? (_instance = new JSON()); }
+        }
+
+        private JSON()
+        {
+        }
+        /// <summary>
+        /// You can set these paramters globally for all calls
+        /// </summary>
+        public JSONParameters Parameters = new JSONParameters();
+        private JSONParameters _params;
+
+        public string ToJSON(object obj)
+        {
+            _params = Parameters;
+            _params.FixValues();
+            Reflection.Instance.ShowReadOnlyProperties = _params.ShowReadOnlyProperties;
+            return ToJSON(obj, Parameters);
+        }
+
+        public string ToJSON(object obj, JSONParameters param)
+        {
+            _params = param;
+            _params.FixValues();
+            Reflection.Instance.ShowReadOnlyProperties = _params.ShowReadOnlyProperties;
+            Type t = null;
+
+            if (obj == null)
+                return "null";
+
+            if (obj.GetType().IsGenericType)
+                t = obj.GetType().GetGenericTypeDefinition();
+            if (t == typeof(Dictionary<,>) || t == typeof(List<>))
+                _params.UsingGlobalTypes = false;
+
+            // FEATURE : enable extensions when you can deserialize anon types
+            if (_params.EnableAnonymousTypes) { _params.UseExtensions = false; _params.UsingGlobalTypes = false; Reflection.Instance.ShowReadOnlyProperties = true; }
+            _usingglobals = _params.UsingGlobalTypes;
+            return new JSONSerializer(_params).ConvertToJSON(obj);
+        }
+
+        public object Parse(string json)
+        {
+            _params = Parameters;
+            Reflection.Instance.ShowReadOnlyProperties = _params.ShowReadOnlyProperties;
+            return new JsonParser(json, _params.IgnoreCaseOnDeserialize).Decode();
+        }
+
+        public T ToObject<T>(string json)
+        {
+            return (T)ToObject(json, typeof(T));
+        }
+
+        public object ToObject(string json)
+        {
+            return ToObject(json, null);
+        }
+
+        public object ToObject(string json, Type type)
+        {
+            _params = Parameters;
+            _params.FixValues();
+            Reflection.Instance.ShowReadOnlyProperties = _params.ShowReadOnlyProperties;
+            Type t = null;
+            if (type != null && type.IsGenericType)
+                t = type.GetGenericTypeDefinition();
+            if (t == typeof(Dictionary<,>) || t == typeof(List<>))
+                _params.UsingGlobalTypes = false;
+            _usingglobals = _params.UsingGlobalTypes;
+
+            object o = new JsonParser(json, Parameters.IgnoreCaseOnDeserialize).Decode();
+            if (o == null)
+                return null;
+
+#if !SILVERLIGHT
+            if (type != null && type == typeof(DataSet))
+                return CreateDataset(o as Dictionary<string, object>, null);
+
+            if (type != null && type == typeof(DataTable))
+                return CreateDataTable(o as Dictionary<string, object>, null);
+#endif
+            if (o is IDictionary)
+            {
+                if (type != null && type.IsGenericType && type.GetGenericTypeDefinition() == typeof(Dictionary<,>)) // deserialize a dictionary
+                    return RootDictionary(o, type);
+                else // deserialize an object
+                    return ParseDictionary(o as Dictionary<string, object>, null, type, null);
+            }
+
+            if (o is List<object>)
+            {
+                if (type != null && type.IsGenericType && type.GetGenericTypeDefinition() == typeof(Dictionary<,>)) // kv format
+                    return RootDictionary(o, type);
+
+                if (type != null && type.IsGenericType && type.GetGenericTypeDefinition() == typeof(List<>)) // deserialize to generic list
+                    return RootList(o, type);
+                else
+                    return (o as List<object>).ToArray();
+            }
+
+            if (type != null && o.GetType() != type)
+                return ChangeType(o, type);
+
+            return o;
+        }
+
+        public string Beautify(string input)
+        {
+            return Formatter.PrettyPrint(input);
+        }
+
+        public object FillObject(object input, string json)
+        {
+            _params = Parameters;
+            _params.FixValues();
+            Reflection.Instance.ShowReadOnlyProperties = _params.ShowReadOnlyProperties;
+            Dictionary<string, object> ht = new JsonParser(json, Parameters.IgnoreCaseOnDeserialize).Decode() as Dictionary<string, object>;
+            if (ht == null) return null;
+            return ParseDictionary(ht, null, input.GetType(), input);
+        }
+
+        public object DeepCopy(object obj)
+        {
+            return ToObject(ToJSON(obj));
+        }
+
+        public T DeepCopy<T>(T obj)
+        {
+            return ToObject<T>(ToJSON(obj));
+        }
+
+        internal SafeDictionary<Type, Serialize> _customSerializer = new SafeDictionary<Type, Serialize>();
+        internal SafeDictionary<Type, Deserialize> _customDeserializer = new SafeDictionary<Type, Deserialize>();
+
+        public void RegisterCustomType(Type type, Serialize serializer, Deserialize deserializer)
+        {
+            if (type != null && serializer != null && deserializer != null)
+            {
+                _customSerializer.Add(type, serializer);
+                _customDeserializer.Add(type, deserializer);
+                // reset property cache
+                _propertycache = new SafeDictionary<string, SafeDictionary<string, myPropInfo>>();
+            }
+        }
+
+        internal bool IsTypeRegistered(Type t)
+        {
+            if (_customSerializer.Count == 0)
+                return false;
+            Serialize s;
+            return _customSerializer.TryGetValue(t, out s);
+        }
+
+        #region [   JSON specific reflection   ]
+
+        private enum myPropInfoType
+        {
+            Int,
+            Long,
+            String,
+            Bool,
+            DateTime,
+            Enum,
+            Guid,
+
+            Array,
+            ByteArray,
+            Dictionary,
+            StringDictionary,
+#if !SILVERLIGHT
+            Hashtable,
+            DataSet,
+            DataTable,
+#endif
+			Custom,
+            Unknown,
+        }
+
+        [Flags]
+        private enum myPropInfoFlags
+        {
+            Filled = 1 << 0,
+            CanWrite = 1 << 1
+        }
+
+        private struct myPropInfo
+        {
+            public Type pt;
+            public Type bt;
+            public Type changeType;
+            public Reflection.GenericSetter setter;
+            public Reflection.GenericGetter getter;
+            public Type[] GenericTypes;
+            public string Name;
+            public myPropInfoType Type;
+            public myPropInfoFlags Flags;
+
+            public bool IsClass;
+            public bool IsValueType;
+            public bool IsGenericType;
+        }
+
+        SafeDictionary<string, SafeDictionary<string, myPropInfo>> _propertycache = new SafeDictionary<string, SafeDictionary<string, myPropInfo>>();
+        private SafeDictionary<string, myPropInfo> Getproperties(Type type, string typename)
+        {
+            SafeDictionary<string, myPropInfo> sd = null;
+            if (_propertycache.TryGetValue(typename, out sd))
+            {
+                return sd;
+            }
+            else
+            {
+                sd = new SafeDictionary<string, myPropInfo>();
+                PropertyInfo[] pr = type.GetProperties(BindingFlags.Public | BindingFlags.Instance);
+                foreach (PropertyInfo p in pr)
+                {
+                    myPropInfo d = CreateMyProp(p.PropertyType, p.Name);
+                    d.Flags |= myPropInfoFlags.CanWrite;
+                    d.setter = Reflection.CreateSetMethod(type, p);
+                    d.getter = Reflection.CreateGetMethod(type, p);
+                    sd.Add(p.Name, d);
+                }
+                FieldInfo[] fi = type.GetFields(BindingFlags.Public | BindingFlags.Instance);
+                foreach (FieldInfo f in fi)
+                {
+                    myPropInfo d = CreateMyProp(f.FieldType, f.Name);
+                    d.setter = Reflection.CreateSetField(type, f);
+                    d.getter = Reflection.CreateGetField(type, f);
+                    sd.Add(f.Name, d);
+                }
+
+                _propertycache.Add(typename, sd);
+                return sd;
+            }
+        }
+
+        private myPropInfo CreateMyProp(Type t, string name)
+        {
+            myPropInfo d = new myPropInfo();
+            myPropInfoType d_type = myPropInfoType.Unknown;
+            myPropInfoFlags d_flags = myPropInfoFlags.Filled | myPropInfoFlags.CanWrite;
+
+            if (t == typeof(int) || t == typeof(int?)) d_type = myPropInfoType.Int;
+            else if (t == typeof(long) || t == typeof(long?)) d_type = myPropInfoType.Long;
+            else if (t == typeof(string)) d_type = myPropInfoType.String;
+            else if (t == typeof(bool) || t == typeof(bool?)) d_type = myPropInfoType.Bool;
+            else if (t == typeof(DateTime) || t == typeof(DateTime?)) d_type = myPropInfoType.DateTime;
+            else if (t.IsEnum) d_type = myPropInfoType.Enum;
+            else if (t == typeof(Guid) || t == typeof(Guid?)) d_type = myPropInfoType.Guid;
+            else if (t.IsArray)
+            {
+                d.bt = t.GetElementType();
+                if (t == typeof(byte[]))
+                    d_type = myPropInfoType.ByteArray;
+                else
+                    d_type = myPropInfoType.Array;
+            }
+            else if (t.Name.Contains("Dictionary"))
+            {
+                d.GenericTypes = t.GetGenericArguments();
+                if (d.GenericTypes.Length > 0 && d.GenericTypes[0] == typeof(string))
+                    d_type = myPropInfoType.StringDictionary;
+                else
+                    d_type = myPropInfoType.Dictionary;
+            }
+#if !SILVERLIGHT
+            else if (t == typeof(Hashtable)) d_type = myPropInfoType.Hashtable;
+            else if (t == typeof(DataSet)) d_type = myPropInfoType.DataSet;
+            else if (t == typeof(DataTable)) d_type = myPropInfoType.DataTable;
+#endif
+
+            else if (IsTypeRegistered(t))								
+                d_type = myPropInfoType.Custom;
+
+            d.IsClass = t.IsClass;
+            d.IsValueType = t.IsValueType;
+            if (t.IsGenericType)
+            {
+                d.IsGenericType = true;
+                d.bt = t.GetGenericArguments()[0];
+            }
+
+            d.pt = t;
+            d.Name = name;
+            d.changeType = GetChangeType(t);
+            d.Type = d_type;
+            d.Flags = d_flags;
+
+            return d;
+        }
+
+        private object ChangeType(object value, Type conversionType)
+        {
+            if (conversionType == typeof(int))
+                return (int)((long)value);
+
+            else if (conversionType == typeof(long))
+                return (long)value;
+
+            else if (conversionType == typeof(string))
+                return (string)value;
+
+            else if (conversionType == typeof(Guid))
+                return CreateGuid((string)value);
+
+            else if (conversionType.IsEnum)
+                return CreateEnum(conversionType, (string)value);
+
+            else if (IsTypeRegistered(conversionType))
+                return CreateCustom((string)value, conversionType);
+
+            return Convert.ChangeType(value, conversionType, CultureInfo.InvariantCulture);
+        }
+        #endregion
+
+        #region [   p r i v a t e   m e t h o d s   ]
+
+        private object RootList(object parse, Type type)
+        {
+            Type[] gtypes = type.GetGenericArguments();
+            IList o = (IList)Reflection.Instance.FastCreateInstance(type);
+            foreach (var k in (IList)parse)
+            {
+                _usingglobals = false;
+                object v = k;
+                if (k is Dictionary<string, object>)
+                    v = ParseDictionary(k as Dictionary<string, object>, null, gtypes[0], null);
+                else
+                    v = ChangeType(k, gtypes[0]);
+
+                o.Add(v);
+            }
+            return o;
+        }
+
+        private object RootDictionary(object parse, Type type)
+        {
+            Type[] gtypes = type.GetGenericArguments();
+            if (parse is Dictionary<string, object>)
+            {
+                IDictionary o = (IDictionary)Reflection.Instance.FastCreateInstance(type);
+
+                foreach (var kv in (Dictionary<string, object>)parse)
+                {
+                    object v;
+                    object k = ChangeType(kv.Key, gtypes[0]);
+                    if (kv.Value is Dictionary<string, object>)
+                        v = ParseDictionary(kv.Value as Dictionary<string, object>, null, gtypes[1], null);
+                    else if (kv.Value is List<object>)
+                        v = CreateArray(kv.Value as List<object>, typeof(object), typeof(object), null);
+                    else
+                        v = ChangeType(kv.Value, gtypes[1]);
+                    o.Add(k, v);
+                }
+
+                return o;
+            }
+            if (parse is List<object>)
+                return CreateDictionary(parse as List<object>, type, gtypes, null);
+
+            return null;
+        }
+
+        bool _usingglobals = false;
+        private object ParseDictionary(Dictionary<string, object> d, Dictionary<string, object> globaltypes, Type type, object input)
+        {
+            object tn = "";
+
+            if (d.TryGetValue("$types", out tn))
+            {
+                _usingglobals = true;
+                globaltypes = new Dictionary<string, object>();
+                foreach (var kv in (Dictionary<string, object>)tn)
+                {
+                    globaltypes.Add((string)kv.Value, kv.Key);
+                }
+            }
+
+            bool found = d.TryGetValue("$type", out tn);
+#if !SILVERLIGHT
+            if (found == false && type == typeof(System.Object))
+            {
+                return CreateDataset(d, globaltypes);
+            }
+#endif
+            if (found)
+            {
+                if (_usingglobals)
+                {
+                    object tname = "";
+                    if (globaltypes.TryGetValue((string)tn, out tname))
+                        tn = tname;
+                }
+                type = Reflection.Instance.GetTypeFromCache((string)tn);
+            }
+
+            if (type == null)
+                throw new Exception("Cannot determine type");
+
+            string typename = type.FullName;
+            object o = input;
+            if (o == null)
+                o = Reflection.Instance.FastCreateInstance(type);
+
+            SafeDictionary<string, myPropInfo> props = Getproperties(type, typename);
+            foreach (string n in d.Keys)
+            {
+                string name = n;
+                if (_params.IgnoreCaseOnDeserialize) name = name.ToLower();
+                if (name == "$map")
+                {
+                    ProcessMap(o, props, (Dictionary<string, object>)d[name]);
+                    continue;
+                }
+                myPropInfo pi;
+                if (props.TryGetValue(name, out pi) == false)
+                    continue;
+                if ((pi.Flags & (myPropInfoFlags.Filled | myPropInfoFlags.CanWrite)) != 0)
+                {
+                    object v = d[name];
+
+                    if (v != null)
+                    {
+                        object oset = null;
+
+                        switch (pi.Type)
+                        {
+                            case myPropInfoType.Int: oset = (int)((long)v); break;
+                            case myPropInfoType.Long: oset = (long)v; break;
+                            case myPropInfoType.String: oset = (string)v; break;
+                            case myPropInfoType.Bool: oset = (bool)v; break;
+                            case myPropInfoType.DateTime: oset = CreateDateTime((string)v); break;
+                            case myPropInfoType.Enum: oset = CreateEnum(pi.pt, (string)v); break;
+                            case myPropInfoType.Guid: oset = CreateGuid((string)v); break;
+
+                            case myPropInfoType.Array:
+                                if (!pi.IsValueType)
+                                    oset = CreateArray((List<object>)v, pi.pt, pi.bt, globaltypes);
+                                // what about 'else'?
+                                break;
+                            case myPropInfoType.ByteArray: oset = Convert.FromBase64String((string)v); break;
+#if !SILVERLIGHT
+                            case myPropInfoType.DataSet: oset = CreateDataset((Dictionary<string, object>)v, globaltypes); break;
+                            case myPropInfoType.DataTable: oset = this.CreateDataTable((Dictionary<string, object>)v, globaltypes); break;
+                            case myPropInfoType.Hashtable: // same case as Dictionary
+#endif
+                            case myPropInfoType.Dictionary: oset = CreateDictionary((List<object>)v, pi.pt, pi.GenericTypes, globaltypes); break;
+                            case myPropInfoType.StringDictionary: oset = CreateStringKeyDictionary((Dictionary<string, object>)v, pi.pt, pi.GenericTypes, globaltypes); break;
+
+							case myPropInfoType.Custom: oset = CreateCustom((string)v, pi.pt); break;
+                            default:
+                                {
+                                    if (pi.IsGenericType && pi.IsValueType == false && v is List<object>)
+                                        oset = CreateGenericList((List<object>)v, pi.pt, pi.bt, globaltypes);
+
+                                    else if (pi.IsClass && v is Dictionary<string, object>)
+                                        oset = ParseDictionary((Dictionary<string, object>)v, globaltypes, pi.pt, pi.getter(o));
+
+                                    else if (v is List<object>)
+                                        oset = CreateArray((List<object>)v, pi.pt, typeof(object), globaltypes);
+
+                                    else if (pi.IsValueType)
+                                        oset = ChangeType(v, pi.changeType);
+
+                                    else
+                                        oset = v;
+                                }
+                                break;
+                        }
+
+                        o = pi.setter(o, oset);
+                    }
+                }
+            }
+            return o;
+        }
+
+        private object CreateCustom(string v, Type type)
+        {
+            Deserialize d;
+            _customDeserializer.TryGetValue(type, out d);
+            return d(v);
+        }
+
+        private void ProcessMap(object obj, SafeDictionary<string, JSON.myPropInfo> props, Dictionary<string, object> dic)
+        {
+            foreach (KeyValuePair<string, object> kv in dic)
+            {
+                myPropInfo p = props[kv.Key];
+                object o = p.getter(obj);
+                Type t = Reflection.Instance.GetTypeFromCache((string)kv.Value);
+                if (t == typeof(Guid))
+                    p.setter(obj, CreateGuid((string)o));
+            }
+        }
+
+        static int CreateInteger(out int num, string s, int index, int count)
+        {
+            num = 0;
+            bool neg = false;
+            for (int x = 0; x < count; x++, index++)
+            {
+                char cc = s[index];
+
+                if (cc == '-')
+                    neg = true;
+                else if (cc == '+')
+                    neg = false;
+                else
+                {
+                    num *= 10;
+                    num += (int)(cc - '0');
+                }
+            }
+            if (neg) num = -num;
+
+            return num;
+        }
+
+        internal static long CreateLong(out long num, char[] s, int index, int count)
+        {
+            num = 0;
+            bool neg = false;
+            for (int x = 0; x < count; x++, index++)
+            {
+                char cc = s[index];
+
+                if (cc == '-')
+                    neg = true;
+                else if (cc == '+')
+                    neg = false;
+                else
+                {
+                    num *= 10;
+                    num += (int)(cc - '0');
+                }
+            }
+            if (neg) num = -num;
+
+            return num;
+        }
+
+        private object CreateEnum(Type pt, string v)
+        {
+            // TODO : optimize create enum
+#if !SILVERLIGHT
+            return Enum.Parse(pt, v);
+#else
+            return Enum.Parse(pt, v, true);
+#endif
+        }
+
+        private Guid CreateGuid(string s)
+        {
+            if (s.Length > 30)
+                return new Guid(s);
+            else
+                return new Guid(Convert.FromBase64String(s));
+        }
+
+        private DateTime CreateDateTime(string value)
+        {
+            bool utc = false;
+            //                   0123456789012345678
+            // datetime format = yyyy-MM-dd HH:mm:ss
+            int year;// = (int)CreateLong(value.Substring(0, 4));
+            int month;// = (int)CreateLong(value.Substring(5, 2));
+            int day;// = (int)CreateLong(value.Substring(8, 2));
+            int hour;// = (int)CreateLong(value.Substring(11, 2));
+            int min;// = (int)CreateLong(value.Substring(14, 2));
+            int sec;// = (int)CreateLong(value.Substring(17, 2));
+            CreateInteger(out year, value, 0, 4);
+            CreateInteger(out month, value, 5, 2);
+            CreateInteger(out day, value, 8, 2);
+            CreateInteger(out hour, value, 11, 2);
+            CreateInteger(out min, value, 14, 2);
+            CreateInteger(out sec, value, 17, 2);
+
+            //if (value.EndsWith("Z"))
+            if (value[value.Length - 1] == 'Z')
+                utc = true;
+
+            if (_params.UseUTCDateTime == false && utc == false)
+                return new DateTime(year, month, day, hour, min, sec);
+            else
+                return new DateTime(year, month, day, hour, min, sec, DateTimeKind.Utc).ToLocalTime();
+        }
+
+        private object CreateArray(List<object> data, Type pt, Type bt, Dictionary<string, object> globalTypes)
+        {
+            Array col = Array.CreateInstance(bt, data.Count);
+            // create an array of objects
+            for (int i = 0; i < data.Count; i++)// each (object ob in data)
+            {
+                object ob = data[i];
+                if (ob is IDictionary)
+                    //col.SetValue(ParseDictionary((Dictionary<string, object>)ob, globalTypes, bt, null), i); // Original line
+                    col.SetValue(RootDictionary(ob, typeof(Dictionary<string,object>)), i);	// Change for UE4
+                else
+                    col.SetValue(ChangeType(ob, bt), i);
+            }
+
+            return col;
+        }
+
+
+        private object CreateGenericList(List<object> data, Type pt, Type bt, Dictionary<string, object> globalTypes)
+        {
+            IList col = (IList)Reflection.Instance.FastCreateInstance(pt);
+            // create an array of objects
+            foreach (object ob in data)
+            {
+                if (ob is IDictionary)
+                    col.Add(ParseDictionary((Dictionary<string, object>)ob, globalTypes, bt, null));
+
+                else if (ob is List<object>)
+                    col.Add(((List<object>)ob).ToArray());
+
+                else
+                    col.Add(ChangeType(ob, bt));
+            }
+            return col;
+        }
+
+        private object CreateStringKeyDictionary(Dictionary<string, object> reader, Type pt, Type[] types, Dictionary<string, object> globalTypes)
+        {
+            var col = (IDictionary)Reflection.Instance.FastCreateInstance(pt);
+            //Type t1 = null;
+            Type t2 = null;
+            if (types != null)
+            {
+                //t1 = types[0];
+                t2 = types[1];
+            }
+
+            foreach (KeyValuePair<string, object> values in reader)
+            {
+                var key = values.Key;//ChangeType(values.Key, t1);
+                object val = null;
+                if (values.Value is Dictionary<string, object>)
+                    val = ParseDictionary((Dictionary<string, object>)values.Value, globalTypes, t2, null);
+                else
+                    val = ChangeType(values.Value, t2);
+                col.Add(key, val);
+            }
+
+            return col;
+        }
+
+        private object CreateDictionary(List<object> reader, Type pt, Type[] types, Dictionary<string, object> globalTypes)
+        {
+            IDictionary col = (IDictionary)Reflection.Instance.FastCreateInstance(pt);
+            Type t1 = null;
+            Type t2 = null;
+            if (types != null)
+            {
+                t1 = types[0];
+                t2 = types[1];
+            }
+
+            foreach (Dictionary<string, object> values in reader)
+            {
+                object key = values["k"];
+                object val = values["v"];
+
+                if (key is Dictionary<string, object>)
+                    key = ParseDictionary((Dictionary<string, object>)key, globalTypes, t1, null);
+                else
+                    key = ChangeType(key, t1);
+
+                if (val is Dictionary<string, object>)
+                    val = ParseDictionary((Dictionary<string, object>)val, globalTypes, t2, null);
+                else
+                    val = ChangeType(val, t2);
+
+                col.Add(key, val);
+            }
+
+            return col;
+        }
+
+        private Type GetChangeType(Type conversionType)
+        {
+            if (conversionType.IsGenericType && conversionType.GetGenericTypeDefinition().Equals(typeof(Nullable<>)))
+                return conversionType.GetGenericArguments()[0];
+
+            return conversionType;
+        }
+
+#if !SILVERLIGHT
+        private DataSet CreateDataset(Dictionary<string, object> reader, Dictionary<string, object> globalTypes)
+        {
+            DataSet ds = new DataSet();
+            ds.EnforceConstraints = false;
+            ds.BeginInit();
+
+            // read dataset schema here
+            ReadSchema(reader, ds, globalTypes);
+
+            foreach (KeyValuePair<string, object> pair in reader)
+            {
+                if (pair.Key == "$type" || pair.Key == "$schema") continue;
+
+                List<object> rows = (List<object>)pair.Value;
+                if (rows == null) continue;
+
+                DataTable dt = ds.Tables[pair.Key];
+                ReadDataTable(rows, dt);
+            }
+
+            ds.EndInit();
+
+            return ds;
+        }
+
+        private void ReadSchema(Dictionary<string, object> reader, DataSet ds, Dictionary<string, object> globalTypes)
+        {
+            var schema = reader["$schema"];
+
+            if (schema is string)
+            {
+                TextReader tr = new StringReader((string)schema);
+                ds.ReadXmlSchema(tr);
+            }
+            else
+            {
+                DatasetSchema ms = (DatasetSchema)ParseDictionary((Dictionary<string, object>)schema, globalTypes, typeof(DatasetSchema), null);
+                ds.DataSetName = ms.Name;
+                for (int i = 0; i < ms.Info.Count; i += 3)
+                {
+                    if (ds.Tables.Contains(ms.Info[i]) == false)
+                        ds.Tables.Add(ms.Info[i]);
+                    ds.Tables[ms.Info[i]].Columns.Add(ms.Info[i + 1], Reflection.Instance.GetTypeFromCache(ms.Info[i + 2]));
+                }
+            }
+        }
+
+        private void ReadDataTable(List<object> rows, DataTable dt)
+        {
+            dt.BeginInit();
+            dt.BeginLoadData();
+            List<int> guidcols = new List<int>();
+            List<int> datecol = new List<int>();
+
+            foreach (DataColumn c in dt.Columns)
+            {
+                if (c.DataType == typeof(Guid) || c.DataType == typeof(Guid?))
+                    guidcols.Add(c.Ordinal);
+                if (_params.UseUTCDateTime && (c.DataType == typeof(DateTime) || c.DataType == typeof(DateTime?)))
+                    datecol.Add(c.Ordinal);
+            }
+
+            foreach (List<object> row in rows)
+            {
+                object[] v = new object[row.Count];
+                row.CopyTo(v, 0);
+                foreach (int i in guidcols)
+                {
+                    string s = (string)v[i];
+                    if (s != null && s.Length < 36)
+                        v[i] = new Guid(Convert.FromBase64String(s));
+                }
+                if (_params.UseUTCDateTime)
+                {
+                    foreach (int i in datecol)
+                    {
+                        string s = (string)v[i];
+                        if (s != null)
+                            v[i] = CreateDateTime(s);
+                    }
+                }
+                dt.Rows.Add(v);
+            }
+
+            dt.EndLoadData();
+            dt.EndInit();
+        }
+
+        DataTable CreateDataTable(Dictionary<string, object> reader, Dictionary<string, object> globalTypes)
+        {
+            var dt = new DataTable();
+
+            // read dataset schema here
+            var schema = reader["$schema"];
+
+            if (schema is string)
+            {
+                TextReader tr = new StringReader((string)schema);
+                dt.ReadXmlSchema(tr);
+            }
+            else
+            {
+                var ms = (DatasetSchema)this.ParseDictionary((Dictionary<string, object>)schema, globalTypes, typeof(DatasetSchema), null);
+                dt.TableName = ms.Info[0];
+                for (int i = 0; i < ms.Info.Count; i += 3)
+                {
+                    dt.Columns.Add(ms.Info[i + 1], Reflection.Instance.GetTypeFromCache(ms.Info[i + 2]));
+                }
+            }
+
+            foreach (var pair in reader)
+            {
+                if (pair.Key == "$type" || pair.Key == "$schema")
+                    continue;
+
+                var rows = (List<object>)pair.Value;
+                if (rows == null)
+                    continue;
+
+                if (!dt.TableName.Equals(pair.Key, StringComparison.InvariantCultureIgnoreCase))
+                    continue;
+
+                ReadDataTable(rows, dt);
+            }
+
+            return dt;
+        }
+#endif
+        #endregion
+    }
+
+}
+#if !__MonoCS__
+#pragma warning restore CS1591
+#endif
diff --git a/Engine/Source/Programs/UnrealGameSync/UnrealGameSync/FastJSON/JsonParser.cs b/Engine/Source/Programs/UnrealGameSync/UnrealGameSync/FastJSON/JsonParser.cs
new file mode 100644
--- /dev/null
+++ b/Engine/Source/Programs/UnrealGameSync/UnrealGameSync/FastJSON/JsonParser.cs
@@ -0,0 +1,441 @@
+using System;
+using System.Collections;
+using System.Collections.Generic;
+using System.Globalization;
+using System.Text;
+
+namespace fastJSON
+{
+    /// <summary>
+    /// This class encodes and decodes JSON strings.
+    /// Spec. details, see http://www.json.org/
+    /// </summary>
+    internal sealed class JsonParser
+    {
+        enum Token
+        {
+            None = -1,           // Used to denote no Lookahead available
+            Curly_Open,
+            Curly_Close,
+            Squared_Open,
+            Squared_Close,
+            Colon,
+            Comma,
+            String,
+            Number,
+            True,
+            False,
+            Null
+        }
+
+        readonly char[] json;
+        readonly StringBuilder s = new StringBuilder();
+        Token lookAheadToken = Token.None;
+        int index;
+        bool _ignorecase = false;
+
+
+        internal JsonParser(string json, bool ignorecase)
+        {
+            this.json = json.ToCharArray();
+            _ignorecase = ignorecase;
+        }
+
+        public object Decode()
+        {
+            return ParseValue();
+        }
+
+        private Dictionary<string, object> ParseObject()
+        {
+            Dictionary<string, object> table = new Dictionary<string, object>();
+
+            ConsumeToken(); // {
+
+            while (true)
+            {
+                switch (LookAhead())
+                {
+
+                    case Token.Comma:
+                        ConsumeToken();
+                        break;
+
+                    case Token.Curly_Close:
+                        ConsumeToken();
+                        return table;
+
+                    default:
+                        {
+
+                            // name
+                            string name = ParseString();
+                            if (_ignorecase)
+                                name = name.ToLower();
+
+                            // :
+                            if (NextToken() != Token.Colon)
+                            {
+                                throw new Exception("Expected colon at index " + index);
+                            }
+
+                            // value
+                            object value = ParseValue();
+
+                            table[name] = value;
+                        }
+                        break;
+                }
+            }
+        }
+
+        private List<object> ParseArray()
+        {
+            List<object> array = new List<object>();
+            ConsumeToken(); // [
+
+            while (true)
+            {
+                switch (LookAhead())
+                {
+
+                    case Token.Comma:
+                        ConsumeToken();
+                        break;
+
+                    case Token.Squared_Close:
+                        ConsumeToken();
+                        return array;
+
+                    default:
+                        array.Add(ParseValue());
+                        break;
+                }
+            }
+        }
+
+        private object ParseValue()
+        {
+            switch (LookAhead())
+            {
+                case Token.Number:
+                    return ParseNumber();
+
+                case Token.String:
+                    return ParseString();
+
+                case Token.Curly_Open:
+                    return ParseObject();
+
+                case Token.Squared_Open:
+                    return ParseArray();
+
+                case Token.True:
+                    ConsumeToken();
+                    return true;
+
+                case Token.False:
+                    ConsumeToken();
+                    return false;
+
+                case Token.Null:
+                    ConsumeToken();
+                    return null;
+            }
+
+            throw new Exception("Unrecognized token at index" + index);
+        }
+
+        private string ParseString()
+        {
+            ConsumeToken(); // "
+
+            s.Length = 0;
+
+            int runIndex = -1;
+
+            while (index < json.Length)
+            {
+                var c = json[index++];
+
+                if (c == '"')
+                {
+                    if (runIndex != -1)
+                    {
+                        if (s.Length == 0)
+                            return new string(json, runIndex, index - runIndex - 1);
+
+                        s.Append(json, runIndex, index - runIndex - 1);
+                    }
+                    return s.ToString();
+                }
+
+                if (c != '\\')
+                {
+                    if (runIndex == -1)
+                        runIndex = index - 1;
+
+                    continue;
+                }
+
+                if (index == json.Length) break;
+
+                if (runIndex != -1)
+                {
+                    s.Append(json, runIndex, index - runIndex - 1);
+                    runIndex = -1;
+                }
+
+                switch (json[index++])
+                {
+                    case '"':
+                        s.Append('"');
+                        break;
+
+                    case '\\':
+                        s.Append('\\');
+                        break;
+
+                    case '/':
+                        s.Append('/');
+                        break;
+
+                    case 'b':
+                        s.Append('\b');
+                        break;
+
+                    case 'f':
+                        s.Append('\f');
+                        break;
+
+                    case 'n':
+                        s.Append('\n');
+                        break;
+
+                    case 'r':
+                        s.Append('\r');
+                        break;
+
+                    case 't':
+                        s.Append('\t');
+                        break;
+
+                    case 'u':
+                        {
+                            int remainingLength = json.Length - index;
+                            if (remainingLength < 4) break;
+
+                            // parse the 32 bit hex into an integer codepoint
+                            uint codePoint = ParseUnicode(json[index], json[index + 1], json[index + 2], json[index + 3]);
+                            s.Append((char)codePoint);
+
+                            // skip 4 chars
+                            index += 4;
+                        }
+                        break;
+                }
+            }
+
+            throw new Exception("Unexpectedly reached end of string");
+        }
+
+        private uint ParseSingleChar(char c1, uint multipliyer)
+        {
+            uint p1 = 0;
+            if (c1 >= '0' && c1 <= '9')
+                p1 = (uint)(c1 - '0') * multipliyer;
+            else if (c1 >= 'A' && c1 <= 'F')
+                p1 = (uint)((c1 - 'A') + 10) * multipliyer;
+            else if (c1 >= 'a' && c1 <= 'f')
+                p1 = (uint)((c1 - 'a') + 10) * multipliyer;
+            return p1;
+        }
+
+        private uint ParseUnicode(char c1, char c2, char c3, char c4)
+        {
+            uint p1 = ParseSingleChar(c1, 0x1000);
+            uint p2 = ParseSingleChar(c2, 0x100);
+            uint p3 = ParseSingleChar(c3, 0x10);
+            uint p4 = ParseSingleChar(c4, 1);
+
+            return p1 + p2 + p3 + p4;
+        }
+
+        private long CreateLong(string s)
+        {
+            long num = 0;
+            bool neg = false;
+            foreach (char cc in s)
+            {
+                if (cc == '-')
+                    neg = true;
+                else if (cc == '+')
+                    neg = false;
+                else
+                {
+                    num *= 10;
+                    num += (int)(cc - '0');
+                }
+            }
+
+            return neg ? -num : num;
+        }
+
+        private object ParseNumber()
+        {
+            ConsumeToken();
+
+            // Need to start back one place because the first digit is also a token and would have been consumed
+            var startIndex = index - 1;
+            bool dec = false;
+            do
+            {
+                if (index == json.Length)
+                    break;
+                var c = json[index];
+
+                if ((c >= '0' && c <= '9') || c == '.' || c == '-' || c == '+' || c == 'e' || c == 'E')
+                {
+                    if (c == '.' || c == 'e' || c == 'E')
+                        dec = true;
+                    if (++index == json.Length)
+                        break;                        //throw new Exception("Unexpected end of string whilst parsing number");
+                    continue;
+                }
+                break;
+            } while (true);
+
+			if (dec)
+			{
+				string s = new string(json, startIndex, index - startIndex);
+				return double.Parse(s, NumberFormatInfo.InvariantInfo);
+			}
+			long num;
+			return JSON.CreateLong(out num, json, startIndex, index - startIndex);
+        }
+
+        private Token LookAhead()
+        {
+            if (lookAheadToken != Token.None) return lookAheadToken;
+
+            return lookAheadToken = NextTokenCore();
+        }
+
+        private void ConsumeToken()
+        {
+            lookAheadToken = Token.None;
+        }
+
+        private Token NextToken()
+        {
+            var result = lookAheadToken != Token.None ? lookAheadToken : NextTokenCore();
+
+            lookAheadToken = Token.None;
+
+            return result;
+        }
+
+        private Token NextTokenCore()
+        {
+            char c;
+
+            // Skip past whitespace
+            do
+            {
+                c = json[index];
+
+                if (c > ' ') break;
+                if (c != ' ' && c != '\t' && c != '\n' && c != '\r') break;
+
+            } while (++index < json.Length);
+
+            if (index == json.Length)
+            {
+                throw new Exception("Reached end of string unexpectedly");
+            }
+
+            c = json[index];
+
+            index++;
+
+            //if (c >= '0' && c <= '9')
+            //    return Token.Number;
+
+            switch (c)
+            {
+                case '{':
+                    return Token.Curly_Open;
+
+                case '}':
+                    return Token.Curly_Close;
+
+                case '[':
+                    return Token.Squared_Open;
+
+                case ']':
+                    return Token.Squared_Close;
+
+                case ',':
+                    return Token.Comma;
+
+                case '"':
+                    return Token.String;
+
+                case '0':
+                case '1':
+                case '2':
+                case '3':
+                case '4':
+                case '5':
+                case '6':
+                case '7':
+                case '8':
+                case '9':
+                case '-':
+                case '+':
+                case '.':
+                    return Token.Number;
+
+                case ':':
+                    return Token.Colon;
+
+                case 'f':
+                    if (json.Length - index >= 4 &&
+                        json[index + 0] == 'a' &&
+                        json[index + 1] == 'l' &&
+                        json[index + 2] == 's' &&
+                        json[index + 3] == 'e')
+                    {
+                        index += 4;
+                        return Token.False;
+                    }
+                    break;
+
+                case 't':
+                    if (json.Length - index >= 3 &&
+                        json[index + 0] == 'r' &&
+                        json[index + 1] == 'u' &&
+                        json[index + 2] == 'e')
+                    {
+                        index += 3;
+                        return Token.True;
+                    }
+                    break;
+
+                case 'n':
+                    if (json.Length - index >= 3 &&
+                        json[index + 0] == 'u' &&
+                        json[index + 1] == 'l' &&
+                        json[index + 2] == 'l')
+                    {
+                        index += 3;
+                        return Token.Null;
+                    }
+                    break;
+
+            }
+
+            throw new Exception("Could not find token at index " + --index);
+        }
+    }
+}
diff --git a/Engine/Source/Programs/UnrealGameSync/UnrealGameSync/FastJSON/JsonSerializer.cs b/Engine/Source/Programs/UnrealGameSync/UnrealGameSync/FastJSON/JsonSerializer.cs
new file mode 100644
--- /dev/null
+++ b/Engine/Source/Programs/UnrealGameSync/UnrealGameSync/FastJSON/JsonSerializer.cs
@@ -0,0 +1,505 @@
+﻿using System;
+using System.Collections;
+using System.Collections.Generic;
+#if !SILVERLIGHT
+using System.Data;
+#endif
+using System.Globalization;
+using System.IO;
+using System.Text;
+
+namespace fastJSON
+{
+    internal sealed class JSONSerializer
+    {
+        private StringBuilder _output = new StringBuilder();
+        private StringBuilder _before = new StringBuilder();
+        readonly int _MAX_DEPTH = 10;
+        int _current_depth = 0;
+        private Dictionary<string, int> _globalTypes = new Dictionary<string, int>();
+        private JSONParameters _params;
+        private bool _useEscapedUnicode = false;
+
+        internal JSONSerializer(JSONParameters param)
+        {
+            _params = param;
+            _useEscapedUnicode = _params.UseEscapedUnicode;
+        }
+
+        internal string ConvertToJSON(object obj)
+        {
+            WriteValue(obj);
+
+            string str = "";
+            if (_params.UsingGlobalTypes && _globalTypes != null && _globalTypes.Count > 0)
+            {
+                StringBuilder sb = _before;
+                sb.Append("\"$types\":{");
+                bool pendingSeparator = false;
+                foreach (var kv in _globalTypes)
+                {
+                    if (pendingSeparator) sb.Append(',');
+                    pendingSeparator = true;
+                    sb.Append('\"');
+                    sb.Append(kv.Key);
+                    sb.Append("\":\"");
+                    sb.Append(kv.Value);
+                    sb.Append('\"');
+                }
+                sb.Append("},");
+                sb.Append(_output.ToString());
+                str = sb.ToString();
+            }
+            else
+                str = _output.ToString();
+
+            return str;
+        }
+
+        private void WriteValue(object obj)
+        {
+            if (obj == null || obj is DBNull)
+                _output.Append("null");
+
+            else if (obj is string || obj is char)
+                WriteString(obj.ToString());
+
+            else if (obj is Guid)
+                WriteGuid((Guid)obj);
+
+            else if (obj is bool)
+                _output.Append(((bool)obj) ? "true" : "false"); // conform to standard
+
+            else if (
+                obj is int || obj is long || obj is double ||
+                obj is decimal || obj is float ||
+                obj is byte || obj is short ||
+                obj is sbyte || obj is ushort ||
+                obj is uint || obj is ulong
+            )
+                _output.Append(((IConvertible)obj).ToString(NumberFormatInfo.InvariantInfo));
+
+            else if (obj is DateTime)
+                WriteDateTime((DateTime)obj);
+
+            else if (obj is IDictionary && obj.GetType().IsGenericType && obj.GetType().GetGenericArguments()[0] == typeof(string))
+                WriteStringDictionary((IDictionary)obj);
+
+            else if (obj is IDictionary)
+                WriteDictionary((IDictionary)obj);
+#if !SILVERLIGHT
+            else if (obj is DataSet)
+                WriteDataset((DataSet)obj);
+
+            else if (obj is DataTable)
+                this.WriteDataTable((DataTable)obj);
+#endif
+            else if (obj is byte[])
+                WriteBytes((byte[])obj);
+
+            else if (obj is IEnumerable)// Array || obj is IList || obj is ICollection)
+                WriteArray((IEnumerable)obj);
+
+            else if (obj is Enum)
+                WriteEnum((Enum)obj);
+
+            else if (JSON.Instance.IsTypeRegistered(obj.GetType()))
+                WriteCustom(obj);
+
+            else
+                WriteObject(obj);
+        }
+
+        private void WriteCustom(object obj)
+        {
+            Serialize s;
+            JSON.Instance._customSerializer.TryGetValue(obj.GetType(), out s);
+            WriteStringFast(s(obj));
+        }
+
+        private void WriteEnum(Enum e)
+        {
+            // TODO : optimize enum write
+            WriteStringFast(e.ToString());
+        }
+
+        private void WriteGuid(Guid g)
+        {
+            if (_params.UseFastGuid == false)
+                WriteStringFast(g.ToString());
+            else
+                WriteBytes(g.ToByteArray());
+        }
+
+        private void WriteBytes(byte[] bytes)
+        {
+#if !SILVERLIGHT
+            WriteStringFast(Convert.ToBase64String(bytes, 0, bytes.Length, Base64FormattingOptions.None));
+#else
+            WriteStringFast(Convert.ToBase64String(bytes, 0, bytes.Length));
+#endif
+        }
+
+        private void WriteDateTime(DateTime dateTime)
+        {
+            // datetime format standard : yyyy-MM-dd HH:mm:ss
+            DateTime dt = dateTime;
+            if (_params.UseUTCDateTime)
+                dt = dateTime.ToUniversalTime();
+
+            _output.Append('\"');
+            _output.Append(dt.Year.ToString("0000", NumberFormatInfo.InvariantInfo));
+            _output.Append('-');
+            _output.Append(dt.Month.ToString("00", NumberFormatInfo.InvariantInfo));
+            _output.Append('-');
+            _output.Append(dt.Day.ToString("00", NumberFormatInfo.InvariantInfo));
+            _output.Append(' ');
+            _output.Append(dt.Hour.ToString("00", NumberFormatInfo.InvariantInfo));
+            _output.Append(':');
+            _output.Append(dt.Minute.ToString("00", NumberFormatInfo.InvariantInfo));
+            _output.Append(':');
+            _output.Append(dt.Second.ToString("00", NumberFormatInfo.InvariantInfo));
+
+            if (_params.UseUTCDateTime)
+                _output.Append('Z');
+
+            _output.Append('\"');
+        }
+
+#if !SILVERLIGHT
+        private DatasetSchema GetSchema(DataTable ds)
+        {
+            if (ds == null) return null;
+
+            DatasetSchema m = new DatasetSchema();
+            m.Info = new List<string>();
+            m.Name = ds.TableName;
+
+            foreach (DataColumn c in ds.Columns)
+            {
+                m.Info.Add(ds.TableName);
+                m.Info.Add(c.ColumnName);
+                m.Info.Add(c.DataType.ToString());
+            }
+            // FEATURE : serialize relations and constraints here
+
+            return m;
+        }
+
+        private DatasetSchema GetSchema(DataSet ds)
+        {
+            if (ds == null) return null;
+
+            DatasetSchema m = new DatasetSchema();
+            m.Info = new List<string>();
+            m.Name = ds.DataSetName;
+
+            foreach (DataTable t in ds.Tables)
+            {
+                foreach (DataColumn c in t.Columns)
+                {
+                    m.Info.Add(t.TableName);
+                    m.Info.Add(c.ColumnName);
+                    m.Info.Add(c.DataType.ToString());
+                }
+            }
+            // FEATURE : serialize relations and constraints here
+
+            return m;
+        }
+
+        private string GetXmlSchema(DataTable dt)
+        {
+            using (var writer = new StringWriter())
+            {
+                dt.WriteXmlSchema(writer);
+                return dt.ToString();
+            }
+        }
+
+        private void WriteDataset(DataSet ds)
+        {
+            _output.Append('{');
+            if (_params.UseExtensions)
+            {
+                WritePair("$schema", _params.UseOptimizedDatasetSchema ? (object)GetSchema(ds) : ds.GetXmlSchema());
+                _output.Append(',');
+            }
+            bool tablesep = false;
+            foreach (DataTable table in ds.Tables)
+            {
+                if (tablesep) _output.Append(',');
+                tablesep = true;
+                WriteDataTableData(table);
+            }
+            // end dataset
+            _output.Append('}');
+        }
+
+        private void WriteDataTableData(DataTable table)
+        {
+            _output.Append('\"');
+            _output.Append(table.TableName);
+            _output.Append("\":[");
+            DataColumnCollection cols = table.Columns;
+            bool rowseparator = false;
+            foreach (DataRow row in table.Rows)
+            {
+                if (rowseparator) _output.Append(',');
+                rowseparator = true;
+                _output.Append('[');
+
+                bool pendingSeperator = false;
+                foreach (DataColumn column in cols)
+                {
+                    if (pendingSeperator) _output.Append(',');
+                    WriteValue(row[column]);
+                    pendingSeperator = true;
+                }
+                _output.Append(']');
+            }
+
+            _output.Append(']');
+        }
+
+        void WriteDataTable(DataTable dt)
+        {
+            this._output.Append('{');
+            if (_params.UseExtensions)
+            {
+                this.WritePair("$schema", _params.UseOptimizedDatasetSchema ? (object)this.GetSchema(dt) : this.GetXmlSchema(dt));
+                this._output.Append(',');
+            }
+
+            WriteDataTableData(dt);
+
+            // end datatable
+            this._output.Append('}');
+        }
+#endif
+
+        bool _TypesWritten = false;
+        private void WriteObject(object obj)
+        {
+            if (_params.UsingGlobalTypes == false)
+                _output.Append('{');
+            else
+            {
+                if (_TypesWritten == false)
+                {
+                    _output.Append('{');
+                    _before = _output;
+                    _output = new StringBuilder();
+                }
+                else
+                    _output.Append('{');
+            }
+            _TypesWritten = true;
+            _current_depth++;
+            if (_current_depth > _MAX_DEPTH)
+                throw new Exception("Serializer encountered maximum depth of " + _MAX_DEPTH);
+
+
+            Dictionary<string, string> map = new Dictionary<string, string>();
+            Type t = obj.GetType();
+            bool append = false;
+            if (_params.UseExtensions)
+            {
+                if (_params.UsingGlobalTypes == false)
+                    WritePairFast("$type", Reflection.Instance.GetTypeAssemblyName(t));
+                else
+                {
+                    int dt = 0;
+                    string ct = Reflection.Instance.GetTypeAssemblyName(t);
+                    if (_globalTypes.TryGetValue(ct, out dt) == false)
+                    {
+                        dt = _globalTypes.Count + 1;
+                        _globalTypes.Add(ct, dt);
+                    }
+                    WritePairFast("$type", dt.ToString());
+                }
+                append = true;
+            }
+
+            List<Getters> g = Reflection.Instance.GetGetters(t);
+
+            foreach (var p in g)
+            {
+                object o = p.Getter(obj);
+                if ((o == null || o is DBNull) && _params.SerializeNullValues == false)
+                {
+                    //append = false;
+                }
+                else
+                {
+                    if (append)
+                        _output.Append(',');
+
+                    WritePair(p.Name, o);
+                    if (o != null && _params.UseExtensions)
+                    {
+                        Type tt = o.GetType();
+                        if (tt == typeof(System.Object))
+                            map.Add(p.Name, tt.ToString());
+                    }
+                    append = true;
+                }
+            }
+            if (map.Count > 0 && _params.UseExtensions)
+            {
+                _output.Append(",\"$map\":");
+                WriteStringDictionary(map);
+            }
+            _current_depth--;
+            _output.Append('}');
+            _current_depth--;
+        }
+
+        private void WritePairFast(string name, string value)
+        {
+            if ((value == null) && _params.SerializeNullValues == false)
+                return;
+            WriteStringFast(name);
+
+            _output.Append(':');
+
+            WriteStringFast(value);
+        }
+
+        private void WritePair(string name, object value)
+        {
+            if ((value == null || value is DBNull) && _params.SerializeNullValues == false)
+                return;
+            WriteStringFast(name);
+
+            _output.Append(':');
+
+            WriteValue(value);
+        }
+
+        private void WriteArray(IEnumerable array)
+        {
+            _output.Append('[');
+
+            bool pendingSeperator = false;
+
+            foreach (object obj in array)
+            {
+                if (pendingSeperator) _output.Append(',');
+
+                WriteValue(obj);
+
+                pendingSeperator = true;
+            }
+            _output.Append(']');
+        }
+
+        private void WriteStringDictionary(IDictionary dic)
+        {
+            _output.Append('{');
+
+            bool pendingSeparator = false;
+
+            foreach (DictionaryEntry entry in dic)
+            {
+                if (pendingSeparator) _output.Append(',');
+
+                WritePair((string)entry.Key, entry.Value);
+
+                pendingSeparator = true;
+            }
+            _output.Append('}');
+        }
+
+        private void WriteDictionary(IDictionary dic)
+        {
+            _output.Append('[');
+
+            bool pendingSeparator = false;
+
+            foreach (DictionaryEntry entry in dic)
+            {
+                if (pendingSeparator) _output.Append(',');
+                _output.Append('{');
+                WritePair("k", entry.Key);
+                _output.Append(",");
+                WritePair("v", entry.Value);
+                _output.Append('}');
+
+                pendingSeparator = true;
+            }
+            _output.Append(']');
+        }
+
+        private void WriteStringFast(string s)
+        {
+            _output.Append('\"');
+            _output.Append(s);
+            _output.Append('\"');
+        }
+
+        private void WriteString(string s)
+        {
+            _output.Append('\"');
+
+            int runIndex = -1;
+
+            for (var index = 0; index < s.Length; ++index)
+            {
+                var c = s[index];
+
+                if (_useEscapedUnicode)
+                {
+                    if (c >= ' ' && c < 128 && c != '\"' && c != '\\')
+                    {
+                        if (runIndex == -1)
+                            runIndex = index;
+
+                        continue;
+                    }
+                }
+                else
+                {
+                    if (c != '\t' && c != '\n' && c != '\r' && c != '\"' && c != '\\')// && c != ':' && c!=',')
+                    {
+                        if (runIndex == -1)
+                            runIndex = index;
+
+                        continue;
+                    }
+                }
+
+                if (runIndex != -1)
+                {
+                    _output.Append(s, runIndex, index - runIndex);
+                    runIndex = -1;
+                }
+
+                switch (c)
+                {
+                    case '\t': _output.Append("\\t"); break;
+                    case '\r': _output.Append("\\r"); break;
+                    case '\n': _output.Append("\\n"); break;
+                    case '"':
+                    case '\\': _output.Append('\\'); _output.Append(c); break;
+                    default:
+                        if (_useEscapedUnicode)
+                        {
+                            _output.Append("\\u");
+                            _output.Append(((int)c).ToString("X4", NumberFormatInfo.InvariantInfo));
+                        }
+                        else
+                            _output.Append(c);
+
+                        break;
+                }
+            }
+
+            if (runIndex != -1)
+                _output.Append(s, runIndex, s.Length - runIndex);
+
+
+            _output.Append('\"');
+        }
+    }
+}
diff --git a/Engine/Source/Programs/UnrealGameSync/UnrealGameSync/FastJSON/Reflection.cs b/Engine/Source/Programs/UnrealGameSync/UnrealGameSync/FastJSON/Reflection.cs
new file mode 100644
--- /dev/null
+++ b/Engine/Source/Programs/UnrealGameSync/UnrealGameSync/FastJSON/Reflection.cs
@@ -0,0 +1,335 @@
+﻿using System;
+using System.Collections.Generic;
+using System.Linq;
+using System.Text;
+using System.Reflection.Emit;
+using System.Reflection;
+using System.Collections;
+
+namespace fastJSON
+{
+    internal struct Getters
+    {
+        public string Name;
+        public Reflection.GenericGetter Getter;
+        public Type propertyType;
+    }
+
+    internal sealed class Reflection
+    {
+        public readonly static Reflection Instance = new Reflection();
+        private Reflection()
+        {
+        }
+
+        public bool ShowReadOnlyProperties = false;
+        internal delegate object GenericSetter(object target, object value);
+        internal delegate object GenericGetter(object obj);
+        private delegate object CreateObject();
+
+        private SafeDictionary<Type, string> _tyname = new SafeDictionary<Type, string>();
+        private SafeDictionary<string, Type> _typecache = new SafeDictionary<string, Type>();
+        private SafeDictionary<Type, CreateObject> _constrcache = new SafeDictionary<Type, CreateObject>();
+        private SafeDictionary<Type, List<Getters>> _getterscache = new SafeDictionary<Type, List<Getters>>();
+
+        #region [   PROPERTY GET SET   ]
+        internal string GetTypeAssemblyName(Type t)
+        {
+            string val = "";
+            if (_tyname.TryGetValue(t, out val))
+                return val;
+            else
+            {
+                string s = t.AssemblyQualifiedName;
+                _tyname.Add(t, s);
+                return s;
+            }
+        }
+
+        internal Type GetTypeFromCache(string typename)
+        {
+            Type val = null;
+            if (_typecache.TryGetValue(typename, out val))
+                return val;
+            else
+            {
+                Type t = LoadType(typename);
+                _typecache.Add(typename, t);
+                return t;
+            }
+        }
+
+        internal object FastCreateInstance(Type objtype)
+        {
+            try
+            {
+                CreateObject c = null;
+                if (_constrcache.TryGetValue(objtype, out c))
+                {
+                    return c();
+                }
+                else
+                {
+                    if (objtype.IsClass)
+                    {
+                        DynamicMethod dynMethod = new DynamicMethod("_", objtype, null);
+                        ILGenerator ilGen = dynMethod.GetILGenerator();
+                        ilGen.Emit(OpCodes.Newobj, objtype.GetConstructor(Type.EmptyTypes));
+                        ilGen.Emit(OpCodes.Ret);
+                        c = (CreateObject)dynMethod.CreateDelegate(typeof(CreateObject));
+                        _constrcache.Add(objtype, c);
+                    }
+                    else // structs
+                    {
+                        DynamicMethod dynMethod = new DynamicMethod("_", typeof(object), null);
+                        ILGenerator ilGen = dynMethod.GetILGenerator();
+                        var lv = ilGen.DeclareLocal(objtype);
+                        ilGen.Emit(OpCodes.Ldloca_S, lv);
+                        ilGen.Emit(OpCodes.Initobj, objtype);
+                        ilGen.Emit(OpCodes.Ldloc_0);
+                        ilGen.Emit(OpCodes.Box, objtype);
+                        ilGen.Emit(OpCodes.Ret);
+                        c = (CreateObject)dynMethod.CreateDelegate(typeof(CreateObject));
+                        _constrcache.Add(objtype, c);
+                    }
+                    return c();
+                }
+            }
+            catch (Exception exc)
+            {
+                throw new Exception(string.Format("Failed to fast create instance for type '{0}' from assemebly '{1}'",
+                    objtype.FullName, objtype.AssemblyQualifiedName), exc);
+            }
+        }
+
+        internal static GenericSetter CreateSetField(Type type, FieldInfo fieldInfo)
+        {
+            Type[] arguments = new Type[2];
+            arguments[0] = arguments[1] = typeof(object);
+
+            DynamicMethod dynamicSet = new DynamicMethod("_", typeof(object), arguments, type);
+
+            ILGenerator il = dynamicSet.GetILGenerator();
+
+            if (!type.IsClass) // structs
+            {
+                var lv = il.DeclareLocal(type);
+                il.Emit(OpCodes.Ldarg_0);
+                il.Emit(OpCodes.Unbox_Any, type);
+                il.Emit(OpCodes.Stloc_0);
+                il.Emit(OpCodes.Ldloca_S, lv);
+                il.Emit(OpCodes.Ldarg_1);
+                if (fieldInfo.FieldType.IsClass)
+                    il.Emit(OpCodes.Castclass, fieldInfo.FieldType);
+                else
+                    il.Emit(OpCodes.Unbox_Any, fieldInfo.FieldType);
+                il.Emit(OpCodes.Stfld, fieldInfo);
+                il.Emit(OpCodes.Ldloc_0);
+                il.Emit(OpCodes.Box, type);
+                il.Emit(OpCodes.Ret);
+            }
+            else
+            {
+                il.Emit(OpCodes.Ldarg_0);
+                il.Emit(OpCodes.Ldarg_1);
+                if (fieldInfo.FieldType.IsValueType)
+                    il.Emit(OpCodes.Unbox_Any, fieldInfo.FieldType);
+                il.Emit(OpCodes.Stfld, fieldInfo);
+                il.Emit(OpCodes.Ldarg_0);
+                il.Emit(OpCodes.Ret);
+            }
+            return (GenericSetter)dynamicSet.CreateDelegate(typeof(GenericSetter));
+        }
+
+        internal static GenericSetter CreateSetMethod(Type type, PropertyInfo propertyInfo)
+        {
+            MethodInfo setMethod = propertyInfo.GetSetMethod();
+            if (setMethod == null)
+                return null;
+
+            Type[] arguments = new Type[2];
+            arguments[0] = arguments[1] = typeof(object);
+
+            DynamicMethod setter = new DynamicMethod("_", typeof(object), arguments);
+            ILGenerator il = setter.GetILGenerator();
+
+            if (!type.IsClass) // structs
+            {
+                var lv = il.DeclareLocal(type);
+                il.Emit(OpCodes.Ldarg_0);
+                il.Emit(OpCodes.Unbox_Any, type);
+                il.Emit(OpCodes.Stloc_0);
+                il.Emit(OpCodes.Ldloca_S, lv);
+                il.Emit(OpCodes.Ldarg_1);
+                if (propertyInfo.PropertyType.IsClass)
+                    il.Emit(OpCodes.Castclass, propertyInfo.PropertyType);
+                else
+                    il.Emit(OpCodes.Unbox_Any, propertyInfo.PropertyType);
+                il.EmitCall(OpCodes.Call, setMethod, null);
+                il.Emit(OpCodes.Ldloc_0);
+                il.Emit(OpCodes.Box, type);
+            }
+            else
+            {
+                il.Emit(OpCodes.Ldarg_0);
+                il.Emit(OpCodes.Castclass, propertyInfo.DeclaringType);
+                il.Emit(OpCodes.Ldarg_1);
+                if (propertyInfo.PropertyType.IsClass)
+                    il.Emit(OpCodes.Castclass, propertyInfo.PropertyType);
+                else
+                    il.Emit(OpCodes.Unbox_Any, propertyInfo.PropertyType);
+                il.EmitCall(OpCodes.Callvirt, setMethod, null);
+                il.Emit(OpCodes.Ldarg_0);
+            }
+
+            il.Emit(OpCodes.Ret);
+
+            return (GenericSetter)setter.CreateDelegate(typeof(GenericSetter));
+        }
+
+        internal static GenericGetter CreateGetField(Type type, FieldInfo fieldInfo)
+        {
+            DynamicMethod dynamicGet = new DynamicMethod("_", typeof(object), new Type[] { typeof(object) }, type);
+
+            ILGenerator il = dynamicGet.GetILGenerator();
+
+            if (!type.IsClass) // structs
+            {
+                var lv = il.DeclareLocal(type);
+                il.Emit(OpCodes.Ldarg_0);
+                il.Emit(OpCodes.Unbox_Any, type);
+                il.Emit(OpCodes.Stloc_0);
+                il.Emit(OpCodes.Ldloca_S, lv);
+                il.Emit(OpCodes.Ldfld, fieldInfo);
+                if (fieldInfo.FieldType.IsValueType)
+                    il.Emit(OpCodes.Box, fieldInfo.FieldType);
+            }
+            else
+            {
+                il.Emit(OpCodes.Ldarg_0);
+                il.Emit(OpCodes.Ldfld, fieldInfo);
+                if (fieldInfo.FieldType.IsValueType)
+                    il.Emit(OpCodes.Box, fieldInfo.FieldType);
+            }
+
+            il.Emit(OpCodes.Ret);
+
+            return (GenericGetter)dynamicGet.CreateDelegate(typeof(GenericGetter));
+        }
+
+        internal static GenericGetter CreateGetMethod(Type type, PropertyInfo propertyInfo)
+        {
+            MethodInfo getMethod = propertyInfo.GetGetMethod();
+            if (getMethod == null)
+                return null;
+
+            DynamicMethod getter = new DynamicMethod("_", typeof(object), new Type[] { typeof(object) }, type);
+
+            ILGenerator il = getter.GetILGenerator();
+
+            if (!type.IsClass) // structs
+            {
+                var lv = il.DeclareLocal(type);
+                il.Emit(OpCodes.Ldarg_0);
+                il.Emit(OpCodes.Unbox_Any, type);
+                il.Emit(OpCodes.Stloc_0);
+                il.Emit(OpCodes.Ldloca_S, lv);
+                il.EmitCall(OpCodes.Call, getMethod, null);
+                if (propertyInfo.PropertyType.IsValueType)
+                    il.Emit(OpCodes.Box, propertyInfo.PropertyType);
+            }
+            else
+            {
+                il.Emit(OpCodes.Ldarg_0);
+                il.Emit(OpCodes.Castclass, propertyInfo.DeclaringType);
+                il.EmitCall(OpCodes.Callvirt, getMethod, null);
+                if (propertyInfo.PropertyType.IsValueType)
+                    il.Emit(OpCodes.Box, propertyInfo.PropertyType);
+            }
+
+            il.Emit(OpCodes.Ret);
+
+            return (GenericGetter)getter.CreateDelegate(typeof(GenericGetter));
+        }
+
+        internal List<Getters> GetGetters(Type type)
+        {
+            List<Getters> val = null;
+            if (_getterscache.TryGetValue(type, out val))
+                return val;
+
+            PropertyInfo[] props = type.GetProperties(BindingFlags.Public | BindingFlags.Instance);
+            List<Getters> getters = new List<Getters>();
+            foreach (PropertyInfo p in props)
+            {
+                if (!p.CanWrite && ShowReadOnlyProperties == false) continue;
+
+                object[] att = p.GetCustomAttributes(typeof(System.Xml.Serialization.XmlIgnoreAttribute), false);
+                if (att != null && att.Length > 0)
+                    continue;
+
+                GenericGetter g = CreateGetMethod(type, p);
+                if (g != null)
+                {
+                    Getters gg = new Getters();
+                    gg.Name = p.Name;
+                    gg.Getter = g;
+                    gg.propertyType = p.PropertyType;
+                    getters.Add(gg);
+                }
+            }
+
+            FieldInfo[] fi = type.GetFields(BindingFlags.Instance | BindingFlags.Public);
+            foreach (var f in fi)
+            {
+                object[] att = f.GetCustomAttributes(typeof(System.Xml.Serialization.XmlIgnoreAttribute), false);
+                if (att != null && att.Length > 0)
+                    continue;
+
+                GenericGetter g = CreateGetField(type, f);
+                if (g != null)
+                {
+                    Getters gg = new Getters();
+                    gg.Name = f.Name;
+                    gg.Getter = g;
+                    gg.propertyType = f.FieldType;
+                    getters.Add(gg);
+                }
+            }
+
+            _getterscache.Add(type, getters);
+            return getters;
+        }
+
+		private Type LoadType(string TypeName)
+		{
+			Type t = Type.GetType(TypeName);
+			if (t == null && TypeName.Contains(","))
+			{
+				// Split our assembly on first comma, which gives us namespaced typename, and assembly name / version / culture / public key info
+				// E.g. this class' qualified name would be "fastJSON.Reflection, UnrealBuildTool, Version=4.0.0.0, Culture=neutral, PublicKeyToken=null".
+				// This will be split into the type "jastJSON.Reflection" and the assembly "UnrealBuildTool, Version=4.0.0.0, Culture=neutral, PublicKeyToken=null".
+				string[] TypeParts = TypeName.Split(new char[] { ',' }, 2, StringSplitOptions.RemoveEmptyEntries).Select(x => x.Trim()).ToArray();
+				if (TypeParts.Length != 2)
+				{
+					return null;
+				}
+				Assembly Asm = null;
+				try
+				{
+					Asm = Assembly.Load(TypeParts[1]);
+				}
+				catch (System.IO.FileNotFoundException)
+				{
+					Asm = AppDomain.CurrentDomain.GetAssemblies().FirstOrDefault(x => x.FullName == TypeParts[1]);
+				}
+				if (Asm != null)
+				{
+					t = Asm.GetType(TypeParts[0]);
+				}
+			}
+			return t;
+		}
+		#endregion
+	}
+}
diff --git a/Engine/Source/Programs/UnrealGameSync/UnrealGameSync/FastJSON/SafeDictionary.cs b/Engine/Source/Programs/UnrealGameSync/UnrealGameSync/FastJSON/SafeDictionary.cs
new file mode 100644
--- /dev/null
+++ b/Engine/Source/Programs/UnrealGameSync/UnrealGameSync/FastJSON/SafeDictionary.cs
@@ -0,0 +1,52 @@
+﻿using System;
+using System.Collections.Generic;
+
+namespace fastJSON
+{
+    internal sealed class SafeDictionary<TKey, TValue>
+    {
+        private readonly object _Padlock = new object();
+        private readonly Dictionary<TKey, TValue> _Dictionary;
+
+        public SafeDictionary(int capacity)
+        {
+            _Dictionary = new Dictionary<TKey, TValue>(capacity);
+        }
+
+        public SafeDictionary()
+        {
+            _Dictionary = new Dictionary<TKey, TValue>();
+        }
+
+        public bool TryGetValue(TKey key, out TValue value)
+        {
+            lock (_Padlock)
+                return _Dictionary.TryGetValue(key, out value);
+        }
+
+        public int Count { get { lock (_Padlock) return _Dictionary.Count; } }
+
+        public TValue this[TKey key]
+        {
+            get
+            {
+                lock (_Padlock)
+                    return _Dictionary[key];
+            }
+            set
+            {
+                lock (_Padlock)
+                    _Dictionary[key] = value;
+            }
+        }
+
+        public void Add(TKey key, TValue value)
+        {
+            lock (_Padlock)
+            {
+                if (_Dictionary.ContainsKey(key) == false)
+                    _Dictionary.Add(key, value);
+            }
+        }
+    }
+}
diff --git a/Engine/Source/Programs/UnrealGameSync/UnrealGameSync/Forms/IssueSettingsWindow.cs b/Engine/Source/Programs/UnrealGameSync/UnrealGameSync/Forms/IssueSettingsWindow.cs
--- a/Engine/Source/Programs/UnrealGameSync/UnrealGameSync/Forms/IssueSettingsWindow.cs
+++ b/Engine/Source/Programs/UnrealGameSync/UnrealGameSync/Forms/IssueSettingsWindow.cs
@@ -70,6 +70,11 @@ namespace UnrealGameSync
 			UpdateEnabledTextBoxes();
 		}
 
+		private void NotifyUnacknowledgedCheckBox_CheckedChanged(object sender, EventArgs e)
+		{
+			UpdateEnabledTextBoxes();
+		}
+
 		private void NotifyUnresolvedCheckBox_CheckedChanged(object sender, EventArgs e)
 		{
 			UpdateEnabledTextBoxes();
@@ -127,10 +132,5 @@ namespace UnrealGameSync
 			DialogResult = DialogResult.Cancel;
 			Close();
 		}
-
-		private void groupBox2_Enter(object sender, EventArgs e)
-		{
-
-		}
 	}
 }
diff --git a/Engine/Source/Programs/UnrealGameSync/UnrealGameSync/Json.cs b/Engine/Source/Programs/UnrealGameSync/UnrealGameSync/Json.cs
new file mode 100644
--- /dev/null
+++ b/Engine/Source/Programs/UnrealGameSync/UnrealGameSync/Json.cs
@@ -0,0 +1,231 @@
+// Copyright 1998-2019 Epic Games, Inc. All Rights Reserved.
+
+using System;
+using System.Collections.Generic;
+using System.IO;
+using System.Linq;
+using System.Text;
+using System.Threading.Tasks;
+
+namespace UnrealGameSync
+{
+	/// <summary>
+	/// Options for serializing an object to JSON
+	/// </summary>
+	public enum JsonSerializeOptions
+	{
+		/// <summary>
+		/// No options
+		/// </summary>
+		None,
+
+		/// <summary>
+		/// Whether to format the output for readability
+		/// </summary>
+		PrettyPrint = 1,
+	}
+
+	/// <summary>
+	/// Wrapper class around fastJSON methods.
+	/// </summary>
+	public static class Json
+	{
+		/// <summary>
+		/// Serializes an object to JSON
+		/// </summary>
+		/// <param name="Object">Object to be serialized</param>
+		/// <returns>JSON representation of the object</returns>
+		public static string Serialize(object Object)
+		{
+			return Serialize(Object, JsonSerializeOptions.PrettyPrint);
+		}
+
+		/// <summary>
+		/// Serializes an object to JSON
+		/// </summary>
+		/// <param name="Object">Object to be serialized</param>
+		/// <param name="Options">Options for formatting the object</param>
+		/// <returns>JSON representation of the object</returns>
+		public static string Serialize(object Object, JsonSerializeOptions Options)
+		{
+			fastJSON.JSONParameters Params = new fastJSON.JSONParameters();
+			Params.EnableAnonymousTypes = true;
+
+			string Text = fastJSON.JSON.Instance.ToJSON(Object, Params);
+			if((Options & JsonSerializeOptions.PrettyPrint) != 0)
+			{
+				Text = Format(Text);
+			}
+			return Text;
+		}
+
+		/// <summary>
+		/// Deserialize an object from text
+		/// </summary>
+		/// <typeparam name="T">Type of the object to deserialize</typeparam>
+		/// <param name="Text">String to parse</param>
+		/// <returns>Constructed object</returns>
+		public static Dictionary<string, object> Deserialize(string Text)
+		{
+			return (Dictionary<string, object>)fastJSON.JSON.Instance.Parse(Text);
+		}
+
+		/// <summary>
+		/// Deserialize an object from text
+		/// </summary>
+		/// <typeparam name="T">Type of the object to deserialize</typeparam>
+		/// <param name="Text">String to parse</param>
+		/// <returns>Constructed object</returns>
+		public static T Deserialize<T>(string Text)
+		{
+			return fastJSON.JSON.Instance.ToObject<T>(Text);
+		}
+
+		/// <summary>
+		/// Formats the given JSON string in 'Epic' format (parentheses/brackets on new lines, tabs instead of spaces, no space before colons in key/value pairs)
+		/// </summary>
+		/// <param name="Input">The input string</param>
+		/// <returns>The formatted string</returns>
+		public static string Format(string Input)
+		{
+			StringBuilder Result = new StringBuilder(Input.Length * 2);
+
+			bool bOnEmptyLine = true;
+
+			int Depth = 0;
+			for(int Idx = 0; Idx < Input.Length; Idx++)
+			{
+				switch(Input[Idx])
+				{
+					case '{':
+					case '[':
+						if(!bOnEmptyLine)
+						{
+							Result.AppendLine();
+							bOnEmptyLine = true;
+						}
+
+						AppendIndent(Result, Depth++);
+						Result.Append(Input[Idx]);
+						Result.AppendLine();
+						break;
+					case '}':
+					case ']':
+						if(!bOnEmptyLine)
+						{
+							Result.AppendLine();
+						}
+
+						AppendIndent(Result, --Depth);
+						Result.Append(Input[Idx]);
+						bOnEmptyLine = false;
+						break;
+					case ':':
+						Result.Append(Input[Idx]);
+						Result.Append(' ');
+						break;
+					case ',':
+						Result.Append(Input[Idx]);
+						Result.AppendLine();
+						bOnEmptyLine = true;
+						break;
+					default:
+						if(!Char.IsWhiteSpace(Input[Idx]))
+						{
+							// If this is an empty line, append the indent
+							if(bOnEmptyLine)
+							{
+								AppendIndent(Result, Depth);
+								bOnEmptyLine = false;
+							}
+
+							// Append the character
+							Result.Append(Input[Idx]);
+
+							// If this was a string, also append everything to the end of the string
+							if(Input[Idx] == '\"')
+							{
+								Idx++;
+								for(; Idx < Input.Length; Idx++)
+								{
+									Result.Append(Input[Idx]);
+									if(Input[Idx] == '\\' && Idx + 1 < Input.Length)
+									{
+										Result.Append(Input[++Idx]);
+									}
+									else if(Input[Idx] == '\"')
+									{
+										break;
+									}
+								}
+							}
+						}
+						break;
+				}
+			}
+
+			return Result.ToString();
+		}
+
+		/// <summary>
+		/// Appends an indent to the output string
+		/// </summary>
+		/// <param name="Builder">The string builder to append to</param>
+		/// <param name="Depth">The number of tabs to append</param>
+		private static void AppendIndent(StringBuilder Builder, int Depth)
+		{
+			for(int Idx = 0; Idx < Depth; Idx++)
+			{
+				Builder.Append('\t');
+			}
+		}
+
+		/// <summary>
+		/// Escapes a string for serializing to JSON
+		/// </summary>
+		/// <param name="Value">The string to escape</param>
+		/// <returns>The escaped string</returns>
+		public static string EscapeString(string Value)
+		{
+			StringBuilder Result = new StringBuilder();
+			for (int Idx = 0; Idx < Value.Length; Idx++)
+			{
+				switch (Value[Idx])
+				{
+					case '\"':
+						Result.Append("\\\"");
+						break;
+					case '\\':
+						Result.Append("\\\\");
+						break;
+					case '\b':
+						Result.Append("\\b");
+						break;
+					case '\f':
+						Result.Append("\\f");
+						break;
+					case '\n':
+						Result.Append("\\n");
+						break;
+					case '\r':
+						Result.Append("\\r");
+						break;
+					case '\t':
+						Result.Append("\\t");
+						break;
+					default:
+						if (Char.IsControl(Value[Idx]))
+						{
+							Result.AppendFormat("\\u{0:X4}", (int)Value[Idx]);
+						}
+						else
+						{
+							Result.Append(Value[Idx]);
+						}
+						break;
+				}
+			}
+			return Result.ToString();
+		}
+	}
+}
diff --git a/Engine/Source/Programs/UnrealGameSync/UnrealGameSync/Properties/AssemblyInfo.cs b/Engine/Source/Programs/UnrealGameSync/UnrealGameSync/Properties/AssemblyInfo.cs
--- a/Engine/Source/Programs/UnrealGameSync/UnrealGameSync/Properties/AssemblyInfo.cs
+++ b/Engine/Source/Programs/UnrealGameSync/UnrealGameSync/Properties/AssemblyInfo.cs
@@ -34,5 +34,5 @@ using System.Runtime.InteropServices;
 // You can specify all the values or you can default the Build and Revision Numbers 
 // by using the '*' as shown below:
 // [assembly: AssemblyVersion("1.0.*")]
-[assembly: AssemblyVersion("1.182.*")]
-[assembly: AssemblyFileVersion("1.182.0.0")]
+[assembly: AssemblyVersion("1.184.*")]
+[assembly: AssemblyFileVersion("1.184.0.0")]
diff --git a/Engine/Source/Programs/UnrealGameSync/UnrealGameSync/UnrealGameSync.csproj b/Engine/Source/Programs/UnrealGameSync/UnrealGameSync/UnrealGameSync.csproj
--- a/Engine/Source/Programs/UnrealGameSync/UnrealGameSync/UnrealGameSync.csproj
+++ b/Engine/Source/Programs/UnrealGameSync/UnrealGameSync/UnrealGameSync.csproj
@@ -159,6 +159,14 @@
     <Compile Include="DeleteFilesTask.cs" />
     <Compile Include="DeploymentSettings.cs" />
     <Compile Include="DetectProjectSettingsTask.cs" />
+    <Compile Include="FastJSON\Formatter.cs" />
+    <Compile Include="FastJSON\Getters.cs" />
+    <Compile Include="FastJSON\JSON.cs" />
+    <Compile Include="FastJSON\JsonParser.cs" />
+    <Compile Include="FastJSON\JsonSerializer.cs" />
+    <Compile Include="FastJSON\Reflection.cs" />
+    <Compile Include="FastJSON\SafeDictionary.cs" />
+    <Compile Include="Json.cs" />
     <Compile Include="FileFilter.cs" />
     <Compile Include="FindFoldersToCleanTask.cs" />
     <Compile Include="Forms\ArgumentsWindow.cs">
@@ -567,7 +575,9 @@
     <Content Include="Resources\DefaultNotificationLogo.png" />
     <Content Include="Resources\StatusIcons.png" />
   </ItemGroup>
-  <ItemGroup />
+  <ItemGroup>
+    <Folder Include="DotNETCommon\" />
+  </ItemGroup>
   <Import Project="$(MSBuildToolsPath)\Microsoft.CSharp.targets" />
   <!-- To modify your build process, add your task inside one of the targets below and uncomment it. 
        Other similar extension points exist, see Microsoft.Common.targets.
diff --git a/Engine/Source/Programs/UnrealGameSync/UnrealGameSync/Workspace.cs b/Engine/Source/Programs/UnrealGameSync/UnrealGameSync/Workspace.cs
--- a/Engine/Source/Programs/UnrealGameSync/UnrealGameSync/Workspace.cs
+++ b/Engine/Source/Programs/UnrealGameSync/UnrealGameSync/Workspace.cs
@@ -730,13 +730,7 @@ namespace UnrealGameSync
 							bool bIsEpicInternal;
 							Perforce.FileExists(ClientRootPath + "/Engine/Build/NotForLicensees/EpicInternal.txt", out bIsEpicInternal, Log);
 
-							Dictionary<string, string> BuildVersionStrings = new Dictionary<string,string>();
-							BuildVersionStrings["\"Changelist\":"] = String.Format(" {0},", PendingChangeNumber);
-							BuildVersionStrings["\"CompatibleChangelist\":"] = String.Format(" {0},", VersionChangeNumber);
-							BuildVersionStrings["\"BranchName\":"] = String.Format(" \"{0}\"", BranchOrStreamName.Replace('/', '+'));
-							BuildVersionStrings["\"IsPromotedBuild\":"] = " 0,";
-							BuildVersionStrings["\"IsLicenseeVersion\":"] = bIsEpicInternal? "0," : "1,";
-							if(!UpdateVersionFile(ClientRootPath + BuildVersionFileName, BuildVersionStrings, PendingChangeNumber))
+							if (!UpdateVersionFile(ClientRootPath + BuildVersionFileName, PendingChangeNumber, Text => UpdateBuildVersion(Text, PendingChangeNumber, VersionChangeNumber, BranchOrStreamName, !bIsEpicInternal)))
 							{
 								StatusMessage = String.Format("Failed to update {0}.", BuildVersionFileName);
 								return WorkspaceUpdateResult.FailedToSync;
@@ -747,13 +741,7 @@ namespace UnrealGameSync
 							bool bIsEpicInternal;
 							Perforce.FileExists(ClientRootPath + "/Engine/Build/NotForLicensees/EpicInternal.txt", out bIsEpicInternal, Log);
 
-							Dictionary<string, string> BuildVersionStrings = new Dictionary<string,string>();
-							BuildVersionStrings["\"Changelist\":"] = String.Format(" {0},", PendingChangeNumber);
-							BuildVersionStrings["\"CompatibleChangelist\":"] = String.Format(" {0},", VersionChangeNumber);
-							BuildVersionStrings["\"BranchName\":"] = String.Format(" \"{0}\"", BranchOrStreamName.Replace('/', '+'));
-							BuildVersionStrings["\"IsPromotedBuild\":"] = " 0,";
-							BuildVersionStrings["\"IsLicenseeVersion\":"] = bIsEpicInternal? "0," : "1,";
-							if(!UpdateVersionFile(ClientRootPath + BuildVersionFileName, BuildVersionStrings, PendingChangeNumber))
+							if (!UpdateVersionFile(ClientRootPath + BuildVersionFileName, PendingChangeNumber, Text => UpdateBuildVersion(Text, PendingChangeNumber, VersionChangeNumber, BranchOrStreamName, !bIsEpicInternal)))
 							{
 								StatusMessage = String.Format("Failed to update {0}.", BuildVersionFileName);
 								return WorkspaceUpdateResult.FailedToSync;
@@ -1399,6 +1387,11 @@ namespace UnrealGameSync
 		}
 
 		bool UpdateVersionFile(string ClientPath, Dictionary<string, string> VersionStrings, int ChangeNumber)
+		{
+			return UpdateVersionFile(ClientPath, ChangeNumber, Text => UpdateVersionStrings(Text, VersionStrings));
+		}
+
+		bool UpdateVersionFile(string ClientPath, int ChangeNumber, Func<string, string> Update)
 		{
 			List<PerforceFileRecord> Records;
 			if(!Perforce.Stat(ClientPath, out Records, Log))
@@ -1434,21 +1427,50 @@ namespace UnrealGameSync
 				return false;
 			}
 
+			string Text = String.Join("\n", Lines);
+			Text = Update(Text);
+			return WriteVersionFile(LocalPath, DepotPath, Text);
+		}
+
+		string UpdateVersionStrings(string Text, Dictionary<string, string> VersionStrings)
+		{
 			StringWriter Writer = new StringWriter();
-			foreach(string Line in Lines)
+			foreach (string Line in Text.Split('\n'))
 			{
 				string NewLine = Line;
-				foreach(KeyValuePair<string, string> VersionString in VersionStrings)
+				foreach (KeyValuePair<string, string> VersionString in VersionStrings)
 				{
-					if(UpdateVersionLine(ref NewLine, VersionString.Key, VersionString.Value))
+					if (UpdateVersionLine(ref NewLine, VersionString.Key, VersionString.Value))
 					{
 						break;
 					}
 				}
 				Writer.WriteLine(NewLine);
 			}
+			return Writer.ToString();
+		}
+
+		string UpdateBuildVersion(string Text, int Changelist, int CodeChangelist, string BranchOrStreamName, bool bIsLicenseeVersion)
+		{
+			Dictionary<string, object> Object = Json.Deserialize(Text);
+
+			object PrevCompatibleChangelistObj;
+			int PrevCompatibleChangelist = Object.TryGetValue("CompatibleChangelist", out PrevCompatibleChangelistObj) ? (int)Convert.ChangeType(PrevCompatibleChangelistObj, typeof(int)) : 0;
+
+			object PrevIsLicenseeVersionObj;
+			bool PrevIsLicenseeVersion = Object.TryGetValue("IsLicenseeVersion", out PrevIsLicenseeVersionObj)? ((int)Convert.ChangeType(PrevIsLicenseeVersionObj, typeof(int)) != 0) : false;
+
+			Object["Changelist"] = Changelist;
+			if(PrevCompatibleChangelist == 0 || PrevIsLicenseeVersion != bIsLicenseeVersion)
+			{
+				// Don't overwrite the compatible changelist if we're in a hotfix release
+				Object["CompatibleChangelist"] = CodeChangelist;
+			}
+			Object["BranchName"] = BranchOrStreamName.Replace('/', '+');
+			Object["IsPromotedBuild"] = 0;
+			Object["IsLicenseeVersion"] = bIsLicenseeVersion ? 1 : 0;
 
-			return WriteVersionFile(LocalPath, DepotPath, Writer.ToString());
+			return Json.Serialize(Object, JsonSerializeOptions.PrettyPrint);
 		}
 
 		bool WriteVersionFile(string LocalPath, string DepotPath, string NewText)
