Updating copyrights for Engine Programs.
#rnx

diff --git a/Engine/Source/Programs/UnrealGameSync/InstallerCustomAction/Properties/AssemblyInfo.cs b/Engine/Source/Programs/UnrealGameSync/InstallerCustomAction/Properties/AssemblyInfo.cs
--- a/Engine/Source/Programs/UnrealGameSync/InstallerCustomAction/Properties/AssemblyInfo.cs
+++ b/Engine/Source/Programs/UnrealGameSync/InstallerCustomAction/Properties/AssemblyInfo.cs
@@ -1,4 +1,4 @@
-﻿// Copyright 1998-2019 Epic Games, Inc. All Rights Reserved.
+﻿// Copyright Epic Games, Inc. All Rights Reserved.
 
 using System.Reflection;
 using System.Runtime.CompilerServices;
@@ -11,7 +11,7 @@ using System.Runtime.InteropServices;
 [assembly: AssemblyDescription("")]
 [assembly: AssemblyCompany("")]
 [assembly: AssemblyProduct("InstallerCustomAction")]
-[assembly: AssemblyCopyright("Copyright 1998-2019 Epic Games, Inc. All rights reserved.")]
+[assembly: AssemblyCopyright("Copyright Epic Games, Inc. All rights reserved.")]
 [assembly: AssemblyTrademark("")]
 [assembly: AssemblyCulture("")]
 
diff --git a/Engine/Source/Programs/UnrealHeaderTool/Private/CodeGenerator.cpp b/Engine/Source/Programs/UnrealHeaderTool/Private/CodeGenerator.cpp
--- a/Engine/Source/Programs/UnrealHeaderTool/Private/CodeGenerator.cpp
+++ b/Engine/Source/Programs/UnrealHeaderTool/Private/CodeGenerator.cpp
@@ -1,4 +1,4 @@
-// Copyright 1998-2019 Epic Games, Inc. All Rights Reserved.
+// Copyright Epic Games, Inc. All Rights Reserved.
 
 #include "UnrealHeaderTool.h"
 #include "CoreMinimal.h"
@@ -91,7 +91,7 @@ namespace
 	}
 
 	const TCHAR HeaderCopyright[] =
-		TEXT("// Copyright 1998-2019 Epic Games, Inc. All Rights Reserved.\r\n")
+		TEXT("// Copyright Epic Games, Inc. All Rights Reserved.\r\n")
 		TEXT("/*===========================================================================\r\n")
 		TEXT("\tGenerated code exported from UnrealHeaderTool.\r\n")
 		TEXT("\tDO NOT modify this manually! Edit the corresponding .h files instead!\r\n")
diff --git a/Engine/Source/Programs/UnrealSwarm/SwarmCommonUtils/Properties/AssemblyInfo.cs b/Engine/Source/Programs/UnrealSwarm/SwarmCommonUtils/Properties/AssemblyInfo.cs
--- a/Engine/Source/Programs/UnrealSwarm/SwarmCommonUtils/Properties/AssemblyInfo.cs
+++ b/Engine/Source/Programs/UnrealSwarm/SwarmCommonUtils/Properties/AssemblyInfo.cs
@@ -1,4 +1,4 @@
-﻿// Copyright 1998-2019 Epic Games, Inc. All Rights Reserved.
+﻿// Copyright Epic Games, Inc. All Rights Reserved.
 
 using System.Reflection;
 using System.Runtime.CompilerServices;
@@ -12,7 +12,7 @@ using System.Runtime.InteropServices;
 [assembly: AssemblyConfiguration("")]
 [assembly: AssemblyCompany("")]
 [assembly: AssemblyProduct("SwarmCommonUtils")]
-[assembly: AssemblyCopyright("Copyright 1998-2019 Epic Games, Inc. All rights reserved.")]
+[assembly: AssemblyCopyright("Copyright Epic Games, Inc. All rights reserved.")]
 [assembly: AssemblyTrademark("")]
 [assembly: AssemblyCulture("")]
 
diff --git a/Engine/Source/Programs/UnrealVS/BatchBuilderToolControl.xaml b/Engine/Source/Programs/UnrealVS/BatchBuilderToolControl.xaml
--- a/Engine/Source/Programs/UnrealVS/BatchBuilderToolControl.xaml
+++ b/Engine/Source/Programs/UnrealVS/BatchBuilderToolControl.xaml
@@ -1,4 +1,4 @@
-﻿<!--Copyright 1998-2019 Epic Games, Inc. All Rights Reserved.-->
+﻿<!--Copyright Epic Games, Inc. All Rights Reserved.-->
 	
 <UserControl x:Class="UnrealVS.BatchBuilderToolControl"
              xmlns="http://schemas.microsoft.com/winfx/2006/xaml/presentation"
diff --git a/Engine/Source/Programs/UnrealVS/BuildStatusIcon.xaml b/Engine/Source/Programs/UnrealVS/BuildStatusIcon.xaml
--- a/Engine/Source/Programs/UnrealVS/BuildStatusIcon.xaml
+++ b/Engine/Source/Programs/UnrealVS/BuildStatusIcon.xaml
@@ -1,4 +1,4 @@
-﻿<!--Copyright 1998-2019 Epic Games, Inc. All Rights Reserved.-->
+﻿<!--Copyright Epic Games, Inc. All Rights Reserved.-->
 	
 <UserControl x:Class="UnrealVS.BuildStatusIcon"
              xmlns="http://schemas.microsoft.com/winfx/2006/xaml/presentation"
diff --git a/Engine/Source/Programs/UnrealVS/Properties/AssemblyInfo.cs b/Engine/Source/Programs/UnrealVS/Properties/AssemblyInfo.cs
--- a/Engine/Source/Programs/UnrealVS/Properties/AssemblyInfo.cs
+++ b/Engine/Source/Programs/UnrealVS/Properties/AssemblyInfo.cs
@@ -1,4 +1,4 @@
-// Copyright 1998-2019 Epic Games, Inc. All Rights Reserved.
+// Copyright Epic Games, Inc. All Rights Reserved.
 
 using System;
 using System.Reflection;
@@ -14,7 +14,7 @@ using System.Runtime.InteropServices;
 [assembly: AssemblyConfiguration("")]
 [assembly: AssemblyCompany("Epic Games, Inc.")]
 [assembly: AssemblyProduct("UnrealVS")]
-[assembly: AssemblyCopyright("Copyright 2014-2019 Epic Games, Inc.")]
+[assembly: AssemblyCopyright("Copyright Epic Games, Inc.")]
 [assembly: AssemblyTrademark("")]
 [assembly: AssemblyCulture("")]   
 [assembly: ComVisible(false)]     
