UGS: Allow specifying a sync filter for whenever you're using precompiled binaries.

#rnx

diff --git a/Engine/Source/Programs/UnrealGameSync/UnrealGameSync/Controls/WorkspaceControl.cs b/Engine/Source/Programs/UnrealGameSync/UnrealGameSync/Controls/WorkspaceControl.cs
--- a/Engine/Source/Programs/UnrealGameSync/UnrealGameSync/Controls/WorkspaceControl.cs
+++ b/Engine/Source/Programs/UnrealGameSync/UnrealGameSync/Controls/WorkspaceControl.cs
@@ -871,6 +871,12 @@ namespace UnrealGameSync
 						return;
 					}
 					Context.Options &= ~(WorkspaceUpdateOptions.Build | WorkspaceUpdateOptions.GenerateProjectFiles);
+
+					string[] ZippedBinariesSyncFilter;
+					if (TryGetProjectSetting(PerforceMonitor.LatestProjectConfigFile, "ZippedBinariesSyncFilter", out ZippedBinariesSyncFilter) && ZippedBinariesSyncFilter.Length > 0)
+					{
+						Context.SyncFilter = Enumerable.Concat(Context.SyncFilter, ZippedBinariesSyncFilter).ToArray();
+					}
 				}
 				Context.ArchiveTypeToDepotPath.Add(EditorArchiveType, EditorArchivePath);
 			}
diff --git a/Engine/Source/Programs/UnrealGameSync/UnrealGameSync/Properties/AssemblyInfo.cs b/Engine/Source/Programs/UnrealGameSync/UnrealGameSync/Properties/AssemblyInfo.cs
--- a/Engine/Source/Programs/UnrealGameSync/UnrealGameSync/Properties/AssemblyInfo.cs
+++ b/Engine/Source/Programs/UnrealGameSync/UnrealGameSync/Properties/AssemblyInfo.cs
@@ -34,5 +34,5 @@ using System.Runtime.InteropServices;
 // You can specify all the values or you can default the Build and Revision Numbers 
 // by using the '*' as shown below:
 // [assembly: AssemblyVersion("1.0.*")]
-[assembly: AssemblyVersion("1.181.*")]
-[assembly: AssemblyFileVersion("1.181.0.0")]
+[assembly: AssemblyVersion("1.182.*")]
+[assembly: AssemblyFileVersion("1.182.0.0")]
