UGS: Fix text box not being enabled when checking the "Notify unacknowledged" box in the issue settings window.

#rnx

diff --git a/Engine/Source/Programs/UnrealGameSync/UnrealGameSync/Forms/IssueSettingsWindow.cs b/Engine/Source/Programs/UnrealGameSync/UnrealGameSync/Forms/IssueSettingsWindow.cs
--- a/Engine/Source/Programs/UnrealGameSync/UnrealGameSync/Forms/IssueSettingsWindow.cs
+++ b/Engine/Source/Programs/UnrealGameSync/UnrealGameSync/Forms/IssueSettingsWindow.cs
@@ -70,6 +70,11 @@ namespace UnrealGameSync
 			UpdateEnabledTextBoxes();
 		}
 
+		private void NotifyUnacknowledgedCheckBox_CheckedChanged(object sender, EventArgs e)
+		{
+			UpdateEnabledTextBoxes();
+		}
+
 		private void NotifyUnresolvedCheckBox_CheckedChanged(object sender, EventArgs e)
 		{
 			UpdateEnabledTextBoxes();
@@ -127,10 +132,5 @@ namespace UnrealGameSync
 			DialogResult = DialogResult.Cancel;
 			Close();
 		}
-
-		private void groupBox2_Enter(object sender, EventArgs e)
-		{
-
-		}
 	}
 }
diff --git a/Engine/Source/Programs/UnrealGameSync/UnrealGameSync/Properties/AssemblyInfo.cs b/Engine/Source/Programs/UnrealGameSync/UnrealGameSync/Properties/AssemblyInfo.cs
--- a/Engine/Source/Programs/UnrealGameSync/UnrealGameSync/Properties/AssemblyInfo.cs
+++ b/Engine/Source/Programs/UnrealGameSync/UnrealGameSync/Properties/AssemblyInfo.cs
@@ -34,5 +34,5 @@ using System.Runtime.InteropServices;
 // You can specify all the values or you can default the Build and Revision Numbers 
 // by using the '*' as shown below:
 // [assembly: AssemblyVersion("1.0.*")]
-[assembly: AssemblyVersion("1.182.*")]
-[assembly: AssemblyFileVersion("1.182.0.0")]
+[assembly: AssemblyVersion("1.183.*")]
+[assembly: AssemblyFileVersion("1.183.0.0")]
