Copying latest UGS changes to 4.24

- Fix for service badges not updating
- Fix exception when using context menu to open new tab, if no other tabs are open

#rnx

diff --git a/Engine/Source/Programs/UnrealGameSync/UnrealGameSync/Forms/MainWindow.cs b/Engine/Source/Programs/UnrealGameSync/UnrealGameSync/Forms/MainWindow.cs
--- a/Engine/Source/Programs/UnrealGameSync/UnrealGameSync/Forms/MainWindow.cs
+++ b/Engine/Source/Programs/UnrealGameSync/UnrealGameSync/Forms/MainWindow.cs
@@ -1127,14 +1127,7 @@ namespace UnrealGameSync
 
 		private void TabMenu_OpenProject_Click(object sender, EventArgs e)
 		{
-			if (TabMenu_TabIdx == -1)
-			{
-				OpenNewProject();
-			}
-			else
-			{
-				EditSelectedProject(TabMenu_TabIdx);
-			}
+			EditSelectedProject(TabMenu_TabIdx);
 		}
 
 		private void TabMenu_TabNames_Stream_Click(object sender, EventArgs e)
