UGS: Prevent exception when clicking on an issue with an invalid error URL.

#rnx

diff --git a/Engine/Source/Programs/UnrealGameSync/UnrealGameSync/Forms/IssueAlertWindow.cs b/Engine/Source/Programs/UnrealGameSync/UnrealGameSync/Forms/IssueAlertWindow.cs
--- a/Engine/Source/Programs/UnrealGameSync/UnrealGameSync/Forms/IssueAlertWindow.cs
+++ b/Engine/Source/Programs/UnrealGameSync/UnrealGameSync/Forms/IssueAlertWindow.cs
@@ -246,7 +246,21 @@ namespace UnrealGameSync
 			IssueBuildData LastBuild = Issue.Builds.OrderByDescending(x => x.Change).FirstOrDefault();
 			if(LastBuild != null)
 			{
-				System.Diagnostics.Process.Start(LastBuild.ErrorUrl);
+				if (String.IsNullOrEmpty(LastBuild.ErrorUrl))
+				{
+					MessageBox.Show("No additional information is available");
+				}
+				else
+				{
+					try
+					{
+						System.Diagnostics.Process.Start(LastBuild.ErrorUrl);
+					}
+					catch (Exception Ex)
+					{
+						MessageBox.Show(String.Format("Unable to launch '{0}' (Error: {1})", LastBuild.ErrorUrl, Ex.Message));
+					}
+				}
 			}
 		}
 	}
