UGS: Fix recent projects list not being saved after right clicking on the '+' tab.

diff --git a/Engine/Source/Programs/UnrealGameSync/UnrealGameSync/Forms/MainWindow.cs b/Engine/Source/Programs/UnrealGameSync/UnrealGameSync/Forms/MainWindow.cs
--- a/Engine/Source/Programs/UnrealGameSync/UnrealGameSync/Forms/MainWindow.cs
+++ b/Engine/Source/Programs/UnrealGameSync/UnrealGameSync/Forms/MainWindow.cs
@@ -374,7 +374,7 @@ namespace UnrealGameSync
 
 				foreach(UserSelectedProjectSettings RecentProject in Settings.RecentProjects)
 				{
-					ToolStripMenuItem Item = new ToolStripMenuItem(RecentProject.ToString(), null, new EventHandler((o, e) => TryOpenProject(RecentProject, TabMenu_TabIdx)));
+					ToolStripMenuItem Item = new ToolStripMenuItem(RecentProject.ToString(), null, new EventHandler((o, e) => TabMenu_OpenRecentProject_Click(RecentProject, TabMenu_TabIdx)));
 					TabMenu_RecentProjects.DropDownItems.Insert(InsertIdx, Item);
 					InsertIdx++;
 				}
@@ -1143,6 +1143,12 @@ namespace UnrealGameSync
 			EditSelectedProject(TabMenu_TabIdx);
 		}
 
+		private void TabMenu_OpenRecentProject_Click(UserSelectedProjectSettings RecentProject, int TabIdx)
+		{
+			TryOpenProject(RecentProject, TabIdx);
+			SaveTabSettings();
+		}
+
 		private void TabMenu_TabNames_Stream_Click(object sender, EventArgs e)
 		{
 			SetTabNames(TabLabels.Stream);
