Merging latest UGS fixes from Dev-Build.

#rnx

diff --git a/Engine/Source/Programs/UnrealGameSync/UnrealGameSync/Controls/WorkspaceControl.cs b/Engine/Source/Programs/UnrealGameSync/UnrealGameSync/Controls/WorkspaceControl.cs
--- a/Engine/Source/Programs/UnrealGameSync/UnrealGameSync/Controls/WorkspaceControl.cs
+++ b/Engine/Source/Programs/UnrealGameSync/UnrealGameSync/Controls/WorkspaceControl.cs
@@ -790,15 +790,15 @@ namespace UnrealGameSync
 		{
 			Owner.StreamChanged(this);
 			/*
-						StatusPanel.SuspendDisplay();
+			StatusPanel.SuspendDisplay();
 
-						string PrevSelectedFileName = SelectedFileName;
-						if(TryCloseProject())
-						{
-							OpenProject(PrevSelectedFileName);
-						}
+			string PrevSelectedFileName = SelectedFileName;
+			if(TryCloseProject())
+			{
+				OpenProject(PrevSelectedFileName);
+			}
 
-						StatusPanel.ResumeDisplay();*/
+			StatusPanel.ResumeDisplay();*/
 		}
 
 		private void StreamChangedCallback()
@@ -934,7 +934,7 @@ namespace UnrealGameSync
 			{
 				IEnumerable<string> AdditionalPaths = PerforceSection.GetValues("AdditionalPathsToSync", new string[0]);
 				CombinedSyncFilter = CombinedSyncFilter.Union(AdditionalPaths).ToArray();
-			}			
+			}
 
 			WorkspaceUpdateContext Context = new WorkspaceUpdateContext(ChangeNumber, Options, CombinedSyncFilter, GetDefaultBuildStepObjects(), ProjectSettings.BuildSteps, null, GetWorkspaceVariables(ChangeNumber));
 			if (Options.HasFlag(WorkspaceUpdateOptions.SyncArchives))
@@ -957,7 +957,7 @@ namespace UnrealGameSync
 
 					if (Archive.Type == EditorArchiveType)
 					{
-						Context.Options &= ~(WorkspaceUpdateOptions.Build | WorkspaceUpdateOptions.GenerateProjectFiles);
+						Context.Options &= ~(WorkspaceUpdateOptions.Build | WorkspaceUpdateOptions.GenerateProjectFiles | WorkspaceUpdateOptions.OpenSolutionAfterSync);
 					}
 
 					string[] ZippedBinariesSyncFilter;
@@ -1648,7 +1648,7 @@ namespace UnrealGameSync
 
 						Category.Name = Object.GetValue("Name", Category.Name);
 						Category.bEnable = Object.GetValue("Enable", Category.bEnable);
-						Category.Paths = Enumerable.Concat(Category.Paths, Object.GetValue("Paths", "").Split(';').Select(x => x.Trim())).Distinct().OrderBy(x => x).ToArray();
+						Category.Paths = Enumerable.Concat(Category.Paths, Object.GetValue("Paths", "").Split(';').Select(x => x.Trim())).Where(x => x.Length > 0).Distinct().OrderBy(x => x).ToArray();
 					}
 				}
 			}
@@ -1697,20 +1697,20 @@ namespace UnrealGameSync
 			ContentBadges.UnionWith(PerforceMonitor.LatestProjectConfigFile.GetValues("Notifications.ContentBadges", new string[0]));
 
 			// Find the most recent build of each type, and the last time it succeeded
-			Dictionary<string, BadgeData> TypeToLastBuild = new Dictionary<string,BadgeData>();
-			Dictionary<string, BadgeData> TypeToLastSucceededBuild = new Dictionary<string,BadgeData>();
-			for(int Idx = SortedChangeNumbers.Count - 1; Idx >= 0; Idx--)
+			Dictionary<string, BadgeData> TypeToLastBuild = new Dictionary<string, BadgeData>();
+			Dictionary<string, BadgeData> TypeToLastSucceededBuild = new Dictionary<string, BadgeData>();
+			for (int Idx = SortedChangeNumbers.Count - 1; Idx >= 0; Idx--)
 			{
 				EventSummary Summary = EventMonitor.GetSummaryForChange(SortedChangeNumbers[Idx]);
-				if(Summary != null)
+				if (Summary != null)
 				{
-					foreach(BadgeData Badge in Summary.Badges)
+					foreach (BadgeData Badge in Summary.Badges)
 					{
-						if(!TypeToLastBuild.ContainsKey(Badge.BuildType) && (Badge.Result == BadgeResult.Success || Badge.Result == BadgeResult.Warning || Badge.Result == BadgeResult.Failure))
+						if (!TypeToLastBuild.ContainsKey(Badge.BuildType) && (Badge.Result == BadgeResult.Success || Badge.Result == BadgeResult.Warning || Badge.Result == BadgeResult.Failure))
 						{
 							TypeToLastBuild.Add(Badge.BuildType, Badge);
 						}
-						if(!TypeToLastSucceededBuild.ContainsKey(Badge.BuildType) && Badge.Result == BadgeResult.Success)
+						if (!TypeToLastSucceededBuild.ContainsKey(Badge.BuildType) && Badge.Result == BadgeResult.Success)
 						{
 							TypeToLastSucceededBuild.Add(Badge.BuildType, Badge);
 						}
@@ -1721,13 +1721,13 @@ namespace UnrealGameSync
 			// Find all the build types that the user needs to be notified about.
 			int RequireNotificationForChange = -1;
 			List<BadgeData> NotifyBuilds = new List<BadgeData>();
-			foreach(BadgeData LastBuild in TypeToLastBuild.Values.OrderBy(x => x.BuildType))
+			foreach (BadgeData LastBuild in TypeToLastBuild.Values.OrderBy(x => x.BuildType))
 			{
-				if(LastBuild.Result == BadgeResult.Failure || LastBuild.Result == BadgeResult.Warning)
+				if (LastBuild.Result == BadgeResult.Failure || LastBuild.Result == BadgeResult.Warning)
 				{
 					// Get the last submitted changelist by this user of the correct type
 					int LastChangeByCurrentUserOfType;
-					if(ContentBadges.Contains(LastBuild.BuildType))
+					if (ContentBadges.Contains(LastBuild.BuildType))
 					{
 						LastChangeByCurrentUserOfType = LastChangeByCurrentUser;
 					}
@@ -1737,18 +1737,18 @@ namespace UnrealGameSync
 					}
 
 					// Check if the failed build was after we submitted
-					if(LastChangeByCurrentUserOfType > 0 && LastBuild.ChangeNumber >= LastChangeByCurrentUserOfType)
+					if (LastChangeByCurrentUserOfType > 0 && LastBuild.ChangeNumber >= LastChangeByCurrentUserOfType)
 					{
 						// And check that there wasn't a successful build after we submitted (if there was, we're in the clear)
 						BadgeData LastSuccessfulBuild;
-						if(!TypeToLastSucceededBuild.TryGetValue(LastBuild.BuildType, out LastSuccessfulBuild) || LastSuccessfulBuild.ChangeNumber < LastChangeByCurrentUserOfType)
+						if (!TypeToLastSucceededBuild.TryGetValue(LastBuild.BuildType, out LastSuccessfulBuild) || LastSuccessfulBuild.ChangeNumber < LastChangeByCurrentUserOfType)
 						{
 							// Add it to the list of notifications
 							NotifyBuilds.Add(LastBuild);
 
 							// Check if this is a new notification, rather than one we've already dismissed
 							int NotifiedChangeNumber;
-							if(!NotifiedBuildTypeToChangeNumber.TryGetValue(LastBuild.BuildType, out NotifiedChangeNumber) || NotifiedChangeNumber < LastChangeByCurrentUserOfType)
+							if (!NotifiedBuildTypeToChangeNumber.TryGetValue(LastBuild.BuildType, out NotifiedChangeNumber) || NotifiedChangeNumber < LastChangeByCurrentUserOfType)
 							{
 								RequireNotificationForChange = Math.Max(RequireNotificationForChange, LastChangeByCurrentUserOfType);
 							}
@@ -1758,21 +1758,21 @@ namespace UnrealGameSync
 			}
 
 			// If there's anything we haven't already notified the user about, do so now
-			if(RequireNotificationForChange != -1)
+			if (RequireNotificationForChange != -1)
 			{
 				// Format the platform list
 				StringBuilder PlatformList = new StringBuilder(NotifyBuilds[0].BuildType);
-				for(int Idx = 1; Idx < NotifyBuilds.Count - 1; Idx++)
+				for (int Idx = 1; Idx < NotifyBuilds.Count - 1; Idx++)
 				{
 					PlatformList.AppendFormat(", {0}", NotifyBuilds[Idx].BuildType);
 				}
-				if(NotifyBuilds.Count > 1)
+				if (NotifyBuilds.Count > 1)
 				{
 					PlatformList.AppendFormat(" and {0}", NotifyBuilds[NotifyBuilds.Count - 1].BuildType);
 				}
 
 				// Show the balloon tooltip
-				if(NotifyBuilds.Any(x => x.Result == BadgeResult.Failure))
+				if (NotifyBuilds.Any(x => x.Result == BadgeResult.Failure))
 				{
 					string Title = String.Format("{0} Errors", PlatformList.ToString());
 					string Message = String.Format("CIS failed after your last submitted changelist ({0}).", RequireNotificationForChange);
@@ -1788,9 +1788,9 @@ namespace UnrealGameSync
 				// Set the link to open the right build pages
 				int HighlightChange = NotifyBuilds.Max(x => x.ChangeNumber);
 				NotificationWindow.OnMoreInformation = () => { Owner.ShowAndActivate(); SelectChange(HighlightChange); };
-				
+
 				// Don't show messages for this change again
-				foreach(BadgeData NotifyBuild in NotifyBuilds)
+				foreach (BadgeData NotifyBuild in NotifyBuilds)
 				{
 					NotifiedBuildTypeToChangeNumber[NotifyBuild.BuildType] = RequireNotificationForChange;
 				}
@@ -1817,12 +1817,12 @@ namespace UnrealGameSync
 			string ShowingChanges = String.Format("Showing {0}/{1} changes.", ListIndexToChangeIndex.Count, NumChanges);
 
 			int CurrentMaxChanges = PerforceMonitor.CurrentMaxChanges;
-			if(PerforceMonitor.PendingMaxChanges > CurrentMaxChanges)
+			if (PerforceMonitor.PendingMaxChanges > CurrentMaxChanges)
 			{
 				Layout.MainText = String.Format("{0}. Fetching {1} more from server...  ", ShowingChanges, PerforceMonitor.PendingMaxChanges - CurrentMaxChanges);
 				Layout.LinkText = "Cancel";
 			}
-			else if(PerforceMonitor.CurrentMaxChanges > NumChanges)
+			else if (PerforceMonitor.CurrentMaxChanges > NumChanges)
 			{
 				Layout.MainText = ShowingChanges;
 				Layout.LinkText = "";
@@ -1850,7 +1850,7 @@ namespace UnrealGameSync
 			TextRenderer.DrawText(Graphics, Layout.MainText, BuildFont, Layout.MainRect, SystemColors.WindowText, TextFormatFlags.SingleLine | TextFormatFlags.VerticalCenter | TextFormatFlags.NoPrefix | TextFormatFlags.NoPadding);
 
 			Color LinkColor = SystemColors.HotTrack;
-			if(bMouseOverExpandLink)
+			if (bMouseOverExpandLink)
 			{
 				LinkColor = Color.FromArgb(LinkColor.B, LinkColor.G, LinkColor.R);
 			}
@@ -1860,11 +1860,11 @@ namespace UnrealGameSync
 
 		private bool HitTestExpandLink(Point Location)
 		{
-			if(ExpandItem == null)
+			if (ExpandItem == null)
 			{
 				return false;
 			}
-			using(Graphics Graphics = Graphics.FromHwnd(IntPtr.Zero))
+			using (Graphics Graphics = Graphics.FromHwnd(IntPtr.Zero))
 			{
 				ExpandRowLayout Layout = LayoutExpandRow(Graphics, ExpandItem.Bounds);
 				return Layout.LinkRect.Contains(Location);
@@ -1873,24 +1873,24 @@ namespace UnrealGameSync
 
 		private void BuildList_DrawItem(object sender, DrawListViewItemEventArgs e)
 		{
-			if(e.Item == ExpandItem)
+			if (e.Item == ExpandItem)
 			{
 				BuildList.DrawDefaultBackground(e.Graphics, e.Bounds);
 
 				ExpandRowLayout ExpandLayout = LayoutExpandRow(e.Graphics, e.Bounds);
 				DrawExpandRow(e.Graphics, ExpandLayout);
 			}
-			else if(Workspace != null)
+			else if (Workspace != null)
 			{
-				if(e.State.HasFlag(ListViewItemStates.Selected))
+				if (e.State.HasFlag(ListViewItemStates.Selected))
 				{
 					BuildList.DrawSelectedBackground(e.Graphics, e.Bounds);
 				}
-				else if(e.ItemIndex == BuildList.HoverItem)
+				else if (e.ItemIndex == BuildList.HoverItem)
 				{
 					BuildList.DrawTrackedBackground(e.Graphics, e.Bounds);
 				}
-				else if(((PerforceChangeSummary)e.Item.Tag).Number == Workspace.PendingChangeNumber)
+				else if (((PerforceChangeSummary)e.Item.Tag).Number == Workspace.PendingChangeNumber)
 				{
 					BuildList.DrawTrackedBackground(e.Graphics, e.Bounds);
 				}
@@ -1912,17 +1912,17 @@ namespace UnrealGameSync
 				ArchiveToChangeNumberToArchivePath[Archive.Name] = ChangeNumberToArchivePath;
 			}
 
-			if(!ChangeNumberToArchivePath.TryGetValue(ChangeNumber, out ArchivePath))
+			if (!ChangeNumberToArchivePath.TryGetValue(ChangeNumber, out ArchivePath))
 			{
 				PerforceChangeDetails Details;
-				if(PerforceMonitor.TryGetChangeDetails(ChangeNumber, out Details))
+				if (PerforceMonitor.TryGetChangeDetails(ChangeNumber, out Details))
 				{
 					// Try to get the archive for this CL
-					if(!Archive.TryGetArchivePathForChangeNumber(ChangeNumber, out ArchivePath) && !Details.bContainsCode)
+					if (!Archive.TryGetArchivePathForChangeNumber(ChangeNumber, out ArchivePath) && !Details.bContainsCode)
 					{
 						// Otherwise if it's a content-only change, find the previous build any use the archive path from that
 						int Index = SortedChangeNumbers.BinarySearch(ChangeNumber);
-						if(Index > 0)
+						if (Index > 0)
 						{
 							ArchivePath = GetArchivePathForChangeNumber(Archive, SortedChangeNumbers[Index - 1]);
 						}
@@ -1953,7 +1953,7 @@ namespace UnrealGameSync
 		private ChangeLayoutInfo GetChangeLayoutInfo(PerforceChangeSummary Change)
 		{
 			ChangeLayoutInfo LayoutInfo;
-			if(!ChangeNumberToLayoutInfo.TryGetValue(Change.Number, out LayoutInfo))
+			if (!ChangeNumberToLayoutInfo.TryGetValue(Change.Number, out LayoutInfo))
 			{
 				LayoutInfo = new ChangeLayoutInfo();
 
@@ -1972,18 +1972,18 @@ namespace UnrealGameSync
 		private void GetRemainingBisectRange(out int OutPassChangeNumber, out int OutFailChangeNumber)
 		{
 			int PassChangeNumber = -1;
-			foreach(KeyValuePair<int, BisectState> Pair in WorkspaceSettings.ChangeNumberToBisectState)
+			foreach (KeyValuePair<int, BisectState> Pair in WorkspaceSettings.ChangeNumberToBisectState)
 			{
-				if(Pair.Value == BisectState.Pass && (Pair.Key > PassChangeNumber || PassChangeNumber == -1))
+				if (Pair.Value == BisectState.Pass && (Pair.Key > PassChangeNumber || PassChangeNumber == -1))
 				{
 					PassChangeNumber = Pair.Key;
 				}
 			}
 
 			int FailChangeNumber = -1;
-			foreach(KeyValuePair<int, BisectState> Pair in WorkspaceSettings.ChangeNumberToBisectState)
+			foreach (KeyValuePair<int, BisectState> Pair in WorkspaceSettings.ChangeNumberToBisectState)
 			{
-				if(Pair.Value == BisectState.Fail && Pair.Key > PassChangeNumber && (Pair.Key < FailChangeNumber || FailChangeNumber == -1))
+				if (Pair.Value == BisectState.Fail && Pair.Key > PassChangeNumber && (Pair.Key < FailChangeNumber || FailChangeNumber == -1))
 				{
 					FailChangeNumber = Pair.Key;
 				}
@@ -1996,7 +1996,7 @@ namespace UnrealGameSync
 		private void BuildList_DrawSubItem(object sender, DrawListViewSubItemEventArgs e)
 		{
 			PerforceChangeSummary Change = (PerforceChangeSummary)e.Item.Tag;
-			if(Change == null)
+			if (Change == null)
 			{
 				return;
 			}
@@ -2011,87 +2011,87 @@ namespace UnrealGameSync
 			Format.FormatFlags = StringFormatFlags.NoWrap;
 			Format.Trimming = StringTrimming.EllipsisCharacter;
 
-			Font CurrentFont = (Change.Number == Workspace.PendingChangeNumber || Change.Number == Workspace.CurrentChangeNumber)? SelectedBuildFont : BuildFont;
+			Font CurrentFont = (Change.Number == Workspace.PendingChangeNumber || Change.Number == Workspace.CurrentChangeNumber) ? SelectedBuildFont : BuildFont;
 
 			bool bAllowSync = CanSyncChange(Change.Number);
-			int BadgeAlpha = bAllowSync? 255 : 128;
-			Color TextColor = (bAllowSync || Change.Number == Workspace.PendingChangeNumber || Change.Number == Workspace.CurrentChangeNumber || (WorkspaceSettings != null && WorkspaceSettings.AdditionalChangeNumbers.Contains(Change.Number)))? SystemColors.WindowText : Blend(SystemColors.Window, SystemColors.WindowText, 0.25f);
+			int BadgeAlpha = bAllowSync ? 255 : 128;
+			Color TextColor = (bAllowSync || Change.Number == Workspace.PendingChangeNumber || Change.Number == Workspace.CurrentChangeNumber || (WorkspaceSettings != null && WorkspaceSettings.AdditionalChangeNumbers.Contains(Change.Number))) ? SystemColors.WindowText : Blend(SystemColors.Window, SystemColors.WindowText, 0.25f);
 
 			const int FadeRange = 6;
-			if(e.ItemIndex >= BuildList.Items.Count - FadeRange && NumChanges >= PerforceMonitor.CurrentMaxChanges && !IsBisectModeEnabled())
+			if (e.ItemIndex >= BuildList.Items.Count - FadeRange && NumChanges >= PerforceMonitor.CurrentMaxChanges && !IsBisectModeEnabled())
 			{
 				float Opacity = (float)(BuildList.Items.Count - e.ItemIndex - 0.9f) / FadeRange;
 				BadgeAlpha = (int)(BadgeAlpha * Opacity);
 				TextColor = Blend(SystemColors.Window, TextColor, Opacity);
 			}
 
-			if(e.ColumnIndex == IconColumn.Index)
+			if (e.ColumnIndex == IconColumn.Index)
 			{
 				EventSummary Summary = EventMonitor.GetSummaryForChange(Change.Number);
 
 				float MinX = 4 * DpiScaleX;
-				if((Summary != null && EventMonitor.WasSyncedByCurrentUser(Summary.ChangeNumber)) || (Workspace != null && Workspace.CurrentChangeNumber == Change.Number))
+				if ((Summary != null && EventMonitor.WasSyncedByCurrentUser(Summary.ChangeNumber)) || (Workspace != null && Workspace.CurrentChangeNumber == Change.Number))
 				{
 					e.Graphics.DrawImage(Properties.Resources.Icons, MinX * DpiScaleX, IconY, PreviousSyncIcon, GraphicsUnit.Pixel);
 				}
-				else if(WorkspaceSettings != null && WorkspaceSettings.AdditionalChangeNumbers.Contains(Change.Number))
+				else if (WorkspaceSettings != null && WorkspaceSettings.AdditionalChangeNumbers.Contains(Change.Number))
 				{
 					e.Graphics.DrawImage(Properties.Resources.Icons, MinX * DpiScaleX, IconY, AdditionalSyncIcon, GraphicsUnit.Pixel);
 				}
-				else if(bAllowSync && ((Summary != null && Summary.LastStarReview != null && Summary.LastStarReview.Type == EventType.Starred) || PromotedChangeNumbers.Contains(Change.Number)))
+				else if (bAllowSync && ((Summary != null && Summary.LastStarReview != null && Summary.LastStarReview.Type == EventType.Starred) || PromotedChangeNumbers.Contains(Change.Number)))
 				{
 					e.Graphics.DrawImage(Properties.Resources.Icons, MinX * DpiScaleX, IconY, PromotedBuildIcon, GraphicsUnit.Pixel);
 				}
 				MinX += PromotedBuildIcon.Width * DpiScaleX;
 
-				if(bAllowSync)
+				if (bAllowSync)
 				{
 					Rectangle QualityIcon = DefaultBuildIcon;
 
-					if(IsBisectModeEnabled())
+					if (IsBisectModeEnabled())
 					{
 						int FirstPass, FirstFail;
 						GetRemainingBisectRange(out FirstPass, out FirstFail);
 
 						BisectState State;
-						if(!WorkspaceSettings.ChangeNumberToBisectState.TryGetValue(Change.Number, out State) || State == BisectState.Exclude)
+						if (!WorkspaceSettings.ChangeNumberToBisectState.TryGetValue(Change.Number, out State) || State == BisectState.Exclude)
 						{
 							QualityIcon = new Rectangle(0, 0, 0, 0);
 						}
-						if(State == BisectState.Pass)
+						if (State == BisectState.Pass)
 						{
 							QualityIcon = BisectPassIcon;
 						}
-						else if(State == BisectState.Fail)
+						else if (State == BisectState.Fail)
 						{
 							QualityIcon = BisectFailIcon;
 						}
-						else if(FirstFail != -1 && Change.Number > FirstFail)
+						else if (FirstFail != -1 && Change.Number > FirstFail)
 						{
 							QualityIcon = BisectImplicitFailIcon;
 						}
-						else if(FirstPass != -1 && Change.Number < FirstPass)
+						else if (FirstPass != -1 && Change.Number < FirstPass)
 						{
 							QualityIcon = BisectImplicitPassIcon;
 						}
 					}
 					else
 					{
-						if(EventMonitor.IsUnderInvestigation(Change.Number))
+						if (EventMonitor.IsUnderInvestigation(Change.Number))
 						{
 							QualityIcon = BadBuildIcon;
 						}
-						else if(Summary != null)
+						else if (Summary != null)
 						{
-							if(Summary.Verdict == ReviewVerdict.Good)
+							if (Summary.Verdict == ReviewVerdict.Good)
 							{
 								QualityIcon = GoodBuildIcon;
 							}
-							else if(Summary.Verdict == ReviewVerdict.Bad)
+							else if (Summary.Verdict == ReviewVerdict.Bad)
 							{
 								QualityIcon = BadBuildIcon;
 							}
-							else if(Summary.Verdict == ReviewVerdict.Mixed)
+							else if (Summary.Verdict == ReviewVerdict.Mixed)
 							{
 								QualityIcon = MixedBuildIcon;
 							}
@@ -2102,23 +2102,23 @@ namespace UnrealGameSync
 					MinX += QualityIcon.Width * DpiScaleX;
 				}
 			}
-			else if(e.ColumnIndex == ChangeColumn.Index || e.ColumnIndex == TimeColumn.Index)
+			else if (e.ColumnIndex == ChangeColumn.Index || e.ColumnIndex == TimeColumn.Index)
 			{
 				TextRenderer.DrawText(e.Graphics, e.SubItem.Text, CurrentFont, e.Bounds, TextColor, TextFormatFlags.EndEllipsis | TextFormatFlags.SingleLine | TextFormatFlags.HorizontalCenter | TextFormatFlags.VerticalCenter | TextFormatFlags.NoPrefix);
 			}
-			else if(e.ColumnIndex == AuthorColumn.Index)
+			else if (e.ColumnIndex == AuthorColumn.Index)
 			{
 				TextRenderer.DrawText(e.Graphics, e.SubItem.Text, CurrentFont, e.Bounds, TextColor, TextFormatFlags.EndEllipsis | TextFormatFlags.SingleLine | TextFormatFlags.VerticalCenter | TextFormatFlags.NoPrefix);
 			}
-			else if(e.ColumnIndex == DescriptionColumn.Index)
+			else if (e.ColumnIndex == DescriptionColumn.Index)
 			{
 				ChangeLayoutInfo Layout = GetChangeLayoutInfo(Change);
 
 				Rectangle RemainingBounds = e.Bounds;
 				if (Layout.DescriptionBadges.Count > 0)
 				{
-				    e.Graphics.IntersectClip(e.Bounds);
-				    e.Graphics.SmoothingMode = SmoothingMode.HighQuality;
+					e.Graphics.IntersectClip(e.Bounds);
+					e.Graphics.SmoothingMode = SmoothingMode.HighQuality;
 
 					RemainingBounds = new Rectangle(RemainingBounds.Left, RemainingBounds.Top, RemainingBounds.Width - (int)(2 * DpiScaleX), RemainingBounds.Height);
 
@@ -2130,24 +2130,24 @@ namespace UnrealGameSync
 
 				TextRenderer.DrawText(e.Graphics, e.SubItem.Text, CurrentFont, RemainingBounds, TextColor, TextFormatFlags.EndEllipsis | TextFormatFlags.SingleLine | TextFormatFlags.VerticalCenter | TextFormatFlags.NoPrefix);
 			}
-			else if(e.ColumnIndex == TypeColumn.Index)
+			else if (e.ColumnIndex == TypeColumn.Index)
 			{
 				ChangeLayoutInfo Layout = GetChangeLayoutInfo(Change);
-				if(Layout.TypeBadges.Count > 0)
+				if (Layout.TypeBadges.Count > 0)
 				{
-				    e.Graphics.IntersectClip(e.Bounds);
-				    e.Graphics.SmoothingMode = SmoothingMode.HighQuality;
+					e.Graphics.IntersectClip(e.Bounds);
+					e.Graphics.SmoothingMode = SmoothingMode.HighQuality;
 					Point TypeLocation = GetBadgeListLocation(Layout.TypeBadges, e.Bounds, HorizontalAlign.Center, VerticalAlignment.Middle);
 					DrawBadges(e.Graphics, TypeLocation, Layout.TypeBadges, BadgeAlpha);
 				}
 			}
-			else if(e.ColumnIndex == CISColumn.Index)
+			else if (e.ColumnIndex == CISColumn.Index)
 			{
 				ChangeLayoutInfo Layout = GetChangeLayoutInfo(Change);
-				if(Layout.BuildBadges.Count > 0)
+				if (Layout.BuildBadges.Count > 0)
 				{
-				    e.Graphics.IntersectClip(e.Bounds);
-				    e.Graphics.SmoothingMode = SmoothingMode.HighQuality;
+					e.Graphics.IntersectClip(e.Bounds);
+					e.Graphics.SmoothingMode = SmoothingMode.HighQuality;
 
 					Point BuildsLocation = GetBadgeListLocation(Layout.BuildBadges, e.Bounds, HorizontalAlign.Center, VerticalAlignment.Middle);
 					BuildsLocation.X = Math.Max(BuildsLocation.X, e.Bounds.Left);
@@ -2155,20 +2155,20 @@ namespace UnrealGameSync
 					DrawBadges(e.Graphics, BuildsLocation, Layout.BuildBadges, BadgeAlpha);
 				}
 			}
-			else if(e.ColumnIndex == StatusColumn.Index)
+			else if (e.ColumnIndex == StatusColumn.Index)
 			{
 				float MaxX = e.SubItem.Bounds.Right;
 
-				if(Change.Number == Workspace.PendingChangeNumber && Workspace.IsBusy())
+				if (Change.Number == Workspace.PendingChangeNumber && Workspace.IsBusy())
 				{
 					Tuple<string, float> Progress = Workspace.CurrentProgress;
 
 					MaxX -= CancelIcon.Width;
 					e.Graphics.DrawImage(Properties.Resources.Icons, MaxX, IconY, CancelIcon, GraphicsUnit.Pixel);
 
-					if(!Splitter.IsLogVisible())
+					if (!Splitter.IsLogVisible())
 					{
-						MaxX -= InfoIcon.Width; 
+						MaxX -= InfoIcon.Width;
 						e.Graphics.DrawImage(Properties.Resources.Icons, MaxX, IconY, InfoIcon, GraphicsUnit.Pixel);
 					}
 
@@ -2181,25 +2181,25 @@ namespace UnrealGameSync
 					e.Graphics.IntersectClip(e.Bounds);
 					e.Graphics.SmoothingMode = SmoothingMode.HighQuality;
 
-					if(Change.Number == Workspace.CurrentChangeNumber)
+					if (Change.Number == Workspace.CurrentChangeNumber)
 					{
 						EventData Review = EventMonitor.GetReviewByCurrentUser(Change.Number);
 
 						MaxX -= FrownIcon.Width * DpiScaleX;
-						e.Graphics.DrawImage(Properties.Resources.Icons, MaxX, IconY, (Review == null || !EventMonitor.IsPositiveReview(Review.Type))? FrownIcon : DisabledFrownIcon, GraphicsUnit.Pixel);
+						e.Graphics.DrawImage(Properties.Resources.Icons, MaxX, IconY, (Review == null || !EventMonitor.IsPositiveReview(Review.Type)) ? FrownIcon : DisabledFrownIcon, GraphicsUnit.Pixel);
 
 						MaxX -= HappyIcon.Width * DpiScaleX;
-						e.Graphics.DrawImage(Properties.Resources.Icons, MaxX, IconY, (Review == null || !EventMonitor.IsNegativeReview(Review.Type))? HappyIcon : DisabledHappyIcon, GraphicsUnit.Pixel);
+						e.Graphics.DrawImage(Properties.Resources.Icons, MaxX, IconY, (Review == null || !EventMonitor.IsNegativeReview(Review.Type)) ? HappyIcon : DisabledHappyIcon, GraphicsUnit.Pixel);
 					}
-					else if(e.ItemIndex == BuildList.HoverItem && bAllowSync)
+					else if (e.ItemIndex == BuildList.HoverItem && bAllowSync)
 					{
 						Rectangle SyncBadgeRectangle = GetSyncBadgeRectangle(e.SubItem.Bounds);
-						DrawBadge(e.Graphics, SyncBadgeRectangle, "Sync", bHoverSync? Color.FromArgb(140, 180, 230) : Color.FromArgb(112, 146, 190), true, true);
+						DrawBadge(e.Graphics, SyncBadgeRectangle, "Sync", bHoverSync ? Color.FromArgb(140, 180, 230) : Color.FromArgb(112, 146, 190), true, true);
 						MaxX = SyncBadgeRectangle.Left - (int)(2 * DpiScaleX);
 					}
 
 					string SummaryText;
-					if(WorkspaceSettings.LastSyncChangeNumber == -1 || WorkspaceSettings.LastSyncChangeNumber != Change.Number || !GetLastUpdateMessage(WorkspaceSettings.LastSyncResult, WorkspaceSettings.LastSyncResultMessage, out SummaryText))
+					if (WorkspaceSettings.LastSyncChangeNumber == -1 || WorkspaceSettings.LastSyncChangeNumber != Change.Number || !GetLastUpdateMessage(WorkspaceSettings.LastSyncResult, WorkspaceSettings.LastSyncResultMessage, out SummaryText))
 					{
 						StringBuilder SummaryTextBuilder = new StringBuilder();
 
@@ -2207,12 +2207,12 @@ namespace UnrealGameSync
 
 						AppendItemList(SummaryTextBuilder, " ", "Under investigation by {0}.", EventMonitor.GetInvestigatingUsers(Change.Number).Select(x => FormatUserName(x)));
 
-						if(Summary != null)
+						if (Summary != null)
 						{
 							string Comments = String.Join(", ", Summary.Comments.Where(x => !String.IsNullOrWhiteSpace(x.Text)).Select(x => String.Format("{0}: \"{1}\"", FormatUserName(x.UserName), x.Text)));
-							if(Comments.Length > 0)
+							if (Comments.Length > 0)
 							{
-								SummaryTextBuilder.Append(((SummaryTextBuilder.Length == 0)? ""  : " ") + Comments);
+								SummaryTextBuilder.Append(((SummaryTextBuilder.Length == 0) ? "" : " ") + Comments);
 							}
 							else
 							{
@@ -2220,10 +2220,10 @@ namespace UnrealGameSync
 							}
 						}
 
-						SummaryText = (SummaryTextBuilder.Length == 0)? "No current users." : SummaryTextBuilder.ToString();
+						SummaryText = (SummaryTextBuilder.Length == 0) ? "No current users." : SummaryTextBuilder.ToString();
 					}
 
-					if(SummaryText != null && SummaryText.Contains('\n'))
+					if (SummaryText != null && SummaryText.Contains('\n'))
 					{
 						SummaryText = SummaryText.Substring(0, SummaryText.IndexOf('\n')).TrimEnd() + "...";
 					}
@@ -2234,19 +2234,19 @@ namespace UnrealGameSync
 			else
 			{
 				ColumnHeader Column = BuildList.Columns[e.ColumnIndex];
-				if(CustomColumns.Contains(Column))
+				if (CustomColumns.Contains(Column))
 				{
-				    ChangeLayoutInfo Layout = GetChangeLayoutInfo(Change);
+					ChangeLayoutInfo Layout = GetChangeLayoutInfo(Change);
+
+					List<BadgeInfo> Badges;
+					if (Layout.CustomBadges.TryGetValue(Column.Text, out Badges) && Badges.Count > 0)
+					{
+						e.Graphics.IntersectClip(e.Bounds);
+						e.Graphics.SmoothingMode = SmoothingMode.HighQuality;
 
-				    List<BadgeInfo> Badges;
-				    if(Layout.CustomBadges.TryGetValue(Column.Text, out Badges) && Badges.Count > 0)
-				    {
-				        e.Graphics.IntersectClip(e.Bounds);
-				        e.Graphics.SmoothingMode = SmoothingMode.HighQuality;
-    
-					    Point BuildsLocation = GetBadgeListLocation(Badges, e.Bounds, HorizontalAlign.Center, VerticalAlignment.Middle);
-					    DrawBadges(e.Graphics, BuildsLocation, Badges, BadgeAlpha);
-				    }
+						Point BuildsLocation = GetBadgeListLocation(Badges, e.Bounds, HorizontalAlign.Center, VerticalAlignment.Middle);
+						DrawBadges(e.Graphics, BuildsLocation, Badges, BadgeAlpha);
+					}
 				}
 			}
 		}
@@ -2275,7 +2275,7 @@ namespace UnrealGameSync
 		private Point GetBadgeListLocation(List<BadgeInfo> Badges, Rectangle Bounds, HorizontalAlign HorzAlign, VerticalAlignment VertAlign)
 		{
 			Point Location = Bounds.Location;
-			if(Badges.Count == 0)
+			if (Badges.Count == 0)
 			{
 				return Bounds.Location;
 			}
@@ -2283,7 +2283,7 @@ namespace UnrealGameSync
 			BadgeInfo LastBadge = Badges[Badges.Count - 1];
 
 			int X = Bounds.X;
-			switch(HorzAlign)
+			switch (HorzAlign)
 			{
 				case HorizontalAlign.Center:
 					X += (Bounds.Width - LastBadge.Width - LastBadge.Offset) / 2;
@@ -2294,7 +2294,7 @@ namespace UnrealGameSync
 			}
 
 			int Y = Bounds.Y;
-			switch(VertAlign)
+			switch (VertAlign)
 			{
 				case VerticalAlignment.Middle:
 					Y += (Bounds.Height - LastBadge.Height) / 2;
@@ -2309,14 +2309,14 @@ namespace UnrealGameSync
 
 		private void DrawBadges(Graphics Graphics, Point ListLocation, List<BadgeInfo> Badges, int Alpha)
 		{
-			for(int Idx = 0; Idx < Badges.Count; Idx++)
+			for (int Idx = 0; Idx < Badges.Count; Idx++)
 			{
 				BadgeInfo Badge = Badges[Idx];
 				bool bMergeLeft = (Idx > 0 && Badges[Idx - 1].Group == Badge.Group);
 				bool bMergeRight = (Idx < Badges.Count - 1 && Badges[Idx + 1].Group == Badge.Group);
 
 				Rectangle Bounds = new Rectangle(ListLocation.X + Badge.Offset, ListLocation.Y, Badge.Width, Badge.Height);
-				if(Badge.UniqueId != null && Badge.UniqueId == HoverBadgeUniqueId)
+				if (Badge.UniqueId != null && Badge.UniqueId == HoverBadgeUniqueId)
 				{
 					DrawBadge(Graphics, Bounds, Badge.Label, Color.FromArgb((Alpha * Badge.HoverBackgroundColor.A) / 255, Badge.HoverBackgroundColor), bMergeLeft, bMergeRight);
 				}
@@ -2332,7 +2332,7 @@ namespace UnrealGameSync
 			string Description = Change.Description;
 
 			PerforceChangeDetails Details;
-			if(PerforceMonitor.TryGetChangeDetails(Change.Number, out Details))
+			if (PerforceMonitor.TryGetChangeDetails(Change.Number, out Details))
 			{
 				Description = Details.Description;
 			}
@@ -2342,10 +2342,10 @@ namespace UnrealGameSync
 			try
 			{
 				ConfigFile ProjectConfigFile = PerforceMonitor.LatestProjectConfigFile;
-				if(ProjectConfigFile != null)
+				if (ProjectConfigFile != null)
 				{
 					string[] BadgeDefinitions = ProjectConfigFile.GetValues("Badges.DescriptionBadges", new string[0]);
-					foreach(string BadgeDefinition in BadgeDefinitions)
+					foreach (string BadgeDefinition in BadgeDefinitions)
 					{
 						ConfigObject BadgeDefinitionObject = new ConfigObject(BadgeDefinition);
 						string Pattern = BadgeDefinitionObject.GetValue("Pattern", null);
@@ -2355,24 +2355,24 @@ namespace UnrealGameSync
 						string HoverColor = BadgeDefinitionObject.GetValue("HoverColor", "#b0b0b0");
 						string Url = BadgeDefinitionObject.GetValue("Url", null);
 						string Arguments = BadgeDefinitionObject.GetValue("Arguments", null);
-						if(!String.IsNullOrEmpty(Name) && !String.IsNullOrEmpty(Pattern))
+						if (!String.IsNullOrEmpty(Name) && !String.IsNullOrEmpty(Pattern))
 						{
-							foreach(Match MatchResult in Regex.Matches(Description, Pattern, RegexOptions.Multiline))
+							foreach (Match MatchResult in Regex.Matches(Description, Pattern, RegexOptions.Multiline))
 							{
 								Color BadgeColor = System.Drawing.ColorTranslator.FromHtml(Color);
 								Color HoverBadgeColor = System.Drawing.ColorTranslator.FromHtml(HoverColor);
 
-								string UniqueId = String.IsNullOrEmpty(Url)? null : String.Format("Description:{0}:{1}", Change.Number, Badges.Count);
+								string UniqueId = String.IsNullOrEmpty(Url) ? null : String.Format("Description:{0}:{1}", Change.Number, Badges.Count);
 
 								string ExpandedUrl = ReplaceRegexMatches(Url, MatchResult);
 								string ExpandedArguments = ReplaceRegexMatches(Arguments, MatchResult);
 
 								Action ClickHandler;
-								if(String.IsNullOrEmpty(ExpandedUrl))
+								if (String.IsNullOrEmpty(ExpandedUrl))
 								{
 									ClickHandler = null;
 								}
-								else if(String.IsNullOrEmpty(ExpandedArguments))
+								else if (String.IsNullOrEmpty(ExpandedArguments))
 								{
 									ClickHandler = () => SafeProcessStart(ExpandedUrl);
 								}
@@ -2387,7 +2387,7 @@ namespace UnrealGameSync
 					}
 				}
 			}
-			catch(Exception)
+			catch (Exception)
 			{
 			}
 
@@ -2434,9 +2434,9 @@ namespace UnrealGameSync
 
 		private string ReplaceRegexMatches(string Text, Match MatchResult)
 		{
-			if(Text != null)
+			if (Text != null)
 			{
-				for(int Idx = 1; Idx < MatchResult.Groups.Count; Idx++)
+				for (int Idx = 1; Idx < MatchResult.Groups.Count; Idx++)
 				{
 					string SourceText = String.Format("${0}", Idx);
 					Text = Text.Replace(SourceText, MatchResult.Groups[Idx].Value);
@@ -2450,18 +2450,18 @@ namespace UnrealGameSync
 			List<BadgeInfo> Badges = new List<BadgeInfo>();
 
 			PerforceChangeDetails Details;
-			if(PerforceMonitor.TryGetChangeDetails(ChangeNumber, out Details))
+			if (PerforceMonitor.TryGetChangeDetails(ChangeNumber, out Details))
 			{
-				if(Details.bContainsCode)
+				if (Details.bContainsCode)
 				{
 					Badges.Add(new BadgeInfo("Code", "ChangeType", Color.FromArgb(116, 185, 255)));
 				}
-				if(Details.bContainsContent)
+				if (Details.bContainsContent)
 				{
 					Badges.Add(new BadgeInfo("Content", "ChangeType", Color.FromArgb(162, 155, 255)));
 				}
 			}
-			if(Badges.Count == 0)
+			if (Badges.Count == 0)
 			{
 				Badges.Add(new BadgeInfo("Unknown", "ChangeType", Color.FromArgb(192, 192, 192)));
 			}
@@ -2473,10 +2473,10 @@ namespace UnrealGameSync
 		private bool TryGetProjectSetting(ConfigFile ProjectConfigFile, string Name, out string Value)
 		{
 			ConfigSection ProjectSection = ProjectConfigFile.FindSection(SelectedProjectIdentifier);
-			if(ProjectSection != null)
+			if (ProjectSection != null)
 			{
 				string NewValue = ProjectSection.GetValue(Name, null);
-				if(NewValue != null)
+				if (NewValue != null)
 				{
 					Value = NewValue;
 					return true;
@@ -2484,10 +2484,10 @@ namespace UnrealGameSync
 			}
 
 			ConfigSection DefaultSection = ProjectConfigFile.FindSection("Default");
-			if(DefaultSection != null)
+			if (DefaultSection != null)
 			{
 				string NewValue = DefaultSection.GetValue(Name, null);
-				if(NewValue != null)
+				if (NewValue != null)
 				{
 					Value = NewValue;
 					return true;
@@ -2501,7 +2501,7 @@ namespace UnrealGameSync
 		private bool TryGetProjectSetting(ConfigFile ProjectConfigFile, string Name, out string[] Values)
 		{
 			string ValueList;
-			if(TryGetProjectSetting(ProjectConfigFile, Name, out ValueList))
+			if (TryGetProjectSetting(ProjectConfigFile, Name, out ValueList))
 			{
 				Values = ValueList.Split('\n').Select(x => x.Trim()).Where(x => x.Length > 0).ToArray();
 				return true;
@@ -2516,14 +2516,14 @@ namespace UnrealGameSync
 		private bool TryGetProjectSetting(ConfigFile ProjectConfigFile, string Name, string LegacyName, out string Value)
 		{
 			string NewValue;
-			if(TryGetProjectSetting(ProjectConfigFile, Name, out NewValue))
+			if (TryGetProjectSetting(ProjectConfigFile, Name, out NewValue))
 			{
 				Value = NewValue;
 				return true;
 			}
 
 			NewValue = ProjectConfigFile.GetValue(LegacyName, null);
-			if(NewValue != null)
+			if (NewValue != null)
 			{
 				Value = NewValue;
 				return true;
@@ -2537,23 +2537,23 @@ namespace UnrealGameSync
 		{
 			List<BadgeInfo> Badges = new List<BadgeInfo>();
 
-			if(Summary != null && Summary.Badges.Count > 0)
+			if (Summary != null && Summary.Badges.Count > 0)
 			{
 				// Create a lookup for build data for each badge name
 				Dictionary<string, BadgeData> BadgeNameToBuildData = new Dictionary<string, BadgeData>();
-				foreach(BadgeData Badge in Summary.Badges)
+				foreach (BadgeData Badge in Summary.Badges)
 				{
 					BadgeNameToBuildData[Badge.BadgeName] = Badge;
 				}
 
 				// Add all the badges, sorted by group
-				foreach(KeyValuePair<string, string> BadgeNameAndGroup in BadgeNameAndGroupPairs)
+				foreach (KeyValuePair<string, string> BadgeNameAndGroup in BadgeNameAndGroupPairs)
 				{
 					BadgeData BadgeData;
 					BadgeNameToBuildData.TryGetValue(BadgeNameAndGroup.Key, out BadgeData);
 
 					BadgeInfo BadgeInfo = CreateBadge(ChangeNumber, BadgeNameAndGroup.Key, BadgeNameAndGroup.Value, BadgeData);
-					if(MaxBuildBadgeChars != -1 && BadgeInfo.Label.Length > MaxBuildBadgeChars)
+					if (MaxBuildBadgeChars != -1 && BadgeInfo.Label.Length > MaxBuildBadgeChars)
 					{
 						BadgeInfo.ToolTip = BadgeInfo.Label;
 						BadgeInfo.Label = BadgeInfo.Label.Substring(0, MaxBuildBadgeChars);
@@ -2572,7 +2572,7 @@ namespace UnrealGameSync
 			string BadgeLabel = BadgeName;
 			Color BadgeColor = Color.FromArgb(0, Color.White);
 
-			if(BadgeData != null)
+			if (BadgeData != null)
 			{
 				BadgeLabel = BadgeData.BadgeLabel;
 				BadgeColor = GetBuildBadgeColor(BadgeData.Result);
@@ -2581,7 +2581,7 @@ namespace UnrealGameSync
 			Color HoverBadgeColor = Color.FromArgb(BadgeColor.A, Math.Min(BadgeColor.R + 32, 255), Math.Min(BadgeColor.G + 32, 255), Math.Min(BadgeColor.B + 32, 255));
 
 			Action ClickHandler;
-			if(BadgeData == null || String.IsNullOrEmpty(BadgeData.Url))
+			if (BadgeData == null || String.IsNullOrEmpty(BadgeData.Url))
 			{
 				ClickHandler = null;
 			}
@@ -2597,17 +2597,17 @@ namespace UnrealGameSync
 		private Dictionary<string, List<BadgeInfo>> CreateCustomBadges(int ChangeNumber, EventSummary Summary)
 		{
 			Dictionary<string, List<BadgeInfo>> ColumnNameToBadges = new Dictionary<string, List<BadgeInfo>>();
-			if(Summary != null && Summary.Badges.Count > 0)
+			if (Summary != null && Summary.Badges.Count > 0)
 			{
-				foreach(ColumnHeader CustomColumn in CustomColumns)
+				foreach (ColumnHeader CustomColumn in CustomColumns)
 				{
 					ConfigObject Config = (ConfigObject)CustomColumn.Tag;
-					if(Config != null)
+					if (Config != null)
 					{
 						List<BadgeInfo> Badges = new List<BadgeInfo>();
 
-						string[] BadgeNames = Config.GetValue("Badges", "").Split(new char[]{ ',' }, StringSplitOptions.RemoveEmptyEntries);
-						foreach(string BadgeName in BadgeNames)
+						string[] BadgeNames = Config.GetValue("Badges", "").Split(new char[] { ',' }, StringSplitOptions.RemoveEmptyEntries);
+						foreach (string BadgeName in BadgeNames)
 						{
 							BadgeInfo Badge = CreateBadge(ChangeNumber, BadgeName, "XXXX", Summary.Badges.FirstOrDefault(x => x.BadgeName == BadgeName));
 							Badges.Add(Badge);
@@ -2624,19 +2624,19 @@ namespace UnrealGameSync
 
 		private static Color GetBuildBadgeColor(BadgeResult Result)
 		{
-			if(Result == BadgeResult.Starting)
+			if (Result == BadgeResult.Starting)
 			{
 				return Color.FromArgb(128, 192, 255);
 			}
-			else if(Result == BadgeResult.Warning)
+			else if (Result == BadgeResult.Warning)
 			{
 				return Color.FromArgb(255, 192, 0);
 			}
-			else if(Result == BadgeResult.Failure)
+			else if (Result == BadgeResult.Failure)
 			{
 				return Color.FromArgb(192, 64, 0);
 			}
-			else if(Result == BadgeResult.Skipped)
+			else if (Result == BadgeResult.Skipped)
 			{
 				return Color.FromArgb(192, 192, 192);
 			}
@@ -2664,18 +2664,18 @@ namespace UnrealGameSync
 
 		private void DrawBadge(Graphics Graphics, Rectangle BadgeRect, string BadgeText, Color BadgeColor, bool bMergeLeft, bool bMergeRight)
 		{
-			if(BadgeColor.A != 0)
+			if (BadgeColor.A != 0)
 			{
 				using (GraphicsPath Path = new GraphicsPath())
 				{
 					Path.StartFigure();
-					Path.AddLine(BadgeRect.Left + (bMergeLeft? 1 : 0), BadgeRect.Top, BadgeRect.Left - (bMergeLeft? 1 : 0), BadgeRect.Bottom);
-					Path.AddLine(BadgeRect.Left - (bMergeLeft? 1 : 0), BadgeRect.Bottom, BadgeRect.Right - 1 - (bMergeRight? 1 : 0), BadgeRect.Bottom);
-					Path.AddLine(BadgeRect.Right - 1 - (bMergeRight? 1 : 0), BadgeRect.Bottom, BadgeRect.Right - 1 + (bMergeRight? 1 : 0), BadgeRect.Top);
-					Path.AddLine(BadgeRect.Right - 1 + (bMergeRight? 1 : 0), BadgeRect.Top, BadgeRect.Left + (bMergeLeft? 1 : 0), BadgeRect.Top);
+					Path.AddLine(BadgeRect.Left + (bMergeLeft ? 1 : 0), BadgeRect.Top, BadgeRect.Left - (bMergeLeft ? 1 : 0), BadgeRect.Bottom);
+					Path.AddLine(BadgeRect.Left - (bMergeLeft ? 1 : 0), BadgeRect.Bottom, BadgeRect.Right - 1 - (bMergeRight ? 1 : 0), BadgeRect.Bottom);
+					Path.AddLine(BadgeRect.Right - 1 - (bMergeRight ? 1 : 0), BadgeRect.Bottom, BadgeRect.Right - 1 + (bMergeRight ? 1 : 0), BadgeRect.Top);
+					Path.AddLine(BadgeRect.Right - 1 + (bMergeRight ? 1 : 0), BadgeRect.Top, BadgeRect.Left + (bMergeLeft ? 1 : 0), BadgeRect.Top);
 					Path.CloseFigure();
 
-					using(SolidBrush Brush = new SolidBrush(BadgeColor))
+					using (SolidBrush Brush = new SolidBrush(BadgeColor))
 					{
 						Graphics.FillPath(Brush, Path);
 					}
@@ -2688,7 +2688,7 @@ namespace UnrealGameSync
 		private Size GetBadgeSize(string BadgeText)
 		{
 			Size BadgeSize;
-			if(!BadgeLabelToSize.TryGetValue(BadgeText, out BadgeSize))
+			if (!BadgeLabelToSize.TryGetValue(BadgeText, out BadgeSize))
 			{
 				Size LabelSize = TextRenderer.MeasureText(BadgeText, BadgeFont);
 				int BadgeHeight = BadgeFont.Height + 1;
@@ -2701,7 +2701,7 @@ namespace UnrealGameSync
 
 		private static bool GetLastUpdateMessage(WorkspaceUpdateResult Result, string ResultMessage, out string Message)
 		{
-			if(Result != WorkspaceUpdateResult.Success && ResultMessage != null)
+			if (Result != WorkspaceUpdateResult.Success && ResultMessage != null)
 			{
 				Message = ResultMessage;
 				return true;
@@ -2711,7 +2711,7 @@ namespace UnrealGameSync
 
 		private static bool GetGenericLastUpdateMessage(WorkspaceUpdateResult Result, out string Message)
 		{
-			switch(Result)
+			switch (Result)
 			{
 				case WorkspaceUpdateResult.Canceled:
 					Message = "Sync canceled.";
@@ -2754,9 +2754,9 @@ namespace UnrealGameSync
 		private static void AppendItemList(StringBuilder Builder, string Separator, string Format, IEnumerable<string> Items)
 		{
 			string ItemList = FormatItemList(Format, Items);
-			if(ItemList != null)
+			if (ItemList != null)
 			{
-				if(Builder.Length > 0)
+				if (Builder.Length > 0)
 				{
 					Builder.Append(Separator);
 				}
@@ -2767,15 +2767,15 @@ namespace UnrealGameSync
 		private static string FormatItemList(string Format, IEnumerable<string> Items)
 		{
 			string[] ItemsArray = Items.Distinct().ToArray();
-			if(ItemsArray.Length == 0)
+			if (ItemsArray.Length == 0)
 			{
 				return null;
 			}
 
 			StringBuilder Builder = new StringBuilder(ItemsArray[0]);
-			if(ItemsArray.Length > 1)
+			if (ItemsArray.Length > 1)
 			{
-				for(int Idx = 1; Idx < ItemsArray.Length - 1; Idx++)
+				for (int Idx = 1; Idx < ItemsArray.Length - 1; Idx++)
 				{
 					Builder.Append(", ");
 					Builder.Append(ItemsArray[Idx]);
@@ -2788,15 +2788,15 @@ namespace UnrealGameSync
 
 		private void BuildList_MouseDoubleClick(object Sender, MouseEventArgs Args)
 		{
-			if(Args.Button == MouseButtons.Left)
+			if (Args.Button == MouseButtons.Left)
 			{
 				ListViewHitTestInfo HitTest = BuildList.HitTest(Args.Location);
-				if(HitTest.Item != null)
+				if (HitTest.Item != null)
 				{
 					PerforceChangeSummary Change = (PerforceChangeSummary)HitTest.Item.Tag;
-					if(Change != null)
+					if (Change != null)
 					{
-						if(Change.Number == Workspace.CurrentChangeNumber)
+						if (Change.Number == Workspace.CurrentChangeNumber)
 						{
 							LaunchEditor();
 						}
@@ -2811,54 +2811,54 @@ namespace UnrealGameSync
 
 		public void LaunchEditor()
 		{
-			if(!Workspace.IsBusy() && Workspace.CurrentChangeNumber != -1)
+			if (!Workspace.IsBusy() && Workspace.CurrentChangeNumber != -1)
 			{
 				BuildConfig EditorBuildConfig = GetEditorBuildConfig();
 
 				List<string> ReceiptPaths = GetEditorReceiptPaths(EditorBuildConfig);
 
 				string EditorExe = GetEditorExePath(EditorBuildConfig);
-				if(ReceiptPaths.Any(x => File.Exists(x)) && File.Exists(EditorExe))
+				if (ReceiptPaths.Any(x => File.Exists(x)) && File.Exists(EditorExe))
 				{
-					if(Settings.bEditorArgumentsPrompt && !ModifyEditorArguments())
+					if (Settings.bEditorArgumentsPrompt && !ModifyEditorArguments())
 					{
 						return;
 					}
 
 					StringBuilder LaunchArguments = new StringBuilder();
-					if(SelectedFileName.EndsWith(".uproject", StringComparison.InvariantCultureIgnoreCase))
+					if (SelectedFileName.EndsWith(".uproject", StringComparison.InvariantCultureIgnoreCase))
 					{
 						LaunchArguments.AppendFormat("\"{0}\"", SelectedFileName);
 					}
-					foreach(Tuple<string, bool> EditorArgument in Settings.EditorArguments)
+					foreach (Tuple<string, bool> EditorArgument in Settings.EditorArguments)
 					{
-						if(EditorArgument.Item2)
+						if (EditorArgument.Item2)
 						{
 							LaunchArguments.AppendFormat(" {0}", EditorArgument.Item1);
 						}
 					}
-					if(EditorBuildConfig == BuildConfig.Debug || EditorBuildConfig == BuildConfig.DebugGame)
+					if (EditorBuildConfig == BuildConfig.Debug || EditorBuildConfig == BuildConfig.DebugGame)
 					{
 						LaunchArguments.Append(" -debug");
 					}
-					
-					if(!Utility.SpawnProcess(EditorExe, LaunchArguments.ToString()))
+
+					if (!Utility.SpawnProcess(EditorExe, LaunchArguments.ToString()))
 					{
 						ShowErrorDialog("Unable to spawn {0} {1}", EditorExe, LaunchArguments.ToString());
 					}
 				}
 				else
 				{
-					if(MessageBox.Show("The editor needs to be built before you can run it. Build it now?", "Editor out of date", MessageBoxButtons.YesNo) == System.Windows.Forms.DialogResult.Yes)
+					if (MessageBox.Show("The editor needs to be built before you can run it. Build it now?", "Editor out of date", MessageBoxButtons.YesNo) == System.Windows.Forms.DialogResult.Yes)
 					{
 						Owner.ShowAndActivate();
 
 						WorkspaceUpdateOptions Options = WorkspaceUpdateOptions.Build | WorkspaceUpdateOptions.RunAfterSync;
-						if(Settings.bUseIncrementalBuilds)
+						if (Settings.bUseIncrementalBuilds)
 						{
 							Options |= WorkspaceUpdateOptions.UseIncrementalBuilds;
 						}
-		
+
 						WorkspaceUpdateContext Context = new WorkspaceUpdateContext(Workspace.CurrentChangeNumber, Options, null, GetDefaultBuildStepObjects(), ProjectSettings.BuildSteps, null, GetWorkspaceVariables(Workspace.CurrentChangeNumber));
 						StartWorkspaceUpdate(Context, null);
 					}
@@ -2870,9 +2870,9 @@ namespace UnrealGameSync
 		{
 			// Try to read the executable path from the target receipt
 			List<string> ReceiptFileNames = GetEditorReceiptPaths(Config);
-			foreach(string ReceiptFileName in ReceiptFileNames)
+			foreach (string ReceiptFileName in ReceiptFileNames)
 			{
-				if(File.Exists(ReceiptFileName))
+				if (File.Exists(ReceiptFileName))
 				{
 					Log.WriteLine("Reading {0}", ReceiptFileName);
 					try
@@ -2882,10 +2882,10 @@ namespace UnrealGameSync
 						Dictionary<string, object> RawObject = Serializer.Deserialize<Dictionary<string, object>>(Text);
 
 						object LaunchFileNameObject;
-						if(RawObject.TryGetValue("Launch", out LaunchFileNameObject))
+						if (RawObject.TryGetValue("Launch", out LaunchFileNameObject))
 						{
 							string LaunchFileName = LaunchFileNameObject as string;
-							if(LaunchFileName != null)
+							if (LaunchFileName != null)
 							{
 								LaunchFileName = LaunchFileName.Replace("$(EngineDir)", Path.Combine(BranchDirectoryName, "Engine"));
 								LaunchFileName = LaunchFileName.Replace("$(ProjectDir)", Path.GetDirectoryName(SelectedFileName));
@@ -2893,7 +2893,7 @@ namespace UnrealGameSync
 							}
 						}
 					}
-					catch(Exception Ex)
+					catch (Exception Ex)
 					{
 						Log.WriteLine("  Exception while parsing receipt: {0}", Ex.ToString());
 					}
@@ -2903,9 +2903,9 @@ namespace UnrealGameSync
 
 			// Otherwise use the standard editor path
 			string ExeFileName = "UE4Editor.exe";
-			if((Config != BuildConfig.DebugGame || PerforceMonitor.LatestProjectConfigFile.GetValue("Options.DebugGameHasSeparateExecutable", false)) && Config != BuildConfig.Development)
+			if ((Config != BuildConfig.DebugGame || PerforceMonitor.LatestProjectConfigFile.GetValue("Options.DebugGameHasSeparateExecutable", false)) && Config != BuildConfig.Development)
 			{
-				ExeFileName	= String.Format("UE4Editor-Win64-{0}.exe", Config.ToString());
+				ExeFileName = String.Format("UE4Editor-Win64-{0}.exe", Config.ToString());
 			}
 			return Path.Combine(BranchDirectoryName, "Engine", "Binaries", "Win64", ExeFileName);
 		}
@@ -2913,10 +2913,10 @@ namespace UnrealGameSync
 		private bool WaitForProgramsToFinish()
 		{
 			string[] ProcessFileNames = GetProcessesRunningInWorkspace();
-			if(ProcessFileNames.Length > 0)
+			if (ProcessFileNames.Length > 0)
 			{
 				ProgramsRunningWindow ProgramsRunning = new ProgramsRunningWindow(GetProcessesRunningInWorkspace, ProcessFileNames);
-				if(ProgramsRunning.ShowDialog() != DialogResult.OK)
+				if (ProgramsRunning.ShowDialog() != DialogResult.OK)
 				{
 					return false;
 				}
@@ -2946,20 +2946,20 @@ namespace UnrealGameSync
 			try
 			{
 				string RootDirectoryName = Path.GetFullPath(Workspace.LocalRootPath) + Path.DirectorySeparatorChar;
-				foreach(Process ProcessInstance in Process.GetProcesses())
+				foreach (Process ProcessInstance in Process.GetProcesses())
 				{
 					try
 					{
-						if(ProcessNames.Contains(ProcessInstance.ProcessName))
+						if (ProcessNames.Contains(ProcessInstance.ProcessName))
 						{
 							string ProcessFileName = Path.GetFullPath(ProcessInstance.MainModule.FileName);
-							if(ProcessFileName.StartsWith(RootDirectoryName, StringComparison.InvariantCultureIgnoreCase))
+							if (ProcessFileName.StartsWith(RootDirectoryName, StringComparison.InvariantCultureIgnoreCase))
 							{
 								string DisplayTitle = ProcessFileName;
 								try
 								{
 									string MainWindowTitle = ProcessInstance.MainWindowTitle;
-									if(!String.IsNullOrEmpty(MainWindowTitle))
+									if (!String.IsNullOrEmpty(MainWindowTitle))
 									{
 										DisplayTitle = String.Format("{0} ({1})", ProcessInstance.MainWindowTitle, DisplayTitle);
 									}
@@ -2984,7 +2984,7 @@ namespace UnrealGameSync
 
 		private List<string> GetEditorReceiptPaths(BuildConfig Config)
 		{
-			string ConfigSuffix = (Config == BuildConfig.Development)? "" : String.Format("-Win64-{0}", Config.ToString());
+			string ConfigSuffix = (Config == BuildConfig.Development) ? "" : String.Format("-Win64-{0}", Config.ToString());
 
 			List<string> PossiblePaths = new List<string>();
 			if (EditorTargetName != null)
@@ -3004,7 +3004,7 @@ namespace UnrealGameSync
 			}
 			else
 			{
-				PossiblePaths.Add(Path.Combine("Engine", "Binaries", "Win64", String.Format("UE4Editor{0}.target", ConfigSuffix)));
+				PossiblePaths.Add(Path.Combine(Path.GetDirectoryName(SelectedFileName), "Engine", "Binaries", "Win64", String.Format("UE4Editor{0}.target", ConfigSuffix)));
 			}
 			return PossiblePaths;
 		}
@@ -3012,7 +3012,7 @@ namespace UnrealGameSync
 		private void TimerCallback(object Sender, EventArgs Args)
 		{
 			Tuple<string, float> Progress = Workspace.CurrentProgress;
-			if(Progress != null && Progress.Item2 > 0.0f)
+			if (Progress != null && Progress.Item2 > 0.0f)
 			{
 				DesiredTaskbarState = Tuple.Create(TaskbarState.Normal, Progress.Item2);
 			}
@@ -3030,7 +3030,7 @@ namespace UnrealGameSync
 		private void OpenPerforce()
 		{
 			StringBuilder CommandLine = new StringBuilder();
-			if(Workspace != null && Workspace.Perforce != null)
+			if (Workspace != null && Workspace.Perforce != null)
 			{
 				CommandLine.AppendFormat("-p \"{0}\" -c \"{1}\" -u \"{2}\"", Workspace.Perforce.ServerAndPort ?? "perforce:1666", Workspace.Perforce.ClientName, Workspace.Perforce.UserName);
 			}
@@ -3043,13 +3043,13 @@ namespace UnrealGameSync
 			StatusPanel.SetContentWidth(NewContentWidth);
 
 			List<StatusLine> Lines = new List<StatusLine>();
-			if(Workspace.IsBusy())
+			if (Workspace.IsBusy())
 			{
 				// Sync in progress
 				Tuple<string, float> Progress = Workspace.CurrentProgress;
 
 				StatusLine SummaryLine = new StatusLine();
-				if(Workspace.PendingChangeNumber == -1)
+				if (Workspace.PendingChangeNumber == -1)
 				{
 					SummaryLine.AddText("Working... | ");
 				}
@@ -3059,14 +3059,14 @@ namespace UnrealGameSync
 					SummaryLine.AddLink(Workspace.PendingChangeNumber.ToString(), FontStyle.Regular, () => { SelectChange(Workspace.PendingChangeNumber); });
 					SummaryLine.AddText("... | ");
 				}
-				SummaryLine.AddLink(Splitter.IsLogVisible()? "Hide Log" : "Show Log", FontStyle.Bold | FontStyle.Underline, () => { ToggleLogVisibility(); });
+				SummaryLine.AddLink(Splitter.IsLogVisible() ? "Hide Log" : "Show Log", FontStyle.Bold | FontStyle.Underline, () => { ToggleLogVisibility(); });
 				SummaryLine.AddText(" | ");
 				SummaryLine.AddLink("Cancel", FontStyle.Bold | FontStyle.Underline, () => { CancelWorkspaceUpdate(); });
 				Lines.Add(SummaryLine);
 
 				StatusLine ProgressLine = new StatusLine();
 				ProgressLine.AddText(String.Format("{0}  ", Progress.Item1));
-				if(Progress.Item2 > 0.0f) ProgressLine.AddProgressBar(Progress.Item2);
+				if (Progress.Item2 > 0.0f) ProgressLine.AddProgressBar(Progress.Item2);
 				Lines.Add(ProgressLine);
 			}
 			else
@@ -3080,16 +3080,16 @@ namespace UnrealGameSync
 				Lines.Add(ProjectLine);
 
 				// Spacer
-				Lines.Add(new StatusLine(){ LineHeight = 0.5f });
+				Lines.Add(new StatusLine() { LineHeight = 0.5f });
 
 				// Sync status
 				StatusLine SummaryLine = new StatusLine();
-				if(IsBisectModeEnabled())
+				if (IsBisectModeEnabled())
 				{
 					int PassChangeNumber;
 					int FailChangeNumber;
 					GetRemainingBisectRange(out PassChangeNumber, out FailChangeNumber);
-					
+
 					SummaryLine.AddText("Bisecting changes between ");
 					SummaryLine.AddLink(String.Format("{0}", PassChangeNumber), FontStyle.Regular, () => { SelectChange(PassChangeNumber); });
 					SummaryLine.AddText(" and ");
@@ -3097,7 +3097,7 @@ namespace UnrealGameSync
 					SummaryLine.AddText(".  ");
 
 					int BisectChangeNumber = GetBisectChangeNumber();
-					if(BisectChangeNumber != -1)
+					if (BisectChangeNumber != -1)
 					{
 						SummaryLine.AddLink("Sync Next", FontStyle.Bold | FontStyle.Underline, () => { SyncBisectChange(); });
 						SummaryLine.AddText(" | ");
@@ -3107,16 +3107,16 @@ namespace UnrealGameSync
 				}
 				else
 				{
-					if(Workspace.CurrentChangeNumber != -1)
+					if (Workspace.CurrentChangeNumber != -1)
 					{
-						if(StreamName == null)
+						if (StreamName == null)
 						{
 							SummaryLine.AddText("Last synced to changelist ");
 						}
 						else
 						{
 							SummaryLine.AddText("Last synced to ");
-							if(Workspace.ProjectStreamFilter == null || Workspace.ProjectStreamFilter.Count == 0)
+							if (Workspace.ProjectStreamFilter == null || Workspace.ProjectStreamFilter.Count == 0)
 							{
 								SummaryLine.AddLink(StreamName, FontStyle.Regular, (P, R) => { SelectOtherStreamDialog(); });
 							}
@@ -3131,7 +3131,7 @@ namespace UnrealGameSync
 					else
 					{
 						SummaryLine.AddText("You are not currently synced to ");
-						if(StreamName == null)
+						if (StreamName == null)
 						{
 							SummaryLine.AddText("this branch.");
 						}
@@ -3149,14 +3149,14 @@ namespace UnrealGameSync
 
 				// Programs
 				StatusLine ProgramsLine = new StatusLine();
-				if(Workspace.CurrentChangeNumber != -1)
+				if (Workspace.CurrentChangeNumber != -1)
 				{
 					ProgramsLine.AddLink("Unreal Editor", FontStyle.Regular, () => { LaunchEditor(); });
 					ProgramsLine.AddText("  |  ");
 				}
 
 				string[] SdkInfoEntries;
-				if(TryGetProjectSetting(PerforceMonitor.LatestProjectConfigFile, "SdkInfo", out SdkInfoEntries))
+				if (TryGetProjectSetting(PerforceMonitor.LatestProjectConfigFile, "SdkInfo", out SdkInfoEntries))
 				{
 					ProgramsLine.AddLink("SDK Info", FontStyle.Regular, () => { ShowRequiredSdkInfo(); });
 					ProgramsLine.AddText("  |  ");
@@ -3169,10 +3169,10 @@ namespace UnrealGameSync
 				ProgramsLine.AddLink("Windows Explorer", FontStyle.Regular, () => { SafeProcessStart("explorer.exe", String.Format("\"{0}\"", Path.GetDirectoryName(SelectedFileName))); });
 
 				string BuildHealthProject;
-				if(TryGetProjectSetting(PerforceMonitor.LatestProjectConfigFile, "BuildHealthProject", out BuildHealthProject))
+				if (TryGetProjectSetting(PerforceMonitor.LatestProjectConfigFile, "BuildHealthProject", out BuildHealthProject))
 				{
 					ProgramsLine.AddText("  |  ");
-					if(bUserHasOpenIssues)
+					if (bUserHasOpenIssues)
 					{
 						ProgramsLine.AddBadge("Build Health", GetBuildBadgeColor(BadgeResult.Failure), (P, R) => ShowBuildHealthMenu(R, BuildHealthProject));
 					}
@@ -3182,10 +3182,10 @@ namespace UnrealGameSync
 					}
 				}
 
-				foreach(ServiceBadgeInfo ServiceBadge in ServiceBadges)
+				foreach (ServiceBadgeInfo ServiceBadge in ServiceBadges)
 				{
 					ProgramsLine.AddText("  |  ");
-					if(ServiceBadge.Url == null)
+					if (ServiceBadge.Url == null)
 					{
 						ProgramsLine.AddBadge(ServiceBadge.Name, GetBuildBadgeColor(ServiceBadge.Result), null);
 					}
@@ -3199,24 +3199,24 @@ namespace UnrealGameSync
 				Lines.Add(ProgramsLine);
 
 				// Get the summary of the last sync
-				if(WorkspaceSettings.LastSyncChangeNumber > 0)
+				if (WorkspaceSettings.LastSyncChangeNumber > 0)
 				{
 					string SummaryText;
-					if(WorkspaceSettings.LastSyncChangeNumber == Workspace.CurrentChangeNumber && WorkspaceSettings.LastSyncResult == WorkspaceUpdateResult.Success && WorkspaceSettings.LastSyncTime.HasValue)
+					if (WorkspaceSettings.LastSyncChangeNumber == Workspace.CurrentChangeNumber && WorkspaceSettings.LastSyncResult == WorkspaceUpdateResult.Success && WorkspaceSettings.LastSyncTime.HasValue)
 					{
-						Lines.Add(new StatusLine(){ LineHeight = 0.5f });
+						Lines.Add(new StatusLine() { LineHeight = 0.5f });
 
 						StatusLine SuccessLine = new StatusLine();
 						SuccessLine.AddIcon(Properties.Resources.StatusIcons, new Size(16, 16), 0);
-						SuccessLine.AddText(String.Format("  Sync took {0}{1}s, completed at {2}.", (WorkspaceSettings.LastSyncDurationSeconds >= 60)? String.Format("{0}m ", WorkspaceSettings.LastSyncDurationSeconds / 60) : "", WorkspaceSettings.LastSyncDurationSeconds % 60, WorkspaceSettings.LastSyncTime.Value.ToLocalTime().ToString("h\\:mmtt").ToLowerInvariant()));
+						SuccessLine.AddText(String.Format("  Sync took {0}{1}s, completed at {2}.", (WorkspaceSettings.LastSyncDurationSeconds >= 60) ? String.Format("{0}m ", WorkspaceSettings.LastSyncDurationSeconds / 60) : "", WorkspaceSettings.LastSyncDurationSeconds % 60, WorkspaceSettings.LastSyncTime.Value.ToLocalTime().ToString("h\\:mmtt").ToLowerInvariant()));
 						Lines.Add(SuccessLine);
 					}
-					else if(GetLastUpdateMessage(WorkspaceSettings.LastSyncResult, WorkspaceSettings.LastSyncResultMessage, out SummaryText))
+					else if (GetLastUpdateMessage(WorkspaceSettings.LastSyncResult, WorkspaceSettings.LastSyncResultMessage, out SummaryText))
 					{
-						Lines.Add(new StatusLine(){ LineHeight = 0.5f });
+						Lines.Add(new StatusLine() { LineHeight = 0.5f });
 
 						int SummaryTextLength = SummaryText.IndexOf('\n');
-						if(SummaryTextLength == -1)
+						if (SummaryTextLength == -1)
 						{
 							SummaryTextLength = SummaryText.Length;
 						}
@@ -3225,11 +3225,11 @@ namespace UnrealGameSync
 						StatusLine FailLine = new StatusLine();
 						FailLine.AddIcon(Properties.Resources.StatusIcons, new Size(16, 16), 1);
 
-						if(SummaryTextLength == SummaryText.Length)
+						if (SummaryTextLength == SummaryText.Length)
 						{
 							FailLine.AddText(String.Format("  {0}  ", SummaryText));
 						}
-						else 
+						else
 						{
 							FailLine.AddText(String.Format("  {0}...  ", SummaryText.Substring(0, SummaryTextLength).TrimEnd()));
 							FailLine.AddLink("More...", FontStyle.Bold | FontStyle.Underline, () => { ViewLastSyncStatus(); });
@@ -3242,7 +3242,7 @@ namespace UnrealGameSync
 			}
 
 			StatusLine Caption = null;
-			if(StreamName != null && !Workspace.IsBusy())
+			if (StreamName != null && !Workspace.IsBusy())
 			{
 				Caption = new StatusLine();
 				Caption.AddLink(StreamName + "\u25BE", FontStyle.Bold, (P, R) => { SelectOtherStream(R); });
@@ -3252,29 +3252,29 @@ namespace UnrealGameSync
 			Color? TintColor = null;
 
 			ConfigFile ProjectConfigFile = PerforceMonitor.LatestProjectConfigFile;
-			if(ProjectConfigFile != null)
+			if (ProjectConfigFile != null)
 			{
 				string Message;
-				if(TryGetProjectSetting(ProjectConfigFile, "Message", out Message))
+				if (TryGetProjectSetting(ProjectConfigFile, "Message", out Message))
 				{
 					Alert = CreateStatusLineFromMarkdown(Message);
 				}
 
 				string StatusPanelColor;
-				if(TryGetProjectSetting(ProjectConfigFile, "StatusPanelColor", out StatusPanelColor))
+				if (TryGetProjectSetting(ProjectConfigFile, "StatusPanelColor", out StatusPanelColor))
 				{
 					TintColor = System.Drawing.ColorTranslator.FromHtml(StatusPanelColor);
 				}
 			}
 
-			using(Graphics Graphics = CreateGraphics())
+			using (Graphics Graphics = CreateGraphics())
 			{
 				int StatusPanelHeight = (int)(148.0f * Graphics.DpiY / 96.0f);
-				if(Alert != null)
+				if (Alert != null)
 				{
 					StatusPanelHeight += (int)(40.0f * Graphics.DpiY / 96.0f);
 				}
-				if(StatusLayoutPanel.RowStyles[0].Height != StatusPanelHeight)
+				if (StatusLayoutPanel.RowStyles[0].Height != StatusPanelHeight)
 				{
 					StatusLayoutPanel.RowStyles[0].Height = StatusPanelHeight;
 				}
@@ -3287,40 +3287,40 @@ namespace UnrealGameSync
 		{
 			int MinSeparatorIdx = BuildHealthContextMenu.Items.IndexOf(BuildHealthContextMenu_MinSeparator);
 
-			while(BuildHealthContextMenu.Items[MinSeparatorIdx + 1] != BuildHealthContextMenu_MaxSeparator)
+			while (BuildHealthContextMenu.Items[MinSeparatorIdx + 1] != BuildHealthContextMenu_MaxSeparator)
 			{
 				BuildHealthContextMenu.Items.RemoveAt(MinSeparatorIdx + 1);
 			}
 
 			List<IssueData> Issues = new List<IssueData>();
 			bool bHasOtherAssignedIssue = false;
-			foreach(IssueData Issue in IssueMonitor.GetIssues())
+			foreach (IssueData Issue in IssueMonitor.GetIssues())
 			{
-				if(Issue.Project == BuildHealthProject)
+				if (Issue.Project == BuildHealthProject)
 				{
 					Issues.Add(Issue);
 				}
-				else if(Issue.FixChange == 0 && Issue.Owner != null && String.Compare(Issue.Owner, Workspace.Perforce.UserName, StringComparison.OrdinalIgnoreCase) == 0)
+				else if (Issue.FixChange == 0 && Issue.Owner != null && String.Compare(Issue.Owner, Workspace.Perforce.UserName, StringComparison.OrdinalIgnoreCase) == 0)
 				{
 					bHasOtherAssignedIssue = true;
 				}
 			}
-			
-			BuildHealthContextMenu_Browse.Font = new Font(BuildHealthContextMenu_Browse.Font, bHasOtherAssignedIssue? FontStyle.Bold : FontStyle.Regular);
 
-			if(Issues.Count == 0)
+			BuildHealthContextMenu_Browse.Font = new Font(BuildHealthContextMenu_Browse.Font, bHasOtherAssignedIssue ? FontStyle.Bold : FontStyle.Regular);
+
+			if (Issues.Count == 0)
 			{
 				BuildHealthContextMenu_MaxSeparator.Visible = false;
 			}
 			else
 			{
 				BuildHealthContextMenu_MaxSeparator.Visible = true;
-				for(int Idx = 0; Idx < Issues.Count; Idx++)
+				for (int Idx = 0; Idx < Issues.Count; Idx++)
 				{
 					IssueData Issue = Issues[Idx];
 
 					string Summary = Issue.Summary;
-					if(Summary.Length > 100)
+					if (Summary.Length > 100)
 					{
 						Summary = Summary.Substring(0, 100).TrimEnd() + "...";
 					}
@@ -3334,7 +3334,7 @@ namespace UnrealGameSync
 					else
 					{
 						Description.AppendFormat(" - {0}", FormatUserName(Issue.Owner));
-						if(!Issue.AcknowledgedAt.HasValue)
+						if (!Issue.AcknowledgedAt.HasValue)
 						{
 							Description.Append(" (?)");
 						}
@@ -3347,11 +3347,11 @@ namespace UnrealGameSync
 					Description.AppendFormat(" ({0})", Utility.FormatDurationMinutes((int)((Issue.RetrievedAt - Issue.CreatedAt).TotalMinutes + 1)));
 
 					ToolStripMenuItem Item = new ToolStripMenuItem(Description.ToString());
-					if(Issue.FixChange > 0)
+					if (Issue.FixChange > 0)
 					{
 						Item.ForeColor = SystemColors.GrayText;
 					}
-					else if(Issue.Owner != null && String.Compare(Issue.Owner, Workspace.Perforce.UserName, StringComparison.OrdinalIgnoreCase) == 0)
+					else if (Issue.Owner != null && String.Compare(Issue.Owner, Workspace.Perforce.UserName, StringComparison.OrdinalIgnoreCase) == 0)
 					{
 						Item.Font = new Font(Item.Font, FontStyle.Bold);
 					}
@@ -3393,7 +3393,7 @@ namespace UnrealGameSync
 		private void BuildHealthContextMenu_Settings_Click(object sender, EventArgs e)
 		{
 			IssueSettingsWindow IssueSettings = new IssueSettingsWindow(Settings);
-			if(IssueSettings.ShowDialog(this) == DialogResult.OK)
+			if (IssueSettings.ShowDialog(this) == DialogResult.OK)
 			{
 				Owner.UpdateAlertWindows();
 			}
@@ -3402,7 +3402,7 @@ namespace UnrealGameSync
 		private void ShowRequiredSdkInfo()
 		{
 			string[] SdkInfoEntries;
-			if(TryGetProjectSetting(PerforceMonitor.LatestProjectConfigFile, "SdkInfo", out SdkInfoEntries))
+			if (TryGetProjectSetting(PerforceMonitor.LatestProjectConfigFile, "SdkInfo", out SdkInfoEntries))
 			{
 				Dictionary<string, string> Variables = GetWorkspaceVariables(-1);
 				SdkInfoWindow Window = new SdkInfoWindow(SdkInfoEntries, Variables, BadgeFont);
@@ -3417,12 +3417,12 @@ namespace UnrealGameSync
 			FontStyle Style = FontStyle.Regular;
 
 			StringBuilder ElementText = new StringBuilder();
-			for(int Idx = 0; Idx < Text.Length; )
+			for (int Idx = 0; Idx < Text.Length;)
 			{
 				// Bold and italic
-				if(Text[Idx] == '_' || Text[Idx] == '*')
+				if (Text[Idx] == '_' || Text[Idx] == '*')
 				{
-					if(Idx + 1 < Text.Length && Text[Idx + 1] == Text[Idx])
+					if (Idx + 1 < Text.Length && Text[Idx + 1] == Text[Idx])
 					{
 						FlushMarkdownText(Line, ElementText, Style);
 						Idx += 2;
@@ -3439,7 +3439,7 @@ namespace UnrealGameSync
 				}
 
 				// Strikethrough
-				if(Idx + 2 < Text.Length && Text[Idx] == '~' && Text[Idx + 1] == '~')
+				if (Idx + 2 < Text.Length && Text[Idx] == '~' && Text[Idx + 1] == '~')
 				{
 					FlushMarkdownText(Line, ElementText, Style);
 					Idx += 2;
@@ -3450,7 +3450,7 @@ namespace UnrealGameSync
 				// Link
 				if (Text[Idx] == '[')
 				{
-					if(Idx + 1 < Text.Length && Text[Idx + 1] == '[')
+					if (Idx + 1 < Text.Length && Text[Idx + 1] == '[')
 					{
 						ElementText.Append(Text[Idx]);
 						Idx += 2;
@@ -3458,10 +3458,10 @@ namespace UnrealGameSync
 					}
 
 					int EndIdx = Text.IndexOf("](", Idx);
-					if(EndIdx != -1)
+					if (EndIdx != -1)
 					{
 						int UrlEndIdx = Text.IndexOf(')', EndIdx + 2);
-						if(UrlEndIdx != -1)
+						if (UrlEndIdx != -1)
 						{
 							FlushMarkdownText(Line, ElementText, Style);
 							string LinkText = Text.Substring(Idx + 1, EndIdx - (Idx + 1));
@@ -3474,12 +3474,12 @@ namespace UnrealGameSync
 				}
 
 				// Icon
-				if(Text[Idx] == ':')
+				if (Text[Idx] == ':')
 				{
 					int EndIdx = Text.IndexOf(':', Idx + 1);
-					if(EndIdx != -1)
+					if (EndIdx != -1)
 					{
-						if(String.Compare(":alert:", 0, Text, Idx, EndIdx - Idx) == 0)
+						if (String.Compare(":alert:", 0, Text, Idx, EndIdx - Idx) == 0)
 						{
 							FlushMarkdownText(Line, ElementText, Style);
 							Line.AddIcon(Properties.Resources.StatusIcons, new Size(16, 16), 4);
@@ -3499,7 +3499,7 @@ namespace UnrealGameSync
 
 		private void FlushMarkdownText(StatusLine Line, StringBuilder ElementText, FontStyle Style)
 		{
-			if(ElementText.Length > 0)
+			if (ElementText.Length > 0)
 			{
 				Line.AddText(ElementText.ToString(), Style);
 				ElementText.Clear();
@@ -3509,30 +3509,30 @@ namespace UnrealGameSync
 		private void SelectOtherStream(Rectangle Bounds)
 		{
 			bool bShownContextMenu = false;
-			if(StreamName != null)
+			if (StreamName != null)
 			{
 				IReadOnlyList<string> OtherStreamNames = Workspace.ProjectStreamFilter;
-				if(OtherStreamNames != null)
+				if (OtherStreamNames != null)
 				{
 					StreamContextMenu.Items.Clear();
 
 					ToolStripMenuItem CurrentStreamItem = new ToolStripMenuItem(StreamName, null, new EventHandler((S, E) => SelectStream(StreamName)));
 					CurrentStreamItem.Checked = true;
 					StreamContextMenu.Items.Add(CurrentStreamItem);
-					
+
 					StreamContextMenu.Items.Add(new ToolStripSeparator());
 
 					foreach (string OtherStreamName in OtherStreamNames.OrderBy(x => x).Where(x => !x.EndsWith("/Dev-Binaries")))
 					{
 						string ThisStreamName = OtherStreamName; // Local for lambda capture
-						if(String.Compare(StreamName, OtherStreamName, StringComparison.InvariantCultureIgnoreCase) != 0)
+						if (String.Compare(StreamName, OtherStreamName, StringComparison.InvariantCultureIgnoreCase) != 0)
 						{
 							ToolStripMenuItem Item = new ToolStripMenuItem(ThisStreamName, null, new EventHandler((S, E) => SelectStream(ThisStreamName)));
 							StreamContextMenu.Items.Add(Item);
 						}
 					}
 
-					if(StreamContextMenu.Items.Count > 2)
+					if (StreamContextMenu.Items.Count > 2)
 					{
 						StreamContextMenu.Items.Add(new ToolStripSeparator());
 					}
@@ -3546,7 +3546,7 @@ namespace UnrealGameSync
 					bShownContextMenu = true;
 				}
 			}
-			if(!bShownContextMenu)
+			if (!bShownContextMenu)
 			{
 				SelectOtherStreamDialog();
 			}
@@ -3555,7 +3555,7 @@ namespace UnrealGameSync
 		private void SelectOtherStreamDialog()
 		{
 			string NewStreamName;
-			if(SelectStreamWindow.ShowModal(this, Workspace.Perforce, StreamName, Log, out NewStreamName))
+			if (SelectStreamWindow.ShowModal(this, Workspace.Perforce, StreamName, Log, out NewStreamName))
 			{
 				SelectStream(NewStreamName);
 			}
@@ -3563,17 +3563,17 @@ namespace UnrealGameSync
 
 		private void SelectStream(string NewStreamName)
 		{
-			if(StreamName != NewStreamName)
+			if (StreamName != NewStreamName)
 			{
-				if(Workspace.IsBusy())
+				if (Workspace.IsBusy())
 				{
 					MessageBox.Show("Please retry after the current sync has finished.", "Sync in Progress");
 				}
-				else if(Workspace.Perforce != null)
+				else if (Workspace.Perforce != null)
 				{
-					if(!Workspace.Perforce.HasOpenFiles(Log) || MessageBox.Show("You have files open for edit in this workspace. If you continue, you will not be able to submit them until you switch back.\n\nContinue switching streams?", "Files checked out", MessageBoxButtons.YesNo) == DialogResult.Yes)
+					if (!Workspace.Perforce.HasOpenFiles(Log) || MessageBox.Show("You have files open for edit in this workspace. If you continue, you will not be able to submit them until you switch back.\n\nContinue switching streams?", "Files checked out", MessageBoxButtons.YesNo) == DialogResult.Yes)
 					{
-						if(Workspace.Perforce.SwitchStream(NewStreamName, Log))
+						if (Workspace.Perforce.SwitchStream(NewStreamName, Log))
 						{
 							StatusPanel.SuspendLayout();
 							StreamChanged();
@@ -3587,10 +3587,10 @@ namespace UnrealGameSync
 		private void ViewLastSyncStatus()
 		{
 			string SummaryText;
-			if(GetLastUpdateMessage(WorkspaceSettings.LastSyncResult, WorkspaceSettings.LastSyncResultMessage, out SummaryText))
+			if (GetLastUpdateMessage(WorkspaceSettings.LastSyncResult, WorkspaceSettings.LastSyncResultMessage, out SummaryText))
 			{
 				string CaptionText;
-				if(!GetGenericLastUpdateMessage(WorkspaceSettings.LastSyncResult, out CaptionText))
+				if (!GetGenericLastUpdateMessage(WorkspaceSettings.LastSyncResult, out CaptionText))
 				{
 					CaptionText = "Sync error";
 				}
@@ -3600,7 +3600,7 @@ namespace UnrealGameSync
 
 		private void ShowErrorInLog()
 		{
-			if(!Splitter.IsLogVisible())
+			if (!Splitter.IsLogVisible())
 			{
 				ToggleLogVisibility();
 			}
@@ -3618,10 +3618,10 @@ namespace UnrealGameSync
 
 			PendingSelectedChangeNumber = -1;
 
-			foreach(ListViewItem Item in BuildList.Items)
+			foreach (ListViewItem Item in BuildList.Items)
 			{
 				PerforceChangeSummary Summary = (PerforceChangeSummary)Item.Tag;
-				if(Summary != null && Summary.Number <= ChangeNumber)
+				if (Summary != null && Summary.Number <= ChangeNumber)
 				{
 					Item.Selected = true;
 					Item.EnsureVisible();
@@ -3632,10 +3632,10 @@ namespace UnrealGameSync
 			PendingSelectedChangeNumber = ChangeNumber;
 
 			int CurrentMaxChanges = PerforceMonitor.CurrentMaxChanges;
-			if(PerforceMonitor.PendingMaxChanges <= CurrentMaxChanges && BuildList.SelectedItems.Count == 0)
+			if (PerforceMonitor.PendingMaxChanges <= CurrentMaxChanges && BuildList.SelectedItems.Count == 0)
 			{
 				PerforceMonitor.PendingMaxChanges = CurrentMaxChanges + BuildListExpandCount;
-				if(ExpandItem != null)
+				if (ExpandItem != null)
 				{
 					ExpandItem.EnsureVisible();
 				}
@@ -3644,17 +3644,17 @@ namespace UnrealGameSync
 
 		private void BuildList_MouseClick(object Sender, MouseEventArgs Args)
 		{
-			if(Args.Button == MouseButtons.Left)
+			if (Args.Button == MouseButtons.Left)
 			{
 				ListViewHitTestInfo HitTest = BuildList.HitTest(Args.Location);
-				if(HitTest.Item != null)
+				if (HitTest.Item != null)
 				{
-					if(HitTest.Item == ExpandItem)
+					if (HitTest.Item == ExpandItem)
 					{
-						if(HitTestExpandLink(Args.Location))
+						if (HitTestExpandLink(Args.Location))
 						{
 							int CurrentMaxChanges = PerforceMonitor.CurrentMaxChanges;
-							if(PerforceMonitor.PendingMaxChanges > CurrentMaxChanges)
+							if (PerforceMonitor.PendingMaxChanges > CurrentMaxChanges)
 							{
 								PerforceMonitor.PendingMaxChanges = CurrentMaxChanges;
 							}
@@ -3668,19 +3668,19 @@ namespace UnrealGameSync
 					else
 					{
 						PerforceChangeSummary Change = (PerforceChangeSummary)HitTest.Item.Tag;
-						if(Workspace.PendingChangeNumber == Change.Number)
+						if (Workspace.PendingChangeNumber == Change.Number)
 						{
 							Rectangle SubItemRect = HitTest.Item.SubItems[StatusColumn.Index].Bounds;
 
-							if(Workspace.IsBusy())
+							if (Workspace.IsBusy())
 							{
 								Rectangle CancelRect = new Rectangle(SubItemRect.Right - 16, SubItemRect.Top, 16, SubItemRect.Height);
 								Rectangle InfoRect = new Rectangle(SubItemRect.Right - 32, SubItemRect.Top, 16, SubItemRect.Height);
-								if(CancelRect.Contains(Args.Location))
+								if (CancelRect.Contains(Args.Location))
 								{
 									CancelWorkspaceUpdate();
 								}
-								else if(InfoRect.Contains(Args.Location) && !Splitter.IsLogVisible())
+								else if (InfoRect.Contains(Args.Location) && !Splitter.IsLogVisible())
 								{
 									ToggleLogVisibility();
 								}
@@ -3689,12 +3689,12 @@ namespace UnrealGameSync
 							{
 								Rectangle HappyRect = new Rectangle(SubItemRect.Right - 32, SubItemRect.Top, 16, SubItemRect.Height);
 								Rectangle FrownRect = new Rectangle(SubItemRect.Right - 16, SubItemRect.Top, 16, SubItemRect.Height);
-								if(HappyRect.Contains(Args.Location))
+								if (HappyRect.Contains(Args.Location))
 								{
 									EventMonitor.PostEvent(Change.Number, EventType.Good);
 									BuildList.Invalidate();
 								}
-								else if(FrownRect.Contains(Args.Location))
+								else if (FrownRect.Contains(Args.Location))
 								{
 									EventMonitor.PostEvent(Change.Number, EventType.Bad);
 									BuildList.Invalidate();
@@ -3704,16 +3704,16 @@ namespace UnrealGameSync
 						else
 						{
 							Rectangle SyncBadgeRectangle = GetSyncBadgeRectangle(HitTest.Item.SubItems[StatusColumn.Index].Bounds);
-							if(SyncBadgeRectangle.Contains(Args.Location) && CanSyncChange(Change.Number))
+							if (SyncBadgeRectangle.Contains(Args.Location) && CanSyncChange(Change.Number))
 							{
 								StartSync(Change.Number, null);
 							}
 						}
 
-						if(DescriptionColumn.Index < HitTest.Item.SubItems.Count && HitTest.Item.SubItems[DescriptionColumn.Index] == HitTest.SubItem)
+						if (DescriptionColumn.Index < HitTest.Item.SubItems.Count && HitTest.Item.SubItems[DescriptionColumn.Index] == HitTest.SubItem)
 						{
 							ChangeLayoutInfo LayoutInfo = GetChangeLayoutInfo(Change);
-							if(LayoutInfo.DescriptionBadges.Count > 0)
+							if (LayoutInfo.DescriptionBadges.Count > 0)
 							{
 								Point BuildListLocation = GetBadgeListLocation(LayoutInfo.DescriptionBadges, HitTest.SubItem.Bounds, HorizontalAlign.Right, VerticalAlignment.Middle);
 								BuildListLocation.Offset(-2, 0);
@@ -3721,7 +3721,7 @@ namespace UnrealGameSync
 								foreach (BadgeInfo Badge in LayoutInfo.DescriptionBadges)
 								{
 									Rectangle BadgeBounds = Badge.GetBounds(BuildListLocation);
-									if(BadgeBounds.Contains(Args.Location) && Badge.ClickHandler != null)
+									if (BadgeBounds.Contains(Args.Location) && Badge.ClickHandler != null)
 									{
 										Badge.ClickHandler();
 										break;
@@ -3730,35 +3730,35 @@ namespace UnrealGameSync
 							}
 						}
 
-						if(CISColumn.Index < HitTest.Item.SubItems.Count && HitTest.Item.SubItems[CISColumn.Index] == HitTest.SubItem)
+						if (CISColumn.Index < HitTest.Item.SubItems.Count && HitTest.Item.SubItems[CISColumn.Index] == HitTest.SubItem)
 						{
 							ChangeLayoutInfo LayoutInfo = GetChangeLayoutInfo(Change);
-							if(LayoutInfo.BuildBadges.Count > 0)
+							if (LayoutInfo.BuildBadges.Count > 0)
 							{
 								Point BuildListLocation = GetBadgeListLocation(LayoutInfo.BuildBadges, HitTest.SubItem.Bounds, HorizontalAlign.Center, VerticalAlignment.Middle);
 								BuildListLocation.X = Math.Max(BuildListLocation.X, HitTest.SubItem.Bounds.Left);
 
 								BadgeInfo BadgeInfo = HitTestBadge(Args.Location, LayoutInfo.BuildBadges, BuildListLocation);
-								if(BadgeInfo != null && BadgeInfo.ClickHandler != null)
+								if (BadgeInfo != null && BadgeInfo.ClickHandler != null)
 								{
 									BadgeInfo.ClickHandler();
 								}
 							}
 						}
 
-						foreach(ColumnHeader CustomColumn in CustomColumns)
+						foreach (ColumnHeader CustomColumn in CustomColumns)
 						{
-							if(CustomColumn.Index < HitTest.Item.SubItems.Count && HitTest.Item.SubItems[CustomColumn.Index] == HitTest.SubItem)
+							if (CustomColumn.Index < HitTest.Item.SubItems.Count && HitTest.Item.SubItems[CustomColumn.Index] == HitTest.SubItem)
 							{
 								ChangeLayoutInfo LayoutInfo = GetChangeLayoutInfo((PerforceChangeSummary)HitTest.Item.Tag);
 
 								List<BadgeInfo> Badges;
-								if(LayoutInfo.CustomBadges.TryGetValue(CustomColumn.Text, out Badges) && Badges.Count > 0)
+								if (LayoutInfo.CustomBadges.TryGetValue(CustomColumn.Text, out Badges) && Badges.Count > 0)
 								{
 									Point ListLocation = GetBadgeListLocation(Badges, HitTest.SubItem.Bounds, HorizontalAlign.Center, VerticalAlignment.Middle);
 
 									BadgeInfo BadgeInfo = HitTestBadge(Args.Location, Badges, ListLocation);
-									if(BadgeInfo != null && BadgeInfo.ClickHandler != null)
+									if (BadgeInfo != null && BadgeInfo.ClickHandler != null)
 									{
 										BadgeInfo.ClickHandler();
 									}
@@ -3768,12 +3768,12 @@ namespace UnrealGameSync
 					}
 				}
 			}
-			else if(Args.Button == MouseButtons.Right)
+			else if (Args.Button == MouseButtons.Right)
 			{
 				ListViewHitTestInfo HitTest = BuildList.HitTest(Args.Location);
-				if(HitTest.Item != null && HitTest.Item.Tag != null)
+				if (HitTest.Item != null && HitTest.Item.Tag != null)
 				{
-					if(BuildList.SelectedItems.Count > 1 && BuildList.SelectedItems.Contains(HitTest.Item))
+					if (BuildList.SelectedItems.Count > 1 && BuildList.SelectedItems.Contains(HitTest.Item))
 					{
 						bool bIsTimeColumn = (HitTest.Item.SubItems.IndexOf(HitTest.SubItem) == TimeColumn.Index);
 						BuildListMultiContextMenu_TimeZoneSeparator.Visible = bIsTimeColumn;
@@ -3800,7 +3800,7 @@ namespace UnrealGameSync
 						bool bIsBusy = Workspace.IsBusy();
 						bool bIsCurrentChange = (ContextMenuChange.Number == Workspace.CurrentChangeNumber);
 						BuildListContextMenu_Sync.Visible = !bIsBusy;
-						BuildListContextMenu_Sync.Font = new Font(SystemFonts.MenuFont, bIsCurrentChange? FontStyle.Regular : FontStyle.Bold);
+						BuildListContextMenu_Sync.Font = new Font(SystemFonts.MenuFont, bIsCurrentChange ? FontStyle.Regular : FontStyle.Bold);
 						BuildListContextMenu_SyncContentOnly.Visible = !bIsBusy && ShouldSyncPrecompiledEditor;
 						BuildListContextMenu_SyncOnlyThisChange.Visible = !bIsBusy && !bIsCurrentChange && ContextMenuChange.Number > Workspace.CurrentChangeNumber && Workspace.CurrentChangeNumber != -1;
 						BuildListContextMenu_Build.Visible = !bIsBusy && bIsCurrentChange && !ShouldSyncPrecompiledEditor && Settings.bUseIncrementalBuilds;
@@ -3811,7 +3811,7 @@ namespace UnrealGameSync
 						BuildListContextMenu_OpenVisualStudio.Visible = !bIsBusy && bIsCurrentChange;
 						BuildListContextMenu_Cancel.Visible = bIsBusy;
 
-						BisectState State; 
+						BisectState State;
 						WorkspaceSettings.ChangeNumberToBisectState.TryGetValue(ContextMenuChange.Number, out State);
 						bool bIsBisectMode = IsBisectModeEnabled();
 						BuildListContextMenu_Bisect_Pass.Visible = bIsBisectMode && State != BisectState.Pass;
@@ -3838,19 +3838,19 @@ namespace UnrealGameSync
 
 						int CustomToolStart = BuildListContextMenu.Items.IndexOf(BuildListContextMenu_CustomTool_Start) + 1;
 						int CustomToolEnd = BuildListContextMenu.Items.IndexOf(BuildListContextMenu_CustomTool_End);
-						while(CustomToolEnd > CustomToolStart)
+						while (CustomToolEnd > CustomToolStart)
 						{
 							BuildListContextMenu.Items.RemoveAt(CustomToolEnd - 1);
 							CustomToolEnd--;
 						}
 
 						ConfigFile ProjectConfigFile = PerforceMonitor.LatestProjectConfigFile;
-						if(ProjectConfigFile != null)
+						if (ProjectConfigFile != null)
 						{
 							Dictionary<string, string> Variables = GetWorkspaceVariables(ContextMenuChange.Number);
 
 							string[] ChangeContextMenuEntries = ProjectConfigFile.GetValues("Options.ContextMenu", new string[0]);
-							foreach(string ChangeContextMenuEntry in ChangeContextMenuEntries)
+							foreach (string ChangeContextMenuEntry in ChangeContextMenuEntries)
 							{
 								ConfigObject Object = new ConfigObject(ChangeContextMenuEntry);
 
@@ -3858,14 +3858,14 @@ namespace UnrealGameSync
 								string Execute = Object.GetValue("Execute");
 								string Arguments = Object.GetValue("Arguments");
 
-								if(Label != null && Execute != null)
+								if (Label != null && Execute != null)
 								{
 									Label = Utility.ExpandVariables(Label, Variables);
 									Execute = Utility.ExpandVariables(Execute, Variables);
 									Arguments = Utility.ExpandVariables(Arguments ?? "", Variables);
 
 									ToolStripMenuItem Item = new ToolStripMenuItem(Label, null, new EventHandler((o, a) => SafeProcessStart(Execute, Arguments)));
-	
+
 									BuildListContextMenu.Items.Insert(CustomToolEnd, Item);
 									CustomToolEnd++;
 								}
@@ -3885,40 +3885,40 @@ namespace UnrealGameSync
 			ListViewHitTestInfo HitTest = BuildList.HitTest(e.Location);
 
 			bool bNewMouseOverExpandLink = HitTest.Item != null && HitTestExpandLink(e.Location);
-			if(bMouseOverExpandLink != bNewMouseOverExpandLink)
+			if (bMouseOverExpandLink != bNewMouseOverExpandLink)
 			{
 				bMouseOverExpandLink = bNewMouseOverExpandLink;
-				Cursor = bMouseOverExpandLink? NativeCursors.Hand : Cursors.Arrow;
+				Cursor = bMouseOverExpandLink ? NativeCursors.Hand : Cursors.Arrow;
 				BuildList.Invalidate();
 			}
 
 			string NewHoverBadgeUniqueId = null;
-			if(HitTest.Item != null && HitTest.Item.Tag is PerforceChangeSummary)
+			if (HitTest.Item != null && HitTest.Item.Tag is PerforceChangeSummary)
 			{
 				int ColumnIndex = HitTest.Item.SubItems.IndexOf(HitTest.SubItem);
-				if(ColumnIndex == DescriptionColumn.Index)
+				if (ColumnIndex == DescriptionColumn.Index)
 				{
 					ChangeLayoutInfo LayoutInfo = GetChangeLayoutInfo((PerforceChangeSummary)HitTest.Item.Tag);
-					if(LayoutInfo.DescriptionBadges.Count > 0)
+					if (LayoutInfo.DescriptionBadges.Count > 0)
 					{
 						Point ListLocation = GetBadgeListLocation(LayoutInfo.DescriptionBadges, HitTest.SubItem.Bounds, HorizontalAlign.Right, VerticalAlignment.Middle);
 						NewHoverBadgeUniqueId = HitTestBadge(e.Location, LayoutInfo.DescriptionBadges, ListLocation)?.UniqueId;
 					}
 				}
-				else if(ColumnIndex == CISColumn.Index)
+				else if (ColumnIndex == CISColumn.Index)
 				{
 					ChangeLayoutInfo LayoutInfo = GetChangeLayoutInfo((PerforceChangeSummary)HitTest.Item.Tag);
-					if(LayoutInfo.BuildBadges.Count > 0)
+					if (LayoutInfo.BuildBadges.Count > 0)
 					{
 						Point BuildListLocation = GetBadgeListLocation(LayoutInfo.BuildBadges, HitTest.SubItem.Bounds, HorizontalAlign.Center, VerticalAlignment.Middle);
 						BuildListLocation.X = Math.Max(BuildListLocation.X, HitTest.SubItem.Bounds.Left);
 
 						BadgeInfo Badge = HitTestBadge(e.Location, LayoutInfo.BuildBadges, BuildListLocation);
-						NewHoverBadgeUniqueId = (Badge != null)? Badge.UniqueId : null;
+						NewHoverBadgeUniqueId = (Badge != null) ? Badge.UniqueId : null;
 
-						if(HoverBadgeUniqueId != NewHoverBadgeUniqueId)
+						if (HoverBadgeUniqueId != NewHoverBadgeUniqueId)
 						{
-							if(Badge != null && Badge.ToolTip != null && Badge.BackgroundColor.A != 0)
+							if (Badge != null && Badge.ToolTip != null && Badge.BackgroundColor.A != 0)
 							{
 								BuildListToolTip.Show(Badge.ToolTip, BuildList, new Point(BuildListLocation.X + Badge.Offset, HitTest.Item.Bounds.Bottom + 2));
 							}
@@ -3929,32 +3929,32 @@ namespace UnrealGameSync
 						}
 					}
 				}
-				else if(CustomColumns.Contains(BuildList.Columns[ColumnIndex]))
+				else if (CustomColumns.Contains(BuildList.Columns[ColumnIndex]))
 				{
 					ColumnHeader Column = BuildList.Columns[ColumnIndex];
 
 					ChangeLayoutInfo LayoutInfo = GetChangeLayoutInfo((PerforceChangeSummary)HitTest.Item.Tag);
 
 					List<BadgeInfo> Badges;
-					if(LayoutInfo.CustomBadges.TryGetValue(Column.Text, out Badges) && Badges.Count > 0)
+					if (LayoutInfo.CustomBadges.TryGetValue(Column.Text, out Badges) && Badges.Count > 0)
 					{
 						Point ListLocation = GetBadgeListLocation(Badges, HitTest.SubItem.Bounds, HorizontalAlign.Center, VerticalAlignment.Middle);
 						NewHoverBadgeUniqueId = HitTestBadge(e.Location, Badges, ListLocation)?.UniqueId;
 					}
 				}
 			}
-			if(HoverBadgeUniqueId != NewHoverBadgeUniqueId)
+			if (HoverBadgeUniqueId != NewHoverBadgeUniqueId)
 			{
 				HoverBadgeUniqueId = NewHoverBadgeUniqueId;
 				BuildList.Invalidate();
 			}
 
 			bool bNewHoverSync = false;
-			if(HitTest.Item != null)
+			if (HitTest.Item != null)
 			{
 				bNewHoverSync = GetSyncBadgeRectangle(HitTest.Item.SubItems[StatusColumn.Index].Bounds).Contains(e.Location);
 			}
-			if(bNewHoverSync != bHoverSync)
+			if (bNewHoverSync != bHoverSync)
 			{
 				bHoverSync = bNewHoverSync;
 				BuildList.Invalidate();
@@ -3963,10 +3963,10 @@ namespace UnrealGameSync
 
 		private BadgeInfo HitTestBadge(Point Location, List<BadgeInfo> BadgeList, Point ListLocation)
 		{
-			foreach(BadgeInfo Badge in BadgeList)
+			foreach (BadgeInfo Badge in BadgeList)
 			{
 				Rectangle BadgeBounds = Badge.GetBounds(ListLocation);
-				if(BadgeBounds.Contains(Location))
+				if (BadgeBounds.Contains(Location))
 				{
 					return Badge;
 				}
@@ -4012,12 +4012,12 @@ namespace UnrealGameSync
 		private void SetSelectedArchive(IArchiveInfo Archive, bool bSelected)
 		{
 			ArchiveSettings ArchiveSettings = Settings.Archives.FirstOrDefault(x => x.Type == Archive.Type);
-			if(ArchiveSettings == null)
+			if (ArchiveSettings == null)
 			{
 				ArchiveSettings = new ArchiveSettings(bSelected, Archive.Type, new string[] { Archive.Name });
 				Settings.Archives.Add(ArchiveSettings);
 			}
-			else if(bSelected)
+			else if (bSelected)
 			{
 				ArchiveSettings.bEnabled = true;
 				ArchiveSettings.Order.Remove(Archive.Name);
@@ -4033,12 +4033,12 @@ namespace UnrealGameSync
 		}
 
 		private void UpdateSelectedArchives()
-        {
+		{
 			UpdateBuildSteps();
 			UpdateSyncActionCheckboxes();
 
 			BuildList.Invalidate();
-        }
+		}
 
 
 		private void OptionsButton_Click(object sender, EventArgs e)
@@ -4135,7 +4135,7 @@ namespace UnrealGameSync
 			RunAfterSyncCheckBox.Checked = Settings.bRunAfterSync;
 			OpenSolutionAfterSyncCheckBox.Checked = Settings.bOpenSolutionAfterSync;
 
-			if(Workspace == null || Workspace.IsBusy())
+			if (Workspace == null || Workspace.IsBusy())
 			{
 				AfterSyncingLabel.Enabled = false;
 				BuildAfterSyncCheckBox.Enabled = false;
@@ -4147,25 +4147,25 @@ namespace UnrealGameSync
 				AfterSyncingLabel.Enabled = true;
 				BuildAfterSyncCheckBox.Enabled = bHasBuildSteps;
 				RunAfterSyncCheckBox.Enabled = BuildAfterSyncCheckBox.Checked || ShouldSyncPrecompiledEditor;
-				OpenSolutionAfterSyncCheckBox.Enabled = true;
+				OpenSolutionAfterSyncCheckBox.Enabled = !ShouldSyncPrecompiledEditor;
 			}
 		}
 
 		private void BuildList_FontChanged(object sender, EventArgs e)
 		{
-			if(BuildFont != null)
+			if (BuildFont != null)
 			{
 				BuildFont.Dispose();
 			}
 			BuildFont = BuildList.Font;
 
-			if(SelectedBuildFont != null)
+			if (SelectedBuildFont != null)
 			{
 				SelectedBuildFont.Dispose();
 			}
 			SelectedBuildFont = new Font(BuildFont, FontStyle.Bold);
 
-			if(BadgeFont != null)
+			if (BadgeFont != null)
 			{
 				BadgeFont.Dispose();
 			}
@@ -4179,18 +4179,18 @@ namespace UnrealGameSync
 
 		public void SyncLatestChange(WorkspaceUpdateCallback Callback)
 		{
-			if(Workspace != null)
+			if (Workspace != null)
 			{
 				int ChangeNumber;
-				if(!FindChangeToSync(Settings.SyncType, out ChangeNumber))
+				if (!FindChangeToSync(Settings.SyncType, out ChangeNumber))
 				{
-					ShowErrorDialog(String.Format("Couldn't find any {0}changelist. Double-click on the change you want to sync manually.", (Settings.SyncType == LatestChangeType.Starred)? "starred " : (Settings.SyncType == LatestChangeType.Good)? "good " : ""));
+					ShowErrorDialog(String.Format("Couldn't find any {0}changelist. Double-click on the change you want to sync manually.", (Settings.SyncType == LatestChangeType.Starred) ? "starred " : (Settings.SyncType == LatestChangeType.Good) ? "good " : ""));
 				}
-				else if(ChangeNumber < Workspace.CurrentChangeNumber)
+				else if (ChangeNumber < Workspace.CurrentChangeNumber)
 				{
 					ShowErrorDialog("Workspace is already synced to a later change ({0} vs {1}).", Workspace.CurrentChangeNumber, ChangeNumber);
 				}
-				else if(ChangeNumber >= PerforceMonitor.LastChangeByCurrentUser || MessageBox.Show(String.Format("The changelist that would be synced is before the last change you submitted.\n\nIf you continue, your changes submitted after CL {0} will be locally removed from your workspace until you can sync past them again.\n\nAre you sure you want to continue?", ChangeNumber), "Local changes will be removed", MessageBoxButtons.YesNo) == DialogResult.Yes)
+				else if (ChangeNumber >= PerforceMonitor.LastChangeByCurrentUser || MessageBox.Show(String.Format("The changelist that would be synced is before the last change you submitted.\n\nIf you continue, your changes submitted after CL {0} will be locally removed from your workspace until you can sync past them again.\n\nAre you sure you want to continue?", ChangeNumber), "Local changes will be removed", MessageBoxButtons.YesNo) == DialogResult.Yes)
 				{
 					Owner.ShowAndActivate();
 					SelectChange(ChangeNumber);
@@ -4251,7 +4251,7 @@ namespace UnrealGameSync
 
 		private void BuildListContextMenu_MoreInfo_Click(object sender, EventArgs e)
 		{
-			if(!Utility.SpawnHiddenProcess("p4vc.exe", String.Format("-p\"{0}\" -c{1} change {2}", Workspace.Perforce.ServerAndPort ?? "perforce:1666", Workspace.ClientName, ContextMenuChange.Number)))
+			if (!Utility.SpawnHiddenProcess("p4vc.exe", String.Format("-p\"{0}\" -c{1} change {2}", Workspace.Perforce.ServerAndPort ?? "perforce:1666", Workspace.ClientName, ContextMenuChange.Number)))
 			{
 				MessageBox.Show("Unable to spawn p4vc. Check you have P4V installed.");
 			}
@@ -4259,7 +4259,7 @@ namespace UnrealGameSync
 
 		private void BuildListContextMenu_AddStar_Click(object sender, EventArgs e)
 		{
-			if(MessageBox.Show("Starred builds are meant to convey a stable, verified build to the rest of the team. Do not star a build unless it has been fully tested.\n\nAre you sure you want to star this build?", "Confirm star", MessageBoxButtons.YesNo) == DialogResult.Yes)
+			if (MessageBox.Show("Starred builds are meant to convey a stable, verified build to the rest of the team. Do not star a build unless it has been fully tested.\n\nAre you sure you want to star this build?", "Confirm star", MessageBoxButtons.YesNo) == DialogResult.Yes)
 			{
 				EventMonitor.PostEvent(ContextMenuChange.Number, EventType.Starred);
 			}
@@ -4268,10 +4268,10 @@ namespace UnrealGameSync
 		private void BuildListContextMenu_RemoveStar_Click(object sender, EventArgs e)
 		{
 			EventSummary Summary = EventMonitor.GetSummaryForChange(ContextMenuChange.Number);
-			if(Summary != null && Summary.LastStarReview != null && Summary.LastStarReview.Type == EventType.Starred)
+			if (Summary != null && Summary.LastStarReview != null && Summary.LastStarReview.Type == EventType.Starred)
 			{
 				string Message = String.Format("This change was starred by {0}. Are you sure you want to remove it?", FormatUserName(Summary.LastStarReview.UserName));
-				if(MessageBox.Show(Message, "Confirm removing star", MessageBoxButtons.YesNo) == DialogResult.Yes)
+				if (MessageBox.Show(Message, "Confirm removing star", MessageBoxButtons.YesNo) == DialogResult.Yes)
 				{
 					EventMonitor.PostEvent(ContextMenuChange.Number, EventType.Unstarred);
 				}
@@ -4286,7 +4286,7 @@ namespace UnrealGameSync
 		private void BuildListContextMenu_StartInvestigating_Click(object sender, EventArgs e)
 		{
 			string Message = String.Format("All changes from {0} onwards will be marked as bad while you are investigating an issue.\n\nAre you sure you want to continue?", ContextMenuChange.Number);
-			if(MessageBox.Show(Message, "Confirm investigating", MessageBoxButtons.YesNo) == DialogResult.Yes)
+			if (MessageBox.Show(Message, "Confirm investigating", MessageBoxButtons.YesNo) == DialogResult.Yes)
 			{
 				EventMonitor.StartInvestigating(ContextMenuChange.Number);
 			}
@@ -4295,14 +4295,14 @@ namespace UnrealGameSync
 		private void BuildListContextMenu_FinishInvestigating_Click(object sender, EventArgs e)
 		{
 			int StartChangeNumber = EventMonitor.GetInvestigationStartChangeNumber(ContextMenuChange.Number);
-			if(StartChangeNumber != -1)
+			if (StartChangeNumber != -1)
 			{
-				if(ContextMenuChange.Number > StartChangeNumber)
+				if (ContextMenuChange.Number > StartChangeNumber)
 				{
 					string Message = String.Format("Mark all changes between {0} and {1} as bad?", StartChangeNumber, ContextMenuChange.Number);
-					if(MessageBox.Show(Message, "Finish investigating", MessageBoxButtons.YesNo) == DialogResult.Yes)
+					if (MessageBox.Show(Message, "Finish investigating", MessageBoxButtons.YesNo) == DialogResult.Yes)
 					{
-						foreach(PerforceChangeSummary Change in PerforceMonitor.GetChanges())
+						foreach (PerforceChangeSummary Change in PerforceMonitor.GetChanges())
 						{
 							if (Change.Number >= StartChangeNumber && Change.Number < ContextMenuChange.Number)
 							{
@@ -4325,12 +4325,12 @@ namespace UnrealGameSync
 		{
 			UpdateSyncActionCheckboxes();
 
-			if(PerforceMonitor != null)
+			if (PerforceMonitor != null)
 			{
 				PerforceMonitor.Refresh();
 			}
 
-			if(DesiredTaskbarState.Item1 == TaskbarState.Error)
+			if (DesiredTaskbarState.Item1 == TaskbarState.Error)
 			{
 				DesiredTaskbarState = Tuple.Create(TaskbarState.NoProgress, 0.0f);
 				Owner.UpdateProgress();
@@ -4345,27 +4345,27 @@ namespace UnrealGameSync
 		private void BuildList_ItemMouseHover(object sender, ListViewItemMouseHoverEventArgs Args)
 		{
 			Point ClientPoint = BuildList.PointToClient(Cursor.Position);
-			if(Args.Item.SubItems.Count >= 6)
+			if (Args.Item.SubItems.Count >= 6)
 			{
-				if(Args.Item.Bounds.Contains(ClientPoint))
+				if (Args.Item.Bounds.Contains(ClientPoint))
 				{
 					PerforceChangeSummary Change = (PerforceChangeSummary)Args.Item.Tag;
-					if(Change == null)
+					if (Change == null)
 					{
 						return;
 					}
 
 					Rectangle CISBounds = Args.Item.SubItems[CISColumn.Index].Bounds;
-					if(CISBounds.Contains(ClientPoint))
+					if (CISBounds.Contains(ClientPoint))
 					{
 						return;
 					}
 
 					Rectangle DescriptionBounds = Args.Item.SubItems[DescriptionColumn.Index].Bounds;
-					if(DescriptionBounds.Contains(ClientPoint))
+					if (DescriptionBounds.Contains(ClientPoint))
 					{
 						ChangeLayoutInfo LayoutInfo = GetChangeLayoutInfo(Change);
-						if(LayoutInfo.DescriptionBadges.Count == 0 || ClientPoint.X < GetBadgeListLocation(LayoutInfo.DescriptionBadges, DescriptionBounds, HorizontalAlign.Right, VerticalAlignment.Middle).X - 2)
+						if (LayoutInfo.DescriptionBadges.Count == 0 || ClientPoint.X < GetBadgeListLocation(LayoutInfo.DescriptionBadges, DescriptionBounds, HorizontalAlign.Right, VerticalAlignment.Middle).X - 2)
 						{
 							BuildListToolTip.Show(Change.Description, BuildList, new Point(DescriptionBounds.Left, DescriptionBounds.Bottom + 2));
 							return;
@@ -4373,19 +4373,19 @@ namespace UnrealGameSync
 					}
 
 					EventSummary Summary = EventMonitor.GetSummaryForChange(Change.Number);
-					if(Summary != null)
+					if (Summary != null)
 					{
 						StringBuilder SummaryText = new StringBuilder();
-						if(Summary.Comments.Count > 0)
+						if (Summary.Comments.Count > 0)
 						{
-							foreach(CommentData Comment in Summary.Comments)
+							foreach (CommentData Comment in Summary.Comments)
 							{
-								if(!String.IsNullOrWhiteSpace(Comment.Text))
+								if (!String.IsNullOrWhiteSpace(Comment.Text))
 								{
 									SummaryText.AppendFormat("{0}: \"{1}\"\n", FormatUserName(Comment.UserName), Comment.Text);
 								}
 							}
-							if(SummaryText.Length > 0)
+							if (SummaryText.Length > 0)
 							{
 								SummaryText.Append("\n");
 							}
@@ -4394,11 +4394,11 @@ namespace UnrealGameSync
 						AppendUserList(SummaryText, "\n", "Failed to compile for {0}.", Summary.Reviews.Where(x => x.Type == EventType.DoesNotCompile));
 						AppendUserList(SummaryText, "\n", "Marked good by {0}.", Summary.Reviews.Where(x => x.Type == EventType.Good));
 						AppendUserList(SummaryText, "\n", "Marked bad by {0}.", Summary.Reviews.Where(x => x.Type == EventType.Bad));
-						if(Summary.LastStarReview != null)
+						if (Summary.LastStarReview != null)
 						{
-							AppendUserList(SummaryText, "\n", "Starred by {0}.", new EventData[]{ Summary.LastStarReview });
+							AppendUserList(SummaryText, "\n", "Starred by {0}.", new EventData[] { Summary.LastStarReview });
 						}
-						if(SummaryText.Length > 0)
+						if (SummaryText.Length > 0)
 						{
 							Rectangle SummaryBounds = Args.Item.SubItems[StatusColumn.Index].Bounds;
 							BuildListToolTip.Show(SummaryText.ToString(), BuildList, new Point(SummaryBounds.Left, SummaryBounds.Bottom));
@@ -4415,13 +4415,13 @@ namespace UnrealGameSync
 		{
 			BuildListToolTip.Hide(BuildList);
 
-			if(HoverBadgeUniqueId != null)
+			if (HoverBadgeUniqueId != null)
 			{
 				HoverBadgeUniqueId = null;
 				BuildList.Invalidate();
 			}
 
-			if(bMouseOverExpandLink)
+			if (bMouseOverExpandLink)
 			{
 				Cursor = Cursors.Arrow;
 				bMouseOverExpandLink = false;
@@ -4462,7 +4462,7 @@ namespace UnrealGameSync
 		private bool ModifyEditorArguments()
 		{
 			ArgumentsWindow Arguments = new ArgumentsWindow(Settings.EditorArguments, Settings.bEditorArgumentsPrompt);
-			if(Arguments.ShowDialog(this) == DialogResult.OK)
+			if (Arguments.ShowDialog(this) == DialogResult.OK)
 			{
 				Settings.EditorArguments = Arguments.GetItems();
 				Settings.bEditorArgumentsPrompt = Arguments.PromptBeforeLaunch;
@@ -4482,20 +4482,20 @@ namespace UnrealGameSync
 			string MasterProjectName = "UE4";
 
 			string MasterProjectNameFileName = Path.Combine(BranchDirectoryName, "Engine", "Intermediate", "ProjectFiles", "MasterProjectName.txt");
-			if(File.Exists(MasterProjectNameFileName))
+			if (File.Exists(MasterProjectNameFileName))
 			{
 				try
 				{
 					MasterProjectName = File.ReadAllText(MasterProjectNameFileName).Trim();
 				}
-				catch(Exception Ex)
+				catch (Exception Ex)
 				{
 					Log.WriteException(Ex, "Unable to read '{0}'", MasterProjectNameFileName);
 				}
 			}
 
 			string SolutionFileName = Path.Combine(BranchDirectoryName, MasterProjectName + ".sln");
-			if(!File.Exists(SolutionFileName))
+			if (!File.Exists(SolutionFileName))
 			{
 				MessageBox.Show(String.Format("Couldn't find solution at {0}", SolutionFileName));
 			}
@@ -4557,23 +4557,23 @@ namespace UnrealGameSync
 
 		public void ScheduleTimerElapsed()
 		{
-			if(Settings.bScheduleEnabled)
+			if (Settings.bScheduleEnabled)
 			{
 				Log.WriteLine("Scheduled sync at {0} for {1}.", DateTime.Now, Settings.ScheduleChange);
 
 				int ChangeNumber;
-				if(!FindChangeToSync(Settings.ScheduleChange, out ChangeNumber))
+				if (!FindChangeToSync(Settings.ScheduleChange, out ChangeNumber))
 				{
 					Log.WriteLine("Couldn't find any matching change");
 				}
-				else if(Workspace.CurrentChangeNumber >= ChangeNumber)
+				else if (Workspace.CurrentChangeNumber >= ChangeNumber)
 				{
 					Log.WriteLine("Sync ignored; already at or ahead of CL ({0} >= {1})", Workspace.CurrentChangeNumber, ChangeNumber);
 				}
 				else
 				{
 					WorkspaceUpdateOptions Options = WorkspaceUpdateOptions.Sync | WorkspaceUpdateOptions.SyncArchives | WorkspaceUpdateOptions.GenerateProjectFiles | WorkspaceUpdateOptions.Build | WorkspaceUpdateOptions.ScheduledBuild;
-					if(Settings.bAutoResolveConflicts)
+					if (Settings.bAutoResolveConflicts)
 					{
 						Options |= WorkspaceUpdateOptions.AutoResolveChanges;
 					}
@@ -4584,32 +4584,32 @@ namespace UnrealGameSync
 
 		bool FindChangeToSync(LatestChangeType ChangeType, out int ChangeNumber)
 		{
-			for(int Idx = 0; Idx < BuildList.Items.Count; Idx++)
+			for (int Idx = 0; Idx < BuildList.Items.Count; Idx++)
 			{
 				PerforceChangeSummary Change = (PerforceChangeSummary)BuildList.Items[Idx].Tag;
-				if(Change != null)
+				if (Change != null)
 				{
-					if(ChangeType == LatestChangeType.Any)
+					if (ChangeType == LatestChangeType.Any)
 					{
-						if(CanSyncChange(Change.Number))
+						if (CanSyncChange(Change.Number))
 						{
 							ChangeNumber = FindNewestGoodContentChange(Change.Number);
 							return true;
 						}
 					}
-					else if(ChangeType == LatestChangeType.Good)
+					else if (ChangeType == LatestChangeType.Good)
 					{
 						EventSummary Summary = EventMonitor.GetSummaryForChange(Change.Number);
-						if(Summary != null && Summary.Verdict == ReviewVerdict.Good && CanSyncChange(Change.Number))
+						if (Summary != null && Summary.Verdict == ReviewVerdict.Good && CanSyncChange(Change.Number))
 						{
 							ChangeNumber = FindNewestGoodContentChange(Change.Number);
 							return true;
 						}
 					}
-					else if(ChangeType == LatestChangeType.Starred)
+					else if (ChangeType == LatestChangeType.Starred)
 					{
 						EventSummary Summary = EventMonitor.GetSummaryForChange(Change.Number);
-						if(((Summary != null && Summary.LastStarReview != null && Summary.LastStarReview.Type == EventType.Starred) || PromotedChangeNumbers.Contains(Change.Number)) && CanSyncChange(Change.Number))
+						if (((Summary != null && Summary.LastStarReview != null && Summary.LastStarReview.Type == EventType.Starred) || PromotedChangeNumbers.Contains(Change.Number)) && CanSyncChange(Change.Number))
 						{
 							ChangeNumber = FindNewestGoodContentChange(Change.Number);
 							return true;
@@ -4621,27 +4621,27 @@ namespace UnrealGameSync
 			ChangeNumber = -1;
 			return false;
 		}
-		
+
 		int FindNewestGoodContentChange(int ChangeNumber)
 		{
 			int Index = SortedChangeNumbers.BinarySearch(ChangeNumber);
-			if(Index <= 0)
+			if (Index <= 0)
 			{
 				return ChangeNumber;
 			}
 
-			for(int NextIndex = Index + 1; NextIndex < SortedChangeNumbers.Count; NextIndex++)
+			for (int NextIndex = Index + 1; NextIndex < SortedChangeNumbers.Count; NextIndex++)
 			{
 				int NextChangeNumber = SortedChangeNumbers[NextIndex];
 
 				PerforceChangeDetails Details;
-				if(!PerforceMonitor.TryGetChangeDetails(NextChangeNumber, out Details) || Details.bContainsCode)
+				if (!PerforceMonitor.TryGetChangeDetails(NextChangeNumber, out Details) || Details.bContainsCode)
 				{
 					break;
 				}
 
 				EventSummary Summary = EventMonitor.GetSummaryForChange(NextChangeNumber);
-				if(Summary != null && Summary.Verdict == ReviewVerdict.Bad)
+				if (Summary != null && Summary.Verdict == ReviewVerdict.Bad)
 				{
 					break;
 				}
@@ -4654,17 +4654,17 @@ namespace UnrealGameSync
 
 		private void BuildListContextMenu_LeaveOrEditComment_Click(object sender, EventArgs e)
 		{
-			if(ContextMenuChange != null)
+			if (ContextMenuChange != null)
 			{
 				LeaveCommentWindow LeaveComment = new LeaveCommentWindow();
 
 				string CommentText;
-				if(EventMonitor.GetCommentByCurrentUser(ContextMenuChange.Number, out CommentText))
+				if (EventMonitor.GetCommentByCurrentUser(ContextMenuChange.Number, out CommentText))
 				{
 					LeaveComment.CommentTextBox.Text = CommentText;
 				}
 
-				if(LeaveComment.ShowDialog() == System.Windows.Forms.DialogResult.OK)
+				if (LeaveComment.ShowDialog() == System.Windows.Forms.DialogResult.OK)
 				{
 					EventMonitor.PostComment(ContextMenuChange.Number, LeaveComment.CommentTextBox.Text);
 				}
@@ -4689,19 +4689,19 @@ namespace UnrealGameSync
 
 		private void MoreToolsContextMenu_CleanWorkspace_Click(object sender, EventArgs e)
 		{
-			if(!WaitForProgramsToFinish())
+			if (!WaitForProgramsToFinish())
 			{
 				return;
 			}
 
 			string ExtraSafeToDeleteFolders;
-			if(!TryGetProjectSetting(Workspace.ProjectConfigFile, "SafeToDeleteFolders", out ExtraSafeToDeleteFolders))
+			if (!TryGetProjectSetting(Workspace.ProjectConfigFile, "SafeToDeleteFolders", out ExtraSafeToDeleteFolders))
 			{
 				ExtraSafeToDeleteFolders = "";
 			}
 
 			string ExtraSafeToDeleteExtensions;
-			if(!TryGetProjectSetting(Workspace.ProjectConfigFile, "SafeToDeleteExtensions", out ExtraSafeToDeleteExtensions))
+			if (!TryGetProjectSetting(Workspace.ProjectConfigFile, "SafeToDeleteExtensions", out ExtraSafeToDeleteExtensions))
 			{
 				ExtraSafeToDeleteExtensions = "";
 			}
@@ -4716,26 +4716,26 @@ namespace UnrealGameSync
 		{
 			bHasBuildSteps = false;
 
-			foreach(ToolStripMenuItem CustomToolMenuItem in CustomToolMenuItems)
+			foreach (ToolStripMenuItem CustomToolMenuItem in CustomToolMenuItems)
 			{
 				MoreToolsContextMenu.Items.Remove(CustomToolMenuItem);
 			}
 
 			CustomToolMenuItems.Clear();
 
-			if(Workspace != null)
+			if (Workspace != null)
 			{
 				ConfigFile ProjectConfigFile = Workspace.ProjectConfigFile;
-				if(ProjectConfigFile != null)
+				if (ProjectConfigFile != null)
 				{
 					Dictionary<Guid, ConfigObject> ProjectBuildStepObjects = GetProjectBuildStepObjects(ProjectConfigFile);
 
 					int InsertIdx = 0;
 
 					List<BuildStep> UserSteps = GetUserBuildSteps(ProjectBuildStepObjects);
-					foreach(BuildStep Step in UserSteps)
+					foreach (BuildStep Step in UserSteps)
 					{
-						if(Step.bShowAsTool)
+						if (Step.bShowAsTool)
 						{
 							ToolStripMenuItem NewMenuItem = new ToolStripMenuItem(Step.Description.Replace("&", "&&"));
 							NewMenuItem.Click += new EventHandler((sender, e) => { RunCustomTool(Step.UniqueId); });
@@ -4752,15 +4752,15 @@ namespace UnrealGameSync
 
 		private void RunCustomTool(Guid UniqueId)
 		{
-			if(Workspace != null)
+			if (Workspace != null)
 			{
-				if(Workspace.IsBusy())
+				if (Workspace.IsBusy())
 				{
 					MessageBox.Show("Please retry after the current sync has finished.", "Sync in Progress");
 				}
 				else
 				{
-					WorkspaceUpdateContext Context = new WorkspaceUpdateContext(Workspace.CurrentChangeNumber, WorkspaceUpdateOptions.Build, null, GetDefaultBuildStepObjects(), ProjectSettings.BuildSteps, new HashSet<Guid>{ UniqueId }, GetWorkspaceVariables(Workspace.CurrentChangeNumber));
+					WorkspaceUpdateContext Context = new WorkspaceUpdateContext(Workspace.CurrentChangeNumber, WorkspaceUpdateOptions.Build, null, GetDefaultBuildStepObjects(), ProjectSettings.BuildSteps, new HashSet<Guid> { UniqueId }, GetWorkspaceVariables(Workspace.CurrentChangeNumber));
 					StartWorkspaceUpdate(Context, null);
 				}
 			}
@@ -4773,7 +4773,7 @@ namespace UnrealGameSync
 			string SdkInstallerDir;
 			TryGetProjectSetting(PerforceMonitor.LatestProjectConfigFile, "SdkInstallerDir", out SdkInstallerDir);
 
-			Dictionary<string, string> Variables = new Dictionary<string,string>();
+			Dictionary<string, string> Variables = new Dictionary<string, string>();
 			Variables.Add("Stream", StreamName);
 			Variables.Add("Change", ChangeNumber.ToString());
 			Variables.Add("ClientName", ClientName);
@@ -4783,9 +4783,9 @@ namespace UnrealGameSync
 			Variables.Add("UE4EditorExe", GetEditorExePath(EditorBuildConfig));
 			Variables.Add("UE4EditorCmdExe", GetEditorExePath(EditorBuildConfig).Replace(".exe", "-Cmd.exe"));
 			Variables.Add("UE4EditorConfig", EditorBuildConfig.ToString());
-			Variables.Add("UE4EditorDebugArg", (EditorBuildConfig == BuildConfig.Debug || EditorBuildConfig == BuildConfig.DebugGame)? " -debug" : "");
-			Variables.Add("UseIncrementalBuilds", Settings.bUseIncrementalBuilds? "1" : "0");
-			if(!String.IsNullOrEmpty(SdkInstallerDir))
+			Variables.Add("UE4EditorDebugArg", (EditorBuildConfig == BuildConfig.Debug || EditorBuildConfig == BuildConfig.DebugGame) ? " -debug" : "");
+			Variables.Add("UseIncrementalBuilds", Settings.bUseIncrementalBuilds ? "1" : "0");
+			if (!String.IsNullOrEmpty(SdkInstallerDir))
 			{
 				Variables.Add("SdkInstallerDir", SdkInstallerDir);
 			}
@@ -4795,16 +4795,16 @@ namespace UnrealGameSync
 		private void OptionsContextMenu_EditBuildSteps_Click(object sender, EventArgs e)
 		{
 			ConfigFile ProjectConfigFile = Workspace.ProjectConfigFile;
-			if(Workspace != null && ProjectConfigFile != null)
+			if (Workspace != null && ProjectConfigFile != null)
 			{
 				// Find all the target names for this project
 				List<string> TargetNames = new List<string>();
-				if(!String.IsNullOrEmpty(SelectedFileName) && SelectedFileName.EndsWith(".uproject", StringComparison.InvariantCultureIgnoreCase))
+				if (!String.IsNullOrEmpty(SelectedFileName) && SelectedFileName.EndsWith(".uproject", StringComparison.InvariantCultureIgnoreCase))
 				{
 					DirectoryInfo SourceDirectory = new DirectoryInfo(Path.Combine(Path.GetDirectoryName(SelectedFileName), "Source"));
-					if(SourceDirectory.Exists)
+					if (SourceDirectory.Exists)
 					{
-						foreach(FileInfo TargetFile in SourceDirectory.EnumerateFiles("*.target.cs", SearchOption.TopDirectoryOnly))
+						foreach (FileInfo TargetFile in SourceDirectory.EnumerateFiles("*.target.cs", SearchOption.TopDirectoryOnly))
 						{
 							TargetNames.Add(TargetFile.Name.Substring(0, TargetFile.Name.IndexOf('.')));
 						}
@@ -4821,15 +4821,15 @@ namespace UnrealGameSync
 
 				// Update the user settings
 				List<ConfigObject> ModifiedBuildSteps = new List<ConfigObject>();
-				foreach(BuildStep Step in UserSteps)
+				foreach (BuildStep Step in UserSteps)
 				{
-					if(Step.IsValid())
+					if (Step.IsValid())
 					{
 						ConfigObject DefaultObject;
 						ProjectBuildStepObjects.TryGetValue(Step.UniqueId, out DefaultObject);
 
 						ConfigObject UserConfigObject = Step.ToConfigObject(DefaultObject);
-						if(UserConfigObject != null && UserConfigObject.Pairs.Any(x => x.Key != "UniqueId"))
+						if (UserConfigObject != null && UserConfigObject.Pairs.Any(x => x.Key != "UniqueId"))
 						{
 							ModifiedBuildSteps.Add(UserConfigObject);
 						}
@@ -4849,13 +4849,13 @@ namespace UnrealGameSync
 		private void AddOrUpdateBuildStep(Dictionary<Guid, ConfigObject> Steps, ConfigObject Object)
 		{
 			Guid UniqueId;
-			if(Guid.TryParse(Object.GetValue(BuildStep.UniqueIdKey, ""), out UniqueId))
+			if (Guid.TryParse(Object.GetValue(BuildStep.UniqueIdKey, ""), out UniqueId))
 			{
 				// Add or apply Object to the list of steps in Steps. Do not modify Object; make a copy first.
 				ConfigObject NewObject = new ConfigObject(Object);
 
 				ConfigObject DefaultObject;
-				if(Steps.TryGetValue(UniqueId, out DefaultObject))
+				if (Steps.TryGetValue(UniqueId, out DefaultObject))
 				{
 					NewObject.SetDefaults(DefaultObject);
 				}
@@ -4871,23 +4871,23 @@ namespace UnrealGameSync
 
 		public BuildConfig GetEditorBuildConfig()
 		{
-			return ShouldSyncPrecompiledEditor? BuildConfig.Development : Settings.CompiledEditorBuildConfig;
+			return ShouldSyncPrecompiledEditor ? BuildConfig.Development : Settings.CompiledEditorBuildConfig;
 		}
 
-		private Dictionary<Guid,ConfigObject> GetDefaultBuildStepObjects()
+		private Dictionary<Guid, ConfigObject> GetDefaultBuildStepObjects()
 		{
 			string ProjectArgument = "";
-			if(SelectedFileName.EndsWith(".uproject", StringComparison.InvariantCultureIgnoreCase))
+			if (SelectedFileName.EndsWith(".uproject", StringComparison.InvariantCultureIgnoreCase))
 			{
 				ProjectArgument = String.Format("\"{0}\"", SelectedFileName);
 			}
 
 			string ActualEditorTargetName;
-			if(EditorTargetName != null)
+			if (EditorTargetName != null)
 			{
 				ActualEditorTargetName = EditorTargetName;
 			}
-			else if(SelectedFileName.EndsWith(".uproject", StringComparison.InvariantCultureIgnoreCase) && bIsEnterpriseProject)
+			else if (SelectedFileName.EndsWith(".uproject", StringComparison.InvariantCultureIgnoreCase) && bIsEnterpriseProject)
 			{
 				ActualEditorTargetName = "StudioEditor";
 			}
@@ -4912,18 +4912,18 @@ namespace UnrealGameSync
 		private Dictionary<Guid, ConfigObject> GetProjectBuildStepObjects(ConfigFile ProjectConfigFile)
 		{
 			Dictionary<Guid, ConfigObject> ProjectBuildSteps = GetDefaultBuildStepObjects();
-			foreach(string Line in ProjectConfigFile.GetValues("Build.Step", new string[0]))
+			foreach (string Line in ProjectConfigFile.GetValues("Build.Step", new string[0]))
 			{
 				AddOrUpdateBuildStep(ProjectBuildSteps, new ConfigObject(Line));
 			}
 			return ProjectBuildSteps;
 		}
-		
+
 		private List<BuildStep> GetUserBuildSteps(Dictionary<Guid, ConfigObject> ProjectBuildStepObjects)
 		{
 			// Read all the user-defined build tasks and modifications to the default list
 			Dictionary<Guid, ConfigObject> UserBuildStepObjects = ProjectBuildStepObjects.ToDictionary(x => x.Key, y => new ConfigObject(y.Value));
-			foreach(ConfigObject UserBuildStep in ProjectSettings.BuildSteps)
+			foreach (ConfigObject UserBuildStep in ProjectSettings.BuildSteps)
 			{
 				AddOrUpdateBuildStep(UserBuildStepObjects, UserBuildStep);
 			}
@@ -4937,7 +4937,7 @@ namespace UnrealGameSync
 			if (OptionsContextMenu_SyncPrecompiledBinaries.DropDownItems.Count == 0)
 			{
 				IArchiveInfo EditorArchive = PerforceMonitor.AvailableArchives.FirstOrDefault(x => x.Type == EditorArchiveType);
-				if(EditorArchive != null)
+				if (EditorArchive != null)
 				{
 					SetSelectedArchive(EditorArchive, !OptionsContextMenu_SyncPrecompiledBinaries.Checked);
 				}
@@ -4954,15 +4954,15 @@ namespace UnrealGameSync
 			StringBuilder DiagnosticsText = new StringBuilder();
 			DiagnosticsText.AppendFormat("Application version: {0}\n", Assembly.GetExecutingAssembly().GetName().Version);
 			DiagnosticsText.AppendFormat("Synced from: {0}\n", Program.SyncVersion ?? "(unknown)");
-			DiagnosticsText.AppendFormat("Selected file: {0}\n", (SelectedFileName == null)? "(none)" : SelectedFileName);
-			if(Workspace != null)
+			DiagnosticsText.AppendFormat("Selected file: {0}\n", (SelectedFileName == null) ? "(none)" : SelectedFileName);
+			if (Workspace != null)
 			{
 				DiagnosticsText.AppendFormat("P4 server: {0}\n", Workspace.Perforce.ServerAndPort ?? "(default)");
 				DiagnosticsText.AppendFormat("P4 user: {0}\n", Workspace.Perforce.UserName);
 				DiagnosticsText.AppendFormat("P4 workspace: {0}\n", Workspace.Perforce.ClientName);
 			}
-			DiagnosticsText.AppendFormat("Perforce monitor: {0}\n", (PerforceMonitor == null)? "(inactive)" : PerforceMonitor.LastStatusMessage);
-			DiagnosticsText.AppendFormat("Event monitor: {0}\n", (EventMonitor == null)? "(inactive)" : EventMonitor.LastStatusMessage);
+			DiagnosticsText.AppendFormat("Perforce monitor: {0}\n", (PerforceMonitor == null) ? "(inactive)" : PerforceMonitor.LastStatusMessage);
+			DiagnosticsText.AppendFormat("Event monitor: {0}\n", (EventMonitor == null) ? "(inactive)" : EventMonitor.LastStatusMessage);
 
 			DiagnosticsWindow Diagnostics = new DiagnosticsWindow(DataFolder, DiagnosticsText.ToString());
 			Diagnostics.ShowDialog(this);
@@ -4971,7 +4971,7 @@ namespace UnrealGameSync
 		private void OptionsContextMenu_SyncFilter_Click(object sender, EventArgs e)
 		{
 			SyncFilter Filter = new SyncFilter(GetSyncCategories(), Settings.SyncView, Settings.SyncExcludedCategories, Settings.bSyncAllProjects, Settings.bIncludeAllProjectsInSolution, WorkspaceSettings.SyncView, WorkspaceSettings.SyncIncludedCategories, WorkspaceSettings.SyncExcludedCategories, WorkspaceSettings.bSyncAllProjects, WorkspaceSettings.bIncludeAllProjectsInSolution);
-			if(Filter.ShowDialog() == DialogResult.OK)
+			if (Filter.ShowDialog() == DialogResult.OK)
 			{
 				Settings.SyncExcludedCategories = Filter.GlobalExcludedCategories;
 				Settings.SyncView = Filter.GlobalView;
@@ -4982,7 +4982,7 @@ namespace UnrealGameSync
 				WorkspaceSettings.SyncView = Filter.WorkspaceView;
 				WorkspaceSettings.bSyncAllProjects = Filter.bWorkspaceSyncAllProjects;
 				WorkspaceSettings.bIncludeAllProjectsInSolution = Filter.bWorkspaceIncludeAllProjectsInSolution;
-                Settings.Save();
+				Settings.Save();
 			}
 		}
 
@@ -5014,13 +5014,13 @@ namespace UnrealGameSync
 
 		private void SyncContextMenu_EnterChangelist_Click(object sender, EventArgs e)
 		{
-			if(!WaitForProgramsToFinish())
+			if (!WaitForProgramsToFinish())
 			{
 				return;
 			}
 
-			ChangelistWindow ChildWindow = new ChangelistWindow((Workspace == null)? -1 : Workspace.CurrentChangeNumber);
-			if(ChildWindow.ShowDialog() == DialogResult.OK)
+			ChangelistWindow ChildWindow = new ChangelistWindow((Workspace == null) ? -1 : Workspace.CurrentChangeNumber);
+			if (ChildWindow.ShowDialog() == DialogResult.OK)
 			{
 				StartSync(ChildWindow.ChangeNumber);
 			}
@@ -5028,7 +5028,7 @@ namespace UnrealGameSync
 
 		private void BuildList_KeyDown(object Sender, KeyEventArgs Args)
 		{
-			if(Args.Control && Args.KeyCode == Keys.C && BuildList.SelectedItems.Count > 0)
+			if (Args.Control && Args.KeyCode == Keys.C && BuildList.SelectedItems.Count > 0)
 			{
 				int SelectedChange = ((PerforceChangeSummary)BuildList.SelectedItems[0].Tag).Number;
 				Clipboard.SetText(String.Format("{0}", SelectedChange));
@@ -5068,12 +5068,12 @@ namespace UnrealGameSync
 
 		private void SelectRecentProject(Rectangle Bounds)
 		{
-			while(RecentMenu.Items[2] != RecentMenu_Separator)
+			while (RecentMenu.Items[2] != RecentMenu_Separator)
 			{
 				RecentMenu.Items.RemoveAt(2);
 			}
 
-			foreach(UserSelectedProjectSettings RecentProject in Settings.RecentProjects)
+			foreach (UserSelectedProjectSettings RecentProject in Settings.RecentProjects)
 			{
 				ToolStripMenuItem Item = new ToolStripMenuItem(RecentProject.ToString(), null, new EventHandler((o, e) => Owner.RequestProjectChange(this, RecentProject, true)));
 				RecentMenu.Items.Insert(RecentMenu.Items.Count - 2, Item);
@@ -5118,10 +5118,10 @@ namespace UnrealGameSync
 
 		private void BuildList_ColumnWidthChanged(object sender, ColumnWidthChangedEventArgs e)
 		{
-			if(ColumnWidths != null && MinColumnWidths != null)
+			if (ColumnWidths != null && MinColumnWidths != null)
 			{
 				int NewWidth = BuildList.Columns[e.ColumnIndex].Width;
-				if(NewWidth < MinColumnWidths[e.ColumnIndex])
+				if (NewWidth < MinColumnWidths[e.ColumnIndex])
 				{
 					NewWidth = MinColumnWidths[e.ColumnIndex];
 					BuildList.Columns[e.ColumnIndex].Width = NewWidth;
@@ -5140,16 +5140,16 @@ namespace UnrealGameSync
 			if (PrevBuildListWidth != 0 && BuildListWidth != PrevBuildListWidth && ColumnWidths != null)
 			{
 				float SafeWidth = ColumnWidths.Sum() + 50.0f;
-				if(BuildListWidth < PrevBuildListWidth)
+				if (BuildListWidth < PrevBuildListWidth)
 				{
-					if(BuildListWidth <= SafeWidth)
+					if (BuildListWidth <= SafeWidth)
 					{
 						ResizeColumns(ColumnWidths.Sum() + (BuildListWidth - Math.Max(PrevBuildListWidth, SafeWidth)));
 					}
 				}
 				else
 				{
-					if(BuildListWidth >= SafeWidth)
+					if (BuildListWidth >= SafeWidth)
 					{
 						ResizeColumns(ColumnWidths.Sum() + (BuildListWidth - Math.Max(PrevBuildListWidth, SafeWidth)));
 					}
@@ -5166,15 +5166,15 @@ namespace UnrealGameSync
 			float TotalDelta = Math.Abs(NextTotalWidth - PrevTotalWidth);
 
 			float TotalColumnDelta = 0.0f;
-			for(int Idx = 0; Idx < BuildList.Columns.Count; Idx++)
+			for (int Idx = 0; Idx < BuildList.Columns.Count; Idx++)
 			{
 				TotalColumnDelta += Math.Abs(TargetColumnWidths[Idx] - ColumnWidths[Idx]);
 			}
 
-			if(TotalColumnDelta > 0.5f)
+			if (TotalColumnDelta > 0.5f)
 			{
 				float[] NextColumnWidths = new float[BuildList.Columns.Count];
-				for(int Idx = 0; Idx < BuildList.Columns.Count; Idx++)
+				for (int Idx = 0; Idx < BuildList.Columns.Count; Idx++)
 				{
 					float MaxColumnDelta = TotalDelta * Math.Abs(TargetColumnWidths[Idx] - ColumnWidths[Idx]) / TotalColumnDelta;
 					NextColumnWidths[Idx] = Math.Max(Math.Min(TargetColumnWidths[Idx], ColumnWidths[Idx] + MaxColumnDelta), ColumnWidths[Idx] - MaxColumnDelta);
@@ -5182,7 +5182,7 @@ namespace UnrealGameSync
 
 				// Update the control
 				BuildList.BeginUpdate();
-				for(int Idx = 0; Idx < BuildList.Columns.Count; Idx++)
+				for (int Idx = 0; Idx < BuildList.Columns.Count; Idx++)
 				{
 					BuildList.Columns[Idx].Width = (int)NextColumnWidths[Idx];
 				}
@@ -5208,15 +5208,15 @@ namespace UnrealGameSync
 
 			// Handle special cases for shrinking/expanding
 			float PrevTotalWidth = ColumnWidths.Sum();
-			if(NextTotalWidth < PrevTotalWidth)
+			if (NextTotalWidth < PrevTotalWidth)
 			{
 				// If target size is less than current size, keep it at the current size
-				for(int Idx = 0; Idx < BuildList.Columns.Count; Idx++)
+				for (int Idx = 0; Idx < BuildList.Columns.Count; Idx++)
 				{
-					if(!ConstrainedColumns[Idx])
+					if (!ConstrainedColumns[Idx])
 					{
 						float TargetColumnWidth = MinColumnWidths[Idx] + (RemainingWidth * ColumnWeights[Idx] / RemainingTotalWeight);
-						if(TargetColumnWidth > ColumnWidths[Idx])
+						if (TargetColumnWidth > ColumnWidths[Idx])
 						{
 							TargetColumnWidths[Idx] = ColumnWidths[Idx];
 							ConstrainedColumns[Idx] = true;
@@ -5233,12 +5233,12 @@ namespace UnrealGameSync
 			else
 			{
 				// If target size is greater than desired size, clamp it to that
-				for(int Idx = 0; Idx < BuildList.Columns.Count; Idx++)
+				for (int Idx = 0; Idx < BuildList.Columns.Count; Idx++)
 				{
-					if(!ConstrainedColumns[Idx])
+					if (!ConstrainedColumns[Idx])
 					{
 						float TargetColumnWidth = MinColumnWidths[Idx] + (RemainingWidth * ColumnWeights[Idx] / RemainingTotalWeight);
-						if(TargetColumnWidth > DesiredColumnWidths[Idx])
+						if (TargetColumnWidth > DesiredColumnWidths[Idx])
 						{
 							// Don't allow this column to expand above the maximum desired size
 							TargetColumnWidths[Idx] = DesiredColumnWidths[Idx];
@@ -5254,12 +5254,12 @@ namespace UnrealGameSync
 				}
 
 				// If current size is greater than target size, keep it that way
-				for(int Idx = 0; Idx < BuildList.Columns.Count; Idx++)
+				for (int Idx = 0; Idx < BuildList.Columns.Count; Idx++)
 				{
-					if(!ConstrainedColumns[Idx])
+					if (!ConstrainedColumns[Idx])
 					{
 						float TargetColumnWidth = MinColumnWidths[Idx] + (RemainingWidth * ColumnWeights[Idx] / RemainingTotalWeight);
-						if(TargetColumnWidth < ColumnWidths[Idx])
+						if (TargetColumnWidth < ColumnWidths[Idx])
 						{
 							TargetColumnWidths[Idx] = ColumnWidths[Idx];
 							ConstrainedColumns[Idx] = true;
@@ -5275,9 +5275,9 @@ namespace UnrealGameSync
 			}
 
 			// Allocate the remaining space equally
-			for(int Idx = 0; Idx < BuildList.Columns.Count; Idx++)
+			for (int Idx = 0; Idx < BuildList.Columns.Count; Idx++)
 			{
-				if(!ConstrainedColumns[Idx])
+				if (!ConstrainedColumns[Idx])
 				{
 					TargetColumnWidths[Idx] = MinColumnWidths[Idx] + (RemainingWidth * ColumnWeights[Idx] / RemainingTotalWeight);
 				}
@@ -5298,10 +5298,10 @@ namespace UnrealGameSync
 
 		private void EnableBisectMode()
 		{
-			if(BuildList.SelectedItems.Count >= 2)
+			if (BuildList.SelectedItems.Count >= 2)
 			{
 				Dictionary<int, BisectState> ChangeNumberToBisectState = new Dictionary<int, BisectState>();
-				foreach(ListViewItem SelectedItem in BuildList.SelectedItems)
+				foreach (ListViewItem SelectedItem in BuildList.SelectedItems)
 				{
 					PerforceChangeSummary Change = (PerforceChangeSummary)SelectedItem.Tag;
 					ChangeNumberToBisectState[Change.Number] = BisectState.Include;
@@ -5329,10 +5329,10 @@ namespace UnrealGameSync
 
 		private void SetBisectStateForSelection(BisectState State)
 		{
-			foreach(ListViewItem SelectedItem in BuildList.SelectedItems)
+			foreach (ListViewItem SelectedItem in BuildList.SelectedItems)
 			{
 				PerforceChangeSummary Change = (PerforceChangeSummary)SelectedItem.Tag;
-				if(Change != null)
+				if (Change != null)
 				{
 					WorkspaceSettings.ChangeNumberToBisectState[Change.Number] = State;
 				}
@@ -5353,9 +5353,9 @@ namespace UnrealGameSync
 			GetRemainingBisectRange(out PassChangeNumber, out FailChangeNumber);
 
 			List<int> ChangeNumbers = new List<int>();
-			foreach(KeyValuePair<int, BisectState> Pair in WorkspaceSettings.ChangeNumberToBisectState)
+			foreach (KeyValuePair<int, BisectState> Pair in WorkspaceSettings.ChangeNumberToBisectState)
 			{
-				if(Pair.Value == BisectState.Include && Pair.Key > PassChangeNumber && Pair.Key < FailChangeNumber)
+				if (Pair.Value == BisectState.Include && Pair.Key > PassChangeNumber && Pair.Key < FailChangeNumber)
 				{
 					ChangeNumbers.Add(Pair.Key);
 				}
@@ -5363,13 +5363,13 @@ namespace UnrealGameSync
 
 			ChangeNumbers.Sort();
 
-			return (ChangeNumbers.Count > 0)? ChangeNumbers[ChangeNumbers.Count / 2] : -1;
+			return (ChangeNumbers.Count > 0) ? ChangeNumbers[ChangeNumbers.Count / 2] : -1;
 		}
 
 		private void SyncBisectChange()
 		{
 			int BisectChange = GetBisectChangeNumber();
-			if(BisectChange != -1)
+			if (BisectChange != -1)
 			{
 				Owner.ShowAndActivate();
 				SelectChange(BisectChange);
@@ -5399,7 +5399,7 @@ namespace UnrealGameSync
 
 		private void WorkspaceControl_VisibleChanged(object sender, EventArgs e)
 		{
-			if(PerforceMonitor != null && PerforceMonitor.IsActive != Visible)
+			if (PerforceMonitor != null && PerforceMonitor.IsActive != Visible)
 			{
 				PerforceMonitor.IsActive = Visible;
 			}
@@ -5424,11 +5424,11 @@ namespace UnrealGameSync
 			DisplayBadgeNameAndGroupPairs.AddRange(BadgeNames.Select(x => new KeyValuePair<string, string>(x, "User")));
 
 			string LastGroup = null;
-			foreach(KeyValuePair<string, string> BadgeNameAndGroupPair in DisplayBadgeNameAndGroupPairs)
+			foreach (KeyValuePair<string, string> BadgeNameAndGroupPair in DisplayBadgeNameAndGroupPairs)
 			{
-				if(LastGroup != BadgeNameAndGroupPair.Value)
+				if (LastGroup != BadgeNameAndGroupPair.Value)
 				{
-					if(LastGroup != null)
+					if (LastGroup != null)
 					{
 						FilterContextMenu_Badges.DropDownItems.Add(new ToolStripSeparator());
 					}
@@ -5449,7 +5449,7 @@ namespace UnrealGameSync
 
 		private void FilterContextMenu_Badge_Click(string BadgeName)
 		{
-			if(ProjectSettings.FilterBadges.Contains(BadgeName))
+			if (ProjectSettings.FilterBadges.Contains(BadgeName))
 			{
 				ProjectSettings.FilterBadges.Remove(BadgeName);
 			}
diff --git a/Engine/Source/Programs/UnrealGameSync/UnrealGameSync/EventMonitor.cs b/Engine/Source/Programs/UnrealGameSync/UnrealGameSync/EventMonitor.cs
--- a/Engine/Source/Programs/UnrealGameSync/UnrealGameSync/EventMonitor.cs
+++ b/Engine/Source/Programs/UnrealGameSync/UnrealGameSync/EventMonitor.cs
@@ -92,6 +92,11 @@ namespace UnrealGameSync
 		{
 			get
 			{
+				if (BuildType == null)
+				{
+					return "Unknown";
+				}
+
 				int Idx = BuildType.IndexOf(':');
 				if(Idx == -1)
 				{
@@ -108,6 +113,11 @@ namespace UnrealGameSync
 		{
 			get
 			{
+				if (BuildType == null)
+				{
+					return "Unknown";
+				}
+
 				int Idx = BuildType.IndexOf(':');
 				if(Idx == -1)
 				{
diff --git a/Engine/Source/Programs/UnrealGameSync/UnrealGameSync/Forms/MainWindow.cs b/Engine/Source/Programs/UnrealGameSync/UnrealGameSync/Forms/MainWindow.cs
--- a/Engine/Source/Programs/UnrealGameSync/UnrealGameSync/Forms/MainWindow.cs
+++ b/Engine/Source/Programs/UnrealGameSync/UnrealGameSync/Forms/MainWindow.cs
@@ -719,9 +719,16 @@ namespace UnrealGameSync
 			if(Settings.bScheduleEnabled)
 			{
 				DateTime CurrentTime = DateTime.Now;
+				Random Rnd = new Random();
+
+				// add or subtract from the schedule time to distribute scheduled syncs over a little bit more time
+				// this avoids everyone hitting the p4 server at exactly the same time.
+				const int FudgeMinutes = 10;
+				TimeSpan FudgeTime = TimeSpan.FromMinutes(Rnd.Next(FudgeMinutes * -100, FudgeMinutes * 100) / 100.0);
 				DateTime NextScheduleTime = new DateTime(CurrentTime.Year, CurrentTime.Month, CurrentTime.Day, Settings.ScheduleTime.Hours, Settings.ScheduleTime.Minutes, Settings.ScheduleTime.Seconds);
+				NextScheduleTime += FudgeTime;
 
-				if(NextScheduleTime < CurrentTime)
+				if (NextScheduleTime < CurrentTime)
 				{
 					NextScheduleTime = NextScheduleTime.AddDays(1.0);
 				}
diff --git a/Engine/Source/Programs/UnrealGameSync/UnrealGameSync/PerforceMonitor.cs b/Engine/Source/Programs/UnrealGameSync/UnrealGameSync/PerforceMonitor.cs
--- a/Engine/Source/Programs/UnrealGameSync/UnrealGameSync/PerforceMonitor.cs
+++ b/Engine/Source/Programs/UnrealGameSync/UnrealGameSync/PerforceMonitor.cs
@@ -617,7 +617,7 @@ namespace UnrealGameSync
 
 		public static ConfigFile ReadProjectConfigFile(PerforceConnection Perforce, string BranchClientPath, string SelectedClientFileName, string CacheFolder, List<KeyValuePair<string, DateTime>> LocalConfigFiles, TextWriter Log)
 		{
-			List<string> ConfigFilePaths = Utility.GetConfigFileLocations(BranchClientPath, SelectedClientFileName, '/');
+			List<string> ConfigFilePaths = Utility.GetDepotConfigPaths(BranchClientPath + "/Engine", SelectedClientFileName);
 
 			ConfigFile ProjectConfig = new ConfigFile();
 
diff --git a/Engine/Source/Programs/UnrealGameSync/UnrealGameSync/Program.cs b/Engine/Source/Programs/UnrealGameSync/UnrealGameSync/Program.cs
--- a/Engine/Source/Programs/UnrealGameSync/UnrealGameSync/Program.cs
+++ b/Engine/Source/Programs/UnrealGameSync/UnrealGameSync/Program.cs
@@ -99,7 +99,7 @@ namespace UnrealGameSync
 					{
 						Log.WriteLine("Missing server settings; finding defaults.");
 						GetDefaultServerSettings(ref ServerAndPort, ref UserName, Log);
-						Utility.SaveGlobalPerforceSettings(ServerAndPort, UserName, UpdatePath);
+						Utility.SaveGlobalPerforceSettings(ServerAndPort, UserName, BaseUpdatePath);
 					}
 
 					PerforceConnection DefaultConnection = new PerforceConnection(UserName, null, ServerAndPort);
diff --git a/Engine/Source/Programs/UnrealGameSync/UnrealGameSync/Properties/AssemblyInfo.cs b/Engine/Source/Programs/UnrealGameSync/UnrealGameSync/Properties/AssemblyInfo.cs
--- a/Engine/Source/Programs/UnrealGameSync/UnrealGameSync/Properties/AssemblyInfo.cs
+++ b/Engine/Source/Programs/UnrealGameSync/UnrealGameSync/Properties/AssemblyInfo.cs
@@ -34,5 +34,5 @@ using System.Runtime.InteropServices;
 // You can specify all the values or you can default the Build and Revision Numbers 
 // by using the '*' as shown below:
 // [assembly: AssemblyVersion("1.0.*")]
-[assembly: AssemblyVersion("1.189.*")]
-[assembly: AssemblyFileVersion("1.189.0.0")]
+[assembly: AssemblyVersion("1.193.*")]
+[assembly: AssemblyFileVersion("1.193.0.0")]
diff --git a/Engine/Source/Programs/UnrealGameSync/UnrealGameSync/UserSettings.cs b/Engine/Source/Programs/UnrealGameSync/UnrealGameSync/UserSettings.cs
--- a/Engine/Source/Programs/UnrealGameSync/UnrealGameSync/UserSettings.cs
+++ b/Engine/Source/Programs/UnrealGameSync/UnrealGameSync/UserSettings.cs
@@ -51,7 +51,8 @@ namespace UnrealGameSync
 	{
 		Initial = 0,
 		DefaultServerSettings = 1,
-		Latest = DefaultServerSettings
+		XgeShaderCompilation = 2,
+		Latest = XgeShaderCompilation
 	}
 
 	class ArchiveSettings
@@ -400,6 +401,10 @@ namespace UnrealGameSync
 
 			// Editor arguments
 			string[] Arguments = ConfigFile.GetValues("General.EditorArguments", new string[]{ "0:-log", "0:-fastload" });
+			if (Version < UserSettingsVersion.XgeShaderCompilation)
+			{
+				Arguments = Enumerable.Concat(Arguments, new string[] { "0:-noxgeshadercompile" }).ToArray();
+			}
 			foreach(string Argument in Arguments)
 			{
 				if(Argument.StartsWith("0:"))
diff --git a/Engine/Source/Programs/UnrealGameSync/UnrealGameSync/Utility.cs b/Engine/Source/Programs/UnrealGameSync/UnrealGameSync/Utility.cs
--- a/Engine/Source/Programs/UnrealGameSync/UnrealGameSync/Utility.cs
+++ b/Engine/Source/Programs/UnrealGameSync/UnrealGameSync/Utility.cs
@@ -11,6 +11,7 @@ using System.IO;
 using System.Linq;
 using System.Runtime.InteropServices;
 using System.Text;
+using System.Text.RegularExpressions;
 using System.Threading.Tasks;
 using System.Web.Script.Serialization;
 using System.Windows.Forms;
@@ -266,33 +267,94 @@ namespace UnrealGameSync
 			}
 		}
 
-		public static List<string> GetConfigFileLocations(string BaseWorkspacePath, string ProjectPath, char Separator)
+		/******/
+
+		private static void AddLocalConfigPaths_PlatformFolders(DirectoryInfo BaseDir, string FileName, List<FileInfo> Files)
+		{
+			if(BaseDir.Exists)
+			{
+				FileInfo BaseFileInfo = new FileInfo(Path.Combine(BaseDir.FullName, FileName));
+				if(BaseFileInfo.Exists)
+				{
+					Files.Add(BaseFileInfo);
+				}
+
+				foreach (DirectoryInfo SubDirInfo in BaseDir.EnumerateDirectories())
+				{
+					FileInfo SubFile = new FileInfo(Path.Combine(SubDirInfo.FullName, FileName));
+					if (SubFile.Exists)
+					{
+						Files.Add(SubFile);
+					}
+				}
+			}
+		}
+
+		private static void AddLocalConfigPaths_PlatformExtensions(DirectoryInfo BaseDir, string RelativePath, string FileName, List<FileInfo> Files)
 		{
-			List<string> ProjectConfigFileNames = new List<string>();
-			ProjectConfigFileNames.Add(String.Format("{0}{1}Engine{1}Programs{1}UnrealGameSync{1}UnrealGameSync.ini", BaseWorkspacePath, Separator));
-			ProjectConfigFileNames.Add(String.Format("{0}{1}Engine{1}Programs{1}UnrealGameSync{1}PS4{1}UnrealGameSync.ini", BaseWorkspacePath, Separator));
-			ProjectConfigFileNames.Add(String.Format("{0}{1}Engine{1}Programs{1}UnrealGameSync{1}XboxOne{1}UnrealGameSync.ini", BaseWorkspacePath, Separator));
-			ProjectConfigFileNames.Add(String.Format("{0}{1}Engine{1}Programs{1}UnrealGameSync{1}Switch{1}UnrealGameSync.ini", BaseWorkspacePath, Separator));
-			ProjectConfigFileNames.Add(String.Format("{0}{1}Engine{1}Programs{1}UnrealGameSync{1}NotForLicensees{1}UnrealGameSync.ini", BaseWorkspacePath, Separator));
-			if(ProjectPath.EndsWith(".uproject", StringComparison.InvariantCultureIgnoreCase))
-			{
-				ProjectConfigFileNames.Add(String.Format("{0}{1}Build{1}UnrealGameSync.ini", ProjectPath.Substring(0, ProjectPath.LastIndexOf(Separator)), Separator));
-			    ProjectConfigFileNames.Add(String.Format("{0}{1}Build{1}PS4{1}UnrealGameSync.ini", ProjectPath.Substring(0, ProjectPath.LastIndexOf(Separator)), Separator));
-			    ProjectConfigFileNames.Add(String.Format("{0}{1}Build{1}XboxOne{1}UnrealGameSync.ini", ProjectPath.Substring(0, ProjectPath.LastIndexOf(Separator)), Separator));
-			    ProjectConfigFileNames.Add(String.Format("{0}{1}Build{1}Switch{1}UnrealGameSync.ini", ProjectPath.Substring(0, ProjectPath.LastIndexOf(Separator)), Separator));
-				ProjectConfigFileNames.Add(String.Format("{0}{1}Build{1}NotForLicensees{1}UnrealGameSync.ini", ProjectPath.Substring(0, ProjectPath.LastIndexOf(Separator)), Separator));
+			if (BaseDir.Exists)
+			{
+				AddLocalConfigPaths_PlatformFolders(new DirectoryInfo(Path.Combine(BaseDir.FullName, RelativePath)), FileName, Files);
+
+				DirectoryInfo PlatformExtensionsDir = new DirectoryInfo(Path.Combine(BaseDir.FullName, "Platforms"));
+				if (PlatformExtensionsDir.Exists)
+				{
+					foreach (DirectoryInfo PlatformExtensionDir in PlatformExtensionsDir.EnumerateDirectories())
+					{
+						AddLocalConfigPaths_PlatformFolders(new DirectoryInfo(Path.Combine(PlatformExtensionDir.FullName, RelativePath)), FileName, Files);
+					}
+				}
+			}
+		}
+
+		public static List<FileInfo> GetLocalConfigPaths(DirectoryInfo EngineDir, FileInfo ProjectFile)
+		{
+			List<FileInfo> SearchPaths = new List<FileInfo>();
+			AddLocalConfigPaths_PlatformExtensions(EngineDir, "Programs/UnrealGameSync", "UnrealGameSync.ini", SearchPaths);
+
+			if (ProjectFile.Name.EndsWith(".uproject", StringComparison.OrdinalIgnoreCase))
+			{
+				AddLocalConfigPaths_PlatformExtensions(ProjectFile.Directory, "Build", "UnrealGameSync.ini", SearchPaths);
 			}
 			else
 			{
-				ProjectConfigFileNames.Add(String.Format("{0}{1}Engine{1}Programs{1}UnrealGameSync{1}DefaultProject.ini", BaseWorkspacePath, Separator));
-			    ProjectConfigFileNames.Add(String.Format("{0}{1}Engine{1}Programs{1}UnrealGameSync{1}PS4{1}DefaultProject.ini", BaseWorkspacePath, Separator));
-			    ProjectConfigFileNames.Add(String.Format("{0}{1}Engine{1}Programs{1}UnrealGameSync{1}XboxOne{1}DefaultProject.ini", BaseWorkspacePath, Separator));
-			    ProjectConfigFileNames.Add(String.Format("{0}{1}Engine{1}Programs{1}UnrealGameSync{1}Switch{1}DefaultProject.ini", BaseWorkspacePath, Separator));
-				ProjectConfigFileNames.Add(String.Format("{0}{1}Engine{1}Programs{1}UnrealGameSync{1}NotForLicensees{1}DefaultProject.ini", BaseWorkspacePath, Separator));
+				AddLocalConfigPaths_PlatformExtensions(EngineDir, "Programs/UnrealGameSync", "DefaultEngine.ini", SearchPaths);
 			}
-			return ProjectConfigFileNames;
+			return SearchPaths;
+		}
+
+		/******/
+
+		private static void AddDepotConfigPaths_PlatformFolders(string BasePath, string FileName, List<string> SearchPaths)
+		{
+			SearchPaths.Add(String.Format("{0}/{1}", BasePath, FileName));
+			SearchPaths.Add(String.Format("{0}/*/{1}", BasePath, FileName));
+		}
+
+		private static void AddDepotConfigPaths_PlatformExtensions(string BasePath, string RelativePath, string FileName, List<string> SearchPaths)
+		{
+			AddDepotConfigPaths_PlatformFolders(BasePath + RelativePath, FileName, SearchPaths);
+			AddDepotConfigPaths_PlatformFolders(BasePath + "/Platforms/*" + RelativePath, FileName, SearchPaths);
 		}
 
+		public static List<string> GetDepotConfigPaths(string EnginePath, string ProjectPath)
+		{
+			List<string> SearchPaths = new List<string>();
+			AddDepotConfigPaths_PlatformExtensions(EnginePath, "/Programs/UnrealGameSync", "UnrealGameSync.ini", SearchPaths);
+
+			if (ProjectPath.EndsWith(".uproject", StringComparison.OrdinalIgnoreCase))
+			{
+				AddDepotConfigPaths_PlatformExtensions(ProjectPath.Substring(0, ProjectPath.LastIndexOf('/')), "/Build", "UnrealGameSync.ini", SearchPaths);
+			}
+			else
+			{
+				AddDepotConfigPaths_PlatformExtensions(EnginePath, "/Programs/UnrealGameSync", "DefaultEngine.ini", SearchPaths);
+			}
+			return SearchPaths;
+		}
+
+		/******/
+
 		public static void ReadGlobalPerforceSettings(ref string ServerAndPort, ref string UserName, ref string DepotPath)
 		{
 			using (RegistryKey Key = Registry.CurrentUser.OpenSubKey("SOFTWARE\\Epic Games\\UnrealGameSync", false))
@@ -302,6 +364,17 @@ namespace UnrealGameSync
 					ServerAndPort = Key.GetValue("ServerAndPort", ServerAndPort) as string;
 					UserName = Key.GetValue("UserName", UserName) as string;
 					DepotPath = Key.GetValue("DepotPath", DepotPath) as string;
+
+					// Fix corrupted depot path string
+					if(DepotPath != null)
+					{
+						Match Match = Regex.Match(DepotPath, "^(.*)/(Release|UnstableRelease)/\\.\\.\\.@.*$");
+						if (Match.Success)
+						{
+							DepotPath = Match.Groups[1].Value;
+							SaveGlobalPerforceSettings(ServerAndPort, UserName, DepotPath);
+						}
+					}
 				}
 			}
 		}
diff --git a/Engine/Source/Programs/UnrealGameSync/UnrealGameSync/Workspace.cs b/Engine/Source/Programs/UnrealGameSync/UnrealGameSync/Workspace.cs
--- a/Engine/Source/Programs/UnrealGameSync/UnrealGameSync/Workspace.cs
+++ b/Engine/Source/Programs/UnrealGameSync/UnrealGameSync/Workspace.cs
@@ -1234,24 +1234,22 @@ namespace UnrealGameSync
 		static ConfigFile ReadProjectConfigFile(string LocalRootPath, string SelectedLocalFileName, TextWriter Log)
 		{
 			// Find the valid config file paths
-			List<string> ProjectConfigFileNames = Utility.GetConfigFileLocations(LocalRootPath, SelectedLocalFileName, Path.DirectorySeparatorChar);
+			DirectoryInfo EngineDir = new DirectoryInfo(Path.Combine(LocalRootPath, "Engine"));
+			List<FileInfo> LocalConfigFiles = Utility.GetLocalConfigPaths(EngineDir, new FileInfo(SelectedLocalFileName));
 
 			// Read them in
 			ConfigFile ProjectConfig = new ConfigFile();
-			foreach(string ProjectConfigFileName in ProjectConfigFileNames)
+			foreach(FileInfo LocalConfigFile in LocalConfigFiles)
 			{
-				if(File.Exists(ProjectConfigFileName))
+				try
 				{
-					try
-					{
-						string[] Lines = File.ReadAllLines(ProjectConfigFileName);
-						ProjectConfig.Parse(Lines);
-						Log.WriteLine("Read config file from {0}", ProjectConfigFileName);
-					}
-					catch(Exception Ex)
-					{
-						Log.WriteLine("Failed to read config file from {0}: {1}", ProjectConfigFileName, Ex.ToString());
-					}
+					string[] Lines = File.ReadAllLines(LocalConfigFile.FullName);
+					ProjectConfig.Parse(Lines);
+					Log.WriteLine("Read config file from {0}", LocalConfigFile.FullName);
+				}
+				catch(Exception Ex)
+				{
+					Log.WriteLine("Failed to read config file from {0}: {1}", LocalConfigFile.FullName, Ex.ToString());
 				}
 			}
 			return ProjectConfig;
