Merging latest UGS changes (added support for parallel syncing, fixed syncing single-change)

diff --git a/Engine/Source/Programs/UnrealGameSync/UnrealGameSync/Controls/WorkspaceControl.cs b/Engine/Source/Programs/UnrealGameSync/UnrealGameSync/Controls/WorkspaceControl.cs
--- a/Engine/Source/Programs/UnrealGameSync/UnrealGameSync/Controls/WorkspaceControl.cs
+++ b/Engine/Source/Programs/UnrealGameSync/UnrealGameSync/Controls/WorkspaceControl.cs
@@ -5040,12 +5040,6 @@ namespace UnrealGameSync
 			Owner.ModifyApplicationSettings();
 		}
 
-		private void OptionsContextMenu_PerforceSettings_Click(object sender, EventArgs e)
-		{
-			PerforceSyncSettingsWindow Window = new PerforceSyncSettingsWindow(Settings);
-			Window.ShowDialog();
-		}
-
 		private void OptionsContextMenu_TabNames_Stream_Click(object sender, EventArgs e)
 		{
 			Owner.SetTabNames(TabLabels.Stream);
diff --git a/Engine/Source/Programs/UnrealGameSync/UnrealGameSync/Forms/ApplicationSettingsWindow.cs b/Engine/Source/Programs/UnrealGameSync/UnrealGameSync/Forms/ApplicationSettingsWindow.cs
--- a/Engine/Source/Programs/UnrealGameSync/UnrealGameSync/Forms/ApplicationSettingsWindow.cs
+++ b/Engine/Source/Programs/UnrealGameSync/UnrealGameSync/Forms/ApplicationSettingsWindow.cs
@@ -80,6 +80,8 @@ namespace UnrealGameSync
 			this.UserNameTextBox.Select(UserNameTextBox.TextLength, 0);
 			this.UserNameTextBox.CueBanner = (DefaultUserName == null)? "Default" : String.Format("Default ({0})", DefaultUserName);
 
+			this.ParallelSyncThreadsSpinner.Value = Math.Max(Math.Min(Settings.SyncOptions.NumThreads, ParallelSyncThreadsSpinner.Maximum), ParallelSyncThreadsSpinner.Minimum);
+
 			this.DepotPathTextBox.Text = InitialDepotPath;
 			this.DepotPathTextBox.Select(DepotPathTextBox.TextLength, 0);
 			this.DepotPathTextBox.CueBanner = DeploymentSettings.DefaultDepotPath;
@@ -186,8 +188,9 @@ namespace UnrealGameSync
 				Key.SetValue("UnrealGameSync", String.Format("\"{0}\" -RestoreState", OriginalExecutableFileName));
 			}
 
-			if(Settings.bKeepInTray != KeepInTrayCheckBox.Checked)
+			if (Settings.bKeepInTray != KeepInTrayCheckBox.Checked || Settings.SyncOptions.NumThreads != ParallelSyncThreadsSpinner.Value)
 			{
+				Settings.SyncOptions.NumThreads = (int)ParallelSyncThreadsSpinner.Value;
 				Settings.bKeepInTray = KeepInTrayCheckBox.Checked;
 				Settings.Save();
 			}
@@ -206,5 +209,11 @@ namespace UnrealGameSync
 		{
 			AutomationPortTextBox.Enabled = EnableAutomationCheckBox.Checked;
 		}
+
+		private void AdvancedBtn_Click(object sender, EventArgs e)
+		{
+			PerforceSyncSettingsWindow Window = new PerforceSyncSettingsWindow(Settings);
+			Window.ShowDialog();
+		}
 	}
 }
diff --git a/Engine/Source/Programs/UnrealGameSync/UnrealGameSync/Forms/PerforceSyncSettingsWindow.cs b/Engine/Source/Programs/UnrealGameSync/UnrealGameSync/Forms/PerforceSyncSettingsWindow.cs
--- a/Engine/Source/Programs/UnrealGameSync/UnrealGameSync/Forms/PerforceSyncSettingsWindow.cs
+++ b/Engine/Source/Programs/UnrealGameSync/UnrealGameSync/Forms/PerforceSyncSettingsWindow.cs
@@ -1,4 +1,4 @@
-﻿// Copyright Epic Games, Inc. All Rights Reserved.
+// Copyright Epic Games, Inc. All Rights Reserved.
 
 using System;
 using System.Collections.Generic;
@@ -26,7 +26,6 @@ namespace UnrealGameSync
 		{
 			PerforceSyncOptions SyncOptions = Settings.SyncOptions;
 			NumRetriesTextBox.Text = (SyncOptions.NumRetries > 0)? SyncOptions.NumRetries.ToString() : "";
-			NumThreadsTextBox.Text = (SyncOptions.NumThreads > 0)? SyncOptions.NumThreads.ToString() : "";
 			TcpBufferSizeText.Text = (SyncOptions.TcpBufferSize > 0)? SyncOptions.TcpBufferSize.ToString() : "";
 		}
 
@@ -39,13 +38,6 @@ namespace UnrealGameSync
 				return;
 			}
 
-			int NewNumThreads = 0;
-			if(NumThreadsTextBox.Text.Length > 0 && !int.TryParse(NumThreadsTextBox.Text, out NewNumThreads))
-			{
-				MessageBox.Show("Invalid value for number of threads");
-				return;
-			}
-
 			int NewTcpBufferSize = 0;
 			if(TcpBufferSizeText.Text.Length > 0 && !int.TryParse(TcpBufferSizeText.Text, out NewTcpBufferSize))
 			{
@@ -54,7 +46,6 @@ namespace UnrealGameSync
 			}
 
 			Settings.SyncOptions.NumRetries = NewNumRetries;
-			Settings.SyncOptions.NumThreads = NewNumThreads;
 			Settings.SyncOptions.TcpBufferSize = NewTcpBufferSize;
 			Settings.Save();
 
diff --git a/Engine/Source/Programs/UnrealGameSync/UnrealGameSync/OutputAdapters.cs b/Engine/Source/Programs/UnrealGameSync/UnrealGameSync/OutputAdapters.cs
--- a/Engine/Source/Programs/UnrealGameSync/UnrealGameSync/OutputAdapters.cs
+++ b/Engine/Source/Programs/UnrealGameSync/UnrealGameSync/OutputAdapters.cs
@@ -46,6 +46,26 @@ namespace UnrealGameSync
 		}
 	}
 
+	class ThreadSafeTextWriter : LineBasedTextWriter
+	{
+		object LockObject;
+		TextWriter Inner;
+
+		public ThreadSafeTextWriter(TextWriter Inner)
+		{
+			this.LockObject = new object();
+			this.Inner = Inner;
+		}
+
+		protected override void FlushLine(string Line)
+		{
+			lock (LockObject)
+			{
+				Inner.WriteLine(Line);
+			}
+		}
+	}
+
 	class NullTextWriter : LineBasedTextWriter
 	{
 		protected override void FlushLine(string Line)
diff --git a/Engine/Source/Programs/UnrealGameSync/UnrealGameSync/Perforce.cs b/Engine/Source/Programs/UnrealGameSync/UnrealGameSync/Perforce.cs
--- a/Engine/Source/Programs/UnrealGameSync/UnrealGameSync/Perforce.cs
+++ b/Engine/Source/Programs/UnrealGameSync/UnrealGameSync/Perforce.cs
@@ -218,6 +218,7 @@ namespace UnrealGameSync
 		public string ClientPath;
 		public string Path;
 		public string Action;
+		public long FileSize;
 		public PerforceFileFlags Flags;
 		public int Revision;
 		public bool IsMapped;
@@ -234,6 +235,12 @@ namespace UnrealGameSync
 				Tags.TryGetValue("headAction", out Action);
 			}
 
+			string FileSizeText;
+			if (Tags.TryGetValue("fileSize", out FileSizeText))
+			{
+				FileSize = long.Parse(FileSizeText);
+			}
+
 			string Type;
 			if(Tags.TryGetValue("headType", out Type))
 			{
diff --git a/Engine/Source/Programs/UnrealGameSync/UnrealGameSync/Properties/AssemblyInfo.cs b/Engine/Source/Programs/UnrealGameSync/UnrealGameSync/Properties/AssemblyInfo.cs
--- a/Engine/Source/Programs/UnrealGameSync/UnrealGameSync/Properties/AssemblyInfo.cs
+++ b/Engine/Source/Programs/UnrealGameSync/UnrealGameSync/Properties/AssemblyInfo.cs
@@ -34,5 +34,5 @@ using System.Runtime.InteropServices;
 // You can specify all the values or you can default the Build and Revision Numbers 
 // by using the '*' as shown below:
 // [assembly: AssemblyVersion("1.0.*")]
-[assembly: AssemblyVersion("1.193.*")]
-[assembly: AssemblyFileVersion("1.193.0.0")]
+[assembly: AssemblyVersion("1.195.*")]
+[assembly: AssemblyFileVersion("1.195.0.0")]
diff --git a/Engine/Source/Programs/UnrealGameSync/UnrealGameSync/Workspace.cs b/Engine/Source/Programs/UnrealGameSync/UnrealGameSync/Workspace.cs
--- a/Engine/Source/Programs/UnrealGameSync/UnrealGameSync/Workspace.cs
+++ b/Engine/Source/Programs/UnrealGameSync/UnrealGameSync/Workspace.cs
@@ -1,6 +1,7 @@
 // Copyright Epic Games, Inc. All Rights Reserved.
 
 using System;
+using System.Collections.Concurrent;
 using System.Collections.Generic;
 using System.Diagnostics;
 using System.Globalization;
@@ -105,8 +106,8 @@ namespace UnrealGameSync
 		const string VersionHeaderFileName = "/Engine/Source/Runtime/Launch/Resources/Version.h";
 		const string ObjectVersionFileName = "/Engine/Source/Runtime/Core/Private/UObject/ObjectVersion.cpp";
 
-		readonly string LocalVersionHeaderFileName = VersionHeaderFileName.Replace('/', '\\');
-		readonly string LocalObjectVersionFileName = ObjectVersionFileName.Replace('/', '\\');
+		static readonly string LocalVersionHeaderFileName = VersionHeaderFileName.Replace('/', '\\');
+		static readonly string LocalObjectVersionFileName = ObjectVersionFileName.Replace('/', '\\');
 
 		public readonly PerforceConnection Perforce;
 		public readonly string LocalRootPath;
@@ -161,11 +162,43 @@ namespace UnrealGameSync
 			}
 		}
 
+		class SyncBatchBuilder
+		{
+			public int MaxCommandsPerList { get; }
+			public long MaxSizePerList { get; }
+			public Queue<List<string>> Batches { get; }
+
+			List<string> Commands;
+			long Size;
+
+			public SyncBatchBuilder(int MaxCommandsPerList, long MaxSizePerList)
+			{
+				this.MaxCommandsPerList = MaxCommandsPerList;
+				this.MaxSizePerList = MaxSizePerList;
+				this.Batches = new Queue<List<string>>();
+			}
+
+			public void Add(string NewCommand, long NewSize)
+			{
+				if (Commands == null || Commands.Count >= MaxCommandsPerList || Size + NewSize >= MaxSizePerList)
+				{
+					Commands = new List<string>();
+					Batches.Enqueue(Commands);
+					Size = 0;
+				}
+
+				Commands.Add(NewCommand);
+				Size += NewSize;
+			}
+		}
+
 		class SyncTree
 		{
 			public bool bCanUseWildcard;
+			public int TotalIncludedFiles;
+			public long TotalSize;
 			public int TotalExcludedFiles;
-			public HashSet<string> IncludedFiles = new HashSet<string>();
+			public Dictionary<string, long> IncludedFiles = new Dictionary<string, long>();
 			public Dictionary<string, SyncTree> NameToSubTree = new Dictionary<string, SyncTree>(StringComparer.OrdinalIgnoreCase);
 
 			public SyncTree(bool bCanUseWildcard)
@@ -184,18 +217,20 @@ namespace UnrealGameSync
 				return Result;
 			}
 
-			public void IncludeFile(string Path)
+			public void IncludeFile(string Path, long Size)
 			{
 				int Idx = Path.IndexOf('/');
 				if (Idx == -1)
 				{
-					IncludedFiles.Add(Path);
+					IncludedFiles.Add(Path, Size);
 				}
 				else
 				{
 					SyncTree SubTree = FindOrAddSubTree(Path.Substring(0, Idx));
-					SubTree.IncludeFile(Path.Substring(Idx + 1));
+					SubTree.IncludeFile(Path.Substring(Idx + 1), Size);
 				}
+				TotalIncludedFiles++;
+				TotalSize += Size;
 			}
 
 			public void ExcludeFile(string Path)
@@ -209,21 +244,21 @@ namespace UnrealGameSync
 				TotalExcludedFiles++;
 			}
 
-			public void GetOptimizedSyncCommands(string Prefix, int ChangeNumber, SortedSet<string> SyncCommands)
+			public void GetOptimizedSyncCommands(string Prefix, int ChangeNumber, SyncBatchBuilder Builder)
 			{
-				if (bCanUseWildcard && TotalExcludedFiles == 0)
+				if (bCanUseWildcard && TotalExcludedFiles == 0 && TotalSize < Builder.MaxSizePerList)
 				{
-					SyncCommands.Add(String.Format("{0}/...@{1}", Prefix, ChangeNumber));
+					Builder.Add(String.Format("{0}/...@{1}", Prefix, ChangeNumber), TotalSize);
 				}
 				else
 				{
-					foreach(string File in IncludedFiles)
+					foreach (KeyValuePair<string, long> File in IncludedFiles)
 					{
-						SyncCommands.Add(String.Format("{0}/{1}@{2}", Prefix, File, ChangeNumber));
+						Builder.Add(String.Format("{0}/{1}@{2}", Prefix, File.Key, ChangeNumber), File.Value);
 					}
 					foreach (KeyValuePair<string, SyncTree> Pair in NameToSubTree)
 					{
-						Pair.Value.GetOptimizedSyncCommands(String.Format("{0}/{1}", Prefix, PerforceUtils.EscapePath(Pair.Key)), ChangeNumber, SyncCommands);
+						Pair.Value.GetOptimizedSyncCommands(String.Format("{0}/{1}", Prefix, PerforceUtils.EscapePath(Pair.Key)), ChangeNumber, Builder);
 					}
 				}
 			}
@@ -506,25 +541,28 @@ namespace UnrealGameSync
 
 					// Create a tree to store the sync path
 					SyncTree SyncTree = new SyncTree(false);
-					foreach (string RelativeSyncPath in RelativeSyncPaths)
+					if (!Context.Options.HasFlag(WorkspaceUpdateOptions.SyncSingleChange))
 					{
-						const string WildcardSuffix = "/...";
-						if (RelativeSyncPath.EndsWith(WildcardSuffix, StringComparison.Ordinal))
+						foreach (string RelativeSyncPath in RelativeSyncPaths)
 						{
-							SyncTree Leaf = SyncTree;
-
-							string[] Fragments = RelativeSyncPath.Split('/');
-							for (int Idx = 1; Idx < Fragments.Length - 1; Idx++)
+							const string WildcardSuffix = "/...";
+							if (RelativeSyncPath.EndsWith(WildcardSuffix, StringComparison.Ordinal))
 							{
-								Leaf = Leaf.FindOrAddSubTree(Fragments[Idx]);
-							}
+								SyncTree Leaf = SyncTree;
+
+								string[] Fragments = RelativeSyncPath.Split('/');
+								for (int Idx = 1; Idx < Fragments.Length - 1; Idx++)
+								{
+									Leaf = Leaf.FindOrAddSubTree(Fragments[Idx]);
+								}
 
-							Leaf.bCanUseWildcard = true;
+								Leaf.bCanUseWildcard = true;
+							}
 						}
 					}
 
 					// Find all the server changes, and anything that's opened for edit locally. We need to sync files we have open to schedule a resolve.
-					SortedSet<string> SyncCommands = new SortedSet<string>(StringComparer.Ordinal);
+					SyncBatchBuilder BatchBuilder = new SyncBatchBuilder(200, 100 * 1024 * 1024);
 					List<string> SyncDepotPaths = new List<string>();
 					using(RecordCounter Counter = new RecordCounter(Progress, "Filtering files..."))
 					{
@@ -551,7 +589,7 @@ namespace UnrealGameSync
 								// If it doesn't exist locally, just add a sync command for it
 								if (String.IsNullOrEmpty(SyncRecord.ClientPath))
 								{
-									SyncCommands.Add(String.Format("{0}@{1}", SyncRecord.DepotPath, PendingChangeNumber));
+									BatchBuilder.Add(String.Format("{0}@{1}", SyncRecord.DepotPath, PendingChangeNumber), SyncRecord.FileSize);
 									SyncDepotPaths.Add(SyncRecord.DepotPath);
 									continue;
 								}
@@ -572,7 +610,7 @@ namespace UnrealGameSync
 								// Make sure it's under the current directory. Not sure why this would happen, just being safe.
 								if (!FullName.StartsWith(LocalRootPrefix, StringComparison.OrdinalIgnoreCase))
 								{
-									SyncCommands.Add(String.Format("{0}@{1}", SyncRecord.DepotPath, PendingChangeNumber));
+									BatchBuilder.Add(String.Format("{0}@{1}", SyncRecord.DepotPath, PendingChangeNumber), SyncRecord.FileSize);
 									SyncDepotPaths.Add(SyncRecord.DepotPath);
 									continue;
 								}
@@ -581,7 +619,7 @@ namespace UnrealGameSync
 								string RelativePath = FullName.Substring(LocalRootPrefix.Length).Replace('\\', '/');
 								if (Filter.Matches(RelativePath))
 								{
-									SyncTree.IncludeFile(RelativePath);
+									SyncTree.IncludeFile(RelativePath, SyncRecord.FileSize);
 									SyncDepotPaths.Add(SyncRecord.DepotPath);
 								}
 								else
@@ -591,13 +629,13 @@ namespace UnrealGameSync
 							}
 						}
 					}
-					SyncTree.GetOptimizedSyncCommands(ClientRootPath, PendingChangeNumber, SyncCommands);
+					SyncTree.GetOptimizedSyncCommands(ClientRootPath, PendingChangeNumber, BatchBuilder);
 
 					// Find all the depot paths that will be synced
 					HashSet<string> RemainingDepotPaths = new HashSet<string>(StringComparer.OrdinalIgnoreCase);
 					RemainingDepotPaths.UnionWith(SyncDepotPaths);
 
-					WorkspaceUpdateResult SyncResult = SyncFileRevisions("Syncing files...", Context, SyncCommands.ToList(), RemainingDepotPaths, out StatusMessage);
+					WorkspaceUpdateResult SyncResult = SyncFileRevisions("Syncing files...", Context, BatchBuilder.Batches, RemainingDepotPaths, out StatusMessage);
 					if(SyncResult != WorkspaceUpdateResult.Success)
 					{
 						return SyncResult;
@@ -1182,13 +1220,135 @@ namespace UnrealGameSync
 			return bMatch;
 		}
 
+		class SyncState
+		{
+			public int TotalDepotPaths;
+			public HashSet<string> RemainingDepotPaths;
+			public Queue<List<string>> SyncCommandLists;
+			public string StatusMessage;
+			public WorkspaceUpdateResult Result = WorkspaceUpdateResult.Success;
+		}
+
 		WorkspaceUpdateResult SyncFileRevisions(string Prefix, WorkspaceUpdateContext Context, List<string> SyncCommands, HashSet<string> RemainingDepotPaths, out string StatusMessage)
 		{
-			int TotalDepotPaths = RemainingDepotPaths.Count;
+			Queue<List<string>> SyncCommandLists = new Queue<List<string>>();
+			SyncCommandLists.Enqueue(SyncCommands);
+			return SyncFileRevisions(Prefix, Context, SyncCommandLists, RemainingDepotPaths, out StatusMessage);
+		}
+
+		WorkspaceUpdateResult SyncFileRevisions(string Prefix, WorkspaceUpdateContext Context, Queue<List<string>> SyncCommandLists, HashSet<string> RemainingDepotPaths, out string StatusMessage)
+		{
+			// Figure out the number of additional background threads we want to run with. We can run worker on the current thread.
+			int NumExtraThreads = Math.Max(Math.Min(SyncCommandLists.Count, Context.PerforceSyncOptions.NumThreads) - 1, 0);
+
+			List<Thread> ChildThreads = new List<Thread>(NumExtraThreads);
+			try
+			{
+				// Create the state object shared by all the worker threads
+				SyncState State = new SyncState();
+				State.TotalDepotPaths = RemainingDepotPaths.Count;
+				State.RemainingDepotPaths = RemainingDepotPaths;
+				State.SyncCommandLists = SyncCommandLists;
+
+				// Wrapper writer around the log class to prevent multiple threads writing to it at once
+				ThreadSafeTextWriter LogWrapper = new ThreadSafeTextWriter(Log);
+
+				// Delegate for updating the sync state after a file has been synced
+				Action<PerforceFileRecord, LineBasedTextWriter> SyncOutput = (Record, LocalLog) => { UpdateSyncState(Prefix, Record, State, Progress, LocalLog); };
+
+				// Create all the child threads
+				for (int ThreadIdx = 0; ThreadIdx < NumExtraThreads; ThreadIdx++)
+				{
+					int ThreadNumber = ThreadIdx + 2;
+					Thread ChildThread = new Thread(() => StaticSyncWorker(ThreadNumber, Perforce, PendingChangeNumber, Context, State, SyncOutput, LogWrapper));
+					ChildThreads.Add(ChildThread);
+					ChildThread.Start();
+				}
+
+				// Run one worker on the current thread
+				StaticSyncWorker(1, Perforce, PendingChangeNumber, Context, State, SyncOutput, LogWrapper);
+
+				// Wait for all the background threads to finish
+				foreach (Thread ChildThread in ChildThreads)
+				{
+					ChildThread.Join();
+				}
 
+				// Return the result that was set on the state object
+				StatusMessage = State.StatusMessage;
+				return State.Result;
+			}
+			finally
+			{
+				foreach (Thread ChildThread in ChildThreads)
+				{
+					ChildThread.Abort();
+				}
+				foreach (Thread ChildThread in ChildThreads)
+				{
+					ChildThread.Join();
+				}
+			}
+		}
+
+		static void UpdateSyncState(string Prefix, PerforceFileRecord Record, SyncState State, ProgressValue Progress, TextWriter Log)
+		{
+			lock (State)
+			{
+				State.RemainingDepotPaths.Remove(Record.DepotPath);
+
+				string Message = String.Format("{0} ({1}/{2})", Prefix, State.TotalDepotPaths - State.RemainingDepotPaths.Count, State.TotalDepotPaths);
+				float Fraction = Math.Min((float)(State.TotalDepotPaths - State.RemainingDepotPaths.Count) / (float)State.TotalDepotPaths, 1.0f);
+				Progress.Set(Message, Fraction);
+
+				Log.WriteLine("p4>   {0} {1}", Record.Action, Record.ClientPath);
+			}
+		}
+
+		static void StaticSyncWorker(int ThreadNumber, PerforceConnection Perforce, int PendingChangeNumber, WorkspaceUpdateContext Context, SyncState State, Action<PerforceFileRecord, LineBasedTextWriter> SyncOutput, LineBasedTextWriter GlobalLog)
+		{
+			PrefixedTextWriter ThreadLog = new PrefixedTextWriter(String.Format("{0}:", ThreadNumber), GlobalLog);
+			for (; ; )
+			{
+				// Remove the next batch that needs to be synced
+				List<string> SyncCommands;
+				lock (State)
+				{
+					if (State.Result == WorkspaceUpdateResult.Success && State.SyncCommandLists.Count > 0)
+					{
+						SyncCommands = State.SyncCommandLists.Dequeue();
+					}
+					else
+					{
+						break;
+					}
+				}
+
+				// Sync the files
+				string StatusMessage;
+				WorkspaceUpdateResult Result = StaticSyncFileRevisions(Perforce, PendingChangeNumber, Context, SyncCommands, Record => SyncOutput(Record, ThreadLog), ThreadLog, out StatusMessage);
+
+				// If it failed, try to set it on the state if nothing else has failed first
+				if (Result != WorkspaceUpdateResult.Success)
+				{
+					lock (State)
+					{
+						if (State.Result == WorkspaceUpdateResult.Success)
+						{
+							State.Result = Result;
+							State.StatusMessage = StatusMessage;
+						}
+					}
+					break;
+				}
+			}
+		}
+
+		static WorkspaceUpdateResult StaticSyncFileRevisions(PerforceConnection Perforce, int PendingChangeNumber, WorkspaceUpdateContext Context, List<string> SyncCommands, Action<PerforceFileRecord> SyncOutput, TextWriter Log, out string StatusMessage)
+		{
 			// Sync them all
 			List<string> TamperedFiles = new List<string>();
-			if(!Perforce.Sync(SyncCommands, Record => UpdateSyncProgress(Prefix, Record, RemainingDepotPaths, TotalDepotPaths), TamperedFiles, false, Context.PerforceSyncOptions, Log))
+			if(!Perforce.Sync(SyncCommands, SyncOutput, TamperedFiles, false, Context.PerforceSyncOptions, Log))
 			{
 				StatusMessage = "Aborted sync due to errors.";
 				return WorkspaceUpdateResult.FailedToSync;
@@ -1314,17 +1474,6 @@ namespace UnrealGameSync
 			return true;
 		}
 
-		void UpdateSyncProgress(string Prefix, PerforceFileRecord Record, HashSet<string> RemainingFiles, int NumFiles)
-		{
-			RemainingFiles.Remove(Record.DepotPath);
-
-			string Message = String.Format("{0} ({1}/{2})", Prefix, NumFiles - RemainingFiles.Count, NumFiles);
-			float Fraction = Math.Min((float)(NumFiles - RemainingFiles.Count) / (float)NumFiles, 1.0f);
-			Progress.Set(Message, Fraction);
-
-			Log.WriteLine("p4>   {0} {1}", Record.Action, Record.ClientPath);
-		}
-
 		bool UpdateVersionFile(string ClientPath, Dictionary<string, string> VersionStrings, int ChangeNumber)
 		{
 			return UpdateVersionFile(ClientPath, ChangeNumber, Text => UpdateVersionStrings(Text, VersionStrings));
