Merging latest version of UGS to 4.25 (include stream name in notification dialogs, telemetry changes, include username when launching P4VC)

#rnx

diff --git a/Engine/Source/Programs/UnrealGameSync/UnrealGameSync/Controls/WorkspaceControl.cs b/Engine/Source/Programs/UnrealGameSync/UnrealGameSync/Controls/WorkspaceControl.cs
--- a/Engine/Source/Programs/UnrealGameSync/UnrealGameSync/Controls/WorkspaceControl.cs
+++ b/Engine/Source/Programs/UnrealGameSync/UnrealGameSync/Controls/WorkspaceControl.cs
@@ -1761,7 +1761,8 @@ namespace UnrealGameSync
 			if (RequireNotificationForChange != -1)
 			{
 				// Format the platform list
-				StringBuilder PlatformList = new StringBuilder(NotifyBuilds[0].BuildType);
+				StringBuilder PlatformList = new StringBuilder();
+				PlatformList.AppendFormat("{0}: {1}", StreamName, NotifyBuilds[0].BuildType);
 				for (int Idx = 1; Idx < NotifyBuilds.Count - 1; Idx++)
 				{
 					PlatformList.AppendFormat(", {0}", NotifyBuilds[Idx].BuildType);
@@ -4251,12 +4252,24 @@ namespace UnrealGameSync
 
 		private void BuildListContextMenu_MoreInfo_Click(object sender, EventArgs e)
 		{
-			if (!Utility.SpawnHiddenProcess("p4vc.exe", String.Format("-p\"{0}\" -c{1} change {2}", Workspace.Perforce.ServerAndPort ?? "perforce:1666", Workspace.ClientName, ContextMenuChange.Number)))
+			if (!Utility.SpawnHiddenProcess("p4vc.exe", String.Format("-p\"{0}\" -u\"{1}\" -c{2} change {3}", Workspace.Perforce.ServerAndPort ?? "perforce:1666", Workspace.Perforce.UserName, Workspace.ClientName, ContextMenuChange.Number)))
 			{
 				MessageBox.Show("Unable to spawn p4vc. Check you have P4V installed.");
 			}
 		}
 
+		private void BuildListContextMenu_ViewInSwarm_Click(object sender, EventArgs e)
+		{
+			string SwarmURL;
+			if (TryGetProjectSetting(PerforceMonitor.LatestProjectConfigFile, "SwarmURL", out SwarmURL))
+			{
+				Process.Start(String.Format("{0}/changes/{1}", SwarmURL, ContextMenuChange.Number));
+			} else
+			{
+				MessageBox.Show("Swarm URL is not configured.");
+			}
+		}
+
 		private void BuildListContextMenu_AddStar_Click(object sender, EventArgs e)
 		{
 			if (MessageBox.Show("Starred builds are meant to convey a stable, verified build to the rest of the team. Do not star a build unless it has been fully tested.\n\nAre you sure you want to star this build?", "Confirm star", MessageBoxButtons.YesNo) == DialogResult.Yes)
diff --git a/Engine/Source/Programs/UnrealGameSync/UnrealGameSync/DeploymentSettings.cs b/Engine/Source/Programs/UnrealGameSync/UnrealGameSync/DeploymentSettings.cs
--- a/Engine/Source/Programs/UnrealGameSync/UnrealGameSync/DeploymentSettings.cs
+++ b/Engine/Source/Programs/UnrealGameSync/UnrealGameSync/DeploymentSettings.cs
@@ -2,6 +2,7 @@
 
 using System;
 using System.Collections.Generic;
+using System.IO;
 using System.Linq;
 using System.Text;
 using System.Threading.Tasks;
@@ -13,6 +14,15 @@ namespace UnrealGameSync
 	/// </summary>
 	static partial class DeploymentSettings
 	{
+		/// <summary>
+		/// Delegate used to create a telemetry sink
+		/// </summary>
+		/// <param name="UserName">The default Perforce user name</param>
+		/// <param name="SessionId">Unique identifier for this session</param>
+		/// <param name="Log">Log writer</param>
+		/// <returns>New telemetry sink instance</returns>
+		public delegate ITelemetrySink CreateTelemetrySinkDelegate(string UserName, string SessionId, TextWriter Log);
+
 		/// <summary>
 		/// SQL connection string used to connect to the database for telemetry and review data. The 'Program' class is a partial class, to allow an
 		/// opportunistically included C# source file in NotForLicensees/ProgramSettings.cs to override this value in a static constructor.
@@ -26,8 +36,8 @@ namespace UnrealGameSync
 		public static readonly string DefaultDepotPath = null;
 
 		/// <summary>
-		/// Whether to send telemetry data by default. Can be useful for finding users in need of a hardware upgrade, but can bloat the DB for large teams.
+		/// Delegate used to create a new telemetry sink
 		/// </summary>
-		public static bool bSendTelemetry = false;
+		public static readonly CreateTelemetrySinkDelegate CreateTelemetrySink = (UserName, SessionId, Log) => new NullTelemetrySink();
 	}
 }
diff --git a/Engine/Source/Programs/UnrealGameSync/UnrealGameSync/Program.cs b/Engine/Source/Programs/UnrealGameSync/UnrealGameSync/Program.cs
--- a/Engine/Source/Programs/UnrealGameSync/UnrealGameSync/Program.cs
+++ b/Engine/Source/Programs/UnrealGameSync/UnrealGameSync/Program.cs
@@ -85,36 +85,50 @@ namespace UnrealGameSync
 			string DataFolder = Path.Combine(Environment.GetFolderPath(Environment.SpecialFolder.LocalApplicationData), "UnrealGameSync");
 			Directory.CreateDirectory(DataFolder);
 
-			using(TelemetryWriter Telemetry = new TelemetryWriter(DeploymentSettings.ApiUrl, Path.Combine(DataFolder, "Telemetry.log")))
+			// Create the log file
+			using (TimestampLogWriter Log = new TimestampLogWriter(new BoundedLogWriter(Path.Combine(DataFolder, "UnrealGameSync.log"))))
 			{
-				AppDomain.CurrentDomain.UnhandledException += CurrentDomain_UnhandledException;
+				Log.WriteLine("Application version: {0}", Assembly.GetExecutingAssembly().GetName().Version);
+				Log.WriteLine("Started at {0}", DateTime.Now.ToString());
 
-				// Create the log file
-				using (TimestampLogWriter Log = new TimestampLogWriter(new BoundedLogWriter(Path.Combine(DataFolder, "UnrealGameSync.log"))))
+				string SessionId = Guid.NewGuid().ToString();
+				Log.WriteLine("SessionId: {0}", SessionId);
+
+				if (ServerAndPort == null || UserName == null)
 				{
-					Log.WriteLine("Application version: {0}", Assembly.GetExecutingAssembly().GetName().Version);
-					Log.WriteLine("Started at {0}", DateTime.Now.ToString());
+					Log.WriteLine("Missing server settings; finding defaults.");
+					GetDefaultServerSettings(ref ServerAndPort, ref UserName, Log);
+					Utility.SaveGlobalPerforceSettings(ServerAndPort, UserName, BaseUpdatePath);
+				}
 
-					if (ServerAndPort == null || UserName == null)
-					{
-						Log.WriteLine("Missing server settings; finding defaults.");
-						GetDefaultServerSettings(ref ServerAndPort, ref UserName, Log);
-						Utility.SaveGlobalPerforceSettings(ServerAndPort, UserName, BaseUpdatePath);
-					}
+				using (BoundedLogWriter TelemetryLog = new BoundedLogWriter(Path.Combine(DataFolder, "Telemetry.log")))
+				{
+					TelemetryLog.WriteLine("Creating telemetry sink for session {0}", SessionId);
 
-					PerforceConnection DefaultConnection = new PerforceConnection(UserName, null, ServerAndPort);
-					using (UpdateMonitor UpdateMonitor = new UpdateMonitor(DefaultConnection, UpdatePath))
+					ITelemetrySink PrevTelemetrySink = Telemetry.ActiveSink;
+					using (ITelemetrySink TelemetrySink = DeploymentSettings.CreateTelemetrySink(UserName, SessionId, TelemetryLog))
 					{
-						ProgramApplicationContext Context = new ProgramApplicationContext(DefaultConnection, UpdateMonitor, DeploymentSettings.ApiUrl, DataFolder, ActivateEvent, bRestoreState, UpdateSpawn, ProjectFileName, bUnstable, Log);
-						Application.Run(Context);
+						Telemetry.ActiveSink = TelemetrySink;
+
+						Telemetry.SendEvent("Startup", new { User = Environment.UserName, Machine = Environment.MachineName });
+
+						AppDomain.CurrentDomain.UnhandledException += CurrentDomain_UnhandledException;
 
-						if(UpdateMonitor.IsUpdateAvailable && UpdateSpawn != null)
+						PerforceConnection DefaultConnection = new PerforceConnection(UserName, null, ServerAndPort);
+						using (UpdateMonitor UpdateMonitor = new UpdateMonitor(DefaultConnection, UpdatePath))
 						{
-							InstanceMutex.Close();
-							bool bLaunchUnstable = UpdateMonitor.RelaunchUnstable ?? bUnstable;
-							Utility.SpawnProcess(UpdateSpawn, "-restorestate" + (bLaunchUnstable? " -unstable" : ""));
+							ProgramApplicationContext Context = new ProgramApplicationContext(DefaultConnection, UpdateMonitor, DeploymentSettings.ApiUrl, DataFolder, ActivateEvent, bRestoreState, UpdateSpawn, ProjectFileName, bUnstable, Log);
+							Application.Run(Context);
+
+							if (UpdateMonitor.IsUpdateAvailable && UpdateSpawn != null)
+							{
+								InstanceMutex.Close();
+								bool bLaunchUnstable = UpdateMonitor.RelaunchUnstable ?? bUnstable;
+								Utility.SpawnProcess(UpdateSpawn, "-restorestate" + (bLaunchUnstable ? " -unstable" : ""));
+							}
 						}
 					}
+					Telemetry.ActiveSink = PrevTelemetrySink;
 				}
 			}
 		}
@@ -165,7 +179,7 @@ namespace UnrealGameSync
 					ExceptionTrace.Append("\nInner Exception:\n");
 					ExceptionTrace.Append(InnerEx.ToString());
 				}
-				TelemetryWriter.Enqueue(TelemetryErrorType.Crash, ExceptionTrace.ToString(), null, DateTime.Now);
+				Telemetry.SendEvent("Crash", new { Exception = Ex });
 			}
 		}
 
diff --git a/Engine/Source/Programs/UnrealGameSync/UnrealGameSync/Properties/AssemblyInfo.cs b/Engine/Source/Programs/UnrealGameSync/UnrealGameSync/Properties/AssemblyInfo.cs
--- a/Engine/Source/Programs/UnrealGameSync/UnrealGameSync/Properties/AssemblyInfo.cs
+++ b/Engine/Source/Programs/UnrealGameSync/UnrealGameSync/Properties/AssemblyInfo.cs
@@ -34,5 +34,5 @@ using System.Runtime.InteropServices;
 // You can specify all the values or you can default the Build and Revision Numbers 
 // by using the '*' as shown below:
 // [assembly: AssemblyVersion("1.0.*")]
-[assembly: AssemblyVersion("1.195.*")]
-[assembly: AssemblyFileVersion("1.195.0.0")]
+[assembly: AssemblyVersion("1.198.*")]
+[assembly: AssemblyFileVersion("1.198.0.0")]
diff --git a/Engine/Source/Programs/UnrealGameSync/UnrealGameSync/Telemetry.cs b/Engine/Source/Programs/UnrealGameSync/UnrealGameSync/Telemetry.cs
--- a/Engine/Source/Programs/UnrealGameSync/UnrealGameSync/Telemetry.cs
+++ b/Engine/Source/Programs/UnrealGameSync/UnrealGameSync/Telemetry.cs
@@ -3,249 +3,232 @@
 using System;
 using System.Collections.Concurrent;
 using System.Collections.Generic;
-using System.Data.SqlClient;
-using System.Diagnostics;
+using System.IO;
 using System.Linq;
 using System.Net;
-using System.Net.Sockets;
 using System.Reflection;
 using System.Text;
 using System.Threading;
 using System.Threading.Tasks;
+using System.Web;
 using System.Web.Script.Serialization;
 
 namespace UnrealGameSync
 {
-	class TelemetryTimingData
+	/// <summary>
+	/// Interface for a telemetry sink
+	/// </summary>
+	public interface ITelemetrySink : IDisposable
 	{
-		public string Action { get; set; }
-		public string Result { get; set; }
-		public string UserName { get; set; }
-		public string Project { get; set; }
-		public DateTime Timestamp { get; set; }
-		public float Duration { get; set; }
+		/// <summary>
+		/// Sends a telemetry event with the given information
+		/// </summary>
+		/// <param name="EventName">Name of the event</param>
+		/// <param name="Attributes">Arbitrary object to include in the payload</param>
+		void SendEvent(string EventName, object Attributes);
 	}
 
-	enum TelemetryErrorType
+	/// <summary>
+	/// Telemetry sink that discards all events
+	/// </summary>
+	class NullTelemetrySink : ITelemetrySink
 	{
-		Crash,
-	}
-
-	class TelemetryErrorData
-	{
-		public TelemetryErrorType Type { get; set; }
-		public string Text { get; set; }
-		public string UserName { get; set; }
-		public string Project { get; set; }
-		public DateTime Timestamp { get; set; }
-	}
-
-	class TelemetryStopwatch : IDisposable
-	{
-		readonly string Action;
-		readonly string Project;
-		readonly DateTime StartTime;
-		string Result;
-		DateTime EndTime;
-
-		public TelemetryStopwatch(string InAction, string InProject)
-		{
-			Action = InAction;
-			Project = InProject;
-			StartTime = DateTime.UtcNow;
-		}
-
-		public TimeSpan Stop(string InResult)
-		{
-			EndTime = DateTime.UtcNow;
-			Result = InResult;
-			return Elapsed;
-		}
-
+		/// <inheritdoc/>
 		public void Dispose()
 		{
-			if(Result == null)
-			{
-				Stop("Aborted");
-			}
-			TelemetryWriter.Enqueue(Action, Result, Project, StartTime, (float)Elapsed.TotalSeconds);
 		}
 
-		public TimeSpan Elapsed
+		/// <inheritdoc/>
+		public void SendEvent(string EventName, object Attributes)
 		{
-			get { return ((Result == null)? DateTime.UtcNow : EndTime) - StartTime; }
 		}
 	}
 
-	class TelemetryWriter : IDisposable
+	/// <summary>
+	/// Epic internal telemetry sink using the data router
+	/// </summary>
+	class EpicTelemetrySink : ITelemetrySink
 	{
-		static TelemetryWriter Instance;
-
-		string ApiUrl;
-		Thread WorkerThread;
-		BoundedLogWriter LogWriter;
-		bool bDisposing;
-		ConcurrentQueue<TelemetryTimingData> QueuedTimingData = new ConcurrentQueue<TelemetryTimingData>(); 
-		ConcurrentQueue<TelemetryErrorData> QueuedErrorData = new ConcurrentQueue<TelemetryErrorData>(); 
-		AutoResetEvent RefreshEvent = new AutoResetEvent(false);
-
-		public TelemetryWriter(string InApiUrl, string InLogFileName)
+		/// <summary>
+		/// Combined url to post event streams to
+		/// </summary>
+		string Url;
+
+		/// <summary>
+		/// Lock used to modify the event queue
+		/// </summary>
+		object LockObject = new object();
+
+		/// <summary>
+		/// Whether a flush is queued
+		/// </summary>
+		bool bHasPendingFlush = false;
+
+		/// <summary>
+		/// List of pending events
+		/// </summary>
+		List<string> PendingEvents = new List<string>();
+
+		/// <summary>
+		/// The log writer to use
+		/// </summary>
+		TextWriter Log;
+
+		/// <summary>
+		/// Constructor
+		/// </summary>
+		public EpicTelemetrySink(string Url, TextWriter Log)
 		{
-			Instance = this;
-
-			ApiUrl = InApiUrl;
+			this.Url = Url;
+			this.Log = Log;
 
-			LogWriter = new BoundedLogWriter(InLogFileName);
-			LogWriter.WriteLine("Using connection string: {0}", ApiUrl);
-
-			WorkerThread = new Thread(() => WorkerThreadCallback());
-			WorkerThread.Start();
+			Log.WriteLine("Posting to URL: {0}", Url);
 		}
 
+		/// <inheritdoc/>
 		public void Dispose()
 		{
-			bDisposing = true;
+			Flush();
+		}
 
-			if(WorkerThread != null)
+		/// <inheritdoc/>
+		public void Flush()
+		{
+			for (; ; )
 			{
-				RefreshEvent.Set();
-				if(!WorkerThread.Join(10 * 1000))
+				lock (LockObject)
 				{
-					WorkerThread.Abort();
-					WorkerThread.Join();
+					if (!bHasPendingFlush)
+					{
+						break;
+					}
 				}
-				WorkerThread = null;
+				Thread.Sleep(10);
 			}
-			if(LogWriter != null)
+		}
+
+		/// <inheritdoc/>
+		public void SendEvent(string EventName, object Attributes)
+		{
+			string AttributesText = new JavaScriptSerializer().Serialize(Attributes);
+			if (AttributesText[0] != '{')
 			{
-				LogWriter.Dispose();
-				LogWriter = null;
+				throw new Exception("Expected event data with named properties");
 			}
 
-			Instance = null;
+			string EventText = AttributesText.Insert(1, String.Format("\"EventName\":\"{0}\",", HttpUtility.JavaScriptStringEncode(EventName)));
+			lock (PendingEvents)
+			{
+				PendingEvents.Add(EventText);
+				if (!bHasPendingFlush)
+				{
+					ThreadPool.QueueUserWorkItem(Obj => BackgroundFlush());
+					bHasPendingFlush = true;
+				}
+			}
 		}
 
-		void WorkerThreadCallback()
+		/// <summary>
+		/// Synchronously sends a telemetry event
+		/// </summary>
+		void BackgroundFlush()
 		{
-			string Version = Assembly.GetExecutingAssembly().GetName().Version.ToString();
-
-			string IpAddress = "Unknown";
-			try
+			for (; ; )
 			{
-				IPHostEntry HostEntry = Dns.GetHostEntry(Dns.GetHostName());
-				foreach (IPAddress Address in HostEntry.AddressList)
+				try
 				{
-					if (Address.AddressFamily == AddressFamily.InterNetwork)
+					// Generate the content for this event
+					List<string> Events = new List<string>();
+					lock (LockObject)
 					{
-						IpAddress = Address.ToString();
-						break;
+						if (PendingEvents.Count == 0)
+						{
+							bHasPendingFlush = false;
+							break;
+						}
+
+						Events.AddRange(PendingEvents);
+						PendingEvents.Clear();
 					}
-				}
-			}
-			catch
-			{
-			}
 
-			while(!bDisposing)
-			{
-				// Wait for an update
-				RefreshEvent.WaitOne();
+					// Print all the events we're sending
+					foreach (string Event in Events)
+					{
+						Log.WriteLine("Sending Event: {0}", Event);
+					}
 
-				// Send all the timing data
-				TelemetryTimingData TimingData;
-				while(QueuedTimingData.TryDequeue(out TimingData))
-				{
-					while(!bDisposing && !SendTimingData(TimingData, Version, IpAddress) && ApiUrl != null)
+					// Convert the content to UTF8
+					string ContentText = String.Format("{{\"Events\":[{0}]}}", String.Join(",", Events));
+					byte[] Content = Encoding.UTF8.GetBytes(ContentText);
+
+					// Post the event data
+					HttpWebRequest Request = (HttpWebRequest)WebRequest.Create(Url);
+					Request.Method = "POST";
+					Request.ContentType = "application/json";
+					Request.UserAgent = "ue/ugs";
+					Request.Timeout = 5000;
+					Request.ContentLength = Content.Length;
+					Request.ContentType = "application/json";
+					using (Stream RequestStream = Request.GetRequestStream())
 					{
-						RefreshEvent.WaitOne(10 * 1000);
+						RequestStream.Write(Content, 0, Content.Length);
 					}
-				}
 
-				// Send all the error data
-				TelemetryErrorData ErrorData;
-				while(QueuedErrorData.TryDequeue(out ErrorData))
+					// Wait for the response and dispose of it immediately
+					using (HttpWebResponse Response = (HttpWebResponse)Request.GetResponse())
+					{
+						Log.WriteLine("Response: {0}", (int)Response.StatusCode);
+					}
+				}
+				catch (WebException Ex)
 				{
-					while(!SendErrorData(ErrorData, Version, IpAddress) && ApiUrl != null)
+					// Handle errors. Any non-200 responses automatically generate a WebException.
+					HttpWebResponse Response = (HttpWebResponse)Ex.Response;
+					if (Response == null)
 					{
-						if(bDisposing) break;
-						RefreshEvent.WaitOne(10 * 1000);
+						Log.WriteLine("Exception while attempting to send event: {0}", Ex.ToString());
+					}
+					else
+					{
+						string ResponseText;
+						using (Stream ResponseStream = Response.GetResponseStream())
+						{
+							MemoryStream MemoryStream = new MemoryStream();
+							ResponseStream.CopyTo(MemoryStream);
+							ResponseText = Encoding.UTF8.GetString(MemoryStream.ToArray());
+						}
+						Log.WriteLine("EpicTelemetrySink: Failed to send analytics event. Code = {0}. Desc = {1}. Response = {2}.", (int)Response.StatusCode, Response.StatusDescription, ResponseText);
 					}
 				}
+				catch (Exception Ex)
+				{
+					Log.WriteLine("Exception while attempting to send event: {0}", Ex.ToString());
+				}
 			}
 		}
+	}
 
-		bool SendTimingData(TelemetryTimingData Data, string Version, string IpAddress)
-		{
-			if(!DeploymentSettings.bSendTelemetry)
-			{
-				return true;
-			}
-
-			try
-			{
-				Stopwatch Timer = Stopwatch.StartNew();
-				LogWriter.WriteLine("Posting timing data... ({0}, {1}, {2}, {3}, {4}, {5})", Data.Action, Data.Result, Data.UserName, Data.Project, Data.Timestamp, Data.Duration);
-				RESTApi.POST(ApiUrl, "telemetry", new JavaScriptSerializer().Serialize(Data), string.Format("Version={0}", Version), string.Format("IpAddress={0}", IpAddress));
-				LogWriter.WriteLine("Done in {0}ms.", Timer.ElapsedMilliseconds);
-				return true;
-			}
-			catch(Exception Ex)
-			{
-				LogWriter.WriteException(Ex, "Failed with exception.");
-				return false;
-			}
-		}
-
-		bool SendErrorData(TelemetryErrorData Data, string Version, string IpAddress)
-		{
-			try
-			{
-				Stopwatch Timer = Stopwatch.StartNew();
-				LogWriter.WriteLine("Posting error data... ({0}, {1})", Data.Type, Data.Timestamp);
-				RESTApi.POST(ApiUrl, "error", new JavaScriptSerializer().Serialize(Data), string.Format("Version={0}", Version), string.Format("IpAddress={0}", IpAddress));
-				LogWriter.WriteLine("Done in {0}ms.", Timer.ElapsedMilliseconds);
-				return true;
-			}
-			catch(Exception Ex)
-			{
-				LogWriter.WriteException(Ex, "Failed with exception.");
-				return false;
-			}
-		}
-
-		public static void Enqueue(string Action, string Result, string Project, DateTime Timestamp, float Duration)
+	/// <summary>
+	/// Global telemetry static class
+	/// </summary>
+	public static class Telemetry
+	{
+		/// <summary>
+		/// The current telemetry provider
+		/// </summary>
+		public static ITelemetrySink ActiveSink
 		{
-			TelemetryWriter Writer = Instance;
-			if(Writer != null)
-			{
-				TelemetryTimingData Telemetry = new TelemetryTimingData();
-				Telemetry.Action = Action;
-				Telemetry.Result = Result;
-				Telemetry.UserName = Environment.UserName;
-				Telemetry.Project = Project;
-				Telemetry.Timestamp = Timestamp;
-				Telemetry.Duration = Duration;
-
-				Writer.QueuedTimingData.Enqueue(Telemetry);
-				Writer.RefreshEvent.Set();
-			}
+			get; set;
 		}
 
-		public static void Enqueue(TelemetryErrorType Type, string Text, string Project, DateTime Timestamp)
+		/// <summary>
+		/// Sends a telemetry event with the given information
+		/// </summary>
+		/// <param name="EventName">Name of the event</param>
+		/// <param name="Attributes">Arbitrary object to include in the payload</param>
+		public static void SendEvent(string EventName, object Attributes)
 		{
-			TelemetryWriter Writer = Instance;
-			if(Writer != null)
-			{
-				TelemetryErrorData Error = new TelemetryErrorData();
-				Error.Type = Type;
-				Error.Text = Text;
-				Error.UserName = Environment.UserName;
-				Error.Project = Project;
-				Error.Timestamp = Timestamp;
-
-				Writer.QueuedErrorData.Enqueue(Error);
-				Writer.RefreshEvent.Set();
-			}
+			ActiveSink.SendEvent(EventName, Attributes);
 		}
 	}
 }
diff --git a/Engine/Source/Programs/UnrealGameSync/UnrealGameSync/TelemetryStopwatch.cs b/Engine/Source/Programs/UnrealGameSync/UnrealGameSync/TelemetryStopwatch.cs
new file mode 100644
--- /dev/null
+++ b/Engine/Source/Programs/UnrealGameSync/UnrealGameSync/TelemetryStopwatch.cs
@@ -0,0 +1,69 @@
+// Copyright Epic Games, Inc. All Rights Reserved.
+
+using System;
+using System.Collections.Concurrent;
+using System.Collections.Generic;
+using System.Data.SqlClient;
+using System.Diagnostics;
+using System.Linq;
+using System.Net;
+using System.Net.Sockets;
+using System.Reflection;
+using System.Text;
+using System.Threading;
+using System.Threading.Tasks;
+using System.Web.Script.Serialization;
+
+namespace UnrealGameSync
+{
+	class TelemetryStopwatch : IDisposable
+	{
+		readonly string EventName;
+		readonly Dictionary<string, object> EventData;
+		readonly Stopwatch Timer;
+
+		public TelemetryStopwatch(string EventName, string Project)
+		{
+			this.EventName = EventName;
+
+			EventData = new Dictionary<string, object>();
+			EventData["Project"] = Project;
+
+			Timer = Stopwatch.StartNew();
+		}
+
+		public void AddData(object Data)
+		{
+			foreach (PropertyInfo Property in Data.GetType().GetProperties())
+			{
+				EventData[Property.Name] = Property.GetValue(Data);
+			}
+		}
+
+		public TimeSpan Stop(string InResult)
+		{
+			if (Timer.IsRunning)
+			{
+				Timer.Stop();
+
+				EventData["Result"] = InResult;
+				EventData["TimeSeconds"] = Timer.Elapsed.TotalSeconds;
+			}
+			return Elapsed;
+		}
+
+		public void Dispose()
+		{
+			if (Timer.IsRunning)
+			{
+				Stop("Aborted");
+			}
+			Telemetry.SendEvent(EventName, EventData);
+		}
+
+		public TimeSpan Elapsed
+		{
+			get { return Timer.Elapsed; }
+		}
+	}
+}
diff --git a/Engine/Source/Programs/UnrealGameSync/UnrealGameSync/UnrealGameSync.csproj b/Engine/Source/Programs/UnrealGameSync/UnrealGameSync/UnrealGameSync.csproj
--- a/Engine/Source/Programs/UnrealGameSync/UnrealGameSync/UnrealGameSync.csproj
+++ b/Engine/Source/Programs/UnrealGameSync/UnrealGameSync/UnrealGameSync.csproj
@@ -99,6 +99,7 @@
     <Reference Include="System" />
     <Reference Include="System.Activities" />
     <Reference Include="System.Core" />
+    <Reference Include="System.Web" />
     <Reference Include="System.Web.Extensions" />
     <Reference Include="System.Xml.Linq" />
     <Reference Include="System.Data.DataSetExtensions" />
@@ -393,6 +394,7 @@
     </Compile>
     <Compile Include="Taskbar.cs" />
     <Compile Include="Telemetry.cs" />
+    <Compile Include="TelemetryStopwatch.cs" />
     <Compile Include="UpdateMonitor.cs" />
     <Compile Include="UserSettings.cs" />
     <Compile Include="RESTApi.cs" />
diff --git a/Engine/Source/Programs/UnrealGameSync/UnrealGameSync/Workspace.cs b/Engine/Source/Programs/UnrealGameSync/UnrealGameSync/Workspace.cs
--- a/Engine/Source/Programs/UnrealGameSync/UnrealGameSync/Workspace.cs
+++ b/Engine/Source/Programs/UnrealGameSync/UnrealGameSync/Workspace.cs
@@ -386,7 +386,7 @@ namespace UnrealGameSync
 			int NumFilesSynced = 0;
 			if(Context.Options.HasFlag(WorkspaceUpdateOptions.Sync) || Context.Options.HasFlag(WorkspaceUpdateOptions.SyncSingleChange))
 			{
-				using(TelemetryStopwatch Stopwatch = new TelemetryStopwatch("Sync", TelemetryProjectPath))
+				using(TelemetryStopwatch SyncTelemetryStopwatch = new TelemetryStopwatch("Workspace_Sync", TelemetryProjectPath))
 				{
 					Log.WriteLine("Syncing to {0}...", PendingChangeNumber);
 
@@ -434,100 +434,103 @@ namespace UnrealGameSync
 					// If the hash differs, enumerate everything in the workspace to find what needs to be removed
 					if (NextSyncFilterHash != CurrentSyncFilterHash)
 					{
-						Log.WriteLine("Filter has changed ({0} -> {1}); finding files in workspace that need to be removed.", (String.IsNullOrEmpty(CurrentSyncFilterHash))? "None" : CurrentSyncFilterHash, NextSyncFilterHash);
-
-						// Find all the files that are in this workspace
-						List<PerforceFileRecord> HaveFiles = Context.HaveFiles;
-						if(HaveFiles == null)
+						using (TelemetryStopwatch FilterStopwatch = new TelemetryStopwatch("Workspace_Sync_FilterChanged", TelemetryProjectPath))
 						{
-							HaveFiles = new List<PerforceFileRecord>();
-							using(RecordCounter HaveCounter = new RecordCounter(Progress, "Sync filter changed; checking workspace..."))
+							Log.WriteLine("Filter has changed ({0} -> {1}); finding files in workspace that need to be removed.", (String.IsNullOrEmpty(CurrentSyncFilterHash)) ? "None" : CurrentSyncFilterHash, NextSyncFilterHash);
+
+							// Find all the files that are in this workspace
+							List<PerforceFileRecord> HaveFiles = Context.HaveFiles;
+							if (HaveFiles == null)
 							{
-								if(!Perforce.Have("//...", Record => { HaveFiles.Add(Record); HaveCounter.Increment(); }, Log))
+								HaveFiles = new List<PerforceFileRecord>();
+								using (RecordCounter HaveCounter = new RecordCounter(Progress, "Sync filter changed; checking workspace..."))
 								{
-									StatusMessage = "Unable to query files.";
-									return WorkspaceUpdateResult.FailedToSync;
+									if (!Perforce.Have("//...", Record => { HaveFiles.Add(Record); HaveCounter.Increment(); }, Log))
+									{
+										StatusMessage = "Unable to query files.";
+										return WorkspaceUpdateResult.FailedToSync;
+									}
 								}
+								Context.HaveFiles = HaveFiles;
 							}
-							Context.HaveFiles = HaveFiles;
-						}
 
-						// Build a filter for the current sync paths
-						FileFilter SyncPathsFilter = new FileFilter(FileFilterType.Exclude);
-						foreach(string RelativeSyncPath in RelativeSyncPaths)
-						{
-							SyncPathsFilter.Include(RelativeSyncPath);
-						}
+							// Build a filter for the current sync paths
+							FileFilter SyncPathsFilter = new FileFilter(FileFilterType.Exclude);
+							foreach (string RelativeSyncPath in RelativeSyncPaths)
+							{
+								SyncPathsFilter.Include(RelativeSyncPath);
+							}
 
-						// Remove all the files that are not included by the filter
-						List<string> RemoveDepotPaths = new List<string>();
-						foreach(PerforceFileRecord HaveFile in HaveFiles)
-						{
-							try
+							// Remove all the files that are not included by the filter
+							List<string> RemoveDepotPaths = new List<string>();
+							foreach (PerforceFileRecord HaveFile in HaveFiles)
 							{
-								string FullPath = Path.GetFullPath(HaveFile.Path);
-								if(MatchFilter(FullPath, SyncPathsFilter) && !MatchFilter(FullPath, UserFilter))
+								try
 								{
-									Log.WriteLine("  {0}", HaveFile.DepotPath);
-									RemoveDepotPaths.Add(HaveFile.DepotPath);
+									string FullPath = Path.GetFullPath(HaveFile.Path);
+									if (MatchFilter(FullPath, SyncPathsFilter) && !MatchFilter(FullPath, UserFilter))
+									{
+										Log.WriteLine("  {0}", HaveFile.DepotPath);
+										RemoveDepotPaths.Add(HaveFile.DepotPath);
+									}
+								}
+								catch (PathTooLongException)
+								{
+									// We don't actually care about this when looking for files to remove. Perforce may think that it's synced the path, and silently failed. Just ignore it.
 								}
 							}
-							catch(PathTooLongException)
+
+							// Check if there are any paths outside the regular sync paths
+							if (RemoveDepotPaths.Count > 0)
 							{
-								// We don't actually care about this when looking for files to remove. Perforce may think that it's synced the path, and silently failed. Just ignore it.
-							}
-						}
+								bool bDeleteListMatches = true;
 
-						// Check if there are any paths outside the regular sync paths
-						if(RemoveDepotPaths.Count > 0)
-						{
-							bool bDeleteListMatches = true;
+								Dictionary<string, bool> NewDeleteFiles = new Dictionary<string, bool>(StringComparer.OrdinalIgnoreCase);
+								foreach (string RemoveDepotPath in RemoveDepotPaths)
+								{
+									bool bDelete;
+									if (!Context.DeleteFiles.TryGetValue(RemoveDepotPath, out bDelete))
+									{
+										bDeleteListMatches = false;
+										bDelete = true;
+									}
+									NewDeleteFiles[RemoveDepotPath] = bDelete;
+								}
+								Context.DeleteFiles = NewDeleteFiles;
 
-							Dictionary<string, bool> NewDeleteFiles = new Dictionary<string, bool>(StringComparer.OrdinalIgnoreCase);
-							foreach(string RemoveDepotPath in RemoveDepotPaths)
-							{
-								bool bDelete;
-								if(!Context.DeleteFiles.TryGetValue(RemoveDepotPath, out bDelete))
+								if (!bDeleteListMatches)
 								{
-									bDeleteListMatches = false;
-									bDelete = true;
+									StatusMessage = String.Format("Cancelled after finding {0} files excluded by filter", NewDeleteFiles.Count);
+									return WorkspaceUpdateResult.FilesToDelete;
 								}
-								NewDeleteFiles[RemoveDepotPath] = bDelete;
-							}
-							Context.DeleteFiles = NewDeleteFiles;
 
-							if(!bDeleteListMatches)
-							{
-								StatusMessage = String.Format("Cancelled after finding {0} files excluded by filter", NewDeleteFiles.Count);
-								return WorkspaceUpdateResult.FilesToDelete;
+								RemoveDepotPaths.RemoveAll(x => !Context.DeleteFiles[x]);
 							}
 
-							RemoveDepotPaths.RemoveAll(x => !Context.DeleteFiles[x]);
-						}
-
-						// Actually delete any files that we don't want
-						if(RemoveDepotPaths.Count > 0)
-						{
-							// Clear the current sync filter hash. If the sync is canceled, we'll be in an indeterminate state, and we should always clean next time round.
-							CurrentSyncFilterHash = "INVALID";
+							// Actually delete any files that we don't want
+							if (RemoveDepotPaths.Count > 0)
+							{
+								// Clear the current sync filter hash. If the sync is canceled, we'll be in an indeterminate state, and we should always clean next time round.
+								CurrentSyncFilterHash = "INVALID";
 
-							// Find all the depot paths that will be synced
-							HashSet<string> RemainingDepotPathsToRemove = new HashSet<string>(StringComparer.OrdinalIgnoreCase);
-							RemainingDepotPathsToRemove.UnionWith(RemoveDepotPaths);
+								// Find all the depot paths that will be synced
+								HashSet<string> RemainingDepotPathsToRemove = new HashSet<string>(StringComparer.OrdinalIgnoreCase);
+								RemainingDepotPathsToRemove.UnionWith(RemoveDepotPaths);
 
-							// Build the list of revisions to sync
-							List<string> RevisionsToRemove = new List<string>();
-							RevisionsToRemove.AddRange(RemoveDepotPaths.Select(x => String.Format("{0}#0", x)));
+								// Build the list of revisions to sync
+								List<string> RevisionsToRemove = new List<string>();
+								RevisionsToRemove.AddRange(RemoveDepotPaths.Select(x => String.Format("{0}#0", x)));
 
-							WorkspaceUpdateResult RemoveResult = SyncFileRevisions("Removing files...", Context, RevisionsToRemove, RemainingDepotPathsToRemove, out StatusMessage);
-							if(RemoveResult != WorkspaceUpdateResult.Success)
-							{
-								return RemoveResult;
+								WorkspaceUpdateResult RemoveResult = SyncFileRevisions("Removing files...", Context, RevisionsToRemove, RemainingDepotPathsToRemove, out StatusMessage);
+								if (RemoveResult != WorkspaceUpdateResult.Success)
+								{
+									return RemoveResult;
+								}
 							}
-						}
 
-						// Update the sync filter hash. We've removed any files we need to at this point.
-						CurrentSyncFilterHash = NextSyncFilterHash;
+							// Update the sync filter hash. We've removed any files we need to at this point.
+							CurrentSyncFilterHash = NextSyncFilterHash;
+						}
 					}
 
 					// Create a filter for all the files we don't want
@@ -635,10 +638,19 @@ namespace UnrealGameSync
 					HashSet<string> RemainingDepotPaths = new HashSet<string>(StringComparer.OrdinalIgnoreCase);
 					RemainingDepotPaths.UnionWith(SyncDepotPaths);
 
-					WorkspaceUpdateResult SyncResult = SyncFileRevisions("Syncing files...", Context, BatchBuilder.Batches, RemainingDepotPaths, out StatusMessage);
-					if(SyncResult != WorkspaceUpdateResult.Success)
+					using (TelemetryStopwatch TransferStopwatch = new TelemetryStopwatch("Workspace_Sync_TransferFiles", TelemetryProjectPath))
 					{
-						return SyncResult;
+						TransferStopwatch.AddData(new { MachineName = Environment.MachineName, DomainName = Environment.UserDomainName, ServerAndPort = Perforce.ServerAndPort, UserName = Perforce.UserName, IncludedFiles = SyncTree.TotalIncludedFiles, ExcludedFiles = SyncTree.TotalExcludedFiles, Size = SyncTree.TotalSize, NumThreads = Context.PerforceSyncOptions.NumThreads });
+
+						WorkspaceUpdateResult SyncResult = SyncFileRevisions("Syncing files...", Context, BatchBuilder.Batches, RemainingDepotPaths, out StatusMessage);
+						if (SyncResult != WorkspaceUpdateResult.Success)
+						{
+							TransferStopwatch.AddData(new { SyncResult = SyncResult.ToString(), CompletedFilesFiles = SyncDepotPaths.Count - RemainingDepotPaths.Count });
+							return SyncResult;
+						}
+
+						TransferStopwatch.Stop("Ok");
+						TransferStopwatch.AddData(new { TransferRate = SyncTree.TotalSize / Math.Max(TransferStopwatch.Elapsed.TotalSeconds, 0.0001f) });
 					}
 
 					int VersionChangeNumber = -1;
@@ -843,7 +855,7 @@ namespace UnrealGameSync
 					}
 
 					// Update the timing info
-					Times.Add(new Tuple<string,TimeSpan>("Sync", Stopwatch.Stop("Success")));
+					Times.Add(new Tuple<string,TimeSpan>("Sync", SyncTelemetryStopwatch.Stop("Success")));
 
 					// Save the number of files synced
 					NumFilesSynced = SyncDepotPaths.Count;
@@ -854,7 +866,7 @@ namespace UnrealGameSync
 			// Extract an archive from the depot path
 			if(Context.Options.HasFlag(WorkspaceUpdateOptions.SyncArchives))
 			{
-				using(TelemetryStopwatch Stopwatch = new TelemetryStopwatch("Archives", TelemetryProjectPath))
+				using(TelemetryStopwatch Stopwatch = new TelemetryStopwatch("Workspace_SyncArchives", TelemetryProjectPath))
 				{
 					// Create the directory for extracted archive manifests
 					string ManifestDirectoryName;
@@ -927,7 +939,7 @@ namespace UnrealGameSync
 			// Generate project files in the workspace
 			if(Context.Options.HasFlag(WorkspaceUpdateOptions.GenerateProjectFiles))
 			{
-				using(TelemetryStopwatch Stopwatch = new TelemetryStopwatch("Prj gen", TelemetryProjectPath))
+				using(TelemetryStopwatch Stopwatch = new TelemetryStopwatch("Workspace_GenerateProjectFiles", TelemetryProjectPath))
 				{
 					Progress.Set("Generating project files...", 0.0f);
 
@@ -1001,8 +1013,7 @@ namespace UnrealGameSync
 				}
 
 				// Execute them all
-				string TelemetryEventName = (Context.UserBuildStepObjects.Count > 0)? "CustomBuild" : Context.Options.HasFlag(WorkspaceUpdateOptions.UseIncrementalBuilds) ? "Compile" : "FullCompile";
-				using(TelemetryStopwatch Stopwatch = new TelemetryStopwatch(TelemetryEventName, TelemetryProjectPath))
+				using(TelemetryStopwatch Stopwatch = new TelemetryStopwatch("Workspace_Build", TelemetryProjectPath))
 				{
 					Progress.Set("Starting build...", 0.0f);
 
@@ -1030,8 +1041,10 @@ namespace UnrealGameSync
 							switch(Step.Type)
 							{
 								case BuildStepType.Compile:
-									using(TelemetryStopwatch StepStopwatch = new TelemetryStopwatch("Compile:" + Step.Target, TelemetryProjectPath))
+									using(TelemetryStopwatch StepStopwatch = new TelemetryStopwatch("Workspace_Execute_Compile", TelemetryProjectPath))
 									{
+										StepStopwatch.AddData(new { Target = Step.Target });
+
 										string CommandLine = String.Format("{0} {1} {2} {3} -NoHotReloadFromIDE", Step.Target, Step.Platform, Step.Configuration, Utility.ExpandVariables(Step.Arguments ?? "", Context.Variables));
 										if(!Context.Options.HasFlag(WorkspaceUpdateOptions.UseIncrementalBuilds) || bForceClean)
 										{
@@ -1053,8 +1066,10 @@ namespace UnrealGameSync
 									}
 									break;
 								case BuildStepType.Cook:
-									using(TelemetryStopwatch StepStopwatch = new TelemetryStopwatch("Cook/Launch: " + Path.GetFileNameWithoutExtension(Step.FileName), TelemetryProjectPath))
+									using(TelemetryStopwatch StepStopwatch = new TelemetryStopwatch("Workspace_Execute_Cook", TelemetryProjectPath))
 									{
+										StepStopwatch.AddData(new { Project = Path.GetFileNameWithoutExtension(Step.FileName) });
+
 										string LocalRunUAT = Path.Combine(LocalRootPath, "Engine", "Build", "BatchFiles", "RunUAT.bat");
 										string Arguments = String.Format("/C \"\"{0}\" -profile=\"{1}\"\"", LocalRunUAT, Path.Combine(LocalRootPath, Step.FileName));
 										Log.WriteLine("uat> Running {0} {1}", LocalRunUAT, Arguments);
@@ -1071,8 +1086,10 @@ namespace UnrealGameSync
 									}
 									break;
 								case BuildStepType.Other:
-									using(TelemetryStopwatch StepStopwatch = new TelemetryStopwatch("Custom: " + Path.GetFileNameWithoutExtension(Step.FileName), TelemetryProjectPath))
+									using(TelemetryStopwatch StepStopwatch = new TelemetryStopwatch("Workspace_Execute_Custom", TelemetryProjectPath))
 									{
+										StepStopwatch.AddData(new { FileName = Path.GetFileNameWithoutExtension(Step.FileName) });
+
 										string ToolFileName = Path.Combine(LocalRootPath, Utility.ExpandVariables(Step.FileName, Context.Variables));
 										string ToolWorkingDir = String.IsNullOrWhiteSpace(Step.WorkingDir) ? Path.GetDirectoryName(ToolFileName) : Utility.ExpandVariables(Step.WorkingDir, Context.Variables);
 										string ToolArguments = Utility.ExpandVariables(Step.Arguments ?? "", Context.Variables);
