Merging latest UGS changes (allow disabling sync categories by default, set default number of sync threads to 2, hide "View in Swarm" button if not configured)

diff --git a/Engine/Source/Programs/UnrealGameSync/UnrealGameSync/Controls/WorkspaceControl.cs b/Engine/Source/Programs/UnrealGameSync/UnrealGameSync/Controls/WorkspaceControl.cs
--- a/Engine/Source/Programs/UnrealGameSync/UnrealGameSync/Controls/WorkspaceControl.cs
+++ b/Engine/Source/Programs/UnrealGameSync/UnrealGameSync/Controls/WorkspaceControl.cs
@@ -926,7 +926,7 @@ namespace UnrealGameSync
 				}
 			}
 
-			string[] CombinedSyncFilter = UserSettings.GetCombinedSyncFilter(GetSyncCategories(), Settings.SyncView, Settings.SyncExcludedCategories, WorkspaceSettings.SyncView, WorkspaceSettings.SyncIncludedCategories, WorkspaceSettings.SyncExcludedCategories);
+			string[] CombinedSyncFilter = UserSettings.GetCombinedSyncFilter(GetSyncCategories(), Settings.SyncView, Settings.SyncCategories, WorkspaceSettings.SyncView, WorkspaceSettings.SyncCategories);
 
 			ConfigSection PerforceSection = PerforceMonitor.LatestProjectConfigFile.FindSection("Perforce");
 
@@ -3875,6 +3875,9 @@ namespace UnrealGameSync
 
 						BuildListContextMenu_CustomTool_End.Visible = (CustomToolEnd > CustomToolStart);
 
+						string SwarmURL;
+						BuildListContextMenu_ViewInSwarm.Visible = TryGetProjectSetting(ProjectConfigFile, "SwarmURL", out SwarmURL);
+							
 						BuildListContextMenu.Show(BuildList, Args.Location);
 					}
 				}
@@ -4264,7 +4267,8 @@ namespace UnrealGameSync
 			if (TryGetProjectSetting(PerforceMonitor.LatestProjectConfigFile, "SwarmURL", out SwarmURL))
 			{
 				Process.Start(String.Format("{0}/changes/{1}", SwarmURL, ContextMenuChange.Number));
-			} else
+			}
+			else
 			{
 				MessageBox.Show("Swarm URL is not configured.");
 			}
@@ -4719,7 +4723,7 @@ namespace UnrealGameSync
 				ExtraSafeToDeleteExtensions = "";
 			}
 
-			string[] CombinedSyncFilter = UserSettings.GetCombinedSyncFilter(GetSyncCategories(), Settings.SyncView, Settings.SyncExcludedCategories, WorkspaceSettings.SyncView, WorkspaceSettings.SyncIncludedCategories, WorkspaceSettings.SyncExcludedCategories);
+			string[] CombinedSyncFilter = UserSettings.GetCombinedSyncFilter(GetSyncCategories(), Settings.SyncView, Settings.SyncCategories, WorkspaceSettings.SyncView, WorkspaceSettings.SyncCategories);
 			List<string> SyncPaths = Workspace.GetSyncPaths(WorkspaceSettings.bSyncAllProjects ?? Settings.bSyncAllProjects, CombinedSyncFilter);
 
 			CleanWorkspaceWindow.DoClean(ParentForm, Workspace.Perforce, BranchDirectoryName, Workspace.ClientRootPath, SyncPaths, ExtraSafeToDeleteFolders.Split('\n'), ExtraSafeToDeleteExtensions.Split('\n'), Log);
@@ -4983,15 +4987,14 @@ namespace UnrealGameSync
 
 		private void OptionsContextMenu_SyncFilter_Click(object sender, EventArgs e)
 		{
-			SyncFilter Filter = new SyncFilter(GetSyncCategories(), Settings.SyncView, Settings.SyncExcludedCategories, Settings.bSyncAllProjects, Settings.bIncludeAllProjectsInSolution, WorkspaceSettings.SyncView, WorkspaceSettings.SyncIncludedCategories, WorkspaceSettings.SyncExcludedCategories, WorkspaceSettings.bSyncAllProjects, WorkspaceSettings.bIncludeAllProjectsInSolution);
+			SyncFilter Filter = new SyncFilter(GetSyncCategories(), Settings.SyncView, Settings.SyncCategories, Settings.bSyncAllProjects, Settings.bIncludeAllProjectsInSolution, WorkspaceSettings.SyncView, WorkspaceSettings.SyncCategories, WorkspaceSettings.bSyncAllProjects, WorkspaceSettings.bIncludeAllProjectsInSolution);
 			if (Filter.ShowDialog() == DialogResult.OK)
 			{
-				Settings.SyncExcludedCategories = Filter.GlobalExcludedCategories;
+				Settings.SyncCategories = Filter.GlobalSyncCategories;
 				Settings.SyncView = Filter.GlobalView;
 				Settings.bSyncAllProjects = Filter.bGlobalSyncAllProjects;
 				Settings.bIncludeAllProjectsInSolution = Filter.bGlobalIncludeAllProjectsInSolution;
-				WorkspaceSettings.SyncIncludedCategories = Filter.WorkspaceIncludedCategories;
-				WorkspaceSettings.SyncExcludedCategories = Filter.WorkspaceExcludedCategories;
+				WorkspaceSettings.SyncCategories = Filter.WorkspaceSyncCategories;
 				WorkspaceSettings.SyncView = Filter.WorkspaceView;
 				WorkspaceSettings.bSyncAllProjects = Filter.bWorkspaceSyncAllProjects;
 				WorkspaceSettings.bIncludeAllProjectsInSolution = Filter.bWorkspaceIncludeAllProjectsInSolution;
diff --git a/Engine/Source/Programs/UnrealGameSync/UnrealGameSync/Forms/SyncFilter.cs b/Engine/Source/Programs/UnrealGameSync/UnrealGameSync/Forms/SyncFilter.cs
--- a/Engine/Source/Programs/UnrealGameSync/UnrealGameSync/Forms/SyncFilter.cs
+++ b/Engine/Source/Programs/UnrealGameSync/UnrealGameSync/Forms/SyncFilter.cs
@@ -16,37 +16,39 @@ namespace UnrealGameSync
 	{
 		Dictionary<Guid, WorkspaceSyncCategory> UniqueIdToCategory;
 		public string[] GlobalView;
-		public Guid[] GlobalExcludedCategories;
+		public Dictionary<Guid, bool> GlobalSyncCategories;
 		public bool bGlobalSyncAllProjects;
 		public bool bGlobalIncludeAllProjectsInSolution;
 		public string[] WorkspaceView;
-		public Guid[] WorkspaceIncludedCategories;
-		public Guid[] WorkspaceExcludedCategories;
+		public Dictionary<Guid, bool> WorkspaceSyncCategories;
 		public bool? bWorkspaceSyncAllProjects;
 		public bool? bWorkspaceIncludeAllProjectsInSolution;
 
-		public SyncFilter(Dictionary<Guid, WorkspaceSyncCategory> InUniqueIdToCategory, string[] InGlobalView, Guid[] InGlobalExcludedCategories, bool bInGlobalProjectOnly, bool bInGlobalIncludeAllProjectsInSolution, string[] InWorkspaceView, Guid[] InWorkspaceIncludedCategories, Guid[] InWorkspaceExcludedCategories, bool? bInWorkspaceProjectOnly, bool? bInWorkspaceIncludeAllProjectsInSolution)
+		public SyncFilter(Dictionary<Guid, WorkspaceSyncCategory> InUniqueIdToCategory, string[] InGlobalView, Dictionary<Guid, bool> InGlobalSyncCategories, bool bInGlobalProjectOnly, bool bInGlobalIncludeAllProjectsInSolution, string[] InWorkspaceView, Dictionary<Guid, bool> InWorkspaceSyncCategories, bool? bInWorkspaceProjectOnly, bool? bInWorkspaceIncludeAllProjectsInSolution)
 		{
 			InitializeComponent();
 
 			UniqueIdToCategory = InUniqueIdToCategory;
-			GlobalExcludedCategories = InGlobalExcludedCategories;
+			GlobalSyncCategories = InGlobalSyncCategories;
 			GlobalView = InGlobalView;
 			bGlobalSyncAllProjects = bInGlobalProjectOnly;
 			bGlobalIncludeAllProjectsInSolution = bInGlobalIncludeAllProjectsInSolution;
-			WorkspaceIncludedCategories = InWorkspaceIncludedCategories;
-			WorkspaceExcludedCategories = InWorkspaceExcludedCategories;
+			WorkspaceSyncCategories = InWorkspaceSyncCategories;
 			WorkspaceView = InWorkspaceView;
 			bWorkspaceSyncAllProjects = bInWorkspaceProjectOnly;
 			bWorkspaceIncludeAllProjectsInSolution = bInWorkspaceIncludeAllProjectsInSolution;
 
+			Dictionary<Guid, bool> SyncCategories = WorkspaceSyncCategory.GetDefault(UniqueIdToCategory.Values);
+
+			WorkspaceSyncCategory.ApplyDelta(SyncCategories, InGlobalSyncCategories);
 			GlobalControl.SetView(GlobalView);
-			SetExcludedCategories(GlobalControl.CategoriesCheckList, UniqueIdToCategory, GlobalExcludedCategories);
+			SetExcludedCategories(GlobalControl.CategoriesCheckList, UniqueIdToCategory, SyncCategories);
 			GlobalControl.SyncAllProjects.Checked = bGlobalSyncAllProjects;
 			GlobalControl.IncludeAllProjectsInSolution.Checked = bGlobalIncludeAllProjectsInSolution;
 
+			WorkspaceSyncCategory.ApplyDelta(SyncCategories, InWorkspaceSyncCategories);
 			WorkspaceControl.SetView(WorkspaceView);
-			SetExcludedCategories(WorkspaceControl.CategoriesCheckList, UniqueIdToCategory, UserSettings.GetEffectiveExcludedCategories(GlobalExcludedCategories, WorkspaceIncludedCategories, WorkspaceExcludedCategories));
+			SetExcludedCategories(WorkspaceControl.CategoriesCheckList, UniqueIdToCategory, SyncCategories);
 			WorkspaceControl.SyncAllProjects.Checked = bWorkspaceSyncAllProjects ?? bGlobalSyncAllProjects;
 			WorkspaceControl.IncludeAllProjectsInSolution.Checked = bWorkspaceIncludeAllProjectsInSolution ?? bGlobalIncludeAllProjectsInSolution;
 
@@ -70,13 +72,13 @@ namespace UnrealGameSync
 			WorkspaceControl.IncludeAllProjectsInSolution.Checked = GlobalControl.IncludeAllProjectsInSolution.Checked;
 		}
 
-		private static void SetExcludedCategories(CheckedListBox ListBox, Dictionary<Guid, WorkspaceSyncCategory> UniqueIdToFilter, Guid[] ExcludedCategories)
+		private static void SetExcludedCategories(CheckedListBox ListBox, Dictionary<Guid, WorkspaceSyncCategory> UniqueIdToFilter, Dictionary<Guid, bool> CategoryIdToSetting)
 		{
 			ListBox.Items.Clear();
 			foreach(WorkspaceSyncCategory Filter in UniqueIdToFilter.Values)
 			{
 				CheckState State = CheckState.Checked;
-				if(ExcludedCategories.Contains(Filter.UniqueId))
+				if(!CategoryIdToSetting[Filter.UniqueId])
 				{
 					State = CheckState.Unchecked;
 				}
@@ -84,31 +86,26 @@ namespace UnrealGameSync
 			}
 		}
 
-		private void GetExcludedCategories(out Guid[] NewGlobalExcludedCategories, out Guid[] NewWorkspaceIncludedCategories, out Guid[] NewWorkspaceExcludedCategories)
+		private void GetExcludedCategories(out Dictionary<Guid, bool> NewGlobalSyncCategories, out Dictionary<Guid, bool> NewWorkspaceSyncCategories)
 		{
-			NewGlobalExcludedCategories = GetExcludedCategories(GlobalControl.CategoriesCheckList);
+			Dictionary<Guid, bool> DefaultSyncCategories = WorkspaceSyncCategory.GetDefault(UniqueIdToCategory.Values);
+
+			Dictionary<Guid, bool> GlobalSyncCategories = GetCategorySettings(GlobalControl.CategoriesCheckList);
+			NewGlobalSyncCategories = WorkspaceSyncCategory.GetDelta(DefaultSyncCategories, GlobalSyncCategories);
 
-			Guid[] AllWorkspaceExcludedCategories = GetExcludedCategories(WorkspaceControl.CategoriesCheckList);
-			NewWorkspaceIncludedCategories = NewGlobalExcludedCategories.Except(AllWorkspaceExcludedCategories).ToArray();
-			NewWorkspaceExcludedCategories = AllWorkspaceExcludedCategories.Except(NewGlobalExcludedCategories).ToArray();
+			Dictionary<Guid, bool> WorkspaceSyncCategories = GetCategorySettings(WorkspaceControl.CategoriesCheckList);
+			NewWorkspaceSyncCategories = WorkspaceSyncCategory.GetDelta(GlobalSyncCategories, WorkspaceSyncCategories);
 		}
 
-		private static Guid[] GetExcludedCategories(CheckedListBox ListBox)
+		private static Dictionary<Guid, bool> GetCategorySettings(CheckedListBox ListBox)
 		{
-			HashSet<Guid> ExcludedCategories = new HashSet<Guid>();
+			Dictionary<Guid, bool> Result = new Dictionary<Guid, bool>();
 			for(int Idx = 0; Idx < ListBox.Items.Count; Idx++)
 			{
 				Guid UniqueId = ((WorkspaceSyncCategory)ListBox.Items[Idx]).UniqueId;
-				if(ListBox.GetItemCheckState(Idx) == CheckState.Unchecked)
-				{
-					ExcludedCategories.Add(UniqueId);
-				}
-				else
-				{
-					ExcludedCategories.Remove(UniqueId);
-				}
+				Result[UniqueId] = ListBox.GetItemCheckState(Idx) == CheckState.Checked;
 			}
-			return ExcludedCategories.ToArray();
+			return Result;
 		}
 
 		private static string[] GetView(TextBox FilterText)
@@ -142,7 +139,7 @@ namespace UnrealGameSync
 			bWorkspaceSyncAllProjects = (WorkspaceControl.SyncAllProjects.Checked == bGlobalSyncAllProjects)? (bool?)null : WorkspaceControl.SyncAllProjects.Checked;
 			bWorkspaceIncludeAllProjectsInSolution = (WorkspaceControl.IncludeAllProjectsInSolution.Checked == bGlobalIncludeAllProjectsInSolution)? (bool?)null : WorkspaceControl.IncludeAllProjectsInSolution.Checked;
 
-			GetExcludedCategories(out GlobalExcludedCategories, out WorkspaceIncludedCategories, out WorkspaceExcludedCategories);
+			GetExcludedCategories(out GlobalSyncCategories, out WorkspaceSyncCategories);
 
 			DialogResult = DialogResult.OK;
 		}
@@ -154,12 +151,11 @@ namespace UnrealGameSync
 
 		private void ShowCombinedView_Click(object sender, EventArgs e)
 		{
-			Guid[] NewGlobalExcludedCategories;
-			Guid[] NewWorkspaceIncludedCategories;
-			Guid[] NewWorkspaceExcludedCategories;
-			GetExcludedCategories(out NewGlobalExcludedCategories, out NewWorkspaceIncludedCategories, out NewWorkspaceExcludedCategories);
+			Dictionary<Guid, bool> NewGlobalSyncCategories;
+			Dictionary<Guid, bool> NewWorkspaceSyncCategories;
+			GetExcludedCategories(out NewGlobalSyncCategories, out NewWorkspaceSyncCategories);
 
-			string[] Filter = UserSettings.GetCombinedSyncFilter(UniqueIdToCategory, GlobalControl.GetView(), NewGlobalExcludedCategories, WorkspaceControl.GetView(), NewWorkspaceIncludedCategories, NewWorkspaceExcludedCategories);
+			string[] Filter = UserSettings.GetCombinedSyncFilter(UniqueIdToCategory, GlobalControl.GetView(), NewGlobalSyncCategories, WorkspaceControl.GetView(), NewWorkspaceSyncCategories);
 			if(Filter.Length == 0)
 			{
 				Filter = new string[]{ "All files will be synced." };
diff --git a/Engine/Source/Programs/UnrealGameSync/UnrealGameSync/Perforce.cs b/Engine/Source/Programs/UnrealGameSync/UnrealGameSync/Perforce.cs
--- a/Engine/Source/Programs/UnrealGameSync/UnrealGameSync/Perforce.cs
+++ b/Engine/Source/Programs/UnrealGameSync/UnrealGameSync/Perforce.cs
@@ -325,8 +325,10 @@ namespace UnrealGameSync
 
 	public class PerforceSyncOptions
 	{
+		public const int DefaultNumThreads = 2;
+
 		public int NumRetries;
-		public int NumThreads;
+		public int NumThreads = DefaultNumThreads;
 		public int TcpBufferSize;
 
 		public PerforceSyncOptions Clone()
diff --git a/Engine/Source/Programs/UnrealGameSync/UnrealGameSync/Properties/AssemblyInfo.cs b/Engine/Source/Programs/UnrealGameSync/UnrealGameSync/Properties/AssemblyInfo.cs
--- a/Engine/Source/Programs/UnrealGameSync/UnrealGameSync/Properties/AssemblyInfo.cs
+++ b/Engine/Source/Programs/UnrealGameSync/UnrealGameSync/Properties/AssemblyInfo.cs
@@ -34,5 +34,5 @@ using System.Runtime.InteropServices;
 // You can specify all the values or you can default the Build and Revision Numbers 
 // by using the '*' as shown below:
 // [assembly: AssemblyVersion("1.0.*")]
-[assembly: AssemblyVersion("1.198.*")]
-[assembly: AssemblyFileVersion("1.198.0.0")]
+[assembly: AssemblyVersion("1.200.*")]
+[assembly: AssemblyFileVersion("1.200.0.0")]
diff --git a/Engine/Source/Programs/UnrealGameSync/UnrealGameSync/UserSettings.cs b/Engine/Source/Programs/UnrealGameSync/UnrealGameSync/UserSettings.cs
--- a/Engine/Source/Programs/UnrealGameSync/UnrealGameSync/UserSettings.cs
+++ b/Engine/Source/Programs/UnrealGameSync/UnrealGameSync/UserSettings.cs
@@ -52,7 +52,8 @@ namespace UnrealGameSync
 		Initial = 0,
 		DefaultServerSettings = 1,
 		XgeShaderCompilation = 2,
-		Latest = XgeShaderCompilation
+		DefaultNumberOfThreads = 3,
+		Latest = DefaultNumberOfThreads
 	}
 
 	class ArchiveSettings
@@ -229,8 +230,7 @@ namespace UnrealGameSync
 
 		// Workspace specific SyncFilters
 		public string[] SyncView;
-		public Guid[] SyncIncludedCategories;
-		public Guid[] SyncExcludedCategories;
+		public Dictionary<Guid, bool> SyncCategories;
 		public bool? bSyncAllProjects;
 		public bool? bIncludeAllProjectsInSolution;
 	}
@@ -264,7 +264,7 @@ namespace UnrealGameSync
 		public List<UserSelectedProjectSettings> OpenProjects;
 		public List<UserSelectedProjectSettings> RecentProjects;
 		public string[] SyncView;
-		public Guid[] SyncExcludedCategories;
+		public Dictionary<Guid, bool> SyncCategories;
 		public bool bSyncAllProjects;
 		public bool bIncludeAllProjectsInSolution;
 		public LatestChangeType SyncType;
@@ -375,9 +375,9 @@ namespace UnrealGameSync
 
 			OpenProjects = ReadProjectList("General.OpenProjects", "General.OpenProjectFileNames");
 			RecentProjects = ReadProjectList("General.RecentProjects", "General.OtherProjectFileNames");
-
 			SyncView = ConfigFile.GetValues("General.SyncFilter", new string[0]);
-			SyncExcludedCategories = ConfigFile.GetGuidValues("General.SyncExcludedCategories", new Guid[0]);
+			SyncCategories = GetCategorySettings(ConfigFile.FindSection("General"), "SyncIncludedCategories", "SyncExcludedCategories");
+
 			bSyncAllProjects = ConfigFile.GetValue("General.SyncAllProjects", false);
 			bIncludeAllProjectsInSolution = ConfigFile.GetValue("General.IncludeAllProjectsInSolution", false);
 			if(!Enum.TryParse(ConfigFile.GetValue("General.SyncType", ""), out SyncType))
@@ -469,16 +469,54 @@ namespace UnrealGameSync
 			{
 				SyncOptions.NumRetries = 0;
 			}
-			if(!int.TryParse(ConfigFile.GetValue("Perforce.NumThreads", "0"), out SyncOptions.NumThreads))
+
+			int NumThreads;
+			if(int.TryParse(ConfigFile.GetValue("Perforce.NumThreads", "0"), out NumThreads) && NumThreads > 0)
 			{
-				SyncOptions.NumThreads = 0;
+				if(Version >= UserSettingsVersion.DefaultNumberOfThreads || NumThreads > 1)
+				{
+					SyncOptions.NumThreads = NumThreads;
+				}
 			}
+
 			if(!int.TryParse(ConfigFile.GetValue("Perforce.TcpBufferSize", "0"), out SyncOptions.TcpBufferSize))
 			{
 				SyncOptions.TcpBufferSize = 0;
 			}
 		}
 
+		static Dictionary<Guid, bool> GetCategorySettings(ConfigSection Section, string IncludedKey, string ExcludedKey)
+		{
+			Dictionary<Guid, bool> Result = new Dictionary<Guid, bool>();
+			if (Section != null)
+			{
+				foreach (Guid UniqueId in Section.GetValues(IncludedKey, new Guid[0]))
+				{
+					Result[UniqueId] = true;
+				}
+				foreach (Guid UniqueId in Section.GetValues(ExcludedKey, new Guid[0]))
+				{
+					Result[UniqueId] = false;
+				}
+			}
+			return Result;
+		}
+
+		static void SetCategorySettings(ConfigSection Section, string IncludedKey, string ExcludedKey, Dictionary<Guid, bool> Categories)
+		{
+			Guid[] IncludedCategories = Categories.Where(x => x.Value).Select(x => x.Key).ToArray();
+			if (IncludedCategories.Length > 0)
+			{
+				Section.SetValues(IncludedKey, IncludedCategories);
+			}
+
+			Guid[] ExcludedCategories = Categories.Where(x => !x.Value).Select(x => x.Key).ToArray();
+			if (ExcludedCategories.Length > 0)
+			{
+				Section.SetValues(ExcludedKey, ExcludedCategories);
+			}
+		}
+
 		static Rectangle? ParseRectangleValue(string Text)
 		{
 			ConfigObject Object = new ConfigObject(Text);
@@ -569,8 +607,7 @@ namespace UnrealGameSync
 					}
 
 					CurrentWorkspace.SyncView = new string[0];
-					CurrentWorkspace.SyncIncludedCategories = new Guid[0];
-					CurrentWorkspace.SyncExcludedCategories = new Guid[0];
+					CurrentWorkspace.SyncCategories = new Dictionary<Guid, bool>();
 					CurrentWorkspace.bSyncAllProjects = null;
 					CurrentWorkspace.bIncludeAllProjectsInSolution = null;
 				}
@@ -602,8 +639,7 @@ namespace UnrealGameSync
 					CurrentWorkspace.ExpandedArchiveTypes = WorkspaceSection.GetValues("ExpandedArchiveName", new string[0]);
 
 					CurrentWorkspace.SyncView = WorkspaceSection.GetValues("SyncFilter", new string[0]);
-					CurrentWorkspace.SyncIncludedCategories = WorkspaceSection.GetValues("SyncIncludedCategories", new Guid[0]);
-					CurrentWorkspace.SyncExcludedCategories = WorkspaceSection.GetValues("SyncExcludedCategories", new Guid[0]);
+					CurrentWorkspace.SyncCategories = GetCategorySettings(WorkspaceSection, "SyncIncludedCategories", "SyncExcludedCategories");
 
 					int SyncAllProjects = WorkspaceSection.GetValue("SyncAllProjects", -1);
 					CurrentWorkspace.bSyncAllProjects = (SyncAllProjects == 0)? (bool?)false : (SyncAllProjects == 1)? (bool?)true : (bool?)null;
@@ -675,7 +711,7 @@ namespace UnrealGameSync
 			GeneralSection.SetValue("FilterIndex", FilterIndex);
 			GeneralSection.SetValues("RecentProjects", RecentProjects.Select(x => x.ToConfigEntry()).ToArray());
 			GeneralSection.SetValues("SyncFilter", SyncView);
-			GeneralSection.SetValues("SyncExcludedCategories", SyncExcludedCategories);
+			SetCategorySettings(GeneralSection, "SyncIncludedCategories", "SyncExcludedCategories", SyncCategories);
 			GeneralSection.SetValue("SyncAllProjects", bSyncAllProjects);
 			GeneralSection.SetValue("IncludeAllProjectsInSolution", bIncludeAllProjectsInSolution);
 			GeneralSection.SetValue("SyncType", SyncType.ToString());
@@ -763,8 +799,7 @@ namespace UnrealGameSync
 				WorkspaceSection.SetValue("LastBuiltChangeNumber", CurrentWorkspace.LastBuiltChangeNumber);
 				WorkspaceSection.SetValues("ExpandedArchiveName", CurrentWorkspace.ExpandedArchiveTypes);
 				WorkspaceSection.SetValues("SyncFilter", CurrentWorkspace.SyncView);
-				WorkspaceSection.SetValues("SyncIncludedCategories", CurrentWorkspace.SyncIncludedCategories);
-				WorkspaceSection.SetValues("SyncExcludedCategories", CurrentWorkspace.SyncExcludedCategories);
+				SetCategorySettings(WorkspaceSection, "SyncIncludedCategories", "SyncExcludedCategories", CurrentWorkspace.SyncCategories);
 				if(CurrentWorkspace.bSyncAllProjects.HasValue)
 				{
 					WorkspaceSection.SetValue("SyncAllProjects", CurrentWorkspace.bSyncAllProjects.Value);
@@ -808,7 +843,7 @@ namespace UnrealGameSync
 			{
 				PerforceSection.SetValue("NumRetries", SyncOptions.NumRetries);
 			}
-			if(SyncOptions.NumThreads > 0)
+			if(SyncOptions.NumThreads != PerforceSyncOptions.DefaultNumThreads)
 			{
 				PerforceSection.SetValue("NumThreads", SyncOptions.NumThreads);
 			}
@@ -821,12 +856,7 @@ namespace UnrealGameSync
 			ConfigFile.Save(FileName);
 		}
 
-		public static Guid[] GetEffectiveExcludedCategories(Guid[] GlobalExcludedCategories, Guid[] WorkspaceIncludedCategories, Guid[] WorkspaceExcludedCategories)
-		{
-			return GlobalExcludedCategories.Except(WorkspaceIncludedCategories).Union(WorkspaceExcludedCategories).ToArray();
-		}
-
-		public static string[] GetCombinedSyncFilter(Dictionary<Guid, WorkspaceSyncCategory> UniqueIdToFilter, string[] GlobalView, Guid[] GlobalExcludedCategories, string[] WorkspaceView, Guid[] WorkspaceIncludedCategories, Guid[] WorkspaceExcludedCategories)
+		public static string[] GetCombinedSyncFilter(Dictionary<Guid, WorkspaceSyncCategory> UniqueIdToFilter, string[] GlobalView, Dictionary<Guid, bool> GlobalCategoryIdToSetting, string[] WorkspaceView, Dictionary<Guid, bool> WorkspaceCategoryIdToSetting)
 		{
 			List<string> Lines = new List<string>();
 			foreach(string ViewLine in Enumerable.Concat(GlobalView, WorkspaceView).Select(x => x.Trim()).Where(x => x.Length > 0 && !x.StartsWith(";")))
@@ -834,10 +864,26 @@ namespace UnrealGameSync
 				Lines.Add(ViewLine);
 			}
 
-			HashSet<Guid> ExcludedCategories = new HashSet<Guid>(GetEffectiveExcludedCategories(GlobalExcludedCategories, WorkspaceIncludedCategories, WorkspaceExcludedCategories));
-			foreach(WorkspaceSyncCategory Filter in UniqueIdToFilter.Values.Where(x => x.bEnable && ExcludedCategories.Contains(x.UniqueId)).OrderBy(x => x.Name))
+			foreach(WorkspaceSyncCategory Filter in UniqueIdToFilter.Values.OrderBy(x => x.Name))
 			{
-				Lines.AddRange(Filter.Paths.Select(x => "-" + x.Trim()));
+				bool bEnable = Filter.bEnable;
+
+				bool bGlobalEnable;
+				if (GlobalCategoryIdToSetting.TryGetValue(Filter.UniqueId, out bGlobalEnable))
+				{
+					bEnable = bGlobalEnable;
+				}
+
+				bool bWorkspaceEnable;
+				if (WorkspaceCategoryIdToSetting.TryGetValue(Filter.UniqueId, out bWorkspaceEnable))
+				{
+					bEnable = bWorkspaceEnable;
+				}
+
+				if (!bEnable)
+				{
+					Lines.AddRange(Filter.Paths.Select(x => "-" + x.Trim()));
+				}
 			}
 
 			return Lines.ToArray();
diff --git a/Engine/Source/Programs/UnrealGameSync/UnrealGameSync/Workspace.cs b/Engine/Source/Programs/UnrealGameSync/UnrealGameSync/Workspace.cs
--- a/Engine/Source/Programs/UnrealGameSync/UnrealGameSync/Workspace.cs
+++ b/Engine/Source/Programs/UnrealGameSync/UnrealGameSync/Workspace.cs
@@ -94,6 +94,33 @@ namespace UnrealGameSync
 			this.Paths = Paths;
 		}
 
+		public static Dictionary<Guid, bool> GetDefault(IEnumerable<WorkspaceSyncCategory> Categories)
+		{
+			return Categories.ToDictionary(x => x.UniqueId, x => x.bEnable);
+		}
+
+		public static Dictionary<Guid, bool> GetDelta(Dictionary<Guid, bool> Source, Dictionary<Guid, bool> Target)
+		{
+			Dictionary<Guid, bool> Changes = new Dictionary<Guid, bool>();
+			foreach (KeyValuePair<Guid, bool> Pair in Target)
+			{
+				bool bValue;
+				if (!Source.TryGetValue(Pair.Key, out bValue) || bValue != Pair.Value)
+				{
+					Changes[Pair.Key] = Pair.Value;
+				}
+			}
+			return Changes;
+		}
+
+		public static void ApplyDelta(Dictionary<Guid, bool> Categories, Dictionary<Guid, bool> Delta)
+		{
+			foreach(KeyValuePair<Guid, bool> Pair in Delta)
+			{
+				Categories[Pair.Key] = Pair.Value;
+			}
+		}
+
 		public override string ToString()
 		{
 			return Name;
