Updating UGS.

#rnx

diff --git a/Engine/Source/Programs/UnrealGameSync/UnrealGameSync/Controls/WorkspaceControl.cs b/Engine/Source/Programs/UnrealGameSync/UnrealGameSync/Controls/WorkspaceControl.cs
--- a/Engine/Source/Programs/UnrealGameSync/UnrealGameSync/Controls/WorkspaceControl.cs
+++ b/Engine/Source/Programs/UnrealGameSync/UnrealGameSync/Controls/WorkspaceControl.cs
@@ -978,10 +978,6 @@ namespace UnrealGameSync
 			{
 				Context.Options |= WorkspaceUpdateOptions.AutoResolveChanges;
 			}
-			if (Settings.bUseIncrementalBuilds)
-			{
-				Context.Options |= WorkspaceUpdateOptions.UseIncrementalBuilds;
-			}
 			if (WorkspaceSettings.bSyncAllProjects ?? Settings.bSyncAllProjects)
 			{
 				Context.Options |= WorkspaceUpdateOptions.SyncAllProjects | WorkspaceUpdateOptions.IncludeAllProjectsInSolution;
@@ -2855,11 +2851,6 @@ namespace UnrealGameSync
 						Owner.ShowAndActivate();
 
 						WorkspaceUpdateOptions Options = WorkspaceUpdateOptions.Build | WorkspaceUpdateOptions.RunAfterSync;
-						if (Settings.bUseIncrementalBuilds)
-						{
-							Options |= WorkspaceUpdateOptions.UseIncrementalBuilds;
-						}
-
 						WorkspaceUpdateContext Context = new WorkspaceUpdateContext(Workspace.CurrentChangeNumber, Options, null, GetDefaultBuildStepObjects(), ProjectSettings.BuildSteps, null, GetWorkspaceVariables(Workspace.CurrentChangeNumber));
 						StartWorkspaceUpdate(Context, null);
 					}
@@ -3804,7 +3795,7 @@ namespace UnrealGameSync
 						BuildListContextMenu_Sync.Font = new Font(SystemFonts.MenuFont, bIsCurrentChange ? FontStyle.Regular : FontStyle.Bold);
 						BuildListContextMenu_SyncContentOnly.Visible = !bIsBusy && ShouldSyncPrecompiledEditor;
 						BuildListContextMenu_SyncOnlyThisChange.Visible = !bIsBusy && !bIsCurrentChange && ContextMenuChange.Number > Workspace.CurrentChangeNumber && Workspace.CurrentChangeNumber != -1;
-						BuildListContextMenu_Build.Visible = !bIsBusy && bIsCurrentChange && !ShouldSyncPrecompiledEditor && Settings.bUseIncrementalBuilds;
+						BuildListContextMenu_Build.Visible = !bIsBusy && bIsCurrentChange && !ShouldSyncPrecompiledEditor;
 						BuildListContextMenu_Rebuild.Visible = !bIsBusy && bIsCurrentChange && !ShouldSyncPrecompiledEditor;
 						BuildListContextMenu_GenerateProjectFiles.Visible = !bIsBusy && bIsCurrentChange;
 						BuildListContextMenu_LaunchEditor.Visible = !bIsBusy && ContextMenuChange.Number == Workspace.CurrentChangeNumber;
@@ -4093,8 +4084,6 @@ namespace UnrealGameSync
 
 			OptionsContextMenu_EditorBuildConfiguration.Enabled = !ShouldSyncPrecompiledEditor;
 			UpdateCheckedBuildConfig();
-			OptionsContextMenu_UseIncrementalBuilds.Enabled = !ShouldSyncPrecompiledEditor;
-			OptionsContextMenu_UseIncrementalBuilds.Checked = Settings.bUseIncrementalBuilds;
 			OptionsContextMenu_CustomizeBuildSteps.Enabled = (Workspace != null);
 			OptionsContextMenu_EditorArguments.Checked = Settings.EditorArguments.Any(x => x.Item2);
 			OptionsContextMenu_ScheduledSync.Checked = Settings.bScheduleEnabled;
@@ -4235,12 +4224,12 @@ namespace UnrealGameSync
 
 		private void BuildListContextMenu_Build_Click(object sender, EventArgs e)
 		{
-			StartWorkspaceUpdate(Workspace.CurrentChangeNumber, WorkspaceUpdateOptions.SyncArchives | WorkspaceUpdateOptions.Build | WorkspaceUpdateOptions.UseIncrementalBuilds);
+			StartWorkspaceUpdate(Workspace.CurrentChangeNumber, WorkspaceUpdateOptions.SyncArchives | WorkspaceUpdateOptions.Build);
 		}
 
 		private void BuildListContextMenu_Rebuild_Click(object sender, EventArgs e)
 		{
-			StartWorkspaceUpdate(Workspace.CurrentChangeNumber, WorkspaceUpdateOptions.SyncArchives | WorkspaceUpdateOptions.Build);
+			StartWorkspaceUpdate(Workspace.CurrentChangeNumber, WorkspaceUpdateOptions.SyncArchives | WorkspaceUpdateOptions.Build | WorkspaceUpdateOptions.Clean);
 		}
 
 		private void BuildListContextMenu_GenerateProjectFiles_Click(object sender, EventArgs e)
@@ -4560,13 +4549,6 @@ namespace UnrealGameSync
 			OptionsContextMenu_BuildConfig_Development.Enabled = (!ShouldSyncPrecompiledEditor || EditorBuildConfig == BuildConfig.Development);
 		}
 
-		private void OptionsContextMenu_UseIncrementalBuilds_Click(object sender, EventArgs e)
-		{
-			OptionsContextMenu_UseIncrementalBuilds.Checked ^= true;
-			Settings.bUseIncrementalBuilds = OptionsContextMenu_UseIncrementalBuilds.Checked;
-			Settings.Save();
-		}
-
 		private void OptionsContextMenu_ScheduleSync_Click(object sender, EventArgs e)
 		{
 			Owner.SetupScheduledSync();
@@ -4801,7 +4783,7 @@ namespace UnrealGameSync
 			Variables.Add("UE4EditorCmdExe", GetEditorExePath(EditorBuildConfig).Replace(".exe", "-Cmd.exe"));
 			Variables.Add("UE4EditorConfig", EditorBuildConfig.ToString());
 			Variables.Add("UE4EditorDebugArg", (EditorBuildConfig == BuildConfig.Debug || EditorBuildConfig == BuildConfig.DebugGame) ? " -debug" : "");
-			Variables.Add("UseIncrementalBuilds", Settings.bUseIncrementalBuilds ? "1" : "0");
+			Variables.Add("UseIncrementalBuilds", "1");
 			if (!String.IsNullOrEmpty(SdkInstallerDir))
 			{
 				Variables.Add("SdkInstallerDir", SdkInstallerDir);
diff --git a/Engine/Source/Programs/UnrealGameSync/UnrealGameSync/DeploymentSettings.cs b/Engine/Source/Programs/UnrealGameSync/UnrealGameSync/DeploymentSettings.cs
--- a/Engine/Source/Programs/UnrealGameSync/UnrealGameSync/DeploymentSettings.cs
+++ b/Engine/Source/Programs/UnrealGameSync/UnrealGameSync/DeploymentSettings.cs
@@ -14,6 +14,7 @@ namespace UnrealGameSync
 	/// </summary>
 	static partial class DeploymentSettings
 	{
+#if WITH_TELEMETRY
 		/// <summary>
 		/// Delegate used to create a telemetry sink
 		/// </summary>
@@ -22,6 +23,7 @@ namespace UnrealGameSync
 		/// <param name="Log">Log writer</param>
 		/// <returns>New telemetry sink instance</returns>
 		public delegate ITelemetrySink CreateTelemetrySinkDelegate(string UserName, string SessionId, TextWriter Log);
+#endif
 
 		/// <summary>
 		/// SQL connection string used to connect to the database for telemetry and review data. The 'Program' class is a partial class, to allow an
@@ -35,9 +37,11 @@ namespace UnrealGameSync
 		/// </summary>
 		public static readonly string DefaultDepotPath = null;
 
+#if WITH_TELEMETRY
 		/// <summary>
 		/// Delegate used to create a new telemetry sink
 		/// </summary>
 		public static readonly CreateTelemetrySinkDelegate CreateTelemetrySink = (UserName, SessionId, Log) => new NullTelemetrySink();
+#endif
 	}
 }
diff --git a/Engine/Source/Programs/UnrealGameSync/UnrealGameSync/Properties/AssemblyInfo.cs b/Engine/Source/Programs/UnrealGameSync/UnrealGameSync/Properties/AssemblyInfo.cs
--- a/Engine/Source/Programs/UnrealGameSync/UnrealGameSync/Properties/AssemblyInfo.cs
+++ b/Engine/Source/Programs/UnrealGameSync/UnrealGameSync/Properties/AssemblyInfo.cs
@@ -34,5 +34,5 @@ using System.Runtime.InteropServices;
 // You can specify all the values or you can default the Build and Revision Numbers 
 // by using the '*' as shown below:
 // [assembly: AssemblyVersion("1.0.*")]
-[assembly: AssemblyVersion("1.200.*")]
-[assembly: AssemblyFileVersion("1.200.0.0")]
+[assembly: AssemblyVersion("1.201.*")]
+[assembly: AssemblyFileVersion("1.201.0.0")]
diff --git a/Engine/Source/Programs/UnrealGameSync/UnrealGameSync/UnrealGameSync.csproj b/Engine/Source/Programs/UnrealGameSync/UnrealGameSync/UnrealGameSync.csproj
--- a/Engine/Source/Programs/UnrealGameSync/UnrealGameSync/UnrealGameSync.csproj
+++ b/Engine/Source/Programs/UnrealGameSync/UnrealGameSync/UnrealGameSync.csproj
@@ -51,7 +51,7 @@
   <PropertyGroup Condition="'$(Configuration)|$(Platform)' == 'Debug|x64'">
     <DebugSymbols>true</DebugSymbols>
     <OutputPath>bin\x64\Debug\</OutputPath>
-    <DefineConstants>DEBUG;TRACE</DefineConstants>
+    <DefineConstants>TRACE;DEBUG;WITH_TELEMETRY</DefineConstants>
     <DebugType>full</DebugType>
     <PlatformTarget>x64</PlatformTarget>
     <UseVSHostingProcess>false</UseVSHostingProcess>
@@ -61,7 +61,7 @@
   </PropertyGroup>
   <PropertyGroup Condition="'$(Configuration)|$(Platform)' == 'Release|x64'">
     <OutputPath>bin\UnstableRelease\</OutputPath>
-    <DefineConstants>TRACE</DefineConstants>
+    <DefineConstants>TRACE;WITH_TELEMETRY</DefineConstants>
     <Optimize>true</Optimize>
     <DebugType>pdbonly</DebugType>
     <PlatformTarget>x64</PlatformTarget>
@@ -73,7 +73,7 @@
   <PropertyGroup Condition="'$(Configuration)|$(Platform)' == 'Debug|x86'">
     <DebugSymbols>true</DebugSymbols>
     <OutputPath>bin\x86\Debug\</OutputPath>
-    <DefineConstants>DEBUG;TRACE</DefineConstants>
+    <DefineConstants>TRACE;DEBUG;WITH_TELEMETRY</DefineConstants>
     <DebugType>full</DebugType>
     <PlatformTarget>x86</PlatformTarget>
     <UseVSHostingProcess>false</UseVSHostingProcess>
@@ -83,7 +83,7 @@
   </PropertyGroup>
   <PropertyGroup Condition="'$(Configuration)|$(Platform)' == 'Release|x86'">
     <OutputPath>bin\x86\Release\</OutputPath>
-    <DefineConstants>TRACE</DefineConstants>
+    <DefineConstants>TRACE;WITH_TELEMETRY</DefineConstants>
     <Optimize>true</Optimize>
     <DebugType>pdbonly</DebugType>
     <PlatformTarget>x86</PlatformTarget>
diff --git a/Engine/Source/Programs/UnrealGameSync/UnrealGameSync/UserSettings.cs b/Engine/Source/Programs/UnrealGameSync/UnrealGameSync/UserSettings.cs
--- a/Engine/Source/Programs/UnrealGameSync/UnrealGameSync/UserSettings.cs
+++ b/Engine/Source/Programs/UnrealGameSync/UnrealGameSync/UserSettings.cs
@@ -254,7 +254,6 @@ namespace UnrealGameSync
 		public bool bOpenSolutionAfterSync;
 		public bool bShowLogWindow;
 		public bool bAutoResolveConflicts;
-		public bool bUseIncrementalBuilds;
 		public bool bShowUnreviewedChanges;
 		public bool bShowAutomatedChanges;
 		public bool bShowLocalTimes;
@@ -352,7 +351,6 @@ namespace UnrealGameSync
 			bOpenSolutionAfterSync = (ConfigFile.GetValue("General.OpenSolutionAfterSync", "0") != "0");
 			bShowLogWindow = (ConfigFile.GetValue("General.ShowLogWindow", false));
 			bAutoResolveConflicts = (ConfigFile.GetValue("General.AutoResolveConflicts", "1") != "0");
-			bUseIncrementalBuilds = ConfigFile.GetValue("General.IncrementalBuilds", true);
 			bShowUnreviewedChanges = ConfigFile.GetValue("General.ShowUnreviewed", true);
 			bShowAutomatedChanges = ConfigFile.GetValue("General.ShowAutomated", false);
 			bShowLocalTimes = ConfigFile.GetValue("General.ShowLocalTimes", false);
@@ -698,7 +696,6 @@ namespace UnrealGameSync
 			GeneralSection.SetValue("OpenSolutionAfterSync", bOpenSolutionAfterSync);
 			GeneralSection.SetValue("ShowLogWindow", bShowLogWindow);
 			GeneralSection.SetValue("AutoResolveConflicts", bAutoResolveConflicts);
-			GeneralSection.SetValue("IncrementalBuilds", bUseIncrementalBuilds);
 			GeneralSection.SetValue("ShowUnreviewed", bShowUnreviewedChanges);
 			GeneralSection.SetValue("ShowAutomated", bShowAutomatedChanges);
 			GeneralSection.SetValue("ShowLocalTimes", bShowLocalTimes);
diff --git a/Engine/Source/Programs/UnrealGameSync/UnrealGameSync/Utility.cs b/Engine/Source/Programs/UnrealGameSync/UnrealGameSync/Utility.cs
--- a/Engine/Source/Programs/UnrealGameSync/UnrealGameSync/Utility.cs
+++ b/Engine/Source/Programs/UnrealGameSync/UnrealGameSync/Utility.cs
@@ -269,7 +269,7 @@ namespace UnrealGameSync
 
 		/******/
 
-		private static void AddLocalConfigPaths_PlatformFolders(DirectoryInfo BaseDir, string FileName, List<FileInfo> Files)
+		private static void AddLocalConfigPaths_WithSubFolders(DirectoryInfo BaseDir, string FileName, List<FileInfo> Files)
 		{
 			if(BaseDir.Exists)
 			{
@@ -290,18 +290,27 @@ namespace UnrealGameSync
 			}
 		}
 
-		private static void AddLocalConfigPaths_PlatformExtensions(DirectoryInfo BaseDir, string RelativePath, string FileName, List<FileInfo> Files)
+		private static void AddLocalConfigPaths_WithExtensionDirs(DirectoryInfo BaseDir, string RelativePath, string FileName, List<FileInfo> Files)
 		{
 			if (BaseDir.Exists)
 			{
-				AddLocalConfigPaths_PlatformFolders(new DirectoryInfo(Path.Combine(BaseDir.FullName, RelativePath)), FileName, Files);
+				AddLocalConfigPaths_WithSubFolders(new DirectoryInfo(Path.Combine(BaseDir.FullName, RelativePath)), FileName, Files);
 
 				DirectoryInfo PlatformExtensionsDir = new DirectoryInfo(Path.Combine(BaseDir.FullName, "Platforms"));
 				if (PlatformExtensionsDir.Exists)
 				{
 					foreach (DirectoryInfo PlatformExtensionDir in PlatformExtensionsDir.EnumerateDirectories())
 					{
-						AddLocalConfigPaths_PlatformFolders(new DirectoryInfo(Path.Combine(PlatformExtensionDir.FullName, RelativePath)), FileName, Files);
+						AddLocalConfigPaths_WithSubFolders(new DirectoryInfo(Path.Combine(PlatformExtensionDir.FullName, RelativePath)), FileName, Files);
+					}
+				}
+
+				DirectoryInfo RestrictedBaseDir = new DirectoryInfo(Path.Combine(BaseDir.FullName, "Restricted"));
+				if (RestrictedBaseDir.Exists)
+				{
+					foreach (DirectoryInfo RestrictedDir in RestrictedBaseDir.EnumerateDirectories())
+					{
+						AddLocalConfigPaths_WithSubFolders(new DirectoryInfo(Path.Combine(RestrictedDir.FullName, RelativePath)), FileName, Files);
 					}
 				}
 			}
@@ -310,15 +319,15 @@ namespace UnrealGameSync
 		public static List<FileInfo> GetLocalConfigPaths(DirectoryInfo EngineDir, FileInfo ProjectFile)
 		{
 			List<FileInfo> SearchPaths = new List<FileInfo>();
-			AddLocalConfigPaths_PlatformExtensions(EngineDir, "Programs/UnrealGameSync", "UnrealGameSync.ini", SearchPaths);
+			AddLocalConfigPaths_WithExtensionDirs(EngineDir, "Programs/UnrealGameSync", "UnrealGameSync.ini", SearchPaths);
 
 			if (ProjectFile.Name.EndsWith(".uproject", StringComparison.OrdinalIgnoreCase))
 			{
-				AddLocalConfigPaths_PlatformExtensions(ProjectFile.Directory, "Build", "UnrealGameSync.ini", SearchPaths);
+				AddLocalConfigPaths_WithExtensionDirs(ProjectFile.Directory, "Build", "UnrealGameSync.ini", SearchPaths);
 			}
 			else
 			{
-				AddLocalConfigPaths_PlatformExtensions(EngineDir, "Programs/UnrealGameSync", "DefaultEngine.ini", SearchPaths);
+				AddLocalConfigPaths_WithExtensionDirs(EngineDir, "Programs/UnrealGameSync", "DefaultEngine.ini", SearchPaths);
 			}
 			return SearchPaths;
 		}
@@ -335,6 +344,7 @@ namespace UnrealGameSync
 		{
 			AddDepotConfigPaths_PlatformFolders(BasePath + RelativePath, FileName, SearchPaths);
 			AddDepotConfigPaths_PlatformFolders(BasePath + "/Platforms/*" + RelativePath, FileName, SearchPaths);
+			AddDepotConfigPaths_PlatformFolders(BasePath + "/Restricted/*" + RelativePath, FileName, SearchPaths);
 		}
 
 		public static List<string> GetDepotConfigPaths(string EnginePath, string ProjectPath)
diff --git a/Engine/Source/Programs/UnrealGameSync/UnrealGameSync/Workspace.cs b/Engine/Source/Programs/UnrealGameSync/UnrealGameSync/Workspace.cs
--- a/Engine/Source/Programs/UnrealGameSync/UnrealGameSync/Workspace.cs
+++ b/Engine/Source/Programs/UnrealGameSync/UnrealGameSync/Workspace.cs
@@ -24,7 +24,7 @@ namespace UnrealGameSync
 		GenerateProjectFiles = 0x08,
 		SyncArchives = 0x10,
 		Build = 0x20,
-		UseIncrementalBuilds = 0x40,
+		Clean = 0x40,
 		ScheduledBuild = 0x80,
 		RunAfterSync = 0x100,
 		OpenSolutionAfterSync = 0x200,
@@ -1073,7 +1073,7 @@ namespace UnrealGameSync
 										StepStopwatch.AddData(new { Target = Step.Target });
 
 										string CommandLine = String.Format("{0} {1} {2} {3} -NoHotReloadFromIDE", Step.Target, Step.Platform, Step.Configuration, Utility.ExpandVariables(Step.Arguments ?? "", Context.Variables));
-										if(!Context.Options.HasFlag(WorkspaceUpdateOptions.UseIncrementalBuilds) || bForceClean)
+										if(Context.Options.HasFlag(WorkspaceUpdateOptions.Clean) || bForceClean)
 										{
 											Log.WriteLine("ubt> Running {0} {1} -clean", UnrealBuildToolPath, CommandLine);
 											Utility.ExecuteProcess(UnrealBuildToolPath, null, CommandLine + " -clean", null, new ProgressTextWriter(Progress, new PrefixedTextWriter("ubt> ", Log)));
