Converted UGS into a Net Core application

* Mostly removed unused references which didnt exist in net core.
* Replaced usages of JsonScriptSerializer (System.Web) with the new System.Text.Json (as the old one does not exist in net core).
* Updated installer to install Launcher and publish it as a self contained application (Makes msi 64 MB instead of 600kb but does not require anything to be installed now).

Worked around what looks to be a bug in WinForms for Net Core, in Workspace Control the event state was selected when drawing the item even though the item was not selected (and the state said it isnt selected). Reading the state from the object instead of looking at the event.

Deleted ProjectWindow form as this was never used and code did not compile (net core projects builds all files without them having to be included)

Deleted metadata server as this has been replaced by Horde. Also removed Post Build Status tool for the same reason.

diff --git a/Engine/Source/Programs/UnrealGameSync/MetadataServer/App_Start/WebApiConfig.cs b/Engine/Source/Programs/UnrealGameSync/MetadataServer/App_Start/WebApiConfig.cs
deleted file mode 100644
diff --git a/Engine/Source/Programs/UnrealGameSync/MetadataServer/Connectors/SqlConnector.cs b/Engine/Source/Programs/UnrealGameSync/MetadataServer/Connectors/SqlConnector.cs
deleted file mode 100644
diff --git a/Engine/Source/Programs/UnrealGameSync/MetadataServer/Controllers/BuildController.cs b/Engine/Source/Programs/UnrealGameSync/MetadataServer/Controllers/BuildController.cs
deleted file mode 100644
diff --git a/Engine/Source/Programs/UnrealGameSync/MetadataServer/Controllers/CISController.cs b/Engine/Source/Programs/UnrealGameSync/MetadataServer/Controllers/CISController.cs
deleted file mode 100644
diff --git a/Engine/Source/Programs/UnrealGameSync/MetadataServer/Controllers/CommentController.cs b/Engine/Source/Programs/UnrealGameSync/MetadataServer/Controllers/CommentController.cs
deleted file mode 100644
diff --git a/Engine/Source/Programs/UnrealGameSync/MetadataServer/Controllers/ErrorController.cs b/Engine/Source/Programs/UnrealGameSync/MetadataServer/Controllers/ErrorController.cs
deleted file mode 100644
diff --git a/Engine/Source/Programs/UnrealGameSync/MetadataServer/Controllers/EventController.cs b/Engine/Source/Programs/UnrealGameSync/MetadataServer/Controllers/EventController.cs
deleted file mode 100644
diff --git a/Engine/Source/Programs/UnrealGameSync/MetadataServer/Controllers/IssuesController.cs b/Engine/Source/Programs/UnrealGameSync/MetadataServer/Controllers/IssuesController.cs
deleted file mode 100644
diff --git a/Engine/Source/Programs/UnrealGameSync/MetadataServer/Controllers/LatestController.cs b/Engine/Source/Programs/UnrealGameSync/MetadataServer/Controllers/LatestController.cs
deleted file mode 100644
diff --git a/Engine/Source/Programs/UnrealGameSync/MetadataServer/Controllers/TelemetryController.cs b/Engine/Source/Programs/UnrealGameSync/MetadataServer/Controllers/TelemetryController.cs
deleted file mode 100644
diff --git a/Engine/Source/Programs/UnrealGameSync/MetadataServer/Controllers/UserController.cs b/Engine/Source/Programs/UnrealGameSync/MetadataServer/Controllers/UserController.cs
deleted file mode 100644
diff --git a/Engine/Source/Programs/UnrealGameSync/MetadataServer/Global.asax.cs b/Engine/Source/Programs/UnrealGameSync/MetadataServer/Global.asax.cs
deleted file mode 100644
diff --git a/Engine/Source/Programs/UnrealGameSync/MetadataServer/MetadataServer.csproj b/Engine/Source/Programs/UnrealGameSync/MetadataServer/MetadataServer.csproj
deleted file mode 100644
diff --git a/Engine/Source/Programs/UnrealGameSync/MetadataServer/Models/BuildData.cs b/Engine/Source/Programs/UnrealGameSync/MetadataServer/Models/BuildData.cs
deleted file mode 100644
diff --git a/Engine/Source/Programs/UnrealGameSync/MetadataServer/Models/CommentData.cs b/Engine/Source/Programs/UnrealGameSync/MetadataServer/Models/CommentData.cs
deleted file mode 100644
diff --git a/Engine/Source/Programs/UnrealGameSync/MetadataServer/Models/EventData.cs b/Engine/Source/Programs/UnrealGameSync/MetadataServer/Models/EventData.cs
deleted file mode 100644
diff --git a/Engine/Source/Programs/UnrealGameSync/MetadataServer/Models/EventSummary.cs b/Engine/Source/Programs/UnrealGameSync/MetadataServer/Models/EventSummary.cs
deleted file mode 100644
diff --git a/Engine/Source/Programs/UnrealGameSync/MetadataServer/Models/IssueData.cs b/Engine/Source/Programs/UnrealGameSync/MetadataServer/Models/IssueData.cs
deleted file mode 100644
diff --git a/Engine/Source/Programs/UnrealGameSync/MetadataServer/Models/LatestData.cs b/Engine/Source/Programs/UnrealGameSync/MetadataServer/Models/LatestData.cs
deleted file mode 100644
diff --git a/Engine/Source/Programs/UnrealGameSync/MetadataServer/Models/TelemetryErrorData.cs b/Engine/Source/Programs/UnrealGameSync/MetadataServer/Models/TelemetryErrorData.cs
deleted file mode 100644
diff --git a/Engine/Source/Programs/UnrealGameSync/MetadataServer/Models/TelemetryTimingData.cs b/Engine/Source/Programs/UnrealGameSync/MetadataServer/Models/TelemetryTimingData.cs
deleted file mode 100644
diff --git a/Engine/Source/Programs/UnrealGameSync/MetadataServer/Properties/AssemblyInfo.cs b/Engine/Source/Programs/UnrealGameSync/MetadataServer/Properties/AssemblyInfo.cs
deleted file mode 100644
diff --git a/Engine/Source/Programs/UnrealGameSync/MetadataServer/Web.Debug.config b/Engine/Source/Programs/UnrealGameSync/MetadataServer/Web.Debug.config
deleted file mode 100644
diff --git a/Engine/Source/Programs/UnrealGameSync/MetadataServer/Web.Release.config b/Engine/Source/Programs/UnrealGameSync/MetadataServer/Web.Release.config
deleted file mode 100644
diff --git a/Engine/Source/Programs/UnrealGameSync/MetadataServer/Web.template.config b/Engine/Source/Programs/UnrealGameSync/MetadataServer/Web.template.config
deleted file mode 100644
diff --git a/Engine/Source/Programs/UnrealGameSync/MetadataServer/packages.config b/Engine/Source/Programs/UnrealGameSync/MetadataServer/packages.config
deleted file mode 100644
diff --git a/Engine/Source/Programs/UnrealGameSync/MetadataTool/App.config b/Engine/Source/Programs/UnrealGameSync/MetadataTool/App.config
deleted file mode 100644
diff --git a/Engine/Source/Programs/UnrealGameSync/MetadataTool/BuildHealth/BuildHealthIssue.cs b/Engine/Source/Programs/UnrealGameSync/MetadataTool/BuildHealth/BuildHealthIssue.cs
deleted file mode 100644
diff --git a/Engine/Source/Programs/UnrealGameSync/MetadataTool/BuildHealth/BuildHealthJobStep.cs b/Engine/Source/Programs/UnrealGameSync/MetadataTool/BuildHealth/BuildHealthJobStep.cs
deleted file mode 100644
diff --git a/Engine/Source/Programs/UnrealGameSync/MetadataTool/BuildHealth/BuildHealthJobStepHistory.cs b/Engine/Source/Programs/UnrealGameSync/MetadataTool/BuildHealth/BuildHealthJobStepHistory.cs
deleted file mode 100644
diff --git a/Engine/Source/Programs/UnrealGameSync/MetadataTool/BuildHealth/BuildHealthState.cs b/Engine/Source/Programs/UnrealGameSync/MetadataTool/BuildHealth/BuildHealthState.cs
deleted file mode 100644
diff --git a/Engine/Source/Programs/UnrealGameSync/MetadataTool/BuildHealth/ChangeInfo.cs b/Engine/Source/Programs/UnrealGameSync/MetadataTool/BuildHealth/ChangeInfo.cs
deleted file mode 100644
diff --git a/Engine/Source/Programs/UnrealGameSync/MetadataTool/BuildHealth/InputData.cs b/Engine/Source/Programs/UnrealGameSync/MetadataTool/BuildHealth/InputData.cs
deleted file mode 100644
diff --git a/Engine/Source/Programs/UnrealGameSync/MetadataTool/BuildHealth/Matchers/CompilePatternMatcher.cs b/Engine/Source/Programs/UnrealGameSync/MetadataTool/BuildHealth/Matchers/CompilePatternMatcher.cs
deleted file mode 100644
diff --git a/Engine/Source/Programs/UnrealGameSync/MetadataTool/BuildHealth/Matchers/ContentPatternMatcher.cs b/Engine/Source/Programs/UnrealGameSync/MetadataTool/BuildHealth/Matchers/ContentPatternMatcher.cs
deleted file mode 100644
diff --git a/Engine/Source/Programs/UnrealGameSync/MetadataTool/BuildHealth/Matchers/CopyrightNoticeMatcher.cs b/Engine/Source/Programs/UnrealGameSync/MetadataTool/BuildHealth/Matchers/CopyrightNoticeMatcher.cs
deleted file mode 100644
diff --git a/Engine/Source/Programs/UnrealGameSync/MetadataTool/BuildHealth/Matchers/DefaultPatternMatcher.cs b/Engine/Source/Programs/UnrealGameSync/MetadataTool/BuildHealth/Matchers/DefaultPatternMatcher.cs
deleted file mode 100644
diff --git a/Engine/Source/Programs/UnrealGameSync/MetadataTool/BuildHealth/Matchers/GenericCodePatternMatcher.cs b/Engine/Source/Programs/UnrealGameSync/MetadataTool/BuildHealth/Matchers/GenericCodePatternMatcher.cs
deleted file mode 100644
diff --git a/Engine/Source/Programs/UnrealGameSync/MetadataTool/BuildHealth/Matchers/GenericContentPatternMatcher.cs b/Engine/Source/Programs/UnrealGameSync/MetadataTool/BuildHealth/Matchers/GenericContentPatternMatcher.cs
deleted file mode 100644
diff --git a/Engine/Source/Programs/UnrealGameSync/MetadataTool/BuildHealth/Matchers/UndefinedSymbolPatternMatcher.cs b/Engine/Source/Programs/UnrealGameSync/MetadataTool/BuildHealth/Matchers/UndefinedSymbolPatternMatcher.cs
deleted file mode 100644
diff --git a/Engine/Source/Programs/UnrealGameSync/MetadataTool/BuildHealth/PatternMatcher.cs b/Engine/Source/Programs/UnrealGameSync/MetadataTool/BuildHealth/PatternMatcher.cs
deleted file mode 100644
diff --git a/Engine/Source/Programs/UnrealGameSync/MetadataTool/CommandHandler.cs b/Engine/Source/Programs/UnrealGameSync/MetadataTool/CommandHandler.cs
deleted file mode 100644
diff --git a/Engine/Source/Programs/UnrealGameSync/MetadataTool/CommandHandler_BuildHealth.cs b/Engine/Source/Programs/UnrealGameSync/MetadataTool/CommandHandler_BuildHealth.cs
deleted file mode 100644
diff --git a/Engine/Source/Programs/UnrealGameSync/MetadataTool/CommandHandler_Http.cs b/Engine/Source/Programs/UnrealGameSync/MetadataTool/CommandHandler_Http.cs
deleted file mode 100644
diff --git a/Engine/Source/Programs/UnrealGameSync/MetadataTool/CommandTypes.cs b/Engine/Source/Programs/UnrealGameSync/MetadataTool/CommandTypes.cs
deleted file mode 100644
diff --git a/Engine/Source/Programs/UnrealGameSync/MetadataTool/MetadataTool.csproj b/Engine/Source/Programs/UnrealGameSync/MetadataTool/MetadataTool.csproj
deleted file mode 100644
diff --git a/Engine/Source/Programs/UnrealGameSync/MetadataTool/Program.cs b/Engine/Source/Programs/UnrealGameSync/MetadataTool/Program.cs
deleted file mode 100644
diff --git a/Engine/Source/Programs/UnrealGameSync/MetadataTool/Properties/AssemblyInfo.cs b/Engine/Source/Programs/UnrealGameSync/MetadataTool/Properties/AssemblyInfo.cs
deleted file mode 100644
diff --git a/Engine/Source/Programs/UnrealGameSync/PostBadgeStatus/App.config b/Engine/Source/Programs/UnrealGameSync/PostBadgeStatus/App.config
deleted file mode 100644
diff --git a/Engine/Source/Programs/UnrealGameSync/PostBadgeStatus/PostBadgeStatus.csproj b/Engine/Source/Programs/UnrealGameSync/PostBadgeStatus/PostBadgeStatus.csproj
deleted file mode 100644
diff --git a/Engine/Source/Programs/UnrealGameSync/PostBadgeStatus/Program.cs b/Engine/Source/Programs/UnrealGameSync/PostBadgeStatus/Program.cs
deleted file mode 100644
diff --git a/Engine/Source/Programs/UnrealGameSync/PostBadgeStatus/Properties/AssemblyInfo.cs b/Engine/Source/Programs/UnrealGameSync/PostBadgeStatus/Properties/AssemblyInfo.cs
deleted file mode 100644
diff --git a/Engine/Source/Programs/UnrealGameSync/PostBadgeStatus/bin/Release/PostBadgeStatus.exe.config b/Engine/Source/Programs/UnrealGameSync/PostBadgeStatus/bin/Release/PostBadgeStatus.exe.config
deleted file mode 100644
diff --git a/Engine/Source/Programs/UnrealGameSync/UnrealGameSync/Controls/WorkspaceControl.cs b/Engine/Source/Programs/UnrealGameSync/UnrealGameSync/Controls/WorkspaceControl.cs
--- a/Engine/Source/Programs/UnrealGameSync/UnrealGameSync/Controls/WorkspaceControl.cs
+++ b/Engine/Source/Programs/UnrealGameSync/UnrealGameSync/Controls/WorkspaceControl.cs
@@ -2,10 +2,6 @@
 
 using System;
 using System.Collections.Generic;
-using System.ComponentModel;
-using System.Data;
-using System.Data.SqlClient;
-using System.Deployment.Application;
 using System.Diagnostics;
 using System.Drawing;
 using System.Drawing.Drawing2D;
@@ -14,14 +10,11 @@ using System.Linq;
 using System.Reflection;
 using System.Runtime.InteropServices;
 using System.Text;
+using System.Text.Json;
 using System.Text.RegularExpressions;
-using System.Threading.Tasks;
 using System.Windows.Forms;
 using System.Windows.Forms.VisualStyles;
-using Microsoft.Win32;
 using System.Threading;
-using System.Drawing.Imaging;
-using System.Web.Script.Serialization;
 
 namespace UnrealGameSync
 {
@@ -1879,7 +1872,7 @@ namespace UnrealGameSync
 			}
 			else if (Workspace != null)
 			{
-				if (e.State.HasFlag(ListViewItemStates.Selected))
+				if (e.Item.Selected)
 				{
 					BuildList.DrawSelectedBackground(e.Graphics, e.Bounds);
 				}
@@ -2870,8 +2863,7 @@ namespace UnrealGameSync
 					try
 					{
 						string Text = File.ReadAllText(ReceiptFileName);
-						JavaScriptSerializer Serializer = new JavaScriptSerializer();
-						Dictionary<string, object> RawObject = Serializer.Deserialize<Dictionary<string, object>>(Text);
+						Dictionary<string, object> RawObject = JsonSerializer.Deserialize<Dictionary<string, object>>(Text);
 
 						object LaunchFileNameObject;
 						if (RawObject.TryGetValue("Launch", out LaunchFileNameObject))
diff --git a/Engine/Source/Programs/UnrealGameSync/UnrealGameSync/DetectProjectSettingsTask.cs b/Engine/Source/Programs/UnrealGameSync/UnrealGameSync/DetectProjectSettingsTask.cs
--- a/Engine/Source/Programs/UnrealGameSync/UnrealGameSync/DetectProjectSettingsTask.cs
+++ b/Engine/Source/Programs/UnrealGameSync/UnrealGameSync/DetectProjectSettingsTask.cs
@@ -5,9 +5,7 @@ using System.Collections.Generic;
 using System.Drawing;
 using System.IO;
 using System.Linq;
-using System.Text;
 using System.Threading.Tasks;
-using System.Web.Script.Serialization;
 
 namespace UnrealGameSync
 {
diff --git a/Engine/Source/Programs/UnrealGameSync/UnrealGameSync/EventMonitor.cs b/Engine/Source/Programs/UnrealGameSync/UnrealGameSync/EventMonitor.cs
--- a/Engine/Source/Programs/UnrealGameSync/UnrealGameSync/EventMonitor.cs
+++ b/Engine/Source/Programs/UnrealGameSync/UnrealGameSync/EventMonitor.cs
@@ -3,15 +3,10 @@
 using System;
 using System.Collections.Concurrent;
 using System.Collections.Generic;
-using System.Data;
-using System.Data.SqlClient;
 using System.Diagnostics;
-using System.IO;
 using System.Linq;
-using System.Text;
+using System.Text.Json;
 using System.Threading;
-using System.Threading.Tasks;
-using System.Web.Script.Serialization;
 
 namespace UnrealGameSync
 {
@@ -521,7 +516,8 @@ namespace UnrealGameSync
 			{
 				Stopwatch Timer = Stopwatch.StartNew();
 				LogWriter.WriteLine("Posting event... ({0}, {1}, {2})", Event.Change, Event.UserName, Event.Type);
-				RESTApi.POST(ApiUrl, "event", new JavaScriptSerializer().Serialize(Event));
+				
+				RESTApi.POST(ApiUrl, "event", JsonSerializer.Serialize(Event));
 				return true;
 			}
 			catch(Exception Ex)
@@ -537,7 +533,7 @@ namespace UnrealGameSync
 			{
 				Stopwatch Timer = Stopwatch.StartNew();
 				LogWriter.WriteLine("Posting comment... ({0}, {1}, {2}, {3})", Comment.ChangeNumber, Comment.UserName, Comment.Text, Comment.Project);
-				RESTApi.POST(ApiUrl, "comment", new JavaScriptSerializer().Serialize(Comment));
+				RESTApi.POST(ApiUrl, "comment", JsonSerializer.Serialize(Comment));
 				LogWriter.WriteLine("Done in {0}ms.", Timer.ElapsedMilliseconds);
 				return true;
 			}
diff --git a/Engine/Source/Programs/UnrealGameSync/UnrealGameSync/Forms/DiagnosticsWindow.cs b/Engine/Source/Programs/UnrealGameSync/UnrealGameSync/Forms/DiagnosticsWindow.cs
--- a/Engine/Source/Programs/UnrealGameSync/UnrealGameSync/Forms/DiagnosticsWindow.cs
+++ b/Engine/Source/Programs/UnrealGameSync/UnrealGameSync/Forms/DiagnosticsWindow.cs
@@ -1,16 +1,7 @@
 // Copyright Epic Games, Inc. All Rights Reserved.
 
 using System;
-using System.Collections.Generic;
-using System.ComponentModel;
-using System.Data;
-using System.Deployment.Application;
 using System.Diagnostics;
-using System.Drawing;
-using System.Linq;
-using System.Reflection;
-using System.Text;
-using System.Threading.Tasks;
 using System.Windows.Forms;
 using Ionic.Zip;
 using System.IO;
diff --git a/Engine/Source/Programs/UnrealGameSync/UnrealGameSync/Forms/ModalTaskWindow.cs b/Engine/Source/Programs/UnrealGameSync/UnrealGameSync/Forms/ModalTaskWindow.cs
--- a/Engine/Source/Programs/UnrealGameSync/UnrealGameSync/Forms/ModalTaskWindow.cs
+++ b/Engine/Source/Programs/UnrealGameSync/UnrealGameSync/Forms/ModalTaskWindow.cs
@@ -82,7 +82,7 @@ namespace UnrealGameSync
 		{
 			if(BackgroundThread != null)
 			{
-				BackgroundThread.Abort();
+				BackgroundThread.Interrupt();
 				BackgroundThread.Join();
 				BackgroundThread = null;
 
diff --git a/Engine/Source/Programs/UnrealGameSync/UnrealGameSync/Forms/ProjectWindow.cs b/Engine/Source/Programs/UnrealGameSync/UnrealGameSync/Forms/ProjectWindow.cs
deleted file mode 100644
diff --git a/Engine/Source/Programs/UnrealGameSync/UnrealGameSync/Properties/AssemblyInfo.cs b/Engine/Source/Programs/UnrealGameSync/UnrealGameSync/Properties/AssemblyInfo.cs
--- a/Engine/Source/Programs/UnrealGameSync/UnrealGameSync/Properties/AssemblyInfo.cs
+++ b/Engine/Source/Programs/UnrealGameSync/UnrealGameSync/Properties/AssemblyInfo.cs
@@ -34,5 +34,5 @@ using System.Runtime.InteropServices;
 // You can specify all the values or you can default the Build and Revision Numbers 
 // by using the '*' as shown below:
 // [assembly: AssemblyVersion("1.0.*")]
-[assembly: AssemblyVersion("1.201.*")]
+[assembly: AssemblyVersion("1.201.0.0")]
 [assembly: AssemblyFileVersion("1.201.0.0")]
diff --git a/Engine/Source/Programs/UnrealGameSync/UnrealGameSync/RESTApi.cs b/Engine/Source/Programs/UnrealGameSync/UnrealGameSync/RESTApi.cs
--- a/Engine/Source/Programs/UnrealGameSync/UnrealGameSync/RESTApi.cs
+++ b/Engine/Source/Programs/UnrealGameSync/UnrealGameSync/RESTApi.cs
@@ -3,12 +3,10 @@
 using System;
 using System.Collections.Generic;
 using System.IO;
-using System.Linq;
 using System.Net;
+using System.Net.Mime;
 using System.Text;
-using System.Text.RegularExpressions;
-using System.Threading.Tasks;
-using System.Web.Script.Serialization;
+using System.Text.Json;
 
 namespace UnrealGameSync
 {
@@ -86,12 +84,12 @@ namespace UnrealGameSync
 
 		public static T GET<T>(string URI, string Resource, params string[] QueryParams)
 		{
-			return new JavaScriptSerializer().Deserialize<T>(SendRequestInternal(URI, Resource, "GET", null, QueryParams));
+			return JsonSerializer.Deserialize<T>(SendRequestInternal(URI, Resource, "GET", null, QueryParams));
 		}
 
 		public static string PUT<T>(string URI, string Resource, T Object, params string[] QueryParams)
 		{
-			return SendRequestInternal(URI, Resource, "PUT", new JavaScriptSerializer().Serialize(Object), QueryParams);
+			return SendRequestInternal(URI, Resource, "PUT", JsonSerializer.Serialize(Object), QueryParams);
 		}
 	}
 }
diff --git a/Engine/Source/Programs/UnrealGameSync/UnrealGameSync/Telemetry.cs b/Engine/Source/Programs/UnrealGameSync/UnrealGameSync/Telemetry.cs
--- a/Engine/Source/Programs/UnrealGameSync/UnrealGameSync/Telemetry.cs
+++ b/Engine/Source/Programs/UnrealGameSync/UnrealGameSync/Telemetry.cs
@@ -1,17 +1,13 @@
 // Copyright Epic Games, Inc. All Rights Reserved.
 
 using System;
-using System.Collections.Concurrent;
 using System.Collections.Generic;
 using System.IO;
-using System.Linq;
 using System.Net;
-using System.Reflection;
 using System.Text;
+using System.Text.Json;
 using System.Threading;
-using System.Threading.Tasks;
 using System.Web;
-using System.Web.Script.Serialization;
 
 namespace UnrealGameSync
 {
@@ -110,7 +106,7 @@ namespace UnrealGameSync
 		/// <inheritdoc/>
 		public void SendEvent(string EventName, object Attributes)
 		{
-			string AttributesText = new JavaScriptSerializer().Serialize(Attributes);
+			string AttributesText = JsonSerializer.Serialize(Attributes);
 			if (AttributesText[0] != '{')
 			{
 				throw new Exception("Expected event data with named properties");
diff --git a/Engine/Source/Programs/UnrealGameSync/UnrealGameSync/TelemetryStopwatch.cs b/Engine/Source/Programs/UnrealGameSync/UnrealGameSync/TelemetryStopwatch.cs
--- a/Engine/Source/Programs/UnrealGameSync/UnrealGameSync/TelemetryStopwatch.cs
+++ b/Engine/Source/Programs/UnrealGameSync/UnrealGameSync/TelemetryStopwatch.cs
@@ -1,18 +1,9 @@
 // Copyright Epic Games, Inc. All Rights Reserved.
 
 using System;
-using System.Collections.Concurrent;
 using System.Collections.Generic;
-using System.Data.SqlClient;
 using System.Diagnostics;
-using System.Linq;
-using System.Net;
-using System.Net.Sockets;
 using System.Reflection;
-using System.Text;
-using System.Threading;
-using System.Threading.Tasks;
-using System.Web.Script.Serialization;
 
 namespace UnrealGameSync
 {
diff --git a/Engine/Source/Programs/UnrealGameSync/UnrealGameSync/UnrealGameSync.csproj b/Engine/Source/Programs/UnrealGameSync/UnrealGameSync/UnrealGameSync.csproj
--- a/Engine/Source/Programs/UnrealGameSync/UnrealGameSync/UnrealGameSync.csproj
+++ b/Engine/Source/Programs/UnrealGameSync/UnrealGameSync/UnrealGameSync.csproj
@@ -1,102 +1,26 @@
-﻿<?xml version="1.0" encoding="utf-8"?>
-<Project ToolsVersion="12.0" DefaultTargets="Build" xmlns="http://schemas.microsoft.com/developer/msbuild/2003">
-  <Import Project="$(MSBuildExtensionsPath)\$(MSBuildToolsVersion)\Microsoft.Common.props" Condition="Exists('$(MSBuildExtensionsPath)\$(MSBuildToolsVersion)\Microsoft.Common.props')" />
+<Project Sdk="Microsoft.NET.Sdk.WindowsDesktop">
   <PropertyGroup>
-    <Configuration Condition=" '$(Configuration)' == '' ">Debug</Configuration>
-    <Platform Condition=" '$(Platform)' == '' ">AnyCPU</Platform>
-    <ProjectGuid>{94DC470D-5944-4899-A646-299F5E605A1D}</ProjectGuid>
     <OutputType>WinExe</OutputType>
-    <AppDesignerFolder>Properties</AppDesignerFolder>
+    <TargetFramework>netcoreapp3.1</TargetFramework>
+    <UseWindowsForms>true</UseWindowsForms>
+    <GenerateAssemblyInfo>false</GenerateAssemblyInfo>
     <RootNamespace>UnrealGameSync</RootNamespace>
     <AssemblyName>UnrealGameSync</AssemblyName>
-    <TargetFrameworkVersion>v4.5</TargetFrameworkVersion>
-    <FileAlignment>512</FileAlignment>
-    <IsWebBootstrapper>false</IsWebBootstrapper>
-    <PublishUrl>publish\</PublishUrl>
-    <Install>true</Install>
-    <InstallFrom>Disk</InstallFrom>
-    <UpdateEnabled>false</UpdateEnabled>
-    <UpdateMode>Foreground</UpdateMode>
-    <UpdateInterval>7</UpdateInterval>
-    <UpdateIntervalUnits>Days</UpdateIntervalUnits>
-    <UpdatePeriodically>false</UpdatePeriodically>
-    <UpdateRequired>false</UpdateRequired>
-    <MapFileExtensions>true</MapFileExtensions>
-    <WebPage>publish.htm</WebPage>
-    <OpenBrowserOnPublish>false</OpenBrowserOnPublish>
-    <ApplicationRevision>14</ApplicationRevision>
-    <ApplicationVersion>1.0.0.%2a</ApplicationVersion>
-    <UseApplicationTrust>false</UseApplicationTrust>
-    <PublishWizardCompleted>true</PublishWizardCompleted>
-    <BootstrapperEnabled>false</BootstrapperEnabled>
-  </PropertyGroup>
-  <PropertyGroup>
-    <GenerateManifests>false</GenerateManifests>
-  </PropertyGroup>
-  <PropertyGroup>
-    <SignManifests>false</SignManifests>
-  </PropertyGroup>
-  <PropertyGroup>
+    <StartupObject>UnrealGameSync.Program</StartupObject>
     <ApplicationIcon>Resources\Icon.ico</ApplicationIcon>
-  </PropertyGroup>
-  <PropertyGroup>
-    <SignAssembly>false</SignAssembly>
-  </PropertyGroup>
-  <PropertyGroup>
-    <TargetZone>LocalIntranet</TargetZone>
-  </PropertyGroup>
-  <PropertyGroup>
     <ApplicationManifest>Properties\app.manifest</ApplicationManifest>
   </PropertyGroup>
-  <PropertyGroup Condition="'$(Configuration)|$(Platform)' == 'Debug|x64'">
-    <DebugSymbols>true</DebugSymbols>
-    <OutputPath>bin\x64\Debug\</OutputPath>
-    <DefineConstants>TRACE;DEBUG;WITH_TELEMETRY</DefineConstants>
-    <DebugType>full</DebugType>
-    <PlatformTarget>x64</PlatformTarget>
-    <UseVSHostingProcess>false</UseVSHostingProcess>
-    <ErrorReport>prompt</ErrorReport>
-    <CodeAnalysisRuleSet>MinimumRecommendedRules.ruleset</CodeAnalysisRuleSet>
-    <Prefer32Bit>true</Prefer32Bit>
-  </PropertyGroup>
-  <PropertyGroup Condition="'$(Configuration)|$(Platform)' == 'Release|x64'">
-    <OutputPath>bin\UnstableRelease\</OutputPath>
-    <DefineConstants>TRACE;WITH_TELEMETRY</DefineConstants>
-    <Optimize>true</Optimize>
-    <DebugType>pdbonly</DebugType>
-    <PlatformTarget>x64</PlatformTarget>
-    <UseVSHostingProcess>false</UseVSHostingProcess>
-    <ErrorReport>prompt</ErrorReport>
-    <CodeAnalysisRuleSet>MinimumRecommendedRules.ruleset</CodeAnalysisRuleSet>
-    <Prefer32Bit>true</Prefer32Bit>
-  </PropertyGroup>
-  <PropertyGroup Condition="'$(Configuration)|$(Platform)' == 'Debug|x86'">
-    <DebugSymbols>true</DebugSymbols>
-    <OutputPath>bin\x86\Debug\</OutputPath>
+  <PropertyGroup Condition="'$(Configuration)|$(Platform)'=='Release|AnyCPU'">
     <DefineConstants>TRACE;DEBUG;WITH_TELEMETRY</DefineConstants>
-    <DebugType>full</DebugType>
-    <PlatformTarget>x86</PlatformTarget>
-    <UseVSHostingProcess>false</UseVSHostingProcess>
-    <ErrorReport>prompt</ErrorReport>
-    <CodeAnalysisRuleSet>MinimumRecommendedRules.ruleset</CodeAnalysisRuleSet>
-    <Prefer32Bit>true</Prefer32Bit>
   </PropertyGroup>
-  <PropertyGroup Condition="'$(Configuration)|$(Platform)' == 'Release|x86'">
-    <OutputPath>bin\x86\Release\</OutputPath>
+  <PropertyGroup Condition="'$(Configuration)|$(Platform)'=='Debug|AnyCPU'">
     <DefineConstants>TRACE;WITH_TELEMETRY</DefineConstants>
-    <Optimize>true</Optimize>
-    <DebugType>pdbonly</DebugType>
-    <PlatformTarget>x86</PlatformTarget>
-    <UseVSHostingProcess>false</UseVSHostingProcess>
-    <ErrorReport>prompt</ErrorReport>
-    <CodeAnalysisRuleSet>MinimumRecommendedRules.ruleset</CodeAnalysisRuleSet>
-    <Prefer32Bit>true</Prefer32Bit>
   </PropertyGroup>
   <ItemGroup>
     <Reference Include="Ionic.Zip.Reduced">
       <HintPath>ThirdParty\Ionic.Zip.Reduced.dll</HintPath>
     </Reference>
-    <Reference Include="System" />
+    <!--<Reference Include="System" />
     <Reference Include="System.Activities" />
     <Reference Include="System.Core" />
     <Reference Include="System.Web" />
@@ -108,429 +32,20 @@
     <Reference Include="System.Deployment" />
     <Reference Include="System.Drawing" />
     <Reference Include="System.Windows.Forms" />
-    <Reference Include="System.Xml" />
+    <Reference Include="System.Xml" />-->
   </ItemGroup>
   <ItemGroup>
-    <Compile Include="ActivationListener.cs" />
-    <Compile Include="ArchiveManifest.cs" />
-    <Compile Include="AutomationServer.cs" />
-    <Compile Include="BuildStep.cs" />
-    <Compile Include="ChildProcess.cs" />
-    <Compile Include="Controls\AlertButtonControl.cs">
-      <SubType>Component</SubType>
-    </Compile>
-    <Compile Include="Controls\AlertImageControl.cs">
-      <SubType>UserControl</SubType>
-    </Compile>
-    <Compile Include="Controls\CustomListViewControl.cs" />
-    <Compile Include="Controls\ErrorPanel.cs">
-      <SubType>Component</SubType>
-    </Compile>
-    <Compile Include="Controls\RichTextBox41.cs">
-      <SubType>Component</SubType>
-    </Compile>
-    <Compile Include="Controls\SyncFilterControl.cs">
-      <SubType>UserControl</SubType>
-    </Compile>
-    <Compile Include="Controls\SyncFilterControl.Designer.cs">
-      <DependentUpon>SyncFilterControl.cs</DependentUpon>
-    </Compile>
-    <Compile Include="Controls\TextBoxWithCueBanner.cs">
-      <SubType>Component</SubType>
-    </Compile>
-    <Compile Include="Controls\WorkspaceControl.cs">
-      <SubType>UserControl</SubType>
-    </Compile>
-    <Compile Include="Controls\WorkspaceControl.Designer.cs">
-      <DependentUpon>WorkspaceControl.cs</DependentUpon>
-    </Compile>
-    <Compile Include="Controls\StatusPanel.cs" />
-    <Compile Include="Controls\BuildListControl.cs">
-      <SubType>Component</SubType>
-    </Compile>
-    <Compile Include="Controls\LogControl.cs">
-      <SubType>UserControl</SubType>
-    </Compile>
-    <Compile Include="Controls\LogSplitContainer.cs">
-      <SubType>Component</SubType>
-    </Compile>
-    <Compile Include="Controls\TabControl.cs">
-      <SubType>Component</SubType>
-    </Compile>
-    <Compile Include="DeleteFilesTask.cs" />
-    <Compile Include="DeploymentSettings.cs" />
-    <Compile Include="DetectProjectSettingsTask.cs" />
-    <Compile Include="FastJSON\Formatter.cs" />
-    <Compile Include="FastJSON\Getters.cs" />
-    <Compile Include="FastJSON\JSON.cs" />
-    <Compile Include="FastJSON\JsonParser.cs" />
-    <Compile Include="FastJSON\JsonSerializer.cs" />
-    <Compile Include="FastJSON\Reflection.cs" />
-    <Compile Include="FastJSON\SafeDictionary.cs" />
-    <Compile Include="Json.cs" />
-    <Compile Include="FileFilter.cs" />
-    <Compile Include="FindFoldersToCleanTask.cs" />
-    <Compile Include="Forms\ArgumentsWindow.cs">
-      <SubType>Form</SubType>
-    </Compile>
-    <Compile Include="Forms\ArgumentsWindow.Designer.cs">
-      <DependentUpon>ArgumentsWindow.cs</DependentUpon>
-    </Compile>
-    <Compile Include="Forms\AutomatedSyncWindow.cs">
-      <SubType>Form</SubType>
-    </Compile>
-    <Compile Include="Forms\AutomatedSyncWindow.Designer.cs">
-      <DependentUpon>AutomatedSyncWindow.cs</DependentUpon>
-    </Compile>
-    <Compile Include="Forms\IssueAlertWindow.cs">
-      <SubType>Form</SubType>
-    </Compile>
-    <Compile Include="Forms\IssueAlertWindow.Designer.cs">
-      <DependentUpon>IssueAlertWindow.cs</DependentUpon>
-    </Compile>
-    <Compile Include="Forms\IssueFixedWindow.cs">
-      <SubType>Form</SubType>
-    </Compile>
-    <Compile Include="Forms\IssueFixedWindow.Designer.cs">
-      <DependentUpon>IssueFixedWindow.cs</DependentUpon>
-    </Compile>
-    <Compile Include="Forms\IssueBrowserWindow.cs">
-      <SubType>Form</SubType>
-    </Compile>
-    <Compile Include="Forms\IssueBrowserWindow.Designer.cs">
-      <DependentUpon>IssueBrowserWindow.cs</DependentUpon>
-    </Compile>
-    <Compile Include="Forms\IssueSettingsWindow.cs">
-      <SubType>Form</SubType>
-    </Compile>
-    <Compile Include="Forms\IssueSettingsWindow.Designer.cs">
-      <DependentUpon>IssueSettingsWindow.cs</DependentUpon>
-    </Compile>
-    <Compile Include="Forms\SelectUserWindow.cs">
-      <SubType>Form</SubType>
-    </Compile>
-    <Compile Include="Forms\SelectUserWindow.Designer.cs">
-      <DependentUpon>SelectUserWindow.cs</DependentUpon>
-    </Compile>
-    <Compile Include="Forms\CleanWorkspaceWindow.cs">
-      <SubType>Form</SubType>
-    </Compile>
-    <Compile Include="Forms\CleanWorkspaceWindow.Designer.cs">
-      <DependentUpon>CleanWorkspaceWindow.cs</DependentUpon>
-    </Compile>
-    <Compile Include="Forms\ClobberWindow.cs">
-      <SubType>Form</SubType>
-    </Compile>
-    <Compile Include="Forms\ClobberWindow.Designer.cs">
-      <DependentUpon>ClobberWindow.cs</DependentUpon>
-    </Compile>
-    <Compile Include="Forms\ConnectWindow.cs">
-      <SubType>Form</SubType>
-    </Compile>
-    <Compile Include="Forms\ConnectWindow.Designer.cs">
-      <DependentUpon>ConnectWindow.cs</DependentUpon>
-    </Compile>
-    <Compile Include="Forms\IssueDetailsWindow.cs">
-      <SubType>Form</SubType>
-    </Compile>
-    <Compile Include="Forms\IssueDetailsWindow.Designer.cs">
-      <DependentUpon>IssueDetailsWindow.cs</DependentUpon>
-    </Compile>
-    <Compile Include="Forms\NewWorkspaceWindow.cs">
-      <SubType>Form</SubType>
-    </Compile>
-    <Compile Include="Forms\NewWorkspaceWindow.Designer.cs">
-      <DependentUpon>NewWorkspaceWindow.cs</DependentUpon>
-    </Compile>
-    <Compile Include="Forms\DeleteWindow.cs">
-      <SubType>Form</SubType>
-    </Compile>
-    <Compile Include="Forms\DeleteWindow.Designer.cs">
-      <DependentUpon>DeleteWindow.cs</DependentUpon>
-    </Compile>
-    <Compile Include="Forms\NotificationWindow.cs">
-      <SubType>Form</SubType>
-    </Compile>
-    <Compile Include="Forms\NotificationWindow.Designer.cs">
-      <DependentUpon>NotificationWindow.cs</DependentUpon>
-    </Compile>
-    <Compile Include="Forms\OpenProjectWindow.cs">
-      <SubType>Form</SubType>
-    </Compile>
-    <Compile Include="Forms\OpenProjectWindow.Designer.cs">
-      <DependentUpon>OpenProjectWindow.cs</DependentUpon>
-    </Compile>
-    <Compile Include="Forms\PasswordWindow.cs">
-      <SubType>Form</SubType>
-    </Compile>
-    <Compile Include="Forms\PasswordWindow.Designer.cs">
-      <DependentUpon>PasswordWindow.cs</DependentUpon>
-    </Compile>
-    <Compile Include="Forms\ApplicationSettingsWindow.cs">
-      <SubType>Form</SubType>
-    </Compile>
-    <Compile Include="Forms\ApplicationSettingsWindow.Designer.cs">
-      <DependentUpon>ApplicationSettingsWindow.cs</DependentUpon>
-    </Compile>
-    <Compile Include="Forms\PerforceSyncSettingsWindow.cs">
-      <SubType>Form</SubType>
-    </Compile>
-    <Compile Include="Forms\PerforceSyncSettingsWindow.Designer.cs">
-      <DependentUpon>PerforceSyncSettingsWindow.cs</DependentUpon>
-    </Compile>
-    <Compile Include="Forms\ProgramsRunningWindow.cs">
-      <SubType>Form</SubType>
-    </Compile>
-    <Compile Include="Forms\ProgramsRunningWindow.Designer.cs">
-      <DependentUpon>ProgramsRunningWindow.cs</DependentUpon>
-    </Compile>
-    <Compile Include="Forms\SdkInfoWindow.cs">
-      <SubType>Form</SubType>
-    </Compile>
-    <Compile Include="Forms\SdkInfoWindow.Designer.cs">
-      <DependentUpon>SdkInfoWindow.cs</DependentUpon>
-    </Compile>
-    <Compile Include="Forms\SelectProjectFromWorkspaceWindow.cs">
-      <SubType>Form</SubType>
-    </Compile>
-    <Compile Include="Forms\SelectProjectFromWorkspaceWindow.Designer.cs">
-      <DependentUpon>SelectProjectFromWorkspaceWindow.cs</DependentUpon>
-    </Compile>
-    <Compile Include="Forms\SelectStreamWindow.cs">
-      <SubType>Form</SubType>
-    </Compile>
-    <Compile Include="Forms\SelectStreamWindow.Designer.cs">
-      <DependentUpon>SelectStreamWindow.cs</DependentUpon>
-    </Compile>
-    <Compile Include="Forms\SelectWorkspaceWindow.cs">
-      <SubType>Form</SubType>
-    </Compile>
-    <Compile Include="Forms\SelectWorkspaceWindow.Designer.cs">
-      <DependentUpon>SelectWorkspaceWindow.cs</DependentUpon>
-    </Compile>
-    <Compile Include="Forms\SyncFilterSyntax.cs">
-      <SubType>Form</SubType>
-    </Compile>
-    <Compile Include="Forms\SyncFilterSyntax.Designer.cs">
-      <DependentUpon>SyncFilterSyntax.cs</DependentUpon>
-    </Compile>
-    <Compile Include="Forms\VariablesWindow.cs">
-      <SubType>Form</SubType>
-    </Compile>
-    <Compile Include="Forms\VariablesWindow.Designer.cs">
-      <DependentUpon>VariablesWindow.cs</DependentUpon>
-    </Compile>
-    <Compile Include="Forms\ModifyBuildStepsWindow.cs">
-      <SubType>Form</SubType>
-    </Compile>
-    <Compile Include="Forms\ModifyBuildStepsWindow.Designer.cs">
-      <DependentUpon>ModifyBuildStepsWindow.cs</DependentUpon>
-    </Compile>
-    <Compile Include="Forms\DiagnosticsWindow.cs">
-      <SubType>Form</SubType>
-    </Compile>
-    <Compile Include="Forms\DiagnosticsWindow.Designer.cs">
-      <DependentUpon>DiagnosticsWindow.cs</DependentUpon>
-    </Compile>
-    <Compile Include="Forms\ChangelistWindow.cs">
-      <SubType>Form</SubType>
-    </Compile>
-    <Compile Include="Forms\ChangelistWindow.Designer.cs">
-      <DependentUpon>ChangelistWindow.cs</DependentUpon>
-    </Compile>
-    <Compile Include="Forms\BuildStepWindow.cs">
-      <SubType>Form</SubType>
-    </Compile>
-    <Compile Include="Forms\BuildStepWindow.Designer.cs">
-      <DependentUpon>BuildStepWindow.cs</DependentUpon>
-    </Compile>
-    <Compile Include="Forms\FailedToDeleteWindow.cs">
-      <SubType>Form</SubType>
-    </Compile>
-    <Compile Include="Forms\FailedToDeleteWindow.Designer.cs">
-      <DependentUpon>FailedToDeleteWindow.cs</DependentUpon>
-    </Compile>
-    <Compile Include="Forms\LeaveCommentWindow.cs">
-      <SubType>Form</SubType>
-    </Compile>
-    <Compile Include="Forms\LeaveCommentWindow.Designer.cs">
-      <DependentUpon>LeaveCommentWindow.cs</DependentUpon>
-    </Compile>
-    <Compile Include="Forms\MainWindow.cs">
-      <SubType>Form</SubType>
-    </Compile>
-    <Compile Include="Forms\MainWindow.Designer.cs">
-      <DependentUpon>MainWindow.cs</DependentUpon>
-    </Compile>
-    <Compile Include="ConfigFile.cs" />
-    <Compile Include="EventMonitor.cs" />
-    <Compile Include="Forms\ModalTaskWindow.cs">
-      <SubType>Form</SubType>
-    </Compile>
-    <Compile Include="Forms\ModalTaskWindow.Designer.cs">
-      <DependentUpon>ModalTaskWindow.cs</DependentUpon>
-    </Compile>
-    <Compile Include="Forms\ScheduleWindow.cs">
-      <SubType>Form</SubType>
-    </Compile>
-    <Compile Include="Forms\ScheduleWindow.Designer.cs">
-      <DependentUpon>ScheduleWindow.cs</DependentUpon>
-    </Compile>
-    <Compile Include="Forms\SyncFilter.cs">
-      <SubType>Form</SubType>
-    </Compile>
-    <Compile Include="Forms\SyncFilter.Designer.cs">
-      <DependentUpon>SyncFilter.cs</DependentUpon>
-    </Compile>
-    <Compile Include="IssueMonitor.cs" />
-    <Compile Include="ModalTask.cs" />
-    <Compile Include="PerforceModalTask.cs" />
-    <Compile Include="PerforceMonitor.cs" />
-    <Compile Include="ProgramApplicationContext.cs" />
-    <Compile Include="Properties\Resources.Designer.cs">
+    <Compile Update="Properties\Settings.Designer.cs">
+      <DesignTimeSharedInput>True</DesignTimeSharedInput>
       <AutoGen>True</AutoGen>
-      <DesignTime>True</DesignTime>
-      <DependentUpon>Resources.resx</DependentUpon>
+      <DependentUpon>Settings.settings</DependentUpon>
     </Compile>
-    <Compile Include="Taskbar.cs" />
-    <Compile Include="Telemetry.cs" />
-    <Compile Include="TelemetryStopwatch.cs" />
-    <Compile Include="UpdateMonitor.cs" />
-    <Compile Include="UserSettings.cs" />
-    <Compile Include="RESTApi.cs" />
-    <Compile Include="Workspace.cs" />
-    <Compile Include="OutputAdapters.cs" />
-    <Compile Include="Perforce.cs" />
-    <Compile Include="Program.cs" />
-    <Compile Include="Properties\AssemblyInfo.cs" />
-    <Compile Include="Utility.cs" />
-    <EmbeddedResource Include="Controls\SyncFilterControl.resx">
-      <DependentUpon>SyncFilterControl.cs</DependentUpon>
-    </EmbeddedResource>
-    <EmbeddedResource Include="Controls\WorkspaceControl.resx">
-      <DependentUpon>WorkspaceControl.cs</DependentUpon>
-    </EmbeddedResource>
-    <EmbeddedResource Include="Forms\ArgumentsWindow.resx">
-      <DependentUpon>ArgumentsWindow.cs</DependentUpon>
-    </EmbeddedResource>
-    <EmbeddedResource Include="Forms\AutomatedSyncWindow.resx">
-      <DependentUpon>AutomatedSyncWindow.cs</DependentUpon>
-    </EmbeddedResource>
-    <EmbeddedResource Include="Forms\IssueAlertWindow.resx">
-      <DependentUpon>IssueAlertWindow.cs</DependentUpon>
-    </EmbeddedResource>
-    <EmbeddedResource Include="Forms\IssueFixedWindow.resx">
-      <DependentUpon>IssueFixedWindow.cs</DependentUpon>
-    </EmbeddedResource>
-    <EmbeddedResource Include="Forms\IssueBrowserWindow.resx">
-      <DependentUpon>IssueBrowserWindow.cs</DependentUpon>
-    </EmbeddedResource>
-    <EmbeddedResource Include="Forms\IssueSettingsWindow.resx">
-      <DependentUpon>IssueSettingsWindow.cs</DependentUpon>
-    </EmbeddedResource>
-    <EmbeddedResource Include="Forms\SelectUserWindow.resx">
-      <DependentUpon>SelectUserWindow.cs</DependentUpon>
-    </EmbeddedResource>
-    <EmbeddedResource Include="Forms\CleanWorkspaceWindow.resx">
-      <DependentUpon>CleanWorkspaceWindow.cs</DependentUpon>
-    </EmbeddedResource>
-    <EmbeddedResource Include="Forms\ClobberWindow.resx">
-      <DependentUpon>ClobberWindow.cs</DependentUpon>
-    </EmbeddedResource>
-    <EmbeddedResource Include="Forms\ConnectWindow.resx">
-      <DependentUpon>ConnectWindow.cs</DependentUpon>
-    </EmbeddedResource>
-    <EmbeddedResource Include="Forms\IssueDetailsWindow.resx">
-      <DependentUpon>IssueDetailsWindow.cs</DependentUpon>
-    </EmbeddedResource>
-    <EmbeddedResource Include="Forms\NewWorkspaceWindow.resx">
-      <DependentUpon>NewWorkspaceWindow.cs</DependentUpon>
-    </EmbeddedResource>
-    <EmbeddedResource Include="Forms\DeleteWindow.resx">
-      <DependentUpon>DeleteWindow.cs</DependentUpon>
-    </EmbeddedResource>
-    <EmbeddedResource Include="Forms\NotificationWindow.resx">
-      <DependentUpon>NotificationWindow.cs</DependentUpon>
-    </EmbeddedResource>
-    <EmbeddedResource Include="Forms\OpenProjectWindow.resx">
-      <DependentUpon>OpenProjectWindow.cs</DependentUpon>
-    </EmbeddedResource>
-    <EmbeddedResource Include="Forms\PasswordWindow.resx">
-      <DependentUpon>PasswordWindow.cs</DependentUpon>
-    </EmbeddedResource>
-    <EmbeddedResource Include="Forms\ApplicationSettingsWindow.resx">
-      <DependentUpon>ApplicationSettingsWindow.cs</DependentUpon>
-    </EmbeddedResource>
-    <EmbeddedResource Include="Forms\PerforceSyncSettingsWindow.resx">
-      <DependentUpon>PerforceSyncSettingsWindow.cs</DependentUpon>
-    </EmbeddedResource>
-    <EmbeddedResource Include="Forms\ProgramsRunningWindow.resx">
-      <DependentUpon>ProgramsRunningWindow.cs</DependentUpon>
-    </EmbeddedResource>
-    <EmbeddedResource Include="Forms\SdkInfoWindow.resx">
-      <DependentUpon>SdkInfoWindow.cs</DependentUpon>
-    </EmbeddedResource>
-    <EmbeddedResource Include="Forms\SelectProjectFromWorkspaceWindow.resx">
-      <DependentUpon>SelectProjectFromWorkspaceWindow.cs</DependentUpon>
-    </EmbeddedResource>
-    <EmbeddedResource Include="Forms\SelectStreamWindow.resx">
-      <DependentUpon>SelectStreamWindow.cs</DependentUpon>
-    </EmbeddedResource>
-    <EmbeddedResource Include="Forms\SelectWorkspaceWindow.resx">
-      <DependentUpon>SelectWorkspaceWindow.cs</DependentUpon>
-    </EmbeddedResource>
-    <EmbeddedResource Include="Forms\SyncFilterSyntax.resx">
-      <DependentUpon>SyncFilterSyntax.cs</DependentUpon>
-    </EmbeddedResource>
-    <EmbeddedResource Include="Forms\VariablesWindow.resx">
-      <DependentUpon>VariablesWindow.cs</DependentUpon>
-    </EmbeddedResource>
-    <EmbeddedResource Include="Forms\ModifyBuildStepsWindow.resx">
-      <DependentUpon>ModifyBuildStepsWindow.cs</DependentUpon>
-    </EmbeddedResource>
-    <EmbeddedResource Include="Forms\DiagnosticsWindow.resx">
-      <DependentUpon>DiagnosticsWindow.cs</DependentUpon>
-    </EmbeddedResource>
-    <EmbeddedResource Include="Forms\ChangelistWindow.resx">
-      <DependentUpon>ChangelistWindow.cs</DependentUpon>
-    </EmbeddedResource>
-    <EmbeddedResource Include="Forms\BuildStepWindow.resx">
-      <DependentUpon>BuildStepWindow.cs</DependentUpon>
-    </EmbeddedResource>
-    <EmbeddedResource Include="Forms\FailedToDeleteWindow.resx">
-      <DependentUpon>FailedToDeleteWindow.cs</DependentUpon>
-    </EmbeddedResource>
-    <EmbeddedResource Include="Forms\LeaveCommentWindow.resx">
-      <DependentUpon>LeaveCommentWindow.cs</DependentUpon>
-    </EmbeddedResource>
-    <EmbeddedResource Include="Forms\MainWindow.resx">
-      <DependentUpon>MainWindow.cs</DependentUpon>
-    </EmbeddedResource>
-    <EmbeddedResource Include="Forms\ModalTaskWindow.resx">
-      <DependentUpon>ModalTaskWindow.cs</DependentUpon>
-    </EmbeddedResource>
-    <EmbeddedResource Include="Forms\ScheduleWindow.resx">
-      <DependentUpon>ScheduleWindow.cs</DependentUpon>
-    </EmbeddedResource>
-    <EmbeddedResource Include="Forms\SyncFilter.resx">
-      <DependentUpon>SyncFilter.cs</DependentUpon>
-    </EmbeddedResource>
-    <EmbeddedResource Include="Properties\Resources.resx">
-      <Generator>ResXFileCodeGenerator</Generator>
-      <SubType>Designer</SubType>
-      <LastGenOutput>Resources.Designer.cs</LastGenOutput>
-    </EmbeddedResource>
-    <None Include="Properties\app.manifest" />
-    <None Include="Properties\Settings.settings">
+  </ItemGroup>
+  <ItemGroup>
+    <None Update="Properties\Settings.settings">
       <Generator>SettingsSingleFileGenerator</Generator>
       <LastGenOutput>Settings.Designer.cs</LastGenOutput>
     </None>
-    <Compile Include="Properties\Settings.Designer.cs">
-      <AutoGen>True</AutoGen>
-      <DependentUpon>Settings.settings</DependentUpon>
-      <DesignTimeSharedInput>True</DesignTimeSharedInput>
-    </Compile>
-    <None Include="Resources\App.config" />
   </ItemGroup>
   <Choose>
     <When Condition="Exists('..\..\..\..\Restricted\NotForLicensees\Source\Programs\UnrealGameSync\UnrealGameSync\DeploymentSettingsInternal.cs')">
@@ -541,53 +56,4 @@
       </ItemGroup>
     </When>
   </Choose>
-  <ItemGroup>
-    <EmbeddedResource Include="Resources\Icons.png" />
-  </ItemGroup>
-  <ItemGroup>
-    <EmbeddedResource Include="Resources\DropList.png" />
-  </ItemGroup>
-  <ItemGroup>
-    <EmbeddedResource Include="Resources\Icon.ico" />
-  </ItemGroup>
-  <ItemGroup>
-    <EmbeddedResource Include="Resources\Log.png" />
-  </ItemGroup>
-  <ItemGroup>
-    <BootstrapperPackage Include=".NETFramework,Version=v4.5">
-      <Visible>False</Visible>
-      <ProductName>Microsoft .NET Framework 4.5 %28x86 and x64%29</ProductName>
-      <Install>true</Install>
-    </BootstrapperPackage>
-    <BootstrapperPackage Include="Microsoft.Net.Client.3.5">
-      <Visible>False</Visible>
-      <ProductName>.NET Framework 3.5 SP1 Client Profile</ProductName>
-      <Install>false</Install>
-    </BootstrapperPackage>
-    <BootstrapperPackage Include="Microsoft.Net.Framework.3.5.SP1">
-      <Visible>False</Visible>
-      <ProductName>.NET Framework 3.5 SP1</ProductName>
-      <Install>false</Install>
-    </BootstrapperPackage>
-  </ItemGroup>
-  <ItemGroup>
-    <None Include="Resources\AlertErrorIcon.png" />
-    <None Include="Resources\AlertWarningIcon.png" />
-    <Content Include="Resources\Perforce.png" />
-    <None Include="Resources\DefaultErrorLogo.png" />
-    <Content Include="Resources\DefaultProjectLogo.png" />
-    <Content Include="Resources\DefaultNotificationLogo.png" />
-    <Content Include="Resources\StatusIcons.png" />
-  </ItemGroup>
-  <ItemGroup>
-    <Folder Include="DotNETCommon\" />
-  </ItemGroup>
-  <Import Project="$(MSBuildToolsPath)\Microsoft.CSharp.targets" />
-  <!-- To modify your build process, add your task inside one of the targets below and uncomment it. 
-       Other similar extension points exist, see Microsoft.Common.targets.
-  <Target Name="BeforeBuild">
-  </Target>
-  <Target Name="AfterBuild">
-  </Target>
-  -->
 </Project>
\ No newline at end of file
diff --git a/Engine/Source/Programs/UnrealGameSync/UnrealGameSync/Utility.cs b/Engine/Source/Programs/UnrealGameSync/UnrealGameSync/Utility.cs
--- a/Engine/Source/Programs/UnrealGameSync/UnrealGameSync/Utility.cs
+++ b/Engine/Source/Programs/UnrealGameSync/UnrealGameSync/Utility.cs
@@ -8,12 +8,10 @@ using System.Collections.Generic;
 using System.Diagnostics;
 using System.Drawing;
 using System.IO;
-using System.Linq;
-using System.Runtime.InteropServices;
+using System.Linq;	
 using System.Text;
+using System.Text.Json;
 using System.Text.RegularExpressions;
-using System.Threading.Tasks;
-using System.Web.Script.Serialization;
 using System.Windows.Forms;
 
 namespace UnrealGameSync
@@ -250,8 +248,7 @@ namespace UnrealGameSync
 		{
 			try
 			{
-				JavaScriptSerializer Serializer = new JavaScriptSerializer();
-				Dictionary<string, object> RawObject = Serializer.Deserialize<Dictionary<string, object>>(Text);
+				Dictionary<string, object> RawObject = JsonSerializer.Deserialize<Dictionary<string, object>>(Text);
 
 				object Enterprise;
 				if(RawObject.TryGetValue("Enterprise", out Enterprise) && Enterprise is bool)
diff --git a/Engine/Source/Programs/UnrealGameSync/UnrealGameSyncAutomation/UnrealGameSyncAutomation.csproj b/Engine/Source/Programs/UnrealGameSync/UnrealGameSyncAutomation/UnrealGameSyncAutomation.csproj
--- a/Engine/Source/Programs/UnrealGameSync/UnrealGameSyncAutomation/UnrealGameSyncAutomation.csproj
+++ b/Engine/Source/Programs/UnrealGameSync/UnrealGameSyncAutomation/UnrealGameSyncAutomation.csproj
@@ -1,54 +1,11 @@
 ﻿<?xml version="1.0" encoding="utf-8"?>
-<Project ToolsVersion="15.0" xmlns="http://schemas.microsoft.com/developer/msbuild/2003">
-  <Import Project="$(MSBuildExtensionsPath)\$(MSBuildToolsVersion)\Microsoft.Common.props" Condition="Exists('$(MSBuildExtensionsPath)\$(MSBuildToolsVersion)\Microsoft.Common.props')" />
+<Project Sdk="Microsoft.NET.Sdk.WindowsDesktop">
   <PropertyGroup>
-    <Configuration Condition=" '$(Configuration)' == '' ">Debug</Configuration>
-    <Platform Condition=" '$(Platform)' == '' ">AnyCPU</Platform>
-    <ProjectGuid>{7C9D0A33-C2A8-4415-B9ED-9F40337DBEC6}</ProjectGuid>
-    <OutputType>Exe</OutputType>
+    <OutputType>WinExe</OutputType>
+    <TargetFramework>netcoreapp3.1</TargetFramework>
+    <UseWindowsForms>true</UseWindowsForms>
+    <GenerateAssemblyInfo>false</GenerateAssemblyInfo>
     <RootNamespace>UnrealGameSyncAutomation</RootNamespace>
     <AssemblyName>UnrealGameSyncAutomation</AssemblyName>
-    <TargetFrameworkVersion>v4.5</TargetFrameworkVersion>
-    <FileAlignment>512</FileAlignment>
-    <AutoGenerateBindingRedirects>true</AutoGenerateBindingRedirects>
-    <Deterministic>true</Deterministic>
-    <TargetFrameworkProfile />
   </PropertyGroup>
-  <PropertyGroup Condition=" '$(Configuration)|$(Platform)' == 'Debug|AnyCPU' ">
-    <PlatformTarget>AnyCPU</PlatformTarget>
-    <DebugSymbols>true</DebugSymbols>
-    <DebugType>full</DebugType>
-    <Optimize>false</Optimize>
-    <OutputPath>bin\Debug\</OutputPath>
-    <DefineConstants>DEBUG;TRACE</DefineConstants>
-    <ErrorReport>prompt</ErrorReport>
-    <WarningLevel>4</WarningLevel>
-  </PropertyGroup>
-  <PropertyGroup Condition=" '$(Configuration)|$(Platform)' == 'Release|AnyCPU' ">
-    <PlatformTarget>AnyCPU</PlatformTarget>
-    <DebugType>pdbonly</DebugType>
-    <Optimize>true</Optimize>
-    <OutputPath>bin\Release\</OutputPath>
-    <DefineConstants>TRACE</DefineConstants>
-    <ErrorReport>prompt</ErrorReport>
-    <WarningLevel>4</WarningLevel>
-  </PropertyGroup>
-  <ItemGroup>
-    <Reference Include="System" />
-    <Reference Include="System.Core" />
-    <Reference Include="System.Xml.Linq" />
-    <Reference Include="System.Data.DataSetExtensions" />
-    <Reference Include="Microsoft.CSharp" />
-    <Reference Include="System.Data" />
-    <Reference Include="System.Net.Http" />
-    <Reference Include="System.Xml" />
-  </ItemGroup>
-  <ItemGroup>
-    <Compile Include="Program.cs" />
-    <Compile Include="Properties\AssemblyInfo.cs" />
-  </ItemGroup>
-  <ItemGroup>
-    <None Include="App.config" />
-  </ItemGroup>
-  <Import Project="$(MSBuildToolsPath)\Microsoft.CSharp.targets" />
 </Project>
\ No newline at end of file
diff --git a/Engine/Source/Programs/UnrealGameSync/UnrealGameSyncLauncher/UnrealGameSyncLauncher.csproj b/Engine/Source/Programs/UnrealGameSync/UnrealGameSyncLauncher/UnrealGameSyncLauncher.csproj
--- a/Engine/Source/Programs/UnrealGameSync/UnrealGameSyncLauncher/UnrealGameSyncLauncher.csproj
+++ b/Engine/Source/Programs/UnrealGameSync/UnrealGameSyncLauncher/UnrealGameSyncLauncher.csproj
@@ -1,79 +1,21 @@
-﻿<?xml version="1.0" encoding="utf-8"?>
-<Project ToolsVersion="12.0" DefaultTargets="Build" xmlns="http://schemas.microsoft.com/developer/msbuild/2003">
-  <Import Project="$(MSBuildExtensionsPath)\$(MSBuildToolsVersion)\Microsoft.Common.props" Condition="Exists('$(MSBuildExtensionsPath)\$(MSBuildToolsVersion)\Microsoft.Common.props')" />
+<?xml version="1.0" encoding="utf-8"?>
+<Project Sdk="Microsoft.NET.Sdk.WindowsDesktop">
+
   <PropertyGroup>
-    <Configuration Condition=" '$(Configuration)' == '' ">Debug</Configuration>
-    <Platform Condition=" '$(Platform)' == '' ">AnyCPU</Platform>
-    <ProjectGuid>{28679228-7901-4740-88EE-4DCB2D6A7535}</ProjectGuid>
     <OutputType>WinExe</OutputType>
-    <AppDesignerFolder>Properties</AppDesignerFolder>
+    <TargetFramework>netcoreapp3.1</TargetFramework>
+    <UseWindowsForms>true</UseWindowsForms>
+    <GenerateAssemblyInfo>false</GenerateAssemblyInfo>
     <RootNamespace>UnrealGameSyncLauncher</RootNamespace>
     <AssemblyName>UnrealGameSyncLauncher</AssemblyName>
-    <TargetFrameworkVersion>v4.5</TargetFrameworkVersion>
-    <FileAlignment>512</FileAlignment>
-  </PropertyGroup>
-  <PropertyGroup>
+    <StartupObject>UnrealGameSyncLauncher.Program</StartupObject>
     <ApplicationIcon>Resources\Icon.ico</ApplicationIcon>
-  </PropertyGroup>
-  <PropertyGroup Condition="'$(Configuration)|$(Platform)' == 'Debug|x64'">
-    <DebugSymbols>true</DebugSymbols>
-    <OutputPath>bin\x64\Debug\</OutputPath>
-    <DefineConstants>DEBUG;TRACE</DefineConstants>
-    <DebugType>full</DebugType>
-    <PlatformTarget>x64</PlatformTarget>
-    <ErrorReport>prompt</ErrorReport>
-    <CodeAnalysisRuleSet>MinimumRecommendedRules.ruleset</CodeAnalysisRuleSet>
-    <Prefer32Bit>true</Prefer32Bit>
-  </PropertyGroup>
-  <PropertyGroup Condition="'$(Configuration)|$(Platform)' == 'Release|x64'">
-    <OutputPath>bin\x64\Release\</OutputPath>
-    <DefineConstants>TRACE</DefineConstants>
-    <Optimize>true</Optimize>
-    <DebugType>pdbonly</DebugType>
-    <PlatformTarget>x64</PlatformTarget>
-    <ErrorReport>prompt</ErrorReport>
-    <CodeAnalysisRuleSet>MinimumRecommendedRules.ruleset</CodeAnalysisRuleSet>
-    <Prefer32Bit>true</Prefer32Bit>
-  </PropertyGroup>
-  <PropertyGroup Condition="'$(Configuration)|$(Platform)' == 'Debug|x86'">
-    <DebugSymbols>true</DebugSymbols>
-    <OutputPath>bin\x86\Debug\</OutputPath>
-    <DefineConstants>DEBUG;TRACE</DefineConstants>
-    <DebugType>full</DebugType>
-    <PlatformTarget>x86</PlatformTarget>
-    <ErrorReport>prompt</ErrorReport>
-    <CodeAnalysisRuleSet>MinimumRecommendedRules.ruleset</CodeAnalysisRuleSet>
-    <Prefer32Bit>true</Prefer32Bit>
-  </PropertyGroup>
-  <PropertyGroup Condition="'$(Configuration)|$(Platform)' == 'Release|x86'">
-    <OutputPath>bin\x86\Release\</OutputPath>
-    <DefineConstants>TRACE</DefineConstants>
-    <Optimize>true</Optimize>
-    <DebugType>pdbonly</DebugType>
-    <PlatformTarget>x86</PlatformTarget>
-    <ErrorReport>prompt</ErrorReport>
-    <CodeAnalysisRuleSet>MinimumRecommendedRules.ruleset</CodeAnalysisRuleSet>
-    <Prefer32Bit>true</Prefer32Bit>
-  </PropertyGroup>
-  <PropertyGroup>
     <ApplicationManifest>app.manifest</ApplicationManifest>
+    <RuntimeIdentifiers>win10-x64</RuntimeIdentifiers>
   </PropertyGroup>
-  <ItemGroup>
-    <Reference Include="System" />
-    <Reference Include="System.Core" />
-    <Reference Include="System.Web.Extensions" />
-    <Reference Include="System.Xml.Linq" />
-    <Reference Include="System.Data.DataSetExtensions" />
-    <Reference Include="Microsoft.CSharp" />
-    <Reference Include="System.Data" />
-    <Reference Include="System.Deployment" />
-    <Reference Include="System.Drawing" />
-    <Reference Include="System.Windows.Forms" />
-    <Reference Include="System.Xml" />
-  </ItemGroup>
   <ItemGroup>
     <Compile Include="..\UnrealGameSync\ChildProcess.cs">
-      <Link>Shared\ChildProcess.cs</Link>
+      <Link>UnrealGameSync\ChildProcess.cs</Link>
     </Compile>
     <Compile Include="..\UnrealGameSync\Forms\ModalTaskWindow.cs">
       <Link>UnrealGameSync\ModalTaskWindow.cs</Link>
@@ -92,68 +34,26 @@
       <DependentUpon>PasswordWindow.cs</DependentUpon>
     </Compile>
     <Compile Include="..\UnrealGameSync\DeploymentSettings.cs">
-      <Link>Shared\DeploymentSettings.cs</Link>
+      <Link>UnrealGameSync\DeploymentSettings.cs</Link>
     </Compile>
     <Compile Include="..\UnrealGameSync\ModalTask.cs">
-      <Link>Shared\ModalTask.cs</Link>
+      <Link>UnrealGameSync\ModalTask.cs</Link>
     </Compile>
     <Compile Include="..\UnrealGameSync\Perforce.cs">
-      <Link>Shared\Perforce.cs</Link>
+      <Link>UnrealGameSync\Perforce.cs</Link>
     </Compile>
     <Compile Include="..\UnrealGameSync\PerforceModalTask.cs">
       <Link>UnrealGameSync\PerforceModalTask.cs</Link>
     </Compile>
     <Compile Include="..\UnrealGameSync\Utility.cs">
-      <Link>Shared\Utility.cs</Link>
-    </Compile>
-    <Compile Include="Forms\LogWindow.cs">
-      <SubType>Form</SubType>
+      <Link>UnrealGameSync\Utility.cs</Link>
     </Compile>
-    <Compile Include="Forms\LogWindow.Designer.cs">
-      <DependentUpon>LogWindow.cs</DependentUpon>
-    </Compile>
-    <Compile Include="Forms\SettingsWindow.cs">
-      <SubType>Form</SubType>
-    </Compile>
-    <Compile Include="Forms\SettingsWindow.Designer.cs">
-      <DependentUpon>SettingsWindow.cs</DependentUpon>
-    </Compile>
-    <Compile Include="SyncAndRunPerforceTask.cs" />
-    <Compile Include="Program.cs" />
-    <Compile Include="Properties\AssemblyInfo.cs" />
-    <EmbeddedResource Include="Forms\LogWindow.resx">
-      <DependentUpon>LogWindow.cs</DependentUpon>
-    </EmbeddedResource>
-    <EmbeddedResource Include="Forms\SettingsWindow.resx">
-      <DependentUpon>SettingsWindow.cs</DependentUpon>
-    </EmbeddedResource>
-    <EmbeddedResource Include="Properties\Resources.resx">
-      <Generator>ResXFileCodeGenerator</Generator>
-      <LastGenOutput>Resources.Designer.cs</LastGenOutput>
-      <SubType>Designer</SubType>
-    </EmbeddedResource>
-    <Compile Include="Properties\Resources.Designer.cs">
-      <AutoGen>True</AutoGen>
-      <DependentUpon>Resources.resx</DependentUpon>
-    </Compile>
-    <None Include="app.manifest" />
-    <None Include="Properties\Settings.settings">
-      <Generator>SettingsSingleFileGenerator</Generator>
-      <LastGenOutput>Settings.Designer.cs</LastGenOutput>
-    </None>
-    <Compile Include="Properties\Settings.Designer.cs">
-      <AutoGen>True</AutoGen>
-      <DependentUpon>Settings.settings</DependentUpon>
-      <DesignTimeSharedInput>True</DesignTimeSharedInput>
-    </Compile>
-  </ItemGroup>
-  <ItemGroup>
-    <None Include="App.config" />
   </ItemGroup>
   <ItemGroup>
-    <Content Include="Resources\Icon.ico" />
+    <None Update="Properties\Settings.settings">
+      <Generator>SettingsSingleFileGenerator</Generator>
+    </None>
   </ItemGroup>
-  <ItemGroup />
   <Choose>
     <When Condition="Exists('..\..\..\..\Restricted\NotForLicensees\Source\Programs\UnrealGameSync\UnrealGameSync\DeploymentSettingsInternal.cs')">
       <ItemGroup>
@@ -163,12 +63,4 @@
       </ItemGroup>
     </When>
   </Choose>
-  <Import Project="$(MSBuildToolsPath)\Microsoft.CSharp.targets" />
-  <!-- To modify your build process, add your task inside one of the targets below and uncomment it. 
-       Other similar extension points exist, see Microsoft.Common.targets.
-  <Target Name="BeforeBuild">
-  </Target>
-  <Target Name="AfterBuild">
-  </Target>
-  -->
 </Project>
\ No newline at end of file
