Fixed issue with sending exception for telemetry in UGS due to change to System.Text.Json (which doesnt handle json circular references thus fails when serializing Exception).

diff --git a/Engine/Source/Programs/UnrealGameSync/UnrealGameSync/Program.cs b/Engine/Source/Programs/UnrealGameSync/UnrealGameSync/Program.cs
--- a/Engine/Source/Programs/UnrealGameSync/UnrealGameSync/Program.cs
+++ b/Engine/Source/Programs/UnrealGameSync/UnrealGameSync/Program.cs
@@ -173,13 +173,7 @@ namespace UnrealGameSync
 			Exception Ex = Args.ExceptionObject as Exception;
 			if(Ex != null)
 			{
-				StringBuilder ExceptionTrace = new StringBuilder(Ex.ToString());
-				for(Exception InnerEx = Ex.InnerException; InnerEx != null; InnerEx = InnerEx.InnerException)
-				{
-					ExceptionTrace.Append("\nInner Exception:\n");
-					ExceptionTrace.Append(InnerEx.ToString());
-				}
-				Telemetry.SendEvent("Crash", new { Exception = Ex });
+				Telemetry.SendEvent("Crash", new {Exception = Ex.ToString()});
 			}
 		}
 
