Changed all cases were we use Thread.Abort in UGS as this is not supported in Net Core, changed to Thread.Interupt instead but we should really rewrite this to use the thread pool instead (which is tracked in UE-97800)

Fixed up issues due to behavior difference between System.Text.Json and System.Web Json Serializers.

diff --git a/Engine/Source/Programs/UnrealGameSync/UnrealGameSync/Controls/WorkspaceControl.cs b/Engine/Source/Programs/UnrealGameSync/UnrealGameSync/Controls/WorkspaceControl.cs
--- a/Engine/Source/Programs/UnrealGameSync/UnrealGameSync/Controls/WorkspaceControl.cs
+++ b/Engine/Source/Programs/UnrealGameSync/UnrealGameSync/Controls/WorkspaceControl.cs
@@ -2851,6 +2851,11 @@ namespace UnrealGameSync
 			}
 		}
 
+		private class ReceiptJsonObject
+		{
+			public string Launch { get; set; }
+		}
+
 		private string GetEditorExePath(BuildConfig Config)
 		{
 			// Try to read the executable path from the target receipt
@@ -2863,18 +2868,14 @@ namespace UnrealGameSync
 					try
 					{
 						string Text = File.ReadAllText(ReceiptFileName);
-						Dictionary<string, object> RawObject = JsonSerializer.Deserialize<Dictionary<string, object>>(Text);
+						ReceiptJsonObject Receipt = JsonSerializer.Deserialize<ReceiptJsonObject>(Text);
 
-						object LaunchFileNameObject;
-						if (RawObject.TryGetValue("Launch", out LaunchFileNameObject))
+						string LaunchFileName = Receipt.Launch;
+						if (LaunchFileName != null)
 						{
-							string LaunchFileName = LaunchFileNameObject as string;
-							if (LaunchFileName != null)
-							{
-								LaunchFileName = LaunchFileName.Replace("$(EngineDir)", Path.Combine(BranchDirectoryName, "Engine"));
-								LaunchFileName = LaunchFileName.Replace("$(ProjectDir)", Path.GetDirectoryName(SelectedFileName));
-								return Path.GetFullPath(LaunchFileName);
-							}
+							LaunchFileName = LaunchFileName.Replace("$(EngineDir)", Path.Combine(BranchDirectoryName, "Engine"));
+							LaunchFileName = LaunchFileName.Replace("$(ProjectDir)", Path.GetDirectoryName(SelectedFileName));
+							return Path.GetFullPath(LaunchFileName);
 						}
 					}
 					catch (Exception Ex)
diff --git a/Engine/Source/Programs/UnrealGameSync/UnrealGameSync/EventMonitor.cs b/Engine/Source/Programs/UnrealGameSync/UnrealGameSync/EventMonitor.cs
--- a/Engine/Source/Programs/UnrealGameSync/UnrealGameSync/EventMonitor.cs
+++ b/Engine/Source/Programs/UnrealGameSync/UnrealGameSync/EventMonitor.cs
@@ -202,7 +202,7 @@ namespace UnrealGameSync
 				RefreshEvent.Set();
 				if(!WorkerThread.Join(100))
 				{
-					WorkerThread.Abort();
+					WorkerThread.Interrupt();
 					WorkerThread.Join();
 				}
 				WorkerThread = null;
diff --git a/Engine/Source/Programs/UnrealGameSync/UnrealGameSync/Forms/IssueDetailsWindow.cs b/Engine/Source/Programs/UnrealGameSync/UnrealGameSync/Forms/IssueDetailsWindow.cs
--- a/Engine/Source/Programs/UnrealGameSync/UnrealGameSync/Forms/IssueDetailsWindow.cs
+++ b/Engine/Source/Programs/UnrealGameSync/UnrealGameSync/Forms/IssueDetailsWindow.cs
@@ -90,7 +90,7 @@ namespace UnrealGameSync
 					RefreshEvent.Set();
 					if(!WorkerThread.Join(100))
 					{
-						WorkerThread.Abort();
+						WorkerThread.Interrupt();
 						WorkerThread.Join();
 					}
 					WorkerThread = null;
diff --git a/Engine/Source/Programs/UnrealGameSync/UnrealGameSync/IssueMonitor.cs b/Engine/Source/Programs/UnrealGameSync/UnrealGameSync/IssueMonitor.cs
--- a/Engine/Source/Programs/UnrealGameSync/UnrealGameSync/IssueMonitor.cs
+++ b/Engine/Source/Programs/UnrealGameSync/UnrealGameSync/IssueMonitor.cs
@@ -21,49 +21,49 @@ namespace UnrealGameSync
 
 	public class IssueBuildData
 	{
-		public long Id;
-		public string Stream;
-		public int Change;
-		public string JobName;
-		public string JobUrl;
-		public string JobStepName;
-		public string JobStepUrl;
-		public string ErrorUrl;
-		public IssueBuildOutcome Outcome;
+		public long Id { get; set; }
+		public string Stream { get; set; }
+		public int Change { get; set; }
+		public string JobName { get; set; }
+		public string JobUrl { get; set; }
+		public string JobStepName { get; set; }
+		public string JobStepUrl { get; set; }
+		public string ErrorUrl { get; set; }
+		public IssueBuildOutcome Outcome { get; set; }
 	}
 
 	public class IssueDiagnosticData
 	{
-		public long? BuildId;
-		public string Message;
-		public string Url;
+		public long? BuildId { get; set; }
+		public string Message { get; set; }
+		public string Url { get; set; }
 	}
 
 	public class IssueData
 	{
-		public long Id;
-		public DateTime CreatedAt;
-		public DateTime RetrievedAt;
-		public string Project;
-		public string Summary;
-		public string Details;
-		public string Owner;
-		public string NominatedBy;
-		public DateTime? AcknowledgedAt;
-		public int FixChange;
-		public DateTime? ResolvedAt;
-		public bool bNotify;
-		public List<IssueBuildData> Builds;
+		public long Id { get; set; }
+		public DateTime CreatedAt { get; set; }
+		public DateTime RetrievedAt { get; set; }
+		public string Project { get; set; }
+		public string Summary { get; set; }
+		public string Details { get; set; }
+		public string Owner { get; set; }
+		public string NominatedBy { get; set; }
+		public DateTime? AcknowledgedAt { get; set; }
+		public int FixChange { get; set; }
+		public DateTime? ResolvedAt { get; set; }
+		public bool bNotify { get; set; }
+		public List<IssueBuildData> Builds { get; set; }
 	}
 
 	public class IssueUpdateData
 	{
-		public long Id;
-		public string Owner;
-		public string NominatedBy;
-		public bool? Acknowledged;
-		public int? FixChange;
-		public bool? Resolved;
+		public long Id { get; set; }
+		public string Owner { get; set; }
+		public string NominatedBy { get; set; }
+		public bool? Acknowledged { get; set; }
+		public int? FixChange { get; set; }
+		public bool? Resolved { get; set; }
 	}
 
 	[Flags]
@@ -258,7 +258,7 @@ namespace UnrealGameSync
 				RefreshEvent.Set();
 				if(!WorkerThread.Join(100))
 				{
-					WorkerThread.Abort();
+					WorkerThread.Interrupt();
 					WorkerThread.Join();
 				}
 				WorkerThread = null;
diff --git a/Engine/Source/Programs/UnrealGameSync/UnrealGameSync/PerforceMonitor.cs b/Engine/Source/Programs/UnrealGameSync/UnrealGameSync/PerforceMonitor.cs
--- a/Engine/Source/Programs/UnrealGameSync/UnrealGameSync/PerforceMonitor.cs
+++ b/Engine/Source/Programs/UnrealGameSync/UnrealGameSync/PerforceMonitor.cs
@@ -186,7 +186,7 @@ namespace UnrealGameSync
 				RefreshEvent.Set();
 				if(!WorkerThread.Join(100))
 				{
-					WorkerThread.Abort();
+					WorkerThread.Interrupt();
 				}
 				WorkerThread = null;
 			}
diff --git a/Engine/Source/Programs/UnrealGameSync/UnrealGameSync/UpdateMonitor.cs b/Engine/Source/Programs/UnrealGameSync/UnrealGameSync/UpdateMonitor.cs
--- a/Engine/Source/Programs/UnrealGameSync/UnrealGameSync/UpdateMonitor.cs
+++ b/Engine/Source/Programs/UnrealGameSync/UnrealGameSync/UpdateMonitor.cs
@@ -58,7 +58,7 @@ namespace UnrealGameSync
 			{
 				if(!WorkerThread.Join(30))
 				{
-					WorkerThread.Abort();
+					WorkerThread.Interrupt();
 					WorkerThread.Join();
 				}
 				WorkerThread = null;
diff --git a/Engine/Source/Programs/UnrealGameSync/UnrealGameSync/Utility.cs b/Engine/Source/Programs/UnrealGameSync/UnrealGameSync/Utility.cs
--- a/Engine/Source/Programs/UnrealGameSync/UnrealGameSync/Utility.cs
+++ b/Engine/Source/Programs/UnrealGameSync/UnrealGameSync/Utility.cs
@@ -239,6 +239,11 @@ namespace UnrealGameSync
 			return Result;
 		}
 
+		class ProjectJson
+		{
+			public bool Enterprise { get; set; }
+		}
+
 		/// <summary>
 		/// Determines if a project is an enterprise project
 		/// </summary>
@@ -248,15 +253,9 @@ namespace UnrealGameSync
 		{
 			try
 			{
-				Dictionary<string, object> RawObject = JsonSerializer.Deserialize<Dictionary<string, object>>(Text);
+				ProjectJson Project = JsonSerializer.Deserialize<ProjectJson>(Text);
 
-				object Enterprise;
-				if(RawObject.TryGetValue("Enterprise", out Enterprise) && Enterprise is bool)
-				{
-					return (bool)Enterprise;
-				}
-
-				return false;
+				return Project.Enterprise;
 			}
 			catch
 			{
diff --git a/Engine/Source/Programs/UnrealGameSync/UnrealGameSync/Workspace.cs b/Engine/Source/Programs/UnrealGameSync/UnrealGameSync/Workspace.cs
--- a/Engine/Source/Programs/UnrealGameSync/UnrealGameSync/Workspace.cs
+++ b/Engine/Source/Programs/UnrealGameSync/UnrealGameSync/Workspace.cs
@@ -10,6 +10,8 @@ using System.IO.Compression;
 using System.Linq;
 using System.Security.Cryptography;
 using System.Text;
+using System.Text.Encodings.Web;
+using System.Text.Json;
 using System.Threading;
 using System.Threading.Tasks;
 
@@ -355,7 +357,7 @@ namespace UnrealGameSync
 				Log.WriteLine("OPERATION ABORTED");
 				if(WorkerThread != null)
 				{
-					WorkerThread.Abort();
+					WorkerThread.Interrupt();
 					WorkerThread.Join();
 					WorkerThread = null;
 				}
@@ -1326,7 +1328,7 @@ namespace UnrealGameSync
 			{
 				foreach (Thread ChildThread in ChildThreads)
 				{
-					ChildThread.Abort();
+					ChildThread.Interrupt();
 				}
 				foreach (Thread ChildThread in ChildThreads)
 				{
@@ -1584,13 +1586,13 @@ namespace UnrealGameSync
 
 		string UpdateBuildVersion(string Text, int Changelist, int CodeChangelist, string BranchOrStreamName, bool bIsLicenseeVersion)
 		{
-			Dictionary<string, object> Object = Json.Deserialize(Text);
+			Dictionary<string, object> Object = JsonSerializer.Deserialize<Dictionary<string, object>>(Text);
 
 			object PrevCompatibleChangelistObj;
-			int PrevCompatibleChangelist = Object.TryGetValue("CompatibleChangelist", out PrevCompatibleChangelistObj) ? (int)Convert.ChangeType(PrevCompatibleChangelistObj, typeof(int)) : 0;
+			int PrevCompatibleChangelist = Object.TryGetValue("CompatibleChangelist", out PrevCompatibleChangelistObj) ? (int)Convert.ChangeType(PrevCompatibleChangelistObj.ToString(), typeof(int)) : 0;
 
 			object PrevIsLicenseeVersionObj;
-			bool PrevIsLicenseeVersion = Object.TryGetValue("IsLicenseeVersion", out PrevIsLicenseeVersionObj)? ((int)Convert.ChangeType(PrevIsLicenseeVersionObj, typeof(int)) != 0) : false;
+			bool PrevIsLicenseeVersion = Object.TryGetValue("IsLicenseeVersion", out PrevIsLicenseeVersionObj)? ((int)Convert.ChangeType(PrevIsLicenseeVersionObj.ToString(), typeof(int)) != 0) : false;
 
 			Object["Changelist"] = Changelist;
 			if(PrevCompatibleChangelist == 0 || PrevIsLicenseeVersion != bIsLicenseeVersion)
@@ -1602,7 +1604,12 @@ namespace UnrealGameSync
 			Object["IsPromotedBuild"] = 0;
 			Object["IsLicenseeVersion"] = bIsLicenseeVersion ? 1 : 0;
 
-			return Json.Serialize(Object, JsonSerializeOptions.PrettyPrint);
+			return JsonSerializer.Serialize(Object, new JsonSerializerOptions
+			{
+				WriteIndented = true, 
+				// do not escape +
+				Encoder = JavaScriptEncoder.UnsafeRelaxedJsonEscaping
+			});
 		}
 
 		bool WriteVersionFile(string LocalPath, string DepotPath, string NewText)
