UGS - Precompiled binaries can now be fetched from Jupiter, this also includes adding a general OIDC management to acquire and store user tokens for access.

diff --git a/Engine/Source/Programs/UnrealGameSync/UnrealGameSync/Controls/OIDCControl.cs b/Engine/Source/Programs/UnrealGameSync/UnrealGameSync/Controls/OIDCControl.cs
new file mode 100644
--- /dev/null
+++ b/Engine/Source/Programs/UnrealGameSync/UnrealGameSync/Controls/OIDCControl.cs
@@ -0,0 +1,107 @@
+// Copyright Epic Games, Inc. All Rights Reserved.
+
+using System;
+using System.Drawing;
+using System.Drawing.Drawing2D;
+using System.Windows.Forms;
+
+namespace UnrealGameSync.Controls
+{
+	public partial class OIDCControl : UserControl
+	{
+		private Font BadgeFont;
+		private readonly OIDCTokenManager OidcManager;
+		private readonly string ProviderIdentifier;
+
+		public OIDCControl(OIDCTokenManager InOidcManager, string InProviderIdentifier, string ServiceName)
+		{
+			InitializeComponent();
+
+			OIDCControlGroupBox.Text = ServiceName;
+			OidcManager = InOidcManager;
+			ProviderIdentifier = InProviderIdentifier;
+
+			BadgeFont = new Font(this.Font.FontFamily, this.Font.SizeInPoints - 2, FontStyle.Bold);
+
+			OIDCStatus ServiceStatus = OidcManager.GetStatusForProvider(ProviderIdentifier);
+			LoginButton.Enabled = ServiceStatus == OIDCStatus.NotLoggedIn;
+		}
+
+		private void DrawBadge(Graphics Graphics, Rectangle BadgeRect, string BadgeText, Color BadgeColor, bool bMergeLeft, bool bMergeRight)
+		{
+			if (BadgeColor.A != 0)
+			{
+				using (GraphicsPath Path = new GraphicsPath())
+				{
+					Path.StartFigure();
+					Path.AddLine(BadgeRect.Left + (bMergeLeft ? 1 : 0), BadgeRect.Top, BadgeRect.Left - (bMergeLeft ? 1 : 0), BadgeRect.Bottom);
+					Path.AddLine(BadgeRect.Left - (bMergeLeft ? 1 : 0), BadgeRect.Bottom, BadgeRect.Right - 1 - (bMergeRight ? 1 : 0), BadgeRect.Bottom);
+					Path.AddLine(BadgeRect.Right - 1 - (bMergeRight ? 1 : 0), BadgeRect.Bottom, BadgeRect.Right - 1 + (bMergeRight ? 1 : 0), BadgeRect.Top);
+					Path.AddLine(BadgeRect.Right - 1 + (bMergeRight ? 1 : 0), BadgeRect.Top, BadgeRect.Left + (bMergeLeft ? 1 : 0), BadgeRect.Top);
+					Path.CloseFigure();
+
+					using (SolidBrush Brush = new SolidBrush(BadgeColor))
+					{
+						Graphics.FillPath(Brush, Path);
+					}
+				}
+
+				TextRenderer.DrawText(Graphics, BadgeText, BadgeFont, BadgeRect, Color.White, TextFormatFlags.HorizontalCenter | TextFormatFlags.VerticalCenter | TextFormatFlags.SingleLine | TextFormatFlags.NoPrefix | TextFormatFlags.PreserveGraphicsClipping);
+			}
+		}
+
+		private void StatusPanel_Paint(object sender, PaintEventArgs e)
+		{
+			OIDCStatus ServiceStatus = OidcManager.GetStatusForProvider(ProviderIdentifier);
+			DrawBadge(e.Graphics, e.ClipRectangle, ToServiceDisplayText(ServiceStatus), ToServiceDisplayColor(ServiceStatus), false, false);
+		}
+
+		private Color ToServiceDisplayColor(OIDCStatus OidcStatus)
+		{
+			switch (OidcStatus)
+			{
+				case OIDCStatus.Connected:
+				case OIDCStatus.TokenRefreshRequired:
+					return Color.Green;
+				case OIDCStatus.NotLoggedIn:
+					return Color.Red;
+				default:
+					throw new ArgumentOutOfRangeException(nameof(OidcStatus), OidcStatus, null);
+			}
+		}
+
+		private string ToServiceDisplayText(OIDCStatus OidcStatus)
+		{
+			switch (OidcStatus)
+			{
+				case OIDCStatus.Connected:
+					return "Connected";
+				case OIDCStatus.NotLoggedIn:
+					return "Not Logged In";
+				case OIDCStatus.TokenRefreshRequired:
+					return "Refresh Pending";
+				default:
+					throw new ArgumentOutOfRangeException(nameof(OidcStatus), OidcStatus, null);
+			}
+		}
+
+		private async void LoginButton_Click(object sender, EventArgs e)
+		{
+			try
+			{
+				await OidcManager.Login(ProviderIdentifier);
+
+				Focus();
+
+				OIDCStatus ServiceStatus = OidcManager.GetStatusForProvider(ProviderIdentifier);
+				LoginButton.Enabled = ServiceStatus == OIDCStatus.NotLoggedIn;
+
+				Refresh();
+			}
+			catch (LoginFailedException Exception)
+			{
+				MessageBox.Show(Exception.Message, "Login Failed", MessageBoxButtons.OK);
+			}
+		}
+	}
+}
diff --git a/Engine/Source/Programs/UnrealGameSync/UnrealGameSync/Controls/WorkspaceControl.cs b/Engine/Source/Programs/UnrealGameSync/UnrealGameSync/Controls/WorkspaceControl.cs
--- a/Engine/Source/Programs/UnrealGameSync/UnrealGameSync/Controls/WorkspaceControl.cs
+++ b/Engine/Source/Programs/UnrealGameSync/UnrealGameSync/Controls/WorkspaceControl.cs
@@ -254,6 +254,7 @@ namespace UnrealGameSync
 		bool bUnstable;
 		string EditorTargetName;
 		bool bIsEnterpriseProject;
+		JupiterMonitor JupiterMonitor;
 		PerforceMonitor PerforceMonitor;
 		Workspace Workspace;
 		IssueMonitor IssueMonitor;
@@ -262,7 +263,7 @@ namespace UnrealGameSync
 		HashSet<int> PromotedChangeNumbers = new HashSet<int>();
 		List<int> ListIndexToChangeIndex = new List<int>();
 		List<int> SortedChangeNumbers = new List<int>();
-		Dictionary<string, Dictionary<int, string>> ArchiveToChangeNumberToArchivePath = new Dictionary<string, Dictionary<int, string>>();
+		Dictionary<string, Dictionary<int, string>> ArchiveToChangeNumberToArchiveKey = new Dictionary<string, Dictionary<int, string>>();
 		Dictionary<int, ChangeLayoutInfo> ChangeNumberToLayoutInfo = new Dictionary<int, ChangeLayoutInfo>();
 		List<ToolStripMenuItem> CustomToolMenuItems = new List<ToolStripMenuItem>();
 		int NumChanges;
@@ -312,7 +313,7 @@ namespace UnrealGameSync
 		System.Threading.Timer StartupTimer;
 		List<WorkspaceStartupCallback> StartupCallbacks;
 
-		public WorkspaceControl(IWorkspaceControlOwner InOwner, string InApiUrl, string InOriginalExecutableFileName, bool bInUnstable, DetectProjectSettingsTask DetectSettings, LineBasedTextWriter InLog, UserSettings InSettings)
+		public WorkspaceControl(IWorkspaceControlOwner InOwner, string InApiUrl, string InOriginalExecutableFileName, bool bInUnstable, DetectProjectSettingsTask DetectSettings, LineBasedTextWriter InLog, UserSettings InSettings, OIDCTokenManager InOidcTokenManager)
 		{
 			InitializeComponent();
 
@@ -399,6 +400,8 @@ namespace UnrealGameSync
 			EventMonitor = new EventMonitor(ApiUrl, PerforceUtils.GetClientOrDepotDirectoryName(SelectedProjectIdentifier), DetectSettings.PerforceClient.UserName, ProjectLogBaseName + ".review.log");
 			EventMonitor.OnUpdatesReady += UpdateReviewsCallback;
 
+			JupiterMonitor = JupiterMonitor.CreateFromConfigFile(InOidcTokenManager, ProjectLogBaseName + ".jupiter.log", DetectSettings.LatestProjectConfigFile, SelectedProjectIdentifier);
+
 			UpdateColumnSettings();
 
 			string LogFileName = Path.Combine(DataFolder, ProjectLogBaseName + ".sync.log");
@@ -754,6 +757,11 @@ namespace UnrealGameSync
 				BadgeFont.Dispose();
 				BadgeFont = null;
 			}
+			if (JupiterMonitor != null)
+			{
+				JupiterMonitor.Dispose();
+				JupiterMonitor = null;
+			}
 
 			base.Dispose(disposing);
 		}
@@ -947,16 +955,16 @@ namespace UnrealGameSync
 			WorkspaceUpdateContext Context = new WorkspaceUpdateContext(ChangeNumber, Options, CombinedSyncFilter, GetDefaultBuildStepObjects(), ProjectSettings.BuildSteps, null, GetWorkspaceVariables(ChangeNumber));
 			if (Options.HasFlag(WorkspaceUpdateOptions.SyncArchives))
 			{
-				IReadOnlyList<IArchiveInfo> Archives = PerforceMonitor.AvailableArchives;
+				IReadOnlyList<IArchiveInfo> Archives = GetArchives();
 				foreach (IArchiveInfo Archive in Archives)
 				{
-					Context.ArchiveTypeToDepotPath[Archive.Type] = null;
+					Context.ArchiveTypeToArchive[Archive.Type] = null;
 				}
 
-				List<IArchiveInfo> SelectedArchives = GetSelectedArchives(PerforceMonitor.AvailableArchives);
+				List<IArchiveInfo> SelectedArchives = GetSelectedArchives(Archives);
 				foreach (IArchiveInfo Archive in SelectedArchives)
 				{
-					string ArchivePath = GetArchivePathForChangeNumber(Archive, ChangeNumber);
+					string ArchivePath = GetArchiveKeyForChangeNumber(Archive, ChangeNumber);
 					if (ArchivePath == null)
 					{
 						MessageBox.Show(String.Format("There are no compiled {0} binaries for this change. To sync it, you must disable syncing of precompiled editor binaries.", Archive.Name));
@@ -974,7 +982,7 @@ namespace UnrealGameSync
 						Context.SyncFilter = Enumerable.Concat(Context.SyncFilter, ZippedBinariesSyncFilter).ToArray();
 					}
 
-					Context.ArchiveTypeToDepotPath[Archive.Type] = ArchivePath;
+					Context.ArchiveTypeToArchive[Archive.Type] = new Tuple<IArchiveInfo, string>(Archive, ArchivePath);
 				}
 			}
 			StartWorkspaceUpdate(Context, Callback);
@@ -1085,7 +1093,7 @@ namespace UnrealGameSync
 
 			if (Result == WorkspaceUpdateResult.Success && WorkspaceSettings.ExpandedArchiveTypes != null)
 			{
-				WorkspaceSettings.ExpandedArchiveTypes = WorkspaceSettings.ExpandedArchiveTypes.Except(Context.ArchiveTypeToDepotPath.Where(x => x.Value == null).Select(x => x.Key)).ToArray();
+				WorkspaceSettings.ExpandedArchiveTypes = WorkspaceSettings.ExpandedArchiveTypes.Except(Context.ArchiveTypeToArchive.Where(x => x.Value == null).Select(x => x.Key)).ToArray();
 			}
 
 			WorkspaceSettings.LastSyncChangeNumber = Context.ChangeNumber;
@@ -1178,7 +1186,7 @@ namespace UnrealGameSync
 		{
 			if (SelectedFileName != null)
 			{
-				ArchiveToChangeNumberToArchivePath.Clear();
+				ArchiveToChangeNumberToArchiveKey.Clear();
 				ChangeNumberToLayoutInfo.Clear();
 
 				List<PerforceChangeSummary> Changes = PerforceMonitor.GetChanges();
@@ -1545,7 +1553,7 @@ namespace UnrealGameSync
 			UpdateMaxBuildBadgeChars();
 
 			// Update everything else
-			ArchiveToChangeNumberToArchivePath.Clear();
+			ArchiveToChangeNumberToArchiveKey.Clear();
 			ChangeNumberToLayoutInfo.Clear();
 			BuildList.Invalidate();
 			UpdateServiceBadges();
@@ -1690,7 +1698,7 @@ namespace UnrealGameSync
 
 		void UpdateReviews()
 		{
-			ArchiveToChangeNumberToArchivePath.Clear();
+			ArchiveToChangeNumberToArchiveKey.Clear();
 			ChangeNumberToLayoutInfo.Clear();
 			EventMonitor.ApplyUpdates();
 
@@ -1926,15 +1934,15 @@ namespace UnrealGameSync
 			}
 		}
 
-		private string GetArchivePathForChangeNumber(IArchiveInfo Archive, int ChangeNumber)
+		private string GetArchiveKeyForChangeNumber(IArchiveInfo Archive, int ChangeNumber)
 		{
 			string ArchivePath;
 
 			Dictionary<int, string> ChangeNumberToArchivePath;
-			if (!ArchiveToChangeNumberToArchivePath.TryGetValue(Archive.Name, out ChangeNumberToArchivePath))
+			if (!ArchiveToChangeNumberToArchiveKey.TryGetValue(Archive.Name, out ChangeNumberToArchivePath))
 			{
 				ChangeNumberToArchivePath = new Dictionary<int, string>();
-				ArchiveToChangeNumberToArchivePath[Archive.Name] = ChangeNumberToArchivePath;
+				ArchiveToChangeNumberToArchiveKey[Archive.Name] = ChangeNumberToArchivePath;
 			}
 
 			if (!ChangeNumberToArchivePath.TryGetValue(ChangeNumber, out ArchivePath))
@@ -1943,13 +1951,13 @@ namespace UnrealGameSync
 				if (PerforceMonitor.TryGetChangeDetails(ChangeNumber, out Details))
 				{
 					// Try to get the archive for this CL
-					if (!Archive.TryGetArchivePathForChangeNumber(ChangeNumber, out ArchivePath) && !Details.bContainsCode)
+					if (!Archive.TryGetArchiveKeyForChangeNumber(ChangeNumber, out ArchivePath) && !Details.bContainsCode)
 					{
 						// Otherwise if it's a content-only change, find the previous build any use the archive path from that
 						int Index = SortedChangeNumbers.BinarySearch(ChangeNumber);
 						if (Index > 0)
 						{
-							ArchivePath = GetArchivePathForChangeNumber(Archive, SortedChangeNumbers[Index - 1]);
+							ArchivePath = GetArchiveKeyForChangeNumber(Archive, SortedChangeNumbers[Index - 1]);
 						}
 					}
 				}
@@ -1971,8 +1979,8 @@ namespace UnrealGameSync
 				return false;
 			}
 
-			List<IArchiveInfo> SelectedArchives = GetSelectedArchives(PerforceMonitor.AvailableArchives);
-			return SelectedArchives.Count == 0 || SelectedArchives.All(x => GetArchivePathForChangeNumber(x, ChangeNumber) != null);
+			List<IArchiveInfo> SelectedArchives = GetSelectedArchives(GetArchives());
+			return SelectedArchives.Count == 0 || SelectedArchives.All(x => GetArchiveKeyForChangeNumber(x, ChangeNumber) != null);
 		}
 
 		private ChangeLayoutInfo GetChangeLayoutInfo(PerforceChangeSummary Change)
@@ -4073,7 +4081,7 @@ namespace UnrealGameSync
 
 			OptionsContextMenu_SyncPrecompiledBinaries.DropDownItems.Clear();
 
-			IReadOnlyList<IArchiveInfo> Archives = PerforceMonitor?.AvailableArchives;
+			IReadOnlyList<IArchiveInfo> Archives = GetArchives();
 			if (Archives == null || Archives.Count == 0)
 			{
 				OptionsContextMenu_SyncPrecompiledBinaries.Enabled = false;
@@ -4104,7 +4112,7 @@ namespace UnrealGameSync
 						Item.Enabled = Archive.Exists();
 						if (!Item.Enabled)
 						{
-							Item.ToolTipText = String.Format("No valid archives found at {0}", Archive.DepotPath);
+							Item.ToolTipText = String.Format("No valid archives found at {0}", Archive.BasePath);
 						}
 						Item.Checked = SelectedArchives.Contains(Archive);
 						Item.Click += (Sender, Args) => SetSelectedArchive(Archive, !Item.Checked);
@@ -4129,6 +4137,21 @@ namespace UnrealGameSync
 			OptionsContextMenu.Show(OptionsButton, new Point(OptionsButton.Width - OptionsContextMenu.Size.Width, OptionsButton.Height));
 		}
 
+		private IReadOnlyList<IArchiveInfo> GetArchives()
+		{
+			IReadOnlyList<IArchiveInfo> AvailableArchives = JupiterMonitor?.AvailableArchives;
+			if (AvailableArchives != null && AvailableArchives.Count != 0)
+				return AvailableArchives;
+
+			// if jupiter had no archives we fallback to the perforce monitor
+			if (PerforceMonitor != null)
+			{
+				return PerforceMonitor?.AvailableArchives;
+			}
+
+			return new List<IArchiveInfo>();
+		}
+
 		private void BuildAfterSyncCheckBox_CheckedChanged(object sender, EventArgs e)
 		{
 			Settings.bBuildAfterSync = BuildAfterSyncCheckBox.Checked;
@@ -4925,7 +4948,7 @@ namespace UnrealGameSync
 
 		private bool ShouldSyncPrecompiledEditor
 		{
-			get { return Settings.Archives.Any(x => x.bEnabled && x.Type == EditorArchiveType) && PerforceMonitor != null && PerforceMonitor.AvailableArchives.Any(x => x.Type == "Editor"); }
+			get { return Settings.Archives.Any(x => x.bEnabled && x.Type == EditorArchiveType) && GetArchives().Any(x => x.Type == "Editor"); }
 		}
 
 		public BuildConfig GetEditorBuildConfig()
@@ -4995,7 +5018,7 @@ namespace UnrealGameSync
 		{
 			if (OptionsContextMenu_SyncPrecompiledBinaries.DropDownItems.Count == 0)
 			{
-				IArchiveInfo EditorArchive = PerforceMonitor.AvailableArchives.FirstOrDefault(x => x.Type == EditorArchiveType);
+				IArchiveInfo EditorArchive = GetArchives().FirstOrDefault(x => x.Type == EditorArchiveType);
 				if (EditorArchive != null)
 				{
 					SetSelectedArchive(EditorArchive, !OptionsContextMenu_SyncPrecompiledBinaries.Checked);
diff --git a/Engine/Source/Programs/UnrealGameSync/UnrealGameSync/Forms/MainWindow.cs b/Engine/Source/Programs/UnrealGameSync/UnrealGameSync/Forms/MainWindow.cs
--- a/Engine/Source/Programs/UnrealGameSync/UnrealGameSync/Forms/MainWindow.cs
+++ b/Engine/Source/Programs/UnrealGameSync/UnrealGameSync/Forms/MainWindow.cs
@@ -87,6 +87,8 @@ namespace UnrealGameSync
 
 		bool bRestoreStateOnLoad;
 
+		OIDCTokenManager OIDCTokenManager;
+
 		System.Threading.Timer ScheduleTimer;
 		System.Threading.Timer ScheduleSettledTimer;
 
@@ -103,7 +105,7 @@ namespace UnrealGameSync
 		Rectangle PrimaryWorkArea;
 		List<IssueAlertWindow> AlertWindows = new List<IssueAlertWindow>();
 
-		public MainWindow(UpdateMonitor InUpdateMonitor, string InApiUrl, string InDataFolder, string InCacheFolder, bool bInRestoreStateOnLoad, string InOriginalExecutableFileName, bool bInUnstable, DetectProjectSettingsResult[] StartupProjects, PerforceConnection InDefaultConnection, LineBasedTextWriter InLog, UserSettings InSettings, string InUri)
+		public MainWindow(UpdateMonitor InUpdateMonitor, string InApiUrl, string InDataFolder, string InCacheFolder, bool bInRestoreStateOnLoad, string InOriginalExecutableFileName, bool bInUnstable, DetectProjectSettingsResult[] StartupProjects, PerforceConnection InDefaultConnection, LineBasedTextWriter InLog, UserSettings InSettings, string InUri, OIDCTokenManager InOidcTokenManager)
 		{
 			InitializeComponent();
 
@@ -118,6 +120,7 @@ namespace UnrealGameSync
 			DefaultConnection = InDefaultConnection;
 			Log = InLog;
 			Settings = InSettings;
+			OIDCTokenManager = InOidcTokenManager;
 
 			// While creating tab controls during startup, we need to prevent layout calls resulting in the window handle being created too early. Disable layout calls here.
 			SuspendLayout();
@@ -1133,7 +1136,8 @@ namespace UnrealGameSync
 			}
 
 			// Now that we have the project settings, we can construct the tab
-			WorkspaceControl NewWorkspace = new WorkspaceControl(this, ApiUrl, OriginalExecutableFileName, bUnstable, ProjectSettings, Log, Settings);
+			WorkspaceControl NewWorkspace = new WorkspaceControl(this, ApiUrl, OriginalExecutableFileName, bUnstable, ProjectSettings, Log, Settings, OIDCTokenManager);
+			
 			NewWorkspace.Parent = TabPanel;
 			NewWorkspace.Dock = DockStyle.Fill;
 			NewWorkspace.Hide();
diff --git a/Engine/Source/Programs/UnrealGameSync/UnrealGameSync/Forms/OIDCLoginWindow.cs b/Engine/Source/Programs/UnrealGameSync/UnrealGameSync/Forms/OIDCLoginWindow.cs
new file mode 100644
--- /dev/null
+++ b/Engine/Source/Programs/UnrealGameSync/UnrealGameSync/Forms/OIDCLoginWindow.cs
@@ -0,0 +1,35 @@
+// Copyright Epic Games, Inc. All Rights Reserved.
+
+using System;
+using System.Drawing;
+using System.Windows.Forms;
+using UnrealGameSync.Controls;
+
+namespace UnrealGameSync.Forms
+{
+	public partial class OIDCLoginWindow : Form
+	{
+		public OIDCLoginWindow(OIDCTokenManager OidcManager)
+		{
+			InitializeComponent();
+
+			int LastYPosition = 10;
+			foreach (OIDCTokenManager.ProviderInfo Provider in OidcManager.Providers.Values)
+			{
+				OIDCControl ServiceControl = new OIDCControl(OidcManager, Provider.Identifier, Provider.DisplayName)
+				{
+					Anchor = AnchorStyles.Top | AnchorStyles.Left | AnchorStyles.Right,
+					Location = new Point(10, LastYPosition)
+				};
+				Controls.Add(ServiceControl);
+
+				LastYPosition = ServiceControl.Size.Height + ServiceControl.Location.Y + 10;
+			}
+		}
+
+		private void DoneButton_Click(object sender, EventArgs e)
+		{
+			Close();
+		}
+	}
+}
diff --git a/Engine/Source/Programs/UnrealGameSync/UnrealGameSync/Jupiter/JupiterMonitor.cs b/Engine/Source/Programs/UnrealGameSync/UnrealGameSync/Jupiter/JupiterMonitor.cs
new file mode 100644
--- /dev/null
+++ b/Engine/Source/Programs/UnrealGameSync/UnrealGameSync/Jupiter/JupiterMonitor.cs
@@ -0,0 +1,279 @@
+// Copyright Epic Games, Inc. All Rights Reserved.
+
+using System;
+using System.Collections.Generic;
+using System.IO;
+using System.Linq;
+using System.Net.Http;
+using System.Text.Json;
+using System.Threading;
+using System.Threading.Tasks;
+using IdentityModel.Client;
+using Tools.DotNETCommon;
+using Tools.DotNETCommon.Jupiter;
+
+namespace UnrealGameSync
+{
+	class JupiterMonitor: IArchiveInfoSource, IDisposable
+	{
+		private readonly object ArchivesLock = new object();
+		private IReadOnlyList<IArchiveInfo> Archives;
+
+		private readonly OIDCTokenManager TokenManager;
+		private readonly string JupiterNamespace;
+		private readonly string ProviderIdentifier;
+		private readonly string ExpectedBranch;
+		private readonly Uri JupiterUrl;
+
+		private readonly Timer UpdateTimer;
+		private readonly BoundedLogWriter LogWriter;
+
+		public IReadOnlyList<IArchiveInfo> AvailableArchives
+		{
+			get
+			{
+				lock (ArchivesLock)
+				{
+					return Archives;
+				}
+			}
+		}
+
+		private JupiterMonitor(OIDCTokenManager InTokenManager, string InLogPath, string InNamespace, string InUrl,
+			string InProviderIdentifier, string InExpectedBranch)
+		{
+			TokenManager = InTokenManager;
+			JupiterNamespace = InNamespace;
+			ProviderIdentifier = InProviderIdentifier;
+			ExpectedBranch = InExpectedBranch;
+			JupiterUrl = new Uri(InUrl);
+
+			LogWriter = new BoundedLogWriter(InLogPath);
+			UpdateTimer = new Timer(DoUpdate, null, TimeSpan.Zero, TimeSpan.FromMinutes(5));
+		}
+		public void Dispose()
+		{
+			UpdateTimer?.Dispose();
+			LogWriter?.Dispose();
+		}
+
+		private async void DoUpdate(object State)
+		{
+			await LogWriter.WriteLineAsync($"Starting poll of JupiterMonitor for namespace {JupiterNamespace}");
+			try
+			{
+				IReadOnlyList<IArchiveInfo> NewArchives = await GetAvailableArchives();
+				lock (ArchivesLock)
+				{
+					Archives = NewArchives;
+				}
+
+			}
+			catch (Exception Exception)
+			{
+				await LogWriter.WriteLineAsync($"Exception occured during poll! {Exception}");
+			}
+		}
+
+		public static JupiterMonitor CreateFromConfigFile(OIDCTokenManager TokenManager, string LogPath, ConfigFile ConfigFile, string SelectedProjectIdentifier)
+		{
+			ConfigSection JupiterConfigSection = ConfigFile.FindSection("Jupiter");
+			if (JupiterConfigSection == null)
+				return null;
+
+			string JupiterUrl = JupiterConfigSection.GetValue("JupiterUrl");
+			string OIDCProviderIdentifier = JupiterConfigSection.GetValue("OIDCProviderIdentifier");
+
+			ConfigSection ProjectConfigSection = ConfigFile.FindSection(SelectedProjectIdentifier);
+			if (ProjectConfigSection == null)
+				return null;
+
+			string JupiterNamespace = ProjectConfigSection.GetValue("JupiterNamespace");
+			// Is no namespace has been specified we are unable to fetch builds
+			if (JupiterNamespace == null)
+				return null;
+
+			string ExpectedBranch = ProjectConfigSection.GetValue("ExpectedBranch");
+			// If we do not know which branch to fetch we can not list builds, as it would risk getting binaries from other branches
+			if (ExpectedBranch == null)
+				return null;
+
+			OIDCProviderIdentifier = ProjectConfigSection.GetValue("OIDCProviderIdentifier") ?? OIDCProviderIdentifier;
+
+			// with no oidc provider we are unable to login, thus it is required
+			if (OIDCProviderIdentifier == null)
+				return null;
+
+			// project specific overrides
+			JupiterUrl = ProjectConfigSection.GetValue("JupiterUrl") ?? JupiterUrl;
+
+			return new JupiterMonitor(TokenManager, LogPath, JupiterNamespace, JupiterUrl, OIDCProviderIdentifier, ExpectedBranch);
+		}
+
+		private async Task<IReadOnlyList<IArchiveInfo>> GetAvailableArchives()
+		{
+			string Token = await TokenManager.GetAccessToken(ProviderIdentifier);
+
+			List<JupiterArchiveInfo> NewArchives = new List<JupiterArchiveInfo>();
+			using (HttpClient Client = new HttpClient())
+			{
+				Client.BaseAddress = JupiterUrl;
+				Client.SetBearerToken(Token);
+
+				string ResponseBody = await Client.GetStringAsync($"/api/v1/c/tree-root/{JupiterNamespace}");
+				TreeRootListResponse Response = JsonSerializer.Deserialize<TreeRootListResponse>(ResponseBody, new JsonSerializerOptions {PropertyNameCaseInsensitive = true});
+
+				foreach (string TreeRoot in Response.TreeRoots)
+				{
+					// fetch info on this tree root
+					string ResourceUrl = $"/api/v1/c/tree-root/{JupiterNamespace}/{TreeRoot}";
+					string ResponseBodyTreeReference = await Client.GetStringAsync(ResourceUrl);
+					TreeRootReferenceResponse ResponseTreeReference = JsonSerializer.Deserialize<TreeRootReferenceResponse>(ResponseBodyTreeReference, new JsonSerializerOptions { PropertyNameCaseInsensitive = true });
+
+					Dictionary<string, string> Metadata = ResponseTreeReference.Metadata.ToDictionary(Pair => Pair.Key, Pair => Pair.Value.ToString());
+
+					string ArchiveType, Project, Branch, ChangelistString;
+
+					if (!Metadata.TryGetValue("ArchiveType", out ArchiveType) ||
+					    !Metadata.TryGetValue("Project", out Project) ||
+					    !Metadata.TryGetValue("Branch", out Branch) ||
+						!Metadata.TryGetValue("Changelist", out ChangelistString))
+					{
+						continue;
+					}
+
+					// skip if we do not have metadata as we need that to be able to determine if this root is of the type we want
+					if (Branch == null || Project == null || ArchiveType == null || ChangelistString == null)
+						continue;
+
+					if (!string.Equals(ExpectedBranch, Branch, StringComparison.InvariantCultureIgnoreCase))
+					{
+						continue;
+					}
+
+					int Changelist;
+					if (!int.TryParse(ChangelistString, out Changelist))
+					{
+						continue; // invalid changelist format
+					}
+
+					JupiterArchiveInfo ExistingArchive = NewArchives.FirstOrDefault(Info =>
+						string.Equals(Info.Name, Project, StringComparison.InvariantCultureIgnoreCase) &&
+						string.Equals(Info.Type, ArchiveType, StringComparison.InvariantCultureIgnoreCase));
+					if (ExistingArchive == null)
+					{
+						JupiterArchiveInfo Archive = new JupiterArchiveInfo(Project, ArchiveType, JupiterUrl.ToString(), JupiterNamespace);
+						NewArchives.Add(Archive);
+						ExistingArchive = Archive;
+					}
+
+					ExistingArchive.AddArchiveVersion(Changelist, ResponseTreeReference.TreeReferenceKey);
+				}
+			}
+
+			return NewArchives.AsReadOnly();
+		}
+
+		private class TreeRootListResponse
+		{
+			public List<string> TreeRoots { get; set; }
+		}
+
+		private class TreeRootReferenceResponse
+		{
+			public string TreeReferenceKey { get; set; }
+			public string TreeHash { get; set; }
+
+			public Dictionary<string, object> Metadata { get; set; }
+
+			public string TreeRootState { get; set; }
+		}
+
+		private class JupiterArchiveInfo : IArchiveInfo
+		{
+			private readonly string JupiterUrl;
+			private readonly string JupiterNamespace;
+			private readonly Dictionary<int, string> ChangeToKey = new Dictionary<int, string>();
+
+			public JupiterArchiveInfo(string InName, string InType, string InJupiterUrl, string InJupiterNamespace)
+			{
+				Name = InName;
+				Type = InType;
+
+				JupiterUrl = InJupiterUrl;
+				JupiterNamespace = InJupiterNamespace;
+			}
+
+			public string Name { get; }
+			public string Type { get; }
+
+			public string BasePath
+			{
+				get { return null; } 
+			}
+
+			public bool Exists()
+			{
+				return ChangeToKey.Count > 0;
+			}
+
+			public bool TryGetArchiveKeyForChangeNumber(int ChangeNumber, out string ArchiveKey)
+			{
+				return ChangeToKey.TryGetValue(ChangeNumber, out ArchiveKey);
+			}
+
+			public bool DownloadArchive(string ArchiveKey, string LocalRootPath, string ManifestFileName, TextWriter Log, ProgressValue Progress)
+			{
+				try
+				{
+					Progress<Tuple<float, FileReference>> ProgressCallback = new Progress<Tuple<float, FileReference>>(ProgressUpdate =>
+					{
+						(float ProgressFraction, FileReference FileReference) = ProgressUpdate;
+						Progress.Set(ProgressFraction);
+						lock (Log)
+						{
+							Log.WriteLine("Writing {0}", FileReference.FullName);
+						}
+					});
+
+					// place the manifest for the Jupiter download next to the UGS manifest
+					FileReference UGSManifestFileReference = new FileReference(ManifestFileName);
+					FileReference JupiterManifestFileReference = FileReference.Combine(UGSManifestFileReference.Directory, "Jupiter-Manifest.json");
+
+					DirectoryReference RootDirectory = new DirectoryReference(LocalRootPath);
+					JupiterFileTree FileTree = new JupiterFileTree(RootDirectory, InDeferReadingFiles: true);
+					Task<List<FileReference>> DownloadTask = FileTree.DownloadFromJupiter(JupiterManifestFileReference, JupiterUrl, JupiterNamespace, ArchiveKey, ProgressCallback);
+					DownloadTask.Wait();
+
+					List<FileReference> WrittenFiles = DownloadTask.Result;
+					ArchiveManifest ArchiveManifest = new ArchiveManifest();
+					foreach (FileReference File in WrittenFiles)
+					{
+						ArchiveManifest.Files.Add(new ArchiveManifestFile(File.FullName, File.ToFileInfo().Length, DateTime.Now));
+					}
+
+					// Write it out to a temporary file, then move it into place
+					string TempManifestFileName = ManifestFileName + ".tmp";
+					using (FileStream OutputStream = File.Open(TempManifestFileName, FileMode.Create, FileAccess.Write))
+					{
+						ArchiveManifest.Write(OutputStream);
+					}
+					File.Move(TempManifestFileName, ManifestFileName);
+
+					return true;
+				}
+				catch (Exception Exception)
+				{
+					Log.WriteLine(
+						$"Exception occured when downloading build from Jupiter with key {ArchiveKey}. Exception {Exception}");
+					return false;
+				}
+			}
+
+			public void AddArchiveVersion(int Changelist, string Key)
+			{
+				ChangeToKey[Changelist] = Key;
+			}
+		}
+	}
+}
\ No newline at end of file
diff --git a/Engine/Source/Programs/UnrealGameSync/UnrealGameSync/OIDCTokenManager.cs b/Engine/Source/Programs/UnrealGameSync/UnrealGameSync/OIDCTokenManager.cs
new file mode 100644
--- /dev/null
+++ b/Engine/Source/Programs/UnrealGameSync/UnrealGameSync/OIDCTokenManager.cs
@@ -0,0 +1,407 @@
+// Copyright Epic Games, Inc. All Rights Reserved.
+
+using System;
+using System.Collections.Generic;
+using System.Diagnostics;
+using System.IO;
+using System.Linq;
+using System.Net;
+using System.Net.Http;
+using System.Runtime.InteropServices;
+using System.Security.Cryptography;
+using System.Text;
+using System.Threading.Tasks;
+using IdentityModel.Client;
+using IdentityModel.OidcClient;
+using IdentityModel.OidcClient.Results;
+
+namespace UnrealGameSync
+{
+	public class OIDCTokenManager
+	{
+		private UserSettings Settings;
+		private readonly Dictionary<string, OIDCTokenClient> TokenClients = new Dictionary<string, OIDCTokenClient>();
+
+		public Dictionary<string, ProviderInfo> Providers { get; }
+
+		private OIDCTokenManager(UserSettings InSettings, Dictionary<string, ProviderInfo> InProviders)
+		{
+			Settings = InSettings;
+			Providers = InProviders;
+
+			Dictionary<string, string> RefreshTokens = new Dictionary<string, string>();
+			foreach (KeyValuePair<string, string> Pair in Settings.ProviderToRefreshTokens)
+			{
+				string ProviderIdentifier = Pair.Key;
+				try
+				{
+					byte[] KeyBytes = Convert.FromBase64String(Pair.Value);
+					byte[] UnencryptedBytes = ProtectedData.Unprotect(KeyBytes, null, DataProtectionScope.CurrentUser);
+					RefreshTokens.TryAdd(ProviderIdentifier, Encoding.ASCII.GetString(UnencryptedBytes));
+				}
+				catch (Exception)
+				{
+					// ignore any invalid data in our stored refresh tokens, we will simply prompt the user to login again
+				}
+			}
+
+			foreach (KeyValuePair<string, ProviderInfo> ProviderPair in InProviders)
+			{
+				ProviderInfo Provider = ProviderPair.Value;
+
+				OIDCTokenClient TokenClient = new OIDCTokenClient(new Uri(Provider.ServerUri), Provider.ClientId, Provider.RedirectUri);
+
+				string RefreshToken;
+				if (RefreshTokens.TryGetValue(Provider.Identifier, out RefreshToken))
+				{
+					TokenClient.SetRefreshToken(RefreshToken);
+				}
+				TokenClients.Add(Provider.Identifier, TokenClient);
+			}
+		}
+
+		public bool HasUnfinishedLogin()
+		{
+			return TokenClients.Any(Pair => Pair.Value.GetStatusForProvider() == OIDCStatus.NotLoggedIn);
+		}
+
+		internal static OIDCTokenManager CreateFromConfigFile(UserSettings Settings, List<DetectProjectSettingsTask> ConfigFiles)
+		{
+			// join the provider configuration from all projects
+			Dictionary<string, ProviderInfo> Providers = new Dictionary<string, ProviderInfo>();
+			foreach (DetectProjectSettingsTask DetectProjectSettingsTask in ConfigFiles)
+			{
+				ConfigFile ConfigFile = DetectProjectSettingsTask.LatestProjectConfigFile;
+				ConfigSection ProviderSection = ConfigFile.FindSection("OIDCProvider");
+				if (ProviderSection == null)
+					continue;
+
+				string[] ProviderValues = ProviderSection.GetValues("Provider", (string[]) null);
+
+				foreach (ConfigObject Provider in ProviderValues.Select(s => new ConfigObject(s)).ToList())
+				{
+					string Identifier = Provider.GetValue("Identifier");
+					string ServerUri = Provider.GetValue("ServerUri");
+					string ClientId = Provider.GetValue("ClientId");
+					string DisplayName = Provider.GetValue("DisplayName");
+					string RedirectUri = Provider.GetValue("RedirectUri");
+
+					// we might get a provider with the same identifier from another project, in which case we only keep the first one
+					Providers.TryAdd(Identifier, new ProviderInfo(Identifier, ServerUri, ClientId, DisplayName, RedirectUri));
+				}
+			}
+
+			if (Providers.Count == 0)
+				return null;
+			return new OIDCTokenManager(Settings, Providers);
+		}
+
+		public async Task Login(string ProviderIdentifier)
+		{
+			OIDCTokenClient TokenClient = TokenClients[ProviderIdentifier];
+
+			string RefreshToken = await TokenClient.Login();
+			if (!string.IsNullOrEmpty(RefreshToken))
+			{
+				// if we got a refresh token we store that for future use
+				byte[] RefreshTokenBytes = Encoding.ASCII.GetBytes(RefreshToken);
+				byte[] EncryptedBytes = ProtectedData.Protect(RefreshTokenBytes, null, DataProtectionScope.CurrentUser);
+
+				Settings.ProviderToRefreshTokens[ProviderIdentifier] = Convert.ToBase64String(EncryptedBytes);
+			}
+
+
+		}
+
+		public Task<string> GetAccessToken(string ProviderIdentifier)
+		{
+			return TokenClients[ProviderIdentifier].GetAccessToken();
+		}
+
+		public OIDCStatus GetStatusForProvider(string ProviderIdentifier)
+		{
+			return TokenClients[ProviderIdentifier].GetStatusForProvider();
+		}
+
+		public class ProviderInfo
+		{
+			public string Identifier { get; }
+			public string ServerUri { get; }
+			public string ClientId { get; }
+			public string DisplayName { get; }
+			public string RedirectUri { get; }
+
+			public ProviderInfo(string InIdentifier, string InServerUri, string InClientId, string InDisplayName, string InRedirectUri)
+			{
+				Identifier = InIdentifier; 
+				ServerUri = InServerUri;
+				ClientId = InClientId;
+				DisplayName = InDisplayName;
+				RedirectUri = InRedirectUri;
+			}
+		}
+	}
+
+
+	public enum OIDCStatus
+	{
+		Connected,
+		NotLoggedIn,
+		TokenRefreshRequired
+	}
+
+	class OIDCTokenClient
+	{
+		private readonly Uri AuthorityUri;
+		private readonly string ClientId;
+		private readonly string RedirectUri;
+		private readonly string Scopes;
+
+		private OidcClient OidcClient;
+
+		private string RefreshToken;
+		private string AccessToken;
+		private DateTime TokenExpiry;
+
+		private bool Initialized = false;
+
+
+		public OIDCTokenClient(Uri InAuthorityUri, string InClientId, string InRedirectUri, string InScopes = "openid profile offline_access")
+		{
+			AuthorityUri = InAuthorityUri;
+			ClientId = InClientId;
+			RedirectUri = InRedirectUri;
+			Scopes = InScopes;
+		}
+
+		private async Task Initialize()
+		{
+			if (Initialized)
+				return;
+
+			OidcClientOptions Options = new OidcClientOptions
+			{
+				Authority = AuthorityUri.ToString(),
+				Policy = new Policy { Discovery = new DiscoveryPolicy { Authority = AuthorityUri.ToString() } },
+				ClientId = ClientId,
+				Scope = Scopes,
+				FilterClaims = false,
+				RedirectUri = RedirectUri,
+				Flow = OidcClientOptions.AuthenticationFlow.AuthorizationCode,
+				ResponseMode = OidcClientOptions.AuthorizeResponseMode.Redirect
+			};
+
+			// we need to fetch the discovery document ourselves to support OIDC Authorities which have a subresource for it
+			// with Okta has for authorization servers for instance.
+			DiscoveryDocumentResponse DiscoveryDocument = await GetDiscoveryDocument();
+			Options.ProviderInformation = new ProviderInformation
+			{
+				IssuerName = DiscoveryDocument.Issuer,
+				KeySet = DiscoveryDocument.KeySet,
+
+				AuthorizeEndpoint = DiscoveryDocument.AuthorizeEndpoint,
+				TokenEndpoint = DiscoveryDocument.TokenEndpoint,
+				EndSessionEndpoint = DiscoveryDocument.EndSessionEndpoint,
+				UserInfoEndpoint = DiscoveryDocument.UserInfoEndpoint,
+				TokenEndPointAuthenticationMethods = DiscoveryDocument.TokenEndpointAuthenticationMethodsSupported
+			};
+
+			OidcClient = new OidcClient(Options);
+			Initialized = true;
+		}
+
+		public async Task<string> Login()
+		{
+			await Initialize();
+
+			// setup a local http server to listen for the result of the login
+			using HttpListener Http = new HttpListener();
+			Uri Uri = new Uri(RedirectUri);
+			// build the url the server should be hosted at
+			string Prefix = $"{Uri.Scheme}{Uri.SchemeDelimiter}{Uri.Authority}/";
+			Http.Prefixes.Add(Prefix);
+			Http.Start();
+
+			// generate the appropriate codes we need to login
+			AuthorizeState LoginState = await OidcClient.PrepareLoginAsync();
+			// start the user browser
+			OpenBrowser(LoginState.StartUrl);
+
+			string ResponseData = await ProcessHttpRequest(Http);
+
+			// parse the returned url for the tokens needed to complete the login
+			IdentityModel.OidcClient.LoginResult LoginResult = await OidcClient.ProcessResponseAsync(ResponseData, LoginState);
+
+			Http.Stop();
+
+			if (LoginResult.IsError)
+				throw new LoginFailedException("Failed to login due to error: " + LoginResult.Error);
+
+			RefreshToken = LoginResult.RefreshToken;
+			AccessToken = LoginResult.AccessToken;
+			TokenExpiry = LoginResult.AccessTokenExpiration;
+
+			return RefreshToken;
+		}
+
+		private async Task<string> ProcessHttpRequest(HttpListener Http)
+		{
+			HttpListenerContext Context = await Http.GetContextAsync();
+			string ResponseData;
+			if (Context.Request.HttpMethod == "GET")
+			{
+				ResponseData = Context.Request.RawUrl;
+
+			}
+			else if (Context.Request.HttpMethod == "POST")
+			{
+				var Request = Context.Request;
+				if (Request.ContentType != null && !Request.ContentType.Equals("application/x-www-form-urlencoded", StringComparison.OrdinalIgnoreCase))
+				{
+					// we do not support url encoded return types
+					Context.Response.StatusCode = 415;
+					return null;
+				}
+				
+				// attempt to parse the body
+
+				// if there is no body we can not handle the post
+				if (!Context.Request.HasEntityBody)
+				{
+					Context.Response.StatusCode = 415;
+					return null;
+				}
+
+				await using Stream Body = Request.InputStream;
+				using StreamReader Reader = new StreamReader(Body, Request.ContentEncoding);
+				ResponseData = Reader.ReadToEnd();
+			}
+			else
+			{
+				// if we receive any other http method something is very odd. Tell them to use a different method.
+				Context.Response.StatusCode = 415;
+				return null;
+			}
+
+			// generate a simple http page to show the user
+			HttpListenerResponse Response = Context.Response;
+			string HttpPage = "<html><head><meta http-equiv='refresh' content='10></head><body>Please close this browser and return to UnrealGameSync.</body></html>";
+			byte[] Buffer = Encoding.UTF8.GetBytes(HttpPage);
+			Response.ContentLength64 = Buffer.Length;
+			Stream ResponseOutput = Response.OutputStream;
+			await ResponseOutput.WriteAsync(Buffer, 0, Buffer.Length);
+			ResponseOutput.Close();
+
+			return ResponseData;
+		}
+
+		private void OpenBrowser(string Url)
+		{
+			try
+			{
+				Process.Start(Url);
+			}
+			catch
+			{
+				// hack because of this: [REDACTED_PROFILE_URL]/corefx/issues/10361
+				if (RuntimeInformation.IsOSPlatform(OSPlatform.Windows))
+				{
+					Url = Url.Replace("&", "^&");
+					Process.Start(new ProcessStartInfo("cmd", $"/c start {Url}") { CreateNoWindow = true });
+				}
+				else if (RuntimeInformation.IsOSPlatform(OSPlatform.Linux))
+				{
+					Process.Start("xdg-open", Url);
+				}
+				else if (RuntimeInformation.IsOSPlatform(OSPlatform.OSX))
+				{
+					Process.Start("open", Url);
+				}
+				else
+				{
+					throw;
+				}
+			}
+		}
+
+
+		private async Task<string> DoRefreshToken(string InRefreshToken)
+		{
+			await Initialize();
+
+			// use the refresh token to acquire a new access token
+			RefreshTokenResult RefreshTokenResult = await OidcClient.RefreshTokenAsync(InRefreshToken);
+
+			RefreshToken = RefreshTokenResult.RefreshToken;
+			AccessToken = RefreshTokenResult.AccessToken;
+			TokenExpiry = RefreshTokenResult.AccessTokenExpiration;
+
+			return AccessToken;
+		}
+
+		private async Task<DiscoveryDocumentResponse> GetDiscoveryDocument()
+		{
+			string DiscoUrl = $"{AuthorityUri}/.well-known/openid-configuration";
+
+			HttpClient Client = new HttpClient();
+			DiscoveryDocumentResponse Disco = await Client.GetDiscoveryDocumentAsync(new DiscoveryDocumentRequest
+			{
+				Address = DiscoUrl,
+				Policy =
+				{
+					ValidateEndpoints = false
+				}
+			});
+
+			if (Disco.IsError) throw new Exception(Disco.Error);
+
+			return Disco;
+		}
+
+
+		public async Task<string> GetAccessToken()
+		{
+			if (string.IsNullOrEmpty(RefreshToken))
+				throw new NotLoggedInException();
+
+			// if the token is valid for another few minutes we can use it
+			// we avoid using a token that is about to expire to make sure we can finish the call we expect to do with it before it expires
+			if (!string.IsNullOrEmpty(AccessToken) && TokenExpiry.AddMinutes(2) > DateTime.Now)
+			{
+				return AccessToken;
+			}
+
+			return await DoRefreshToken(RefreshToken);
+		}
+
+		public OIDCStatus GetStatusForProvider()
+		{
+			if (string.IsNullOrEmpty(RefreshToken))
+				return OIDCStatus.NotLoggedIn;
+
+			if (string.IsNullOrEmpty(AccessToken))
+				return OIDCStatus.TokenRefreshRequired;
+
+			if (TokenExpiry < DateTime.Now)
+				return OIDCStatus.TokenRefreshRequired;
+
+			return OIDCStatus.Connected;
+		}
+
+		public void SetRefreshToken(string InRefreshToken)
+		{
+			RefreshToken = InRefreshToken;
+		}
+	}
+
+	internal class LoginFailedException : Exception
+	{
+		public LoginFailedException(string Message) : base(Message)
+		{
+		}
+	}
+
+	internal class NotLoggedInException : Exception
+	{
+	}
+}
\ No newline at end of file
diff --git a/Engine/Source/Programs/UnrealGameSync/UnrealGameSync/PerforceMonitor.cs b/Engine/Source/Programs/UnrealGameSync/UnrealGameSync/PerforceMonitor.cs
--- a/Engine/Source/Programs/UnrealGameSync/UnrealGameSync/PerforceMonitor.cs
+++ b/Engine/Source/Programs/UnrealGameSync/UnrealGameSync/PerforceMonitor.cs
@@ -43,36 +43,49 @@ namespace UnrealGameSync
 		}
 	}
 
+	interface IArchiveInfoSource
+	{
+		IReadOnlyList<IArchiveInfo> AvailableArchives { get; }
+	}
+
 	interface IArchiveInfo
 	{
 		string Name { get; }
 		string Type { get; }
-		string DepotPath { get; }
-
+		string BasePath { get; }
 		bool Exists();
-		bool TryGetArchivePathForChangeNumber(int ChangeNumber, out string ArchivePath);
+		bool TryGetArchiveKeyForChangeNumber(int ChangeNumber, out string ArchiveKey);
+		bool DownloadArchive(string ArchiveKey, string LocalRootPath, string ManifestFileName, TextWriter Log, ProgressValue Progress);
 	}
 
-	class PerforceMonitor : IDisposable
+	class PerforceMonitor : IDisposable, IArchiveInfoSource
 	{
-		class ArchiveInfo : IArchiveInfo
+		internal class PerforceArchiveInfo : IArchiveInfo
 		{
 			public string Name { get; }
 			public string Type { get; }
 			public string DepotPath { get; }
+			public PerforceConnection Perforce { get; }
+
+			public string BasePath
+			{
+				get { return DepotPath; }
+			}
+
 			// TODO: executable/configuration?
 			public SortedList<int, string> ChangeNumberToFileRevision = new SortedList<int, string>();
 
-			public ArchiveInfo(string Name, string Type, string DepotPath)
+			public PerforceArchiveInfo(string Name, string Type, string DepotPath, PerforceConnection Perforce)
 			{
 				this.Name = Name;
 				this.Type = Type;
 				this.DepotPath = DepotPath;
+				this.Perforce = Perforce;
 			}
 
 			public override bool Equals(object Other)
 			{
-				ArchiveInfo OtherArchive = Other as ArchiveInfo;
+				PerforceArchiveInfo OtherArchive = Other as PerforceArchiveInfo;
 				return OtherArchive != null && Name == OtherArchive.Name && Type == OtherArchive.Type && DepotPath == OtherArchive.DepotPath && Enumerable.SequenceEqual(ChangeNumberToFileRevision, OtherArchive.ChangeNumberToFileRevision);
 			}
 
@@ -86,7 +99,7 @@ namespace UnrealGameSync
 				return ChangeNumberToFileRevision.Count > 0;
 			}
 
-			public static bool TryParseConfigEntry(string Text, out ArchiveInfo Info)
+			public static bool TryParseConfigEntry(string Text, PerforceConnection PerforceConnection, out PerforceArchiveInfo Info)
 			{
 				ConfigObject Object = new ConfigObject(Text);
 
@@ -106,13 +119,35 @@ namespace UnrealGameSync
 
 				string Type = Object.GetValue("Type", null) ?? Name;
 
-				Info = new ArchiveInfo(Name, Type, DepotPath);
+				Info = new PerforceArchiveInfo(Name, Type, DepotPath, PerforceConnection);
 				return true;
 			}
 
-			public bool TryGetArchivePathForChangeNumber(int ChangeNumber, out string ArchivePath)
+			public bool TryGetArchiveKeyForChangeNumber(int ChangeNumber, out string ArchiveKey)
+			{
+				return ChangeNumberToFileRevision.TryGetValue(ChangeNumber, out ArchiveKey);
+			}
+
+			public bool DownloadArchive(string ArchiveKey, string LocalRootPath, string ManifestFileName, TextWriter Log, ProgressValue Progress)
 			{
-				return ChangeNumberToFileRevision.TryGetValue(ChangeNumber, out ArchivePath);
+				string TempZipFileName = Path.GetTempFileName();
+				try
+				{
+					if (!Perforce.PrintToFile(ArchiveKey, TempZipFileName, Log) || new FileInfo(TempZipFileName).Length == 0)
+					{
+						return false;
+					}
+
+					ArchiveUtils.ExtractFiles(TempZipFileName, LocalRootPath, ManifestFileName, Progress, Log);
+
+				}
+				finally
+				{
+					File.SetAttributes(TempZipFileName, FileAttributes.Normal);
+					File.Delete(TempZipFileName);
+				}
+
+				return true;
 			}
 
 			public override string ToString()
@@ -140,7 +175,7 @@ namespace UnrealGameSync
 		SortedSet<PerforceChangeSummary> Changes = new SortedSet<PerforceChangeSummary>(new PerforceChangeSorter());
 		SortedDictionary<int, PerforceChangeDetails> ChangeDetails = new SortedDictionary<int,PerforceChangeDetails>();
 		SortedSet<int> PromotedChangeNumbers = new SortedSet<int>();
-		List<ArchiveInfo> Archives = new List<ArchiveInfo>();
+		List<PerforceArchiveInfo> Archives = new List<PerforceArchiveInfo>();
 		AutoResetEvent RefreshEvent = new AutoResetEvent(false);
 		BoundedLogWriter LogWriter;
 		bool bIsEnterpriseProject;
@@ -234,7 +269,7 @@ namespace UnrealGameSync
 				{
 					PollForUpdatesInner();
 				}
-				catch (ThreadAbortException)
+				catch (ThreadInterruptedException)
 				{
 					break;
 				}
@@ -387,7 +422,7 @@ namespace UnrealGameSync
 			// If we are using zipped binaries, make sure we have every change since the last zip containing them. This is necessary for ensuring that content changes show as
 			// syncable in the workspace view if there have been a large number of content changes since the last code change.
 			int MinZippedChangeNumber = -1;
-			foreach (ArchiveInfo Archive in Archives)
+			foreach (PerforceArchiveInfo Archive in Archives)
 			{
 				foreach (int ChangeNumber in Archive.ChangeNumberToFileRevision.Keys)
 				{
@@ -558,7 +593,7 @@ namespace UnrealGameSync
 
 		bool UpdateArchives()
 		{
-			List<ArchiveInfo> NewArchives = new List<ArchiveInfo>();
+			List<PerforceArchiveInfo> NewArchives = new List<PerforceArchiveInfo>();
 
 			// Find all the zipped binaries under this stream
 			ConfigSection ProjectConfigSection = LatestProjectConfigFile.FindSection(SelectedProjectIdentifier);
@@ -568,21 +603,21 @@ namespace UnrealGameSync
 				string LegacyEditorArchivePath = ProjectConfigSection.GetValue("ZippedBinariesPath", null);
 				if (LegacyEditorArchivePath != null)
 				{
-					NewArchives.Add(new ArchiveInfo("Editor", "Editor", LegacyEditorArchivePath));
+					NewArchives.Add(new PerforceArchiveInfo("Editor", "Editor", LegacyEditorArchivePath, Perforce));
 				}
 
 				// New style
 				foreach (string ArchiveValue in ProjectConfigSection.GetValues("Archives", new string[0]))
 				{
-					ArchiveInfo Archive;
-					if (ArchiveInfo.TryParseConfigEntry(ArchiveValue, out Archive))
+					PerforceArchiveInfo Archive;
+					if (PerforceArchiveInfo.TryParseConfigEntry(ArchiveValue, Perforce, out Archive))
 					{
 						NewArchives.Add(Archive);
 					}
 				}
 
 				// Make sure the zipped binaries path exists
-				foreach (ArchiveInfo NewArchive in NewArchives)
+				foreach (PerforceArchiveInfo NewArchive in NewArchives)
 				{
 					bool bExists;
 					if (!Perforce.FileExists(NewArchive.DepotPath, out bExists, LogWriter))
diff --git a/Engine/Source/Programs/UnrealGameSync/UnrealGameSync/ProgramApplicationContext.cs b/Engine/Source/Programs/UnrealGameSync/UnrealGameSync/ProgramApplicationContext.cs
--- a/Engine/Source/Programs/UnrealGameSync/UnrealGameSync/ProgramApplicationContext.cs
+++ b/Engine/Source/Programs/UnrealGameSync/UnrealGameSync/ProgramApplicationContext.cs
@@ -11,6 +11,7 @@ using System.Text;
 using System.Threading;
 using System.Threading.Tasks;
 using System.Windows.Forms;
+using UnrealGameSync.Forms;
 
 namespace UnrealGameSync
 {
@@ -48,6 +49,8 @@ namespace UnrealGameSync
 		ModalTaskWindow DetectStartupProjectSettingsWindow;
 		MainWindow MainWindowInstance;
 
+		OIDCTokenManager OIDCTokenManager;
+
 		public ProgramApplicationContext(PerforceConnection DefaultConnection, UpdateMonitor UpdateMonitor, string ApiUrl, string DataFolder, EventWaitHandle ActivateEvent, bool bRestoreState, string UpdateSpawn, string ProjectFileName, bool bUnstable, TimestampLogWriter Log, string Uri)
 		{
 			this.DefaultConnection = DefaultConnection;
@@ -238,8 +241,16 @@ namespace UnrealGameSync
 			// Get a list of all the valid projects to open
 			DetectProjectSettingsResult[] StartupProjects = DetectStartupProjectSettingsTask.Results.Where(x => x != null).ToArray();
 
+			OIDCTokenManager = OIDCTokenManager.CreateFromConfigFile(Settings, StartupProjects.Select(Result => Result.Task).ToList());
+			// Verify that none of the projects we are opening needs a OIDC login, if they do prompt for the login
+			if (OIDCTokenManager?.HasUnfinishedLogin() ?? false)
+			{
+				OIDCLoginWindow LoginDialog = new OIDCLoginWindow(OIDCTokenManager);
+				LoginDialog.ShowDialog();
+			}
+				
 			// Create the main window
-			MainWindowInstance = new MainWindow(UpdateMonitor, ApiUrl, DataFolder, CacheFolder, bRestoreState, UpdateSpawn ?? Assembly.GetExecutingAssembly().Location, bUnstable, StartupProjects, DefaultConnection, Log, Settings, Uri);
+			MainWindowInstance = new MainWindow(UpdateMonitor, ApiUrl, DataFolder, CacheFolder, bRestoreState, UpdateSpawn ?? Assembly.GetExecutingAssembly().Location, bUnstable, StartupProjects, DefaultConnection, Log, Settings, Uri, OIDCTokenManager);
 			if(bVisible)
 			{
 				MainWindowInstance.Show();
diff --git a/Engine/Source/Programs/UnrealGameSync/UnrealGameSync/UnrealGameSync.csproj b/Engine/Source/Programs/UnrealGameSync/UnrealGameSync/UnrealGameSync.csproj
--- a/Engine/Source/Programs/UnrealGameSync/UnrealGameSync/UnrealGameSync.csproj
+++ b/Engine/Source/Programs/UnrealGameSync/UnrealGameSync/UnrealGameSync.csproj
@@ -18,6 +18,10 @@
   </PropertyGroup>
   <ItemGroup>
     <PackageReference Include="EnvDTE" Version="16.7.30328.74" />
+    <PackageReference Include="IdentityModel.OidcClient" Version="3.1.2" />
+  </ItemGroup>
+  <ItemGroup>
+    <ProjectReference Include="..\..\DotNETCommon\DotNETUtilities\DotNETUtilities.csproj" />
   </ItemGroup>
   <ItemGroup>
     <Reference Include="Ionic.Zip.Reduced">
diff --git a/Engine/Source/Programs/UnrealGameSync/UnrealGameSync/UserSettings.cs b/Engine/Source/Programs/UnrealGameSync/UnrealGameSync/UserSettings.cs
--- a/Engine/Source/Programs/UnrealGameSync/UnrealGameSync/UserSettings.cs
+++ b/Engine/Source/Programs/UnrealGameSync/UnrealGameSync/UserSettings.cs
@@ -273,6 +273,9 @@ namespace UnrealGameSync
 		// Precompiled binaries
 		public List<ArchiveSettings> Archives = new List<ArchiveSettings>();
 
+		// OIDC Settings
+		public Dictionary<string, string> ProviderToRefreshTokens = new Dictionary<string, string>();
+
 		// Window settings
 		public bool bWindowVisible;
 		public FormWindowState WindowState;
@@ -434,6 +437,16 @@ namespace UnrealGameSync
 				Archives.Add(new ArchiveSettings(true, "Editor", new string[0]));
 			}
 
+			// OIDC Settings
+			string[] Tokens = ConfigFile.GetValues("OIDCProviders.Tokens", new string[0]);
+			foreach (string TokenValue in Tokens)
+			{
+				ConfigObject O = new ConfigObject(TokenValue);
+				string Provider = O.GetValue("Provider");
+				string Token = O.GetValue("Token");
+				ProviderToRefreshTokens.TryAdd(Provider, Token);
+			}
+
 			// Window settings
 			bWindowVisible = ConfigFile.GetValue("Window.Visible", true);
 			if(!Enum.TryParse(ConfigFile.GetValue("Window.State", ""), true, out WindowState))
@@ -740,6 +753,19 @@ namespace UnrealGameSync
 			ConfigSection ArchivesSection = ConfigFile.FindOrAddSection("PrecompiledBinaries");
 			ArchivesSection.SetValues("Archives", Archives.Select(x => x.ToConfigEntry()).ToArray());
 
+			// OIDC Settings
+			ConfigSection OIDCSection = ConfigFile.FindOrAddSection("OIDCProviders");
+			List<ConfigObject> TokenObjects = new List<ConfigObject>();
+			foreach (KeyValuePair<string, string> Pair in ProviderToRefreshTokens)
+			{
+				ConfigObject TokenEntryObject = new ConfigObject();
+				TokenEntryObject.SetValue("Provider", Pair.Key);
+				TokenEntryObject.SetValue("Token", Pair.Value);
+				TokenObjects.Add(TokenEntryObject);
+			}
+			OIDCSection.SetValues("Tokens", TokenObjects.Select(x => x.ToString()).ToArray());
+
+
 			// Window settings
 			ConfigSection WindowSection = ConfigFile.FindOrAddSection("Window");
 			WindowSection.Clear();
diff --git a/Engine/Source/Programs/UnrealGameSync/UnrealGameSync/Workspace.cs b/Engine/Source/Programs/UnrealGameSync/UnrealGameSync/Workspace.cs
--- a/Engine/Source/Programs/UnrealGameSync/UnrealGameSync/Workspace.cs
+++ b/Engine/Source/Programs/UnrealGameSync/UnrealGameSync/Workspace.cs
@@ -55,7 +55,7 @@ namespace UnrealGameSync
 		public int ChangeNumber;
 		public WorkspaceUpdateOptions Options;
 		public string[] SyncFilter;
-		public Dictionary<string, string> ArchiveTypeToDepotPath = new Dictionary<string,string>();
+		public Dictionary<string, Tuple<IArchiveInfo, string>> ArchiveTypeToArchive = new Dictionary<string, Tuple<IArchiveInfo, string>>();
 		public Dictionary<string, bool> DeleteFiles = new Dictionary<string,bool>();
 		public Dictionary<string, bool> ClobberFiles = new Dictionary<string,bool>();
 		public Dictionary<Guid,ConfigObject> DefaultBuildSteps;
@@ -911,39 +911,33 @@ namespace UnrealGameSync
 					Directory.CreateDirectory(ManifestDirectoryName);
 
 					// Sync and extract (or just remove) the given archives
-					foreach(KeyValuePair<string, string> ArchiveTypeAndDepotPath in Context.ArchiveTypeToDepotPath)
+					foreach(KeyValuePair<string, Tuple<IArchiveInfo, string>> ArchiveTypeAndArchive in Context.ArchiveTypeToArchive)
 					{
+						string ArchiveType = ArchiveTypeAndArchive.Key;
+
 						// Remove any existing binaries
-						string ManifestFileName = Path.Combine(ManifestDirectoryName, String.Format("{0}.zipmanifest", ArchiveTypeAndDepotPath.Key));
+						string ManifestFileName = Path.Combine(ManifestDirectoryName, String.Format("{0}.zipmanifest", ArchiveType));
 						if(File.Exists(ManifestFileName))
 						{
-							Log.WriteLine("Removing {0} binaries...", ArchiveTypeAndDepotPath.Key);
-							Progress.Set(String.Format("Removing {0} binaries...", ArchiveTypeAndDepotPath.Key), 0.0f);
+							Log.WriteLine("Removing {0} binaries...", ArchiveType);
+							Progress.Set(String.Format("Removing {0} binaries...", ArchiveType), 0.0f);
 							ArchiveUtils.RemoveExtractedFiles(LocalRootPath, ManifestFileName, Progress, Log);
 							File.Delete(ManifestFileName);
 							Log.WriteLine();
 						}
 
 						// If we have a new depot path, sync it down and extract it
-						if(ArchiveTypeAndDepotPath.Value != null)
+						if(ArchiveTypeAndArchive.Value != null)
 						{
-							string TempZipFileName = Path.GetTempFileName();
-							try
-							{
-								Log.WriteLine("Syncing {0} binaries...", ArchiveTypeAndDepotPath.Key.ToLowerInvariant());
-								Progress.Set(String.Format("Syncing {0} binaries...", ArchiveTypeAndDepotPath.Key.ToLowerInvariant()), 0.0f);
-								if(!Perforce.PrintToFile(ArchiveTypeAndDepotPath.Value, TempZipFileName, Log) || new FileInfo(TempZipFileName).Length == 0)
-								{
-									StatusMessage = String.Format("Couldn't read {0}", ArchiveTypeAndDepotPath.Value);
-									return WorkspaceUpdateResult.FailedToSync;
-								}
-								ArchiveUtils.ExtractFiles(TempZipFileName, LocalRootPath, ManifestFileName, Progress, Log);
-								Log.WriteLine();
-							}
-							finally
+							IArchiveInfo ArchiveInfo = ArchiveTypeAndArchive.Value.Item1;
+							string ArchiveKey = ArchiveTypeAndArchive.Value.Item2;
+
+							Log.WriteLine("Syncing {0} binaries...", ArchiveType.ToLowerInvariant());
+							Progress.Set(String.Format("Syncing {0} binaries...", ArchiveType.ToLowerInvariant()), 0.0f);
+							if (!ArchiveInfo.DownloadArchive(ArchiveKey, LocalRootPath, ManifestFileName, Log, Progress))
 							{
-								File.SetAttributes(TempZipFileName, FileAttributes.Normal);
-								File.Delete(TempZipFileName);
+								StatusMessage = String.Format("Couldn't read {0}", ArchiveKey);
+								return WorkspaceUpdateResult.FailedToSync;
 							}
 						}
 					}
