UGS now uses build.bat to run UBT instead of running it directly, this means it can handle our dotnet core location of UBT but also that we can handle things like building UBT automatically if we need to.

diff --git a/Engine/Source/Programs/UnrealGameSync/UnrealGameSync/Workspace.cs b/Engine/Source/Programs/UnrealGameSync/UnrealGameSync/Workspace.cs
--- a/Engine/Source/Programs/UnrealGameSync/UnrealGameSync/Workspace.cs
+++ b/Engine/Source/Programs/UnrealGameSync/UnrealGameSync/Workspace.cs
@@ -1040,14 +1040,6 @@ namespace UnrealGameSync
 				{
 					Progress.Set("Starting build...", 0.0f);
 
-					// Check we've built UBT (it should have been compiled by generating project files)
-					string UnrealBuildToolPath = Path.Combine(LocalRootPath, "Engine", "Binaries", "DotNET", "UnrealBuildTool.exe");
-					if(!File.Exists(UnrealBuildToolPath))
-					{
-						StatusMessage = String.Format("Couldn't find {0}", UnrealBuildToolPath);
-						return WorkspaceUpdateResult.FailedToCompile;
-					}
-
 					// Execute all the steps
 					float MaxProgressFraction = 0.0f;
 					foreach (BuildStep Step in BuildSteps)
@@ -1068,16 +1060,17 @@ namespace UnrealGameSync
 									{
 										StepStopwatch.AddData(new { Target = Step.Target });
 
+										string BuildBat = Path.Combine(LocalRootPath, "Engine", "Build", "BatchFiles", "Build.bat");
 										string CommandLine = String.Format("{0} {1} {2} {3} -NoHotReloadFromIDE", Step.Target, Step.Platform, Step.Configuration, Utility.ExpandVariables(Step.Arguments ?? "", Context.Variables));
 										if(Context.Options.HasFlag(WorkspaceUpdateOptions.Clean) || bForceClean)
 										{
-											Log.WriteLine("ubt> Running {0} {1} -clean", UnrealBuildToolPath, CommandLine);
-											Utility.ExecuteProcess(UnrealBuildToolPath, null, CommandLine + " -clean", null, new ProgressTextWriter(Progress, new PrefixedTextWriter("ubt> ", Log)));
+											Log.WriteLine("ubt> Running {0} {1} -clean", BuildBat, CommandLine);
+											Utility.ExecuteProcess(CmdExe, null, "/C " + BuildBat + " " + CommandLine + " -clean", null, new ProgressTextWriter(Progress, new PrefixedTextWriter("ubt> ", Log)));
 										}
 
-										Log.WriteLine("ubt> Running {0} {1} -progress", UnrealBuildToolPath, CommandLine);
+										Log.WriteLine("ubt> Running {0} {1} -progress", BuildBat, CommandLine);
 
-										int ResultFromBuild = Utility.ExecuteProcess(UnrealBuildToolPath, null, CommandLine + " -progress", null, new ProgressTextWriter(Progress, new PrefixedTextWriter("ubt> ", Log)));
+										int ResultFromBuild = Utility.ExecuteProcess(CmdExe, null, "/C " + BuildBat + " "+ CommandLine + " -progress", null, new ProgressTextWriter(Progress, new PrefixedTextWriter("ubt> ", Log)));
 										if(ResultFromBuild != 0)
 										{
 											StepStopwatch.Stop("Failed");
