Merging latest UGS code from Dev-Build.

#rnx

diff --git a/Engine/Source/Programs/UnrealGameSync/UnrealGameSync/ArchiveManifest.cs b/Engine/Source/Programs/UnrealGameSync/UnrealGameSync/ArchiveManifest.cs
--- a/Engine/Source/Programs/UnrealGameSync/UnrealGameSync/ArchiveManifest.cs
+++ b/Engine/Source/Programs/UnrealGameSync/UnrealGameSync/ArchiveManifest.cs
@@ -1,4 +1,4 @@
-﻿// Copyright Epic Games, Inc. All Rights Reserved.
+// Copyright Epic Games, Inc. All Rights Reserved.
 
 using System;
 using System.Collections.Generic;
@@ -84,25 +84,28 @@ namespace UnrealGameSync
 			DateTime TimeStamp = DateTime.UtcNow;
 			using(ZipFile Zip = new ZipFile(ArchiveFileName))
 			{
-				File.Delete(ManifestFileName);
-
-				// Create the manifest
-				ArchiveManifest Manifest = new ArchiveManifest();
-				foreach(ZipEntry Entry in Zip.Entries)
+				if (ManifestFileName != null)
 				{
-					if(!Entry.FileName.EndsWith("/") && !Entry.FileName.EndsWith("\\"))
-					{ 
-						Manifest.Files.Add(new ArchiveManifestFile(Entry.FileName, Entry.UncompressedSize, TimeStamp));
+					File.Delete(ManifestFileName);
+
+					// Create the manifest
+					ArchiveManifest Manifest = new ArchiveManifest();
+					foreach (ZipEntry Entry in Zip.Entries)
+					{
+						if (!Entry.FileName.EndsWith("/") && !Entry.FileName.EndsWith("\\"))
+						{
+							Manifest.Files.Add(new ArchiveManifestFile(Entry.FileName, Entry.UncompressedSize, TimeStamp));
+						}
 					}
-				}
 
-				// Write it out to a temporary file, then move it into place
-				string TempManifestFileName = ManifestFileName + ".tmp";
-				using(FileStream OutputStream = File.Open(TempManifestFileName, FileMode.Create, FileAccess.Write))
-				{
-					Manifest.Write(OutputStream);
+					// Write it out to a temporary file, then move it into place
+					string TempManifestFileName = ManifestFileName + ".tmp";
+					using (FileStream OutputStream = File.Open(TempManifestFileName, FileMode.Create, FileAccess.Write))
+					{
+						Manifest.Write(OutputStream);
+					}
+					File.Move(TempManifestFileName, ManifestFileName);
 				}
-				File.Move(TempManifestFileName, ManifestFileName);
 
 				// Extract all the files
 				int EntryIdx = 0;
diff --git a/Engine/Source/Programs/UnrealGameSync/UnrealGameSync/Controls/ErrorPanel.cs b/Engine/Source/Programs/UnrealGameSync/UnrealGameSync/Controls/ErrorPanel.cs
--- a/Engine/Source/Programs/UnrealGameSync/UnrealGameSync/Controls/ErrorPanel.cs
+++ b/Engine/Source/Programs/UnrealGameSync/UnrealGameSync/Controls/ErrorPanel.cs
@@ -54,6 +54,10 @@ namespace UnrealGameSync
 		{
 		}
 
+		public void UpdateSettings()
+		{
+		}
+
 		public Color? TintColor
 		{
 			get { return null; }
diff --git a/Engine/Source/Programs/UnrealGameSync/UnrealGameSync/Controls/WorkspaceControl.cs b/Engine/Source/Programs/UnrealGameSync/UnrealGameSync/Controls/WorkspaceControl.cs
--- a/Engine/Source/Programs/UnrealGameSync/UnrealGameSync/Controls/WorkspaceControl.cs
+++ b/Engine/Source/Programs/UnrealGameSync/UnrealGameSync/Controls/WorkspaceControl.cs
@@ -1079,7 +1079,7 @@ namespace UnrealGameSync
 
 		void UpdateCompleteCallback(WorkspaceUpdateContext Context, WorkspaceUpdateResult Result, string ResultMessage)
 		{
-			MainThreadSynchronizationContext.Post((o) => UpdateComplete(Context, Result, ResultMessage), null);
+			MainThreadSynchronizationContext.Post((o) => { if (!bIsDisposing) { UpdateComplete(Context, Result, ResultMessage); } }, null);
 		}
 
 		void UpdateComplete(WorkspaceUpdateContext Context, WorkspaceUpdateResult Result, string ResultMessage)
@@ -2482,6 +2482,27 @@ namespace UnrealGameSync
 			}
 		}
 
+		private static void SafeProcessStart(string Url, string Arguments, string WorkingDir)
+		{
+			try
+			{
+				ProcessStartInfo StartInfo = new ProcessStartInfo();
+				StartInfo.FileName = Url;
+				StartInfo.Arguments = Arguments;
+				StartInfo.WorkingDirectory = WorkingDir;
+				Process.Start(StartInfo).Dispose();
+			}
+			catch
+			{
+				MessageBox.Show(String.Format("Unable to open {0} {1}", Url, Arguments));
+			}
+		}
+
+		public void UpdateSettings()
+		{
+			UpdateStatusPanel();
+		}
+
 		private string ReplaceRegexMatches(string Text, Match MatchResult)
 		{
 			if (Text != null)
@@ -3020,12 +3041,12 @@ namespace UnrealGameSync
 			ProcessNames.Add("UE4Editor-Win64-Debug-Cmd");
 			ProcessNames.Add("UE4Editor-Win64-DebugGame");
 			ProcessNames.Add("UE4Editor-Win64-DebugGame-Cmd");
-			ProcessNames.Add("UE5Editor");
-			ProcessNames.Add("UE5Editor-Cmd");
-			ProcessNames.Add("UE5Editor-Win64-Debug");
-			ProcessNames.Add("UE5Editor-Win64-Debug-Cmd");
-			ProcessNames.Add("UE5Editor-Win64-DebugGame");
-			ProcessNames.Add("UE5Editor-Win64-DebugGame-Cmd");
+			ProcessNames.Add("UnrealEditor");
+			ProcessNames.Add("UnrealEditor-Cmd");
+			ProcessNames.Add("UnrealEditor-Win64-Debug");
+			ProcessNames.Add("UnrealEditor-Win64-Debug-Cmd");
+			ProcessNames.Add("UnrealEditor-Win64-DebugGame");
+			ProcessNames.Add("UnrealEditor-Win64-DebugGame-Cmd");
 			ProcessNames.Add("CrashReportClient");
 			ProcessNames.Add("CrashReportClient-Win64-Development");
 			ProcessNames.Add("CrashReportClient-Win64-Debug");
@@ -3193,6 +3214,18 @@ namespace UnrealGameSync
 			SafeProcessStart("p4v.exe", CommandLine.ToString());
 		}
 
+		private void OpenUshell()
+		{
+			StringBuilder CommandLine = new StringBuilder();
+			CommandLine.AppendFormat("/C \"\"{0}\"", Path.Combine(DataFolder, "Tools", "Ushell", "Current", "ushell.bat"));
+			if (SelectedFileName.EndsWith(".uproject"))
+			{
+				CommandLine.AppendFormat(" --project=\"{0}\"", SelectedFileName);
+			}
+			CommandLine.Append("\"");
+			SafeProcessStart(Environment.GetEnvironmentVariable("COMSPEC"), CommandLine.ToString(), BranchDirectoryName);
+		}
+
 		private void UpdateStatusPanel()
 		{
 			int NewContentWidth = Math.Max(TextRenderer.MeasureText(String.Format("Opened {0}  |  Browse...  |  Connect...", SelectedFileName), StatusPanel.Font).Width, 400);
@@ -3325,6 +3358,11 @@ namespace UnrealGameSync
 					ProgramsLine.AddLink("Visual Studio", FontStyle.Regular, () => { OpenSolution(); });
 					ProgramsLine.AddText("  |  ");
 				}
+				if (Settings.bEnableUshell)
+				{
+					ProgramsLine.AddLink("Ushell", FontStyle.Regular, () => { OpenUshell(); });
+					ProgramsLine.AddText("  |  ");
+				}
 				ProgramsLine.AddLink("Windows Explorer", FontStyle.Regular, () => { SafeProcessStart("explorer.exe", String.Format("\"{0}\"", Path.GetDirectoryName(SelectedFileName))); });
 
 				if (GetDefaultIssueFilter() != null)
diff --git a/Engine/Source/Programs/UnrealGameSync/UnrealGameSync/DeploymentSettings.cs b/Engine/Source/Programs/UnrealGameSync/UnrealGameSync/DeploymentSettings.cs
--- a/Engine/Source/Programs/UnrealGameSync/UnrealGameSync/DeploymentSettings.cs
+++ b/Engine/Source/Programs/UnrealGameSync/UnrealGameSync/DeploymentSettings.cs
@@ -47,6 +47,16 @@ namespace UnrealGameSync
 		/// </summary>
 		public static readonly string DefaultDepotPath = null;
 
+		/// <summary>
+		/// Depot path to sync additional tools from
+		/// </summary>
+		public static readonly string ToolsDepotPath = null;
+
+		/// <summary>
+		/// URL to send info about NET Core versions to.
+		/// </summary>
+		public static readonly string NetCoreTelemetryUrl = null;
+
 #if WITH_TELEMETRY
 		/// <summary>
 		/// Delegate used to create a new telemetry sink
diff --git a/Engine/Source/Programs/UnrealGameSync/UnrealGameSync/EventMonitor.cs b/Engine/Source/Programs/UnrealGameSync/UnrealGameSync/EventMonitor.cs
--- a/Engine/Source/Programs/UnrealGameSync/UnrealGameSync/EventMonitor.cs
+++ b/Engine/Source/Programs/UnrealGameSync/UnrealGameSync/EventMonitor.cs
@@ -5,6 +5,7 @@ using System.Collections.Concurrent;
 using System.Collections.Generic;
 using System.Globalization;
 using System.Diagnostics;
+using System.IO;
 using System.Linq;
 using System.Text.Json;
 using System.Threading;
@@ -665,6 +666,7 @@ namespace UnrealGameSync
 
 		void PollForUpdates()
 		{
+			Stopwatch NetCoreTimer = Stopwatch.StartNew();
 			EventData Event = null;
 			CommentData Comment = null;
 			bool bUpdateThrottledRequests = true;
@@ -705,11 +707,40 @@ namespace UnrealGameSync
 					}
 				}
 
+				// Post info about whether net core is installed
+				if (DeploymentSettings.NetCoreTelemetryUrl != null && NetCoreTimer.Elapsed > TimeSpan.FromMinutes(60.0))
+				{
+					try
+					{
+						RESTApi.POST(DeploymentSettings.NetCoreTelemetryUrl, "netcore", "{ }", "User=" + CurrentUserName, "Machine=" + Environment.MachineName, "NetCore=" + HasNetCore3().ToString());
+					}
+					catch
+					{
+					}
+					NetCoreTimer.Restart();
+				}
+
 				// Wait for something else to do
 				bUpdateThrottledRequests = RefreshEvent.WaitOne(30 * 1000);
 			}
 		}
 
+		public static bool HasNetCore3()
+		{
+			DirectoryInfo BaseDirInfo = new DirectoryInfo(Path.Combine(Environment.GetFolderPath(Environment.SpecialFolder.ProgramFiles), "dotnet", "shared", "Microsoft.NETCore.App"));
+			if (BaseDirInfo.Exists)
+			{
+				foreach (DirectoryInfo SubDir in BaseDirInfo.EnumerateDirectories())
+				{
+					if (SubDir.Name.StartsWith("3.", StringComparison.Ordinal))
+					{
+						return true;
+					}
+				}
+			}
+			return false;
+		}
+
 		bool SendEventToBackend(EventData Event)
 		{
 			try
diff --git a/Engine/Source/Programs/UnrealGameSync/UnrealGameSync/Forms/ApplicationSettingsWindow.cs b/Engine/Source/Programs/UnrealGameSync/UnrealGameSync/Forms/ApplicationSettingsWindow.cs
--- a/Engine/Source/Programs/UnrealGameSync/UnrealGameSync/Forms/ApplicationSettingsWindow.cs
+++ b/Engine/Source/Programs/UnrealGameSync/UnrealGameSync/Forms/ApplicationSettingsWindow.cs
@@ -57,12 +57,15 @@ namespace UnrealGameSync
 
 		bool? bRestartUnstable;
 
-		private ApplicationSettingsWindow(string DefaultServerAndPort, string DefaultUserName, bool bUnstable, string OriginalExecutableFileName, UserSettings Settings, TextWriter Log)
+		ToolUpdateMonitor ToolUpdateMonitor;
+
+		private ApplicationSettingsWindow(string DefaultServerAndPort, string DefaultUserName, bool bUnstable, string OriginalExecutableFileName, UserSettings Settings, ToolUpdateMonitor ToolUpdateMonitor, TextWriter Log)
 		{
 			InitializeComponent();
 
 			this.OriginalExecutableFileName = OriginalExecutableFileName;
 			this.Settings = Settings;
+			this.ToolUpdateMonitor = ToolUpdateMonitor;
 			this.Log = Log;
 
 			Utility.ReadGlobalPerforceSettings(ref InitialServerAndPort, ref InitialUserName, ref InitialDepotPath);
@@ -115,11 +118,14 @@ namespace UnrealGameSync
 			{
 				this.EnableProtocolHandlerCheckBox.CheckState = CheckState.Indeterminate;
 			}
+
+			EnableP4VExtensionsCheckBox.Checked = Settings.bEnableP4VExtensions;
+			EnableUshellCheckBox.Checked = Settings.bEnableUshell;
 		}
 
-		public static bool? ShowModal(IWin32Window Owner, PerforceConnection DefaultConnection, bool bUnstable, string OriginalExecutableFileName, UserSettings Settings, TextWriter Log)
+		public static bool? ShowModal(IWin32Window Owner, PerforceConnection DefaultConnection, bool bUnstable, string OriginalExecutableFileName, UserSettings Settings, ToolUpdateMonitor ToolUpdateMonitor, TextWriter Log)
 		{
-			ApplicationSettingsWindow ApplicationSettings = new ApplicationSettingsWindow(DefaultConnection.ServerAndPort, DefaultConnection.UserName, bUnstable, OriginalExecutableFileName, Settings, Log);
+			ApplicationSettingsWindow ApplicationSettings = new ApplicationSettingsWindow(DefaultConnection.ServerAndPort, DefaultConnection.UserName, bUnstable, OriginalExecutableFileName, Settings, ToolUpdateMonitor, Log);
 			if(ApplicationSettings.ShowDialog() == DialogResult.OK)
 			{
 				return ApplicationSettings.bRestartUnstable;
@@ -194,13 +200,13 @@ namespace UnrealGameSync
 			}
 
 			RegistryKey Key = Registry.CurrentUser.CreateSubKey("SOFTWARE\\Microsoft\\Windows\\CurrentVersion\\Run");
-			if(IsAutomaticallyRunAtStartup())
+			if (AutomaticallyRunAtStartupCheckBox.Checked)
 			{
-	            Key.DeleteValue("UnrealGameSync", false);
+				Key.SetValue("UnrealGameSync", String.Format("\"{0}\" -RestoreState", OriginalExecutableFileName));
 			}
 			else
 			{
-				Key.SetValue("UnrealGameSync", String.Format("\"{0}\" -RestoreState", OriginalExecutableFileName));
+				Key.DeleteValue("UnrealGameSync", false);
 			}
 
 			if (Settings.bKeepInTray != KeepInTrayCheckBox.Checked || Settings.SyncOptions.NumThreads != ParallelSyncThreadsSpinner.Value)
@@ -210,6 +216,20 @@ namespace UnrealGameSync
 				Settings.Save();
 			}
 
+			if (Settings.bEnableP4VExtensions != EnableP4VExtensionsCheckBox.Checked)
+			{
+				Settings.bEnableP4VExtensions = EnableP4VExtensionsCheckBox.Checked;
+				Settings.Save();
+				ToolUpdateMonitor.UpdateNow();
+			}
+
+			if (Settings.bEnableUshell != EnableUshellCheckBox.Checked)
+			{
+				Settings.bEnableUshell = EnableUshellCheckBox.Checked;
+				Settings.Save();
+				ToolUpdateMonitor.UpdateNow();
+			}
+
 			if (EnableProtocolHandlerCheckBox.CheckState == CheckState.Checked)
 			{
 				if (InitialProtocolHandlerState != ProtocolHandlerState.Installed)
diff --git a/Engine/Source/Programs/UnrealGameSync/UnrealGameSync/Forms/MainWindow.cs b/Engine/Source/Programs/UnrealGameSync/UnrealGameSync/Forms/MainWindow.cs
--- a/Engine/Source/Programs/UnrealGameSync/UnrealGameSync/Forms/MainWindow.cs
+++ b/Engine/Source/Programs/UnrealGameSync/UnrealGameSync/Forms/MainWindow.cs
@@ -29,6 +29,7 @@ namespace UnrealGameSync
 		void SyncLatestChange();
 		bool CanLaunchEditor();
 		void LaunchEditor();
+		void UpdateSettings();
 
 		Color? TintColor
 		{
@@ -105,6 +106,10 @@ namespace UnrealGameSync
 		Rectangle PrimaryWorkArea;
 		List<IssueAlertWindow> AlertWindows = new List<IssueAlertWindow>();
 
+		ToolUpdateMonitor ToolUpdateMonitor;
+
+		NetCoreWindow NetCoreWindow;
+
 		public MainWindow(UpdateMonitor InUpdateMonitor, string InApiUrl, string InDataFolder, string InCacheFolder, bool bInRestoreStateOnLoad, string InOriginalExecutableFileName, bool bInUnstable, DetectProjectSettingsResult[] StartupProjects, PerforceConnection InDefaultConnection, LineBasedTextWriter InLog, UserSettings InSettings, string InUri, OIDCTokenManager InOidcTokenManager)
 		{
 			Log = InLog;
@@ -193,6 +198,8 @@ namespace UnrealGameSync
 			{
 				WorkspaceIssueMonitor.IssueMonitor.Start();
 			}
+
+			ToolUpdateMonitor = new ToolUpdateMonitor(DefaultConnection, DataFolder, Settings);
 		}
 
 		void PostAutomationRequest(AutomationRequest Request)
@@ -583,6 +590,12 @@ namespace UnrealGameSync
 				AutomationLog = null;
 			}
 
+			if (ToolUpdateMonitor != null)
+			{
+				ToolUpdateMonitor.Close();
+				ToolUpdateMonitor = null;
+			}
+
 			base.Dispose(disposing);
 		}
 
@@ -1366,11 +1379,17 @@ namespace UnrealGameSync
 
 		public void ModifyApplicationSettings()
 		{
-			bool? bRelaunchUnstable = ApplicationSettingsWindow.ShowModal(this, DefaultConnection, bUnstable, OriginalExecutableFileName, Settings, Log);
+			bool? bRelaunchUnstable = ApplicationSettingsWindow.ShowModal(this, DefaultConnection, bUnstable, OriginalExecutableFileName, Settings, ToolUpdateMonitor, Log);
 			if(bRelaunchUnstable.HasValue)
 			{
 				UpdateMonitor.TriggerUpdate(UpdateType.UserInitiated, bRelaunchUnstable);
 			}
+
+			for (int Idx = 0; Idx < TabControl.GetTabCount(); Idx++)
+			{
+				IMainWindowTabPanel TabPanel = (IMainWindowTabPanel)TabControl.GetTabData(Idx);
+				TabPanel.UpdateSettings();
+			}
 		}
 
 		private void MainWindow_Load(object sender, EventArgs e)
@@ -1648,5 +1667,18 @@ namespace UnrealGameSync
 				}
 			}
 		}
+
+		private void NetCoreTimer_Tick(object sender, EventArgs e)
+		{
+			if (NetCoreWindow == null || !NetCoreWindow.Created)
+			{
+				if (!EventMonitor.HasNetCore3() && Settings.bShowNetCoreInfo)
+				{
+					NetCoreWindow = new NetCoreWindow(Settings);
+					NetCoreWindow.Location = new Point(Location.X + (Width - NetCoreWindow.Width) / 2, Location.Y + (Height - NetCoreWindow.Height) / 2);
+					NetCoreWindow.Show(this);
+				}
+			}
+		}
 	}
 }
diff --git a/Engine/Source/Programs/UnrealGameSync/UnrealGameSync/Forms/NetCoreWindow.cs b/Engine/Source/Programs/UnrealGameSync/UnrealGameSync/Forms/NetCoreWindow.cs
new file mode 100644
--- /dev/null
+++ b/Engine/Source/Programs/UnrealGameSync/UnrealGameSync/Forms/NetCoreWindow.cs
@@ -0,0 +1,42 @@
+using System;
+using System.Collections.Generic;
+using System.ComponentModel;
+using System.Data;
+using System.Drawing;
+using System.Linq;
+using System.Text;
+using System.Threading.Tasks;
+using System.Windows.Forms;
+
+namespace UnrealGameSync.Forms
+{
+	partial class NetCoreWindow : Form
+	{
+		UserSettings Settings;
+
+		public NetCoreWindow(UserSettings Settings)
+		{
+			this.Settings = Settings;
+
+			InitializeComponent();
+		}
+
+		private void SnoozeBtn_Click(object sender, EventArgs e)
+		{
+			Close();
+		}
+
+		private void DismissBtn_Click(object sender, EventArgs e)
+		{
+			Settings.bShowNetCoreInfo = false;
+			Settings.Save();
+
+			Close();
+		}
+
+		private void linkLabel1_LinkClicked(object sender, LinkLabelLinkClickedEventArgs e)
+		{
+			System.Diagnostics.Process.Start("https://dotnet.microsoft.com/download/dotnet-core/current/runtime");
+		}
+	}
+}
diff --git a/Engine/Source/Programs/UnrealGameSync/UnrealGameSync/Json.cs b/Engine/Source/Programs/UnrealGameSync/UnrealGameSync/Json.cs
deleted file mode 100644
diff --git a/Engine/Source/Programs/UnrealGameSync/UnrealGameSync/Jupiter/JupiterMonitor.cs b/Engine/Source/Programs/UnrealGameSync/UnrealGameSync/Jupiter/JupiterMonitor.cs
--- a/Engine/Source/Programs/UnrealGameSync/UnrealGameSync/Jupiter/JupiterMonitor.cs
+++ b/Engine/Source/Programs/UnrealGameSync/UnrealGameSync/Jupiter/JupiterMonitor.cs
@@ -10,7 +10,7 @@ using System.Threading;
 using System.Threading.Tasks;
 using IdentityModel.Client;
 using EpicGames.Core;
-using EpicGames.Core.Jupiter;
+using EpicGames.Jupiter;
 
 namespace UnrealGameSync
 {
@@ -212,6 +212,8 @@ namespace UnrealGameSync
 				get { return null; } 
 			}
 
+			public string Target => throw new NotImplementedException();
+
 			public bool Exists()
 			{
 				return ChangeToKey.Count > 0;
diff --git a/Engine/Source/Programs/UnrealGameSync/UnrealGameSync/Perforce.cs b/Engine/Source/Programs/UnrealGameSync/UnrealGameSync/Perforce.cs
--- a/Engine/Source/Programs/UnrealGameSync/UnrealGameSync/Perforce.cs
+++ b/Engine/Source/Programs/UnrealGameSync/UnrealGameSync/Perforce.cs
@@ -232,6 +232,8 @@ namespace UnrealGameSync
 		public long FileSize;
 		public PerforceFileFlags Flags;
 		public int Revision;
+		public int HeadRevision;
+		public int HeadChange;
 		public bool IsMapped;
 		public bool Unmap;
 		public string Digest;
@@ -283,11 +285,23 @@ namespace UnrealGameSync
 			Unmap = Tags.ContainsKey("unmap");
 
 			string RevisionString;
-			if(Tags.TryGetValue("rev", out RevisionString))
+			if (Tags.TryGetValue("rev", out RevisionString))
 			{
 				int.TryParse(RevisionString, out Revision);
 			}
 
+			string HeadRevisionString;
+			if (Tags.TryGetValue("headRev", out HeadRevisionString))
+			{
+				int.TryParse(HeadRevisionString, out HeadRevision);
+			}
+
+			string HeadChangeString;
+			if (Tags.TryGetValue("headChange", out HeadChangeString))
+			{
+				int.TryParse(HeadChangeString, out HeadChange);
+			}
+
 			Tags.TryGetValue("digest", out Digest);
 		}
 	}
@@ -766,7 +780,7 @@ namespace UnrealGameSync
 
 		public bool FindFiles(string Filter, out List<PerforceFileRecord> FileRecords, TextWriter Log)
 		{
-			bool bResult = RunCommand(String.Format("fstat \"{0}\"", Filter), out FileRecords, CommandOptions.None, Log);
+			bool bResult = RunCommand("", "fstat", new List<string> { Filter }, out FileRecords, CommandOptions.None, Log);
 			if(bResult)
 			{
                  FileRecords.RemoveAll(x => x.Action != null && x.Action.Contains("delete"));
@@ -861,22 +875,27 @@ namespace UnrealGameSync
 
 		public bool FindChanges(IEnumerable<string> Filters, string ByUser, int MaxResults, out List<PerforceChangeSummary> Changes, TextWriter Log)
 		{
-			string Arguments = "changes -s submitted -t -L";
-			if(MaxResults > 0)
+			string Command = "changes";
+			List<string> Arguments = new List<string>();
+			Arguments.Add("-ssubmitted");
+			Arguments.Add("-t");
+			Arguments.Add("-L");
+
+			if (MaxResults > 0)
 			{
-				Arguments += String.Format(" -m {0}", MaxResults);
+				Arguments.Add(String.Format("-m {0}", MaxResults));
 			}
-			if(ByUser != null)
+			if (ByUser != null)
 			{
-				Arguments += String.Format(" -u \"{0}\"", ByUser);
+				Arguments.Add(String.Format("-u \"{0}\"", ByUser));
 			}
-			foreach(string Filter in Filters)
+			foreach (string Filter in Filters)
 			{
-				Arguments += String.Format(" \"{0}\"", Filter);
+				Arguments.Add(String.Format("{0}", Filter));
 			}
 
 			List<string> Lines;
-			if(!RunCommand(Arguments, out Lines, CommandOptions.IgnoreFilesNotInClientViewError, Log))
+			if(!RunCommand("", Command, Arguments, out Lines, CommandOptions.IgnoreFilesNotInClientViewError, Log))
 			{
 				Changes = null;
 				return false;
@@ -1232,18 +1251,15 @@ namespace UnrealGameSync
 			return RunCommand(String.Format("fstat \"{0}\"", Filter), out FileRecords, CommandOptions.IgnoreFilesNotOnClientError | CommandOptions.IgnoreNoSuchFilesError | CommandOptions.IgnoreProtectedNamespaceError | CommandOptions.IgnoreFilesNotInClientViewError, Log);
 		}
 
-		public bool Stat(string Options, List<string> Files, out List<PerforceFileRecord> FileRecords, TextWriter Log)
+		public bool Stat(List<string> Options, List<string> Files, out List<PerforceFileRecord> FileRecords, TextWriter Log)
 		{
-			StringBuilder Arguments = new StringBuilder("fstat");
-			if(!String.IsNullOrEmpty(Options))
-			{
-				Arguments.AppendFormat(" {0}", Options);
-			}
-			foreach(string File in Files)
+			List<string> Arguments = new List<string>();
+			if (Options != null)
 			{
-				Arguments.AppendFormat(" \"{0}\"", File);
+				Arguments.AddRange(Options);
 			}
-			return RunCommand(Arguments.ToString(), out FileRecords, CommandOptions.IgnoreFilesNotOnClientError | CommandOptions.IgnoreNoSuchFilesError | CommandOptions.IgnoreProtectedNamespaceError | CommandOptions.IgnoreFilesNotInClientViewError, Log);
+			Arguments.AddRange(Files);
+			return RunCommand("", "fstat", Arguments, out FileRecords, CommandOptions.IgnoreFilesNotOnClientError | CommandOptions.IgnoreNoSuchFilesError | CommandOptions.IgnoreProtectedNamespaceError | CommandOptions.IgnoreFilesNotInClientViewError, Log);
 		}
 
 		public bool Sync(string Filter, TextWriter Log)
@@ -1345,36 +1361,61 @@ namespace UnrealGameSync
 
 		public bool SyncPreview(string Filter, int ChangeNumber, bool bOnlyFilesInThisChange, out List<PerforceFileRecord> FileRecords, TextWriter Log)
 		{
-			return RunCommand(String.Format("sync --parallel=0 -n {0}@{1}{2}", Filter, bOnlyFilesInThisChange? "=" : "", ChangeNumber), out FileRecords, CommandOptions.IgnoreFilesUpToDateError | CommandOptions.IgnoreNoSuchFilesError | CommandOptions.IgnoreFilesNotInClientViewError, Log);
+			List<string> Arguments = new List<string>();
+			Arguments.Add("--parallel=0");
+			Arguments.Add("-n");
+			Arguments.Add(String.Format("{0}@{1}{2}", Filter, bOnlyFilesInThisChange ? "=" : "", ChangeNumber));
+
+			return RunCommand("", "sync", Arguments, out FileRecords, CommandOptions.IgnoreFilesUpToDateError | CommandOptions.IgnoreNoSuchFilesError | CommandOptions.IgnoreFilesNotInClientViewError, Log);
 		}
 
 		public bool SyncPreview(string Filter, int ChangeNumber, bool bOnlyFilesInThisChange, Action<PerforceFileRecord> SyncOutput, TextWriter Log)
 		{
 			using(PerforceTagRecordParser Parser = new PerforceTagRecordParser(x => SyncOutput(new PerforceFileRecord(x))))
 			{
-				string CommandLine = String.Format("-ztag sync --parallel=0 -n {0}@{1}{2}", Filter, bOnlyFilesInThisChange? "=" : "", ChangeNumber);
-				return RunCommand(CommandLine.ToString(), null, Line => FilterTaggedOutput(Line, Parser, Log), CommandOptions.NoFailOnErrors | CommandOptions.IgnoreFilesUpToDateError | CommandOptions.IgnoreExitCode | CommandOptions.IgnoreNoSuchFilesError | CommandOptions.IgnoreFilesNotInClientViewError, Log);
+				List<string> Arguments = new List<string>();
+				Arguments.Add("--parallel=0");
+				Arguments.Add("-n");
+				Arguments.Add(String.Format("{0}@{1}{2}", Filter, bOnlyFilesInThisChange ? "=" : "", ChangeNumber));
+				return RunCommand("-ztag", "sync", Arguments, null, Line => FilterTaggedOutput(Line, Parser, Log), CommandOptions.NoFailOnErrors | CommandOptions.IgnoreFilesUpToDateError | CommandOptions.IgnoreExitCode | CommandOptions.IgnoreNoSuchFilesError | CommandOptions.IgnoreFilesNotInClientViewError, Log);
 			}
 		}
 
 		public bool ForceSync(string Filter, TextWriter Log)
 		{
-			return RunCommand(String.Format("sync --parallel=0 -f \"{0}\"", Filter), CommandOptions.IgnoreFilesUpToDateError, Log);
+			List<string> Arguments = new List<string>();
+			Arguments.Add("--parallel=0");
+			Arguments.Add("-f");
+			Arguments.Add(Filter);
+
+			return RunCommand("", "sync", Arguments, CommandOptions.IgnoreFilesUpToDateError, Log);
 		}
 
 		public bool ForceSync(string Filter, int ChangeNumber, TextWriter Log)
 		{
-			return RunCommand(String.Format("sync --parallel=0 -f \"{0}\"@{1}", Filter, ChangeNumber), CommandOptions.IgnoreFilesUpToDateError, Log);
+			List<string> Arguments = new List<string>();
+			Arguments.Add("--parallel=0");
+			Arguments.Add("-f");
+			Arguments.Add(String.Format("{0}@{1}", Filter, ChangeNumber));
+
+			return RunCommand("", "sync", Arguments, CommandOptions.IgnoreFilesUpToDateError, Log);
 		}
 
 		public bool GetOpenFiles(string Filter, out List<PerforceFileRecord> FileRecords, TextWriter Log)
 		{
-			return RunCommand(String.Format("opened \"{0}\"", Filter), out FileRecords, CommandOptions.None, Log);
+			List<string> Arguments = new List<string>();
+			Arguments.Add(Filter);
+
+			return RunCommand("", "opened", Arguments, out FileRecords, CommandOptions.None, Log);
 		}
 
 		public bool GetUnresolvedFiles(string Filter, out List<PerforceFileRecord> FileRecords, TextWriter Log)
 		{
-			return RunCommand(String.Format("fstat -Ru \"{0}\"", Filter), out FileRecords, CommandOptions.IgnoreNoSuchFilesError | CommandOptions.IgnoreFilesNotOpenedOnThisClientError | CommandOptions.IgnoreFilesNotInClientViewError, Log);
+			List<string> Arguments = new List<string>();
+			Arguments.Add("-Ru");
+			Arguments.Add(Filter);
+
+			return RunCommand("", "fstat", Arguments, out FileRecords, CommandOptions.IgnoreNoSuchFilesError | CommandOptions.IgnoreFilesNotOpenedOnThisClientError | CommandOptions.IgnoreFilesNotInClientViewError, Log);
 		}
 
 		public bool AutoResolveFile(string File, TextWriter Log)
@@ -1397,6 +1438,21 @@ namespace UnrealGameSync
 			}
 		}
 
+		private bool RunCommand(string GlobalOptions, string Command, List<string> Arguments, out List<PerforceFileRecord> FileRecords, CommandOptions Options, TextWriter Log)
+		{
+			List<Dictionary<string, string>> TagRecords;
+			if (!RunCommand(GlobalOptions, Command, Arguments, out TagRecords, Options, Log))
+			{
+				FileRecords = null;
+				return false;
+			}
+			else
+			{
+				FileRecords = TagRecords.Select(x => new PerforceFileRecord(x)).ToList();
+				return true;
+			}
+		}
+
 		private bool RunCommand(string CommandLine, out List<PerforceFileRecord> FileRecords, CommandOptions Options, TextWriter Log)
 		{
 			List<Dictionary<string, string>> TagRecords;
@@ -1427,6 +1483,28 @@ namespace UnrealGameSync
 			}
 		}
 
+		private bool RunCommand(string GlobalOptions, string Command, List<string> Arguments, out List<Dictionary<string, string>> TagRecords, CommandOptions Options, TextWriter Log)
+		{
+			List<string> Lines;
+			if (!RunCommand(GlobalOptions + " -ztag", Command, Arguments, PerforceOutputChannel.TaggedInfo, out Lines, Options, Log))
+			{
+				TagRecords = null;
+				return false;
+			}
+
+			List<Dictionary<string, string>> LocalOutput = new List<Dictionary<string, string>>();
+			using (PerforceTagRecordParser Parser = new PerforceTagRecordParser(Record => LocalOutput.Add(Record)))
+			{
+				foreach (string Line in Lines)
+				{
+					Parser.OutputLine(Line);
+				}
+			}
+			TagRecords = LocalOutput;
+
+			return true;
+		}
+
 		private bool RunCommand(string CommandLine, out List<Dictionary<string, string>> TagRecords, CommandOptions Options, TextWriter Log)
 		{
 			List<string> Lines;
@@ -1463,14 +1541,63 @@ namespace UnrealGameSync
 			return bResult;
 		}
 
+		private bool RunCommand(string GlobalOptions, string Command, List<string> Arguments, CommandOptions Options, TextWriter Log)
+		{
+			List<string> Lines;
+			bool bResult = RunCommand(GlobalOptions, Command, Arguments, out Lines, Options, Log);
+			if (Lines != null)
+			{
+				foreach (string Line in Lines)
+				{
+					Log.WriteLine("p4>   {0}", Line);
+				}
+			}
+			return bResult;
+		}
+
+		/**
+		 * See the comments for GetFullCommandLine as to why a verison of this function is needed that takes a List of arguments
+		 * instead of a single string. 
+		 */
+		private bool RunCommand(string GlobalOptions, string Command, List<string> Arguments, out List<string> Lines, CommandOptions Options, TextWriter Log)
+		{
+			return RunCommand(GlobalOptions, Command, Arguments, PerforceOutputChannel.Info, out Lines, Options, Log);
+		}
+
+		private bool RunCommand(string GlobalOptions, string CommandLine, List<string> Arguments, PerforceOutputChannel Channel, out List<string> Lines, CommandOptions Options, TextWriter Log)
+		{
+			// get commandline which uses an args file in a temp file
+			string FilenameToDelete;
+			string FullCommandLine = GetFullCommandLine(GlobalOptions, CommandLine, Arguments, Options, out FilenameToDelete);
+
+			bool Result = ExecuteP4Process(FullCommandLine, Channel, out Lines, Options, Log);
+
+			// delete the temp file
+			File.Delete(FilenameToDelete);
+
+			return Result;
+		}
+
 		private bool RunCommand(string CommandLine, out List<string> Lines, CommandOptions Options, TextWriter Log)
 		{
-			return RunCommand(CommandLine, PerforceOutputChannel.Info, out Lines, Options, Log);
+			string FullCommandLine = GetFullCommandLine(CommandLine, Options);
+			return RunCommand(FullCommandLine, PerforceOutputChannel.Info, out Lines, Options, Log);
 		}
 
 		private bool RunCommand(string CommandLine, PerforceOutputChannel Channel, out List<string> Lines, CommandOptions Options, TextWriter Log)
 		{
+			// catch operations that have // and *, but are not using the List of arguments mode (see GetFullCommandLine comment)
+			if (CommandLine.Contains("//") && CommandLine.Contains("*"))
+			{
+				Log.WriteLine("Command {0} contained a filespec that will cause Windows to lockup or never complete(eg. //Foo/*). You must use the version of RunCommand that takes the options in an array of Arguments.", CommandLine);
+			}
+
 			string FullCommandLine = GetFullCommandLine(CommandLine, Options);
+			return ExecuteP4Process(FullCommandLine, Channel, out Lines, Options, Log);
+		}
+
+		private bool ExecuteP4Process(string FullCommandLine, PerforceOutputChannel Channel, out List<string> Lines, CommandOptions Options, TextWriter Log)
+		{
 			Log.WriteLine("p4> p4.exe {0}", FullCommandLine);
 
 			List<string> RawOutputLines;
@@ -1528,6 +1655,45 @@ namespace UnrealGameSync
 			return bResult;
 		}
 
+		private bool RunCommand(string GlobalOptions, string Command, List<string> Arguments, string Input, HandleOutputDelegate HandleOutput, CommandOptions Options, TextWriter Log)
+		{
+			string FilenameToDelete;
+			string FullCommandLine = GetFullCommandLine(GlobalOptions, Command, Arguments, Options, out FilenameToDelete);
+
+			bool bResult = true;
+			Log.WriteLine("p4> p4.exe {0}", FullCommandLine);
+			if (Utility.ExecuteProcess("p4.exe", null, FullCommandLine, Input, Line => { bResult &= ParseCommandOutput(Line, HandleOutput, Options); }) != 0 && !Options.HasFlag(CommandOptions.IgnoreExitCode))
+			{
+				bResult = false;
+			}
+			File.Delete(FilenameToDelete);
+			return bResult;
+		}
+
+		/**
+		 * p4 operations that have a * in a filespec (like //depot/UE4/Engine/*) can cause Windows to either delay a lot (or never complete) as the commandline
+		 * is thought to be a UNC path (\\server\share\ style) and it tries to expand the * - and Windows can be very slow to time out looking for a server named
+		 * depot. The workaround is to put the parameters into a args file with the -x option. However, the args file must have one option per line, while the
+		 * global options must be specified as normal. So, any operation that has a // and * will need to go through this function to put the options into a file.
+		 * The command 'p4 -c someclient changes -s submitted -m 100 //depot/UE4/Engine/*' would be 'p4 -c someclient -x TempFile changes' and the TempFile wouild contain:
+		 *    -ssubmitted
+		 *    -m 100
+		 *    //depot/UE4/Engine/*
+		 * Note that for some reason, -s submitted requires there be no space between -s and submitted when in an args file. Possibly params with a string option, as -m 100 works fine
+		 */
+		private string GetFullCommandLine(string GlobalOptions, string Command, List<string> Arguments, CommandOptions Options, out string FilenameToDelete)
+		{
+			// write arguments to the command, one per line, to temp file
+			FilenameToDelete = Path.Combine(Path.GetTempPath(), Path.GetTempFileName());
+			File.WriteAllLines(FilenameToDelete, Arguments);
+
+			// build up the commandline: <global options> -x <args file> <command>
+			string FullGlobalOptions = string.Format("{0} {1}", GetFullCommandLine("", Options), GlobalOptions);
+			string FullCommandLine = string.Format("{0} -x \"{1}\" {2}", FullGlobalOptions, FilenameToDelete, Command);
+
+			return FullCommandLine;
+		}
+
 		private string GetFullCommandLine(string CommandLine, CommandOptions Options)
 		{
 			StringBuilder FullCommandLine = new StringBuilder();
diff --git a/Engine/Source/Programs/UnrealGameSync/UnrealGameSync/PerforceMonitor.cs b/Engine/Source/Programs/UnrealGameSync/UnrealGameSync/PerforceMonitor.cs
--- a/Engine/Source/Programs/UnrealGameSync/UnrealGameSync/PerforceMonitor.cs
+++ b/Engine/Source/Programs/UnrealGameSync/UnrealGameSync/PerforceMonitor.cs
@@ -686,7 +686,7 @@ namespace UnrealGameSync
 			ConfigFile ProjectConfig = new ConfigFile();
 
 			List<PerforceFileRecord> FileRecords;
-			if(Perforce.Stat("-Ol", ConfigFilePaths, out FileRecords, Log))
+			if(Perforce.Stat(new List<string> { "-Ol" }, ConfigFilePaths, out FileRecords, Log))
 			{
 				foreach(PerforceFileRecord FileRecord in FileRecords)
 				{
diff --git a/Engine/Source/Programs/UnrealGameSync/UnrealGameSync/Properties/AssemblyInfo.cs b/Engine/Source/Programs/UnrealGameSync/UnrealGameSync/Properties/AssemblyInfo.cs
--- a/Engine/Source/Programs/UnrealGameSync/UnrealGameSync/Properties/AssemblyInfo.cs
+++ b/Engine/Source/Programs/UnrealGameSync/UnrealGameSync/Properties/AssemblyInfo.cs
@@ -34,5 +34,5 @@ using System.Runtime.InteropServices;
 // You can specify all the values or you can default the Build and Revision Numbers 
 // by using the '*' as shown below:
 // [assembly: AssemblyVersion("1.0.*")]
-[assembly: AssemblyVersion("1.224.0.0")]
-[assembly: AssemblyFileVersion("1.224.0.0")]
+[assembly: AssemblyVersion("1.234.0.0")]
+[assembly: AssemblyFileVersion("1.234.0.0")]
diff --git a/Engine/Source/Programs/UnrealGameSync/UnrealGameSync/ToolUpdateMonitor.cs b/Engine/Source/Programs/UnrealGameSync/UnrealGameSync/ToolUpdateMonitor.cs
new file mode 100644
--- /dev/null
+++ b/Engine/Source/Programs/UnrealGameSync/UnrealGameSync/ToolUpdateMonitor.cs
@@ -0,0 +1,317 @@
+using System;
+using System.Collections.Generic;
+using System.Diagnostics;
+using System.IO;
+using System.Linq;
+using System.Text;
+using System.Threading;
+using System.Threading.Tasks;
+
+namespace UnrealGameSync
+{
+	class ToolUpdateMonitor : IDisposable
+	{
+		Thread WorkerThread;
+		AutoResetEvent WakeEvent;
+		bool bQuit;
+		string LogFile;
+
+		public PerforceConnection Perforce { get; }
+		string ToolsDir { get; }
+		UserSettings Settings { get; }
+
+		class ToolDefinition
+		{
+			public string Name;
+			public Action<TextWriter> InstallAction;
+			public Action<TextWriter> UninstallAction;
+			public bool bEnabled;
+
+			public ToolDefinition(string Name, Action<TextWriter> InstallAction, Action<TextWriter> UninstallAction, bool bEnabled)
+			{
+				this.Name = Name;
+				this.InstallAction = InstallAction;
+				this.UninstallAction = UninstallAction;
+				this.bEnabled = bEnabled;
+			}
+		}
+
+
+		public ToolUpdateMonitor(PerforceConnection InPerforce, string DataDir, UserSettings Settings)
+		{
+			this.Perforce = InPerforce;
+			this.ToolsDir = Path.Combine(DataDir, "Tools");
+			this.Settings = Settings;
+
+			LogFile = Path.Combine(DataDir, "Tools.log");
+
+			Directory.CreateDirectory(ToolsDir);
+
+			WakeEvent = new AutoResetEvent(false);
+
+			if (DeploymentSettings.ToolsDepotPath != null)
+			{
+				WorkerThread = new Thread(() => PollForUpdates());
+				WorkerThread.Start();
+			}
+		}
+
+		public void Close()
+		{
+			bQuit = true;
+			WakeEvent.Set();
+
+			if (WorkerThread != null)
+			{
+				if (!WorkerThread.Join(30))
+				{
+					WorkerThread.Abort();
+					WorkerThread.Join();
+				}
+				WorkerThread = null;
+			}
+		}
+
+		public void Dispose()
+		{
+			Close();
+		}
+
+		string GetToolPathInternal(string ToolName)
+		{
+			return Path.Combine(ToolsDir, ToolName, "Current");
+		}
+
+		public string GetToolPath(string ToolName)
+		{
+			if (GetToolChange(ToolName) != 0)
+			{
+				return GetToolPathInternal(ToolName);
+			}
+			else
+			{
+				return null;
+			}
+		}
+
+		public void UpdateNow()
+		{
+			WakeEvent.Set();
+		}
+
+		void PollForUpdates()
+		{
+			using (BoundedLogWriter Log = new BoundedLogWriter(LogFile))
+			{
+				while (!bQuit)
+				{
+					try
+					{
+						PollForUpdatesOnce(Log);
+					}
+					catch
+					{
+					}
+
+					WakeEvent.WaitOne(TimeSpan.FromMinutes(60.0));
+				}
+			}
+		}
+
+		void PollForUpdatesOnce(TextWriter Log)
+		{
+			List<ToolDefinition> Tools = new List<ToolDefinition>();
+			Tools.Add(new ToolDefinition("P4VUtils", x => RunP4VUtils("install", x), x => RunP4VUtils("uninstall", x), Settings.bEnableP4VExtensions));
+			Tools.Add(new ToolDefinition("Ushell", null, null, Settings.bEnableUshell));
+
+			if (Tools.Any(x => x.bEnabled))
+			{
+				List<PerforceFileRecord> FileRecords;
+				if (Perforce.Stat(DeploymentSettings.ToolsDepotPath + "/...", out FileRecords, Log))
+				{
+					foreach (ToolDefinition Tool in Tools.Where(x => x.bEnabled))
+					{
+						string Prefix = String.Format("{0}/{1}/", DeploymentSettings.ToolsDepotPath, Tool.Name);
+
+						List<PerforceFileRecord> ToolFileRecords = FileRecords.Where(x => x.DepotPath.StartsWith(Prefix, StringComparison.OrdinalIgnoreCase)).ToList();
+						if (ToolFileRecords.Count == 0)
+						{
+							continue;
+						}
+
+						int HeadChange = ToolFileRecords.Max(x => x.HeadChange);
+						if (HeadChange == GetToolChange(Tool.Name))
+						{
+							continue;
+						}
+
+						List<PerforceFileRecord> SyncFileRecords = ToolFileRecords.Where(x => x.Action != "delete").ToList();
+						try
+						{
+							UpdateTool(Tool.Name, HeadChange, SyncFileRecords, Tool.InstallAction, Log);
+						}
+						catch (Exception Ex)
+						{
+							Log.WriteLine("Exception while updating tool: {0}", Ex.ToString());
+						}
+					}
+				}
+			}
+
+			foreach (ToolDefinition Tool in Tools.Where(x => !x.bEnabled))
+			{
+				if (GetToolChange(Tool.Name) != 0)
+				{
+					try
+					{
+						RemoveTool(Tool.Name, Tool.UninstallAction, Log);
+					}
+					catch (Exception Ex)
+					{
+						Log.WriteLine("Exception while removing tool: {0}", Ex.ToString());
+					}
+				}
+			}
+		}
+
+		void RunP4VUtils(string Command, TextWriter Log)
+		{
+			string ToolPath = GetToolPathInternal("P4VUtils");
+
+			int ExitCode = Utility.ExecuteProcess(Path.Combine(ToolPath, "P4VUtils.exe"), ToolPath, Command, null, Log);
+			Log.WriteLine("(Exit: {0})", ExitCode);
+		}
+
+		void RemoveTool(string ToolName, Action<TextWriter> UninstallAction, TextWriter Log)
+		{
+			Log.WriteLine("Removing {0}", ToolName);
+
+			if (UninstallAction != null)
+			{
+				Log.WriteLine("Running uninstall...");
+				UninstallAction?.Invoke(Log);
+			}
+
+			SetToolChange(ToolName, null);
+
+			string ToolPath = GetToolPath(ToolName);
+			Log.WriteLine("Removing {0}", ToolPath);
+
+			TryDeleteDirectory(ToolPath);
+		}
+
+		static void ForceDeleteDirectory(string DirectoryName)
+		{
+			DirectoryInfo BaseDir = new DirectoryInfo(DirectoryName);
+			if (BaseDir.Exists)
+			{
+				foreach (FileInfo File in BaseDir.EnumerateFiles("*", SearchOption.AllDirectories))
+				{
+					File.Attributes = FileAttributes.Normal;
+				}
+				BaseDir.Delete(true);
+			}
+		}
+
+		static bool TryDeleteDirectory(string DirectoryName)
+		{
+			try
+			{
+				ForceDeleteDirectory(DirectoryName);
+				return true;
+			}
+			catch
+			{
+				return false;
+			}
+		}
+
+		bool UpdateTool(string ToolName, int Change, List<PerforceFileRecord> Records, Action<TextWriter> InstallAction, TextWriter Log)
+		{
+			string ToolDir = Path.Combine(ToolsDir, ToolName);
+			Directory.CreateDirectory(ToolDir);
+
+			foreach (DirectoryInfo ExistingDir in new DirectoryInfo(ToolDir).EnumerateDirectories("Prev-*"))
+			{
+				TryDeleteDirectory(ExistingDir.FullName);
+			}
+
+			string NextToolDir = Path.Combine(ToolDir, "Next");
+			ForceDeleteDirectory(NextToolDir);
+			Directory.CreateDirectory(NextToolDir);
+
+			string NextToolZipsDir = Path.Combine(NextToolDir, ".zips");
+			Directory.CreateDirectory(NextToolZipsDir);
+
+			for (int Idx = 0; Idx < Records.Count; Idx++)
+			{
+				string ZipFile = Path.Combine(NextToolZipsDir, String.Format("{0}.{1}.zip", ToolName, Idx));
+				if (!Perforce.PrintToFile(String.Format("{0}#{1}", Records[Idx].DepotPath, Records[Idx].HeadRevision), ZipFile, Log) || !File.Exists(ZipFile))
+				{
+					Log.WriteLine("Unable to print {0}", Records[Idx].DepotPath);
+					return false;
+				}
+				ArchiveUtils.ExtractFiles(ZipFile, NextToolDir, null, new ProgressValue(), Log);
+			}
+
+			SetToolChange(ToolName, null);
+
+			string CurrentToolDir = Path.Combine(ToolDir, "Current");
+			if (Directory.Exists(CurrentToolDir))
+			{
+				string PrevDirectoryName = Path.Combine(ToolDir, String.Format("Prev-{0:X16}", Stopwatch.GetTimestamp()));
+				Directory.Move(CurrentToolDir, PrevDirectoryName);
+				TryDeleteDirectory(PrevDirectoryName);
+			}
+
+			Directory.Move(NextToolDir, CurrentToolDir);
+
+			if (InstallAction != null)
+			{
+				Log.WriteLine("Running installer...");
+				InstallAction.Invoke(Log);
+			}
+
+			SetToolChange(ToolName, Change);
+			return true;
+		}
+
+		string GetConfigFilePath(string ToolName)
+		{
+			return Path.Combine(ToolsDir, ToolName, ToolName + ".ini");
+		}
+
+		int GetToolChange(string ToolName)
+		{
+			try
+			{
+				string ConfigFilePath = GetConfigFilePath(ToolName);
+				if (File.Exists(ConfigFilePath))
+				{
+					ConfigFile ConfigFile = new ConfigFile();
+					ConfigFile.Load(ConfigFilePath);
+					return ConfigFile.GetValue("Settings.Change", 0);
+				}
+			}
+			catch
+			{
+			}
+			return 0;
+		}
+
+		void SetToolChange(string ToolName, int? Change)
+		{
+			string ConfigFilePath = GetConfigFilePath(ToolName);
+			if (Change.HasValue)
+			{
+				ConfigFile ConfigFile = new ConfigFile();
+				ConfigFile.SetValue("Settings.Change", Change.Value);
+				ConfigFile.Save(ConfigFilePath);
+			}
+			else
+			{
+				File.Delete(ConfigFilePath);
+			}
+		}
+	}
+}
diff --git a/Engine/Source/Programs/UnrealGameSync/UnrealGameSync/UnrealGameSync.csproj b/Engine/Source/Programs/UnrealGameSync/UnrealGameSync/UnrealGameSync.csproj
--- a/Engine/Source/Programs/UnrealGameSync/UnrealGameSync/UnrealGameSync.csproj
+++ b/Engine/Source/Programs/UnrealGameSync/UnrealGameSync/UnrealGameSync.csproj
@@ -1,4 +1,4 @@
-﻿<Project Sdk="Microsoft.NET.Sdk.WindowsDesktop">
+<Project Sdk="Microsoft.NET.Sdk.WindowsDesktop">
   <PropertyGroup>
     <OutputType>WinExe</OutputType>
     <TargetFramework>netcoreapp3.1</TargetFramework>
@@ -22,6 +22,7 @@
   </ItemGroup>
   <ItemGroup>
     <ProjectReference Include="..\..\Shared\EpicGames.Core\EpicGames.Core.csproj" />
+    <ProjectReference Include="..\..\Shared\EpicGames.Jupiter\EpicGames.Jupiter.csproj" />
   </ItemGroup>
   <ItemGroup>
     <Reference Include="Ionic.Zip.Reduced">
diff --git a/Engine/Source/Programs/UnrealGameSync/UnrealGameSync/UriHandlers/P4Handler.cs b/Engine/Source/Programs/UnrealGameSync/UnrealGameSync/UriHandlers/P4Handler.cs
--- a/Engine/Source/Programs/UnrealGameSync/UnrealGameSync/UriHandlers/P4Handler.cs
+++ b/Engine/Source/Programs/UnrealGameSync/UnrealGameSync/UriHandlers/P4Handler.cs
@@ -8,6 +8,19 @@ namespace UnrealGameSync
 	/// </summary>
 	static class P4Handler
 	{
+		[UriHandler(true)]
+		public static UriResult P4V(string DepotPath)
+		{
+			string CommandLine = string.Format("-s \"{0}\"", DepotPath);
+
+			if (!Utility.SpawnHiddenProcess("p4v.exe", CommandLine))
+			{
+				return new UriResult() { Error = string.Format("Error spawning p4vc.exe with command line: {0}", CommandLine) };
+			}
+
+			return new UriResult() { Success = true };
+		}
+
 		[UriHandler(true)]
 		public static UriResult Timelapse(string DepotPath, int Line = -1)
 		{
diff --git a/Engine/Source/Programs/UnrealGameSync/UnrealGameSync/UserSettings.cs b/Engine/Source/Programs/UnrealGameSync/UnrealGameSync/UserSettings.cs
--- a/Engine/Source/Programs/UnrealGameSync/UnrealGameSync/UserSettings.cs
+++ b/Engine/Source/Programs/UnrealGameSync/UnrealGameSync/UserSettings.cs
@@ -258,6 +258,9 @@ namespace UnrealGameSync
 		public bool bShowAutomatedChanges;
 		public bool bShowLocalTimes;
 		public bool bKeepInTray;
+		public bool bEnableP4VExtensions;
+		public bool bEnableUshell;
+		public bool bShowNetCoreInfo;
 		public int FilterIndex;
 		public UserSelectedProjectSettings LastProject;
 		public List<UserSelectedProjectSettings> OpenProjects;
@@ -356,6 +359,9 @@ namespace UnrealGameSync
 			bShowAutomatedChanges = ConfigFile.GetValue("General.ShowAutomated", false);
 			bShowLocalTimes = ConfigFile.GetValue("General.ShowLocalTimes", false);
 			bKeepInTray = ConfigFile.GetValue("General.KeepInTray", true);
+			bEnableP4VExtensions = ConfigFile.GetValue("General.EnableP4VExtensions", false);
+			bEnableUshell = ConfigFile.GetValue("General.EnableUshell", false);
+			bShowNetCoreInfo = ConfigFile.GetValue("General.ShowNetCoreInfo", true);
 			int.TryParse(ConfigFile.GetValue("General.FilterIndex", "0"), out FilterIndex);
 
 			string LastProjectString = ConfigFile.GetValue("General.LastProject", null);
@@ -717,6 +723,9 @@ namespace UnrealGameSync
 			}
 			GeneralSection.SetValues("OpenProjects", OpenProjects.Select(x => x.ToConfigEntry()).ToArray());
 			GeneralSection.SetValue("KeepInTray", bKeepInTray);
+			GeneralSection.SetValue("EnableP4VExtensions", bEnableP4VExtensions);
+			GeneralSection.SetValue("EnableUshell", bEnableUshell);
+			GeneralSection.SetValue("ShowNetCoreInfo", bShowNetCoreInfo);
 			GeneralSection.SetValue("FilterIndex", FilterIndex);
 			GeneralSection.SetValues("RecentProjects", RecentProjects.Select(x => x.ToConfigEntry()).ToArray());
 			GeneralSection.SetValues("SyncFilter", SyncView);
diff --git a/Engine/Source/Programs/UnrealGameSync/UnrealGameSync/Workspace.cs b/Engine/Source/Programs/UnrealGameSync/UnrealGameSync/Workspace.cs
--- a/Engine/Source/Programs/UnrealGameSync/UnrealGameSync/Workspace.cs
+++ b/Engine/Source/Programs/UnrealGameSync/UnrealGameSync/Workspace.cs
@@ -1065,12 +1065,12 @@ namespace UnrealGameSync
 										if(Context.Options.HasFlag(WorkspaceUpdateOptions.Clean) || bForceClean)
 										{
 											Log.WriteLine("ubt> Running {0} {1} -clean", BuildBat, CommandLine);
-											Utility.ExecuteProcess(CmdExe, null, "/C " + BuildBat + " " + CommandLine + " -clean", null, new ProgressTextWriter(Progress, new PrefixedTextWriter("ubt> ", Log)));
+											Utility.ExecuteProcess(CmdExe, null, "/C \"\"" + BuildBat + "\" " + CommandLine + " -clean\"", null, new ProgressTextWriter(Progress, new PrefixedTextWriter("ubt> ", Log)));
 										}
 
 										Log.WriteLine("ubt> Running {0} {1} -progress", BuildBat, CommandLine);
 
-										int ResultFromBuild = Utility.ExecuteProcess(CmdExe, null, "/C " + BuildBat + " "+ CommandLine + " -progress", null, new ProgressTextWriter(Progress, new PrefixedTextWriter("ubt> ", Log)));
+										int ResultFromBuild = Utility.ExecuteProcess(CmdExe, null, "/C \"\"" + BuildBat + "\" "+ CommandLine + " -progress\"", null, new ProgressTextWriter(Progress, new PrefixedTextWriter("ubt> ", Log)));
 										if(ResultFromBuild != 0)
 										{
 											StepStopwatch.Stop("Failed");
