Merging UGS changes from //UE4/Dev-Build.

#rnx

diff --git a/Engine/Source/Programs/UnrealGameSync/UnrealGameSync/Controls/WorkspaceControl.cs b/Engine/Source/Programs/UnrealGameSync/UnrealGameSync/Controls/WorkspaceControl.cs
--- a/Engine/Source/Programs/UnrealGameSync/UnrealGameSync/Controls/WorkspaceControl.cs
+++ b/Engine/Source/Programs/UnrealGameSync/UnrealGameSync/Controls/WorkspaceControl.cs
@@ -3209,7 +3209,7 @@ namespace UnrealGameSync
 			StringBuilder CommandLine = new StringBuilder();
 			if (Workspace != null && Workspace.Perforce != null)
 			{
-				CommandLine.AppendFormat("-p \"{0}\" -c \"{1}\" -u \"{2}\"", Workspace.Perforce.ServerAndPort ?? "perforce:1666", Workspace.Perforce.ClientName, Workspace.Perforce.UserName);
+				CommandLine.Append(Workspace.Perforce.GetConnectionOptions());
 			}
 			SafeProcessStart("p4v.exe", CommandLine.ToString());
 		}
@@ -4158,7 +4158,7 @@ namespace UnrealGameSync
 			}
 
 			bool bNewHoverSync = false;
-			if (HitTest.Item != null)
+			if (HitTest.Item != null && StatusColumn.Index < HitTest.Item.SubItems.Count)
 			{
 				bNewHoverSync = GetSyncBadgeRectangle(HitTest.Item.SubItems[StatusColumn.Index].Bounds).Contains(e.Location);
 			}
@@ -4495,10 +4495,7 @@ namespace UnrealGameSync
 
 		private void BuildListContextMenu_MoreInfo_Click(object sender, EventArgs e)
 		{
-			if (!Utility.SpawnHiddenProcess("p4vc.exe", String.Format("-p\"{0}\" -u\"{1}\" -c{2} change {3}", Workspace.Perforce.ServerAndPort ?? "perforce:1666", Workspace.Perforce.UserName, Workspace.ClientName, ContextMenuChange.Number)))
-			{
-				MessageBox.Show("Unable to spawn p4vc. Check you have P4V installed.");
-			}
+			Utility.SpawnP4VC(String.Format("{0} change {1}", Workspace.Perforce.GetConnectionOptions(), ContextMenuChange.Number));
 		}
 
 		private void BuildListContextMenu_ViewInSwarm_Click(object sender, EventArgs e)
diff --git a/Engine/Source/Programs/UnrealGameSync/UnrealGameSync/EventMonitor.cs b/Engine/Source/Programs/UnrealGameSync/UnrealGameSync/EventMonitor.cs
--- a/Engine/Source/Programs/UnrealGameSync/UnrealGameSync/EventMonitor.cs
+++ b/Engine/Source/Programs/UnrealGameSync/UnrealGameSync/EventMonitor.cs
@@ -200,6 +200,9 @@ namespace UnrealGameSync
 		public EventData LastStarReview;
 		public List<BadgeData> Badges = new List<BadgeData>();
 		public List<CommentData> Comments = new List<CommentData>();
+
+		public GetMetadataResponseV2 SharedMetadata;
+		public GetMetadataResponseV2 ProjectMetadata;
 	}
 
 	class EventMonitor : IDisposable
@@ -537,6 +540,9 @@ namespace UnrealGameSync
 
 		void ConvertMetadataToEvents(GetMetadataResponseV2 Metadata)
 		{
+			EventSummary NewSummary = new EventSummary();
+			NewSummary.ChangeNumber = Metadata.Change;
+
 			EventSummary Summary;
 			if (ChangeNumberToSummary.TryGetValue(Metadata.Change, out Summary))
 			{
@@ -544,9 +550,35 @@ namespace UnrealGameSync
 				{
 					UserNameToLastSyncEvent.Remove(CurrentUser);
 				}
-				ChangeNumberToSummary.Remove(Metadata.Change);
+
+				NewSummary.SharedMetadata = Summary.SharedMetadata;
+				NewSummary.ProjectMetadata = Summary.ProjectMetadata;
+			}
+
+			if (string.IsNullOrEmpty(Metadata.Project))
+			{
+				NewSummary.SharedMetadata = Metadata;
+			}
+			else
+			{
+				NewSummary.ProjectMetadata = Metadata;
 			}
 
+			ChangeNumberToSummary[NewSummary.ChangeNumber] = NewSummary;
+
+			if (NewSummary.SharedMetadata != null)
+			{
+				PostEvents(NewSummary.SharedMetadata);
+			}
+
+			if (NewSummary.ProjectMetadata != null)
+			{
+				PostEvents(NewSummary.ProjectMetadata);
+			}
+		}
+
+		void PostEvents(GetMetadataResponseV2 Metadata)
+		{
 			if (Metadata.Badges != null)
 			{
 				foreach (GetBadgeDataResponseV2 BadgeData in Metadata.Badges)
diff --git a/Engine/Source/Programs/UnrealGameSync/UnrealGameSync/Forms/IssueDetailsWindow.cs b/Engine/Source/Programs/UnrealGameSync/UnrealGameSync/Forms/IssueDetailsWindow.cs
--- a/Engine/Source/Programs/UnrealGameSync/UnrealGameSync/Forms/IssueDetailsWindow.cs
+++ b/Engine/Source/Programs/UnrealGameSync/UnrealGameSync/Forms/IssueDetailsWindow.cs
@@ -1288,10 +1288,7 @@ namespace UnrealGameSync
 
 		private void BuildListContextMenu_MoreInfo_Click(object sender, EventArgs e)
 		{
-			if(!Utility.SpawnHiddenProcess("p4vc.exe", String.Format("-p\"{0}\" change {1}", Perforce.ServerAndPort, ContextMenuChange.Number)))
-			{
-				MessageBox.Show("Unable to spawn p4vc. Check you have P4V installed.");
-			}
+			Utility.SpawnP4VC(String.Format("{0} change {1}", Perforce.GetConnectionOptions(), ContextMenuChange.Number));
 		}
 
 		private void BuildListContextMenu_Blame_Click(object sender, EventArgs e)
diff --git a/Engine/Source/Programs/UnrealGameSync/UnrealGameSync/Forms/IssueFixedWindow.cs b/Engine/Source/Programs/UnrealGameSync/UnrealGameSync/Forms/IssueFixedWindow.cs
--- a/Engine/Source/Programs/UnrealGameSync/UnrealGameSync/Forms/IssueFixedWindow.cs
+++ b/Engine/Source/Programs/UnrealGameSync/UnrealGameSync/Forms/IssueFixedWindow.cs
@@ -49,6 +49,7 @@ namespace UnrealGameSync
 			{
 				if(Thread != null)
 				{
+					WakeEvent.Set();
 					bStopRequested = true;
 					Thread.Join();
 					Thread = null;
@@ -104,6 +105,7 @@ namespace UnrealGameSync
 
 			this.Perforce = Perforce;
 			this.Worker = new FindChangesWorker(Perforce, PopulateChanges);
+			components.Add(Worker);
 
 			UserNameTextBox.Text = Perforce.UserName;
 			UserNameTextBox.SelectionStart = UserNameTextBox.Text.Length;
@@ -244,10 +246,7 @@ namespace UnrealGameSync
 		private void ChangesContextMenu_MoreInfo_Click(object sender, EventArgs e)
 		{
 			PerforceDescribeRecord Record = (PerforceDescribeRecord)ChangesListContextMenu.Tag;
-			if(!Utility.SpawnHiddenProcess("p4vc.exe", String.Format("-p\"{0}\" change {1}", Perforce.ServerAndPort, Record.ChangeNumber)))
-			{
-				MessageBox.Show("Unable to spawn p4vc. Check you have P4V installed.");
-			}
+			Utility.SpawnP4VC(String.Format("{0} change {1}", Perforce.GetConnectionOptions(), Record.ChangeNumber));
 		}
 
 		private void ChangeNumberTextBox_Enter(object sender, EventArgs e)
diff --git a/Engine/Source/Programs/UnrealGameSync/UnrealGameSync/Forms/SyncFilter.cs b/Engine/Source/Programs/UnrealGameSync/UnrealGameSync/Forms/SyncFilter.cs
--- a/Engine/Source/Programs/UnrealGameSync/UnrealGameSync/Forms/SyncFilter.cs
+++ b/Engine/Source/Programs/UnrealGameSync/UnrealGameSync/Forms/SyncFilter.cs
@@ -93,14 +93,14 @@ namespace UnrealGameSync
 		{
 			Dictionary<Guid, bool> DefaultSyncCategories = WorkspaceSyncCategory.GetDefault(UniqueIdToCategory.Values);
 
-			Dictionary<Guid, bool> GlobalSyncCategories = GetCategorySettings(GlobalControl.CategoriesCheckList);
+			Dictionary<Guid, bool> GlobalSyncCategories = GetCategorySettings(GlobalControl.CategoriesCheckList, this.GlobalSyncCategories);
 			NewGlobalSyncCategories = WorkspaceSyncCategory.GetDelta(DefaultSyncCategories, GlobalSyncCategories);
 
-			Dictionary<Guid, bool> WorkspaceSyncCategories = GetCategorySettings(WorkspaceControl.CategoriesCheckList);
+			Dictionary<Guid, bool> WorkspaceSyncCategories = GetCategorySettings(WorkspaceControl.CategoriesCheckList, this.WorkspaceSyncCategories);
 			NewWorkspaceSyncCategories = WorkspaceSyncCategory.GetDelta(GlobalSyncCategories, WorkspaceSyncCategories);
 		}
 
-		private static Dictionary<Guid, bool> GetCategorySettings(CheckedListBox ListBox)
+		private Dictionary<Guid, bool> GetCategorySettings(CheckedListBox ListBox, IEnumerable<KeyValuePair<Guid, bool>> OriginalSettings)
 		{
 			Dictionary<Guid, bool> Result = new Dictionary<Guid, bool>();
 			for(int Idx = 0; Idx < ListBox.Items.Count; Idx++)
@@ -111,7 +111,13 @@ namespace UnrealGameSync
 					Result[UniqueId] = ListBox.GetItemCheckState(Idx) == CheckState.Checked;
 				}
 			}
-
+			foreach (KeyValuePair<Guid, bool> OriginalSetting in OriginalSettings)
+			{
+				if (!UniqueIdToCategory.ContainsKey(OriginalSetting.Key))
+				{
+					Result[OriginalSetting.Key] = OriginalSetting.Value;
+				}
+			}
 			return Result;
 		}
 
diff --git a/Engine/Source/Programs/UnrealGameSync/UnrealGameSync/Perforce.cs b/Engine/Source/Programs/UnrealGameSync/UnrealGameSync/Perforce.cs
--- a/Engine/Source/Programs/UnrealGameSync/UnrealGameSync/Perforce.cs
+++ b/Engine/Source/Programs/UnrealGameSync/UnrealGameSync/Perforce.cs
@@ -883,11 +883,11 @@ namespace UnrealGameSync
 
 			if (MaxResults > 0)
 			{
-				Arguments.Add(String.Format("-m {0}", MaxResults));
+				Arguments.Add(String.Format("-m{0}", MaxResults));
 			}
 			if (ByUser != null)
 			{
-				Arguments.Add(String.Format("-u \"{0}\"", ByUser));
+				Arguments.Add(String.Format("-u{0}", ByUser));
 			}
 			foreach (string Filter in Filters)
 			{
@@ -1694,22 +1694,32 @@ namespace UnrealGameSync
 			return FullCommandLine;
 		}
 
-		private string GetFullCommandLine(string CommandLine, CommandOptions Options)
+		public string GetConnectionOptions(bool bIncludeClient = true)
 		{
-			StringBuilder FullCommandLine = new StringBuilder();
-			if(ServerAndPort != null)
+			StringBuilder BasicCommandArgs = new StringBuilder();
+			if (ServerAndPort != null)
 			{
-				FullCommandLine.AppendFormat("-p{0} ", ServerAndPort);
+				BasicCommandArgs.AppendFormat("-p \"{0}\" ", ServerAndPort);
 			}
-			if(UserName != null)
+			if (UserName != null)
 			{
-				FullCommandLine.AppendFormat("-u{0} ", UserName);
+				BasicCommandArgs.AppendFormat("-u \"{0}\" ", UserName);
 			}
-			if(!Options.HasFlag(CommandOptions.NoClient) && ClientName != null)
+			if (bIncludeClient && ClientName != null)
 			{
-				FullCommandLine.AppendFormat("-c{0} ", ClientName);
+				BasicCommandArgs.AppendFormat("-c \"{0}\" ", ClientName);
 			}
-			if(!Options.HasFlag(CommandOptions.NoChannels))
+
+			return BasicCommandArgs.ToString().TrimEnd();
+		}
+
+		private string GetFullCommandLine(string CommandLine, CommandOptions Options)
+		{
+			StringBuilder FullCommandLine = new StringBuilder();
+			FullCommandLine.Append(GetConnectionOptions(bIncludeClient: !Options.HasFlag(CommandOptions.NoClient)));
+			FullCommandLine.Append(" ");
+
+			if (!Options.HasFlag(CommandOptions.NoChannels))
 			{
 				FullCommandLine.Append("-s ");
 			}
diff --git a/Engine/Source/Programs/UnrealGameSync/UnrealGameSync/Properties/AssemblyInfo.cs b/Engine/Source/Programs/UnrealGameSync/UnrealGameSync/Properties/AssemblyInfo.cs
--- a/Engine/Source/Programs/UnrealGameSync/UnrealGameSync/Properties/AssemblyInfo.cs
+++ b/Engine/Source/Programs/UnrealGameSync/UnrealGameSync/Properties/AssemblyInfo.cs
@@ -34,5 +34,5 @@ using System.Runtime.InteropServices;
 // You can specify all the values or you can default the Build and Revision Numbers 
 // by using the '*' as shown below:
 // [assembly: AssemblyVersion("1.0.*")]
-[assembly: AssemblyVersion("1.234.0.0")]
-[assembly: AssemblyFileVersion("1.234.0.0")]
+[assembly: AssemblyVersion("1.236.0.0")]
+[assembly: AssemblyFileVersion("1.236.0.0")]
diff --git a/Engine/Source/Programs/UnrealGameSync/UnrealGameSync/UriHandlers/P4Handler.cs b/Engine/Source/Programs/UnrealGameSync/UnrealGameSync/UriHandlers/P4Handler.cs
--- a/Engine/Source/Programs/UnrealGameSync/UnrealGameSync/UriHandlers/P4Handler.cs
+++ b/Engine/Source/Programs/UnrealGameSync/UnrealGameSync/UriHandlers/P4Handler.cs
@@ -15,7 +15,7 @@ namespace UnrealGameSync
 
 			if (!Utility.SpawnHiddenProcess("p4v.exe", CommandLine))
 			{
-				return new UriResult() { Error = string.Format("Error spawning p4vc.exe with command line: {0}", CommandLine) };
+				return new UriResult() { Error = string.Format("Error spawning p4v.exe with command line: {0}", CommandLine) };
 			}
 
 			return new UriResult() { Success = true };
@@ -26,10 +26,7 @@ namespace UnrealGameSync
 		{
 			string CommandLine = string.Format("timelapse {0}{1}", Line == -1 ? "" : string.Format(" -l {0} ", Line), DepotPath);
 
-			if (!Utility.SpawnHiddenProcess("p4vc.exe", CommandLine))
-			{
-				return new UriResult() { Error = string.Format("Error spawning p4vc.exe with command line: {0}", CommandLine) };
-			}
+			Utility.SpawnP4VC(CommandLine);
 
 			return new UriResult() { Success = true };
 		}		
diff --git a/Engine/Source/Programs/UnrealGameSync/UnrealGameSync/Utility.cs b/Engine/Source/Programs/UnrealGameSync/UnrealGameSync/Utility.cs
--- a/Engine/Source/Programs/UnrealGameSync/UnrealGameSync/Utility.cs
+++ b/Engine/Source/Programs/UnrealGameSync/UnrealGameSync/Utility.cs
@@ -87,7 +87,14 @@ namespace UnrealGameSync
 				ChildProcess.StartInfo.RedirectStandardOutput = true;
 				ChildProcess.StartInfo.RedirectStandardError = true;
 				ChildProcess.StartInfo.CreateNoWindow = true;
-				return ChildProcess.Start();
+				try
+				{
+					return ChildProcess.Start();
+				}
+				catch
+				{
+					return false;
+				}
 			}
 		}
 
@@ -613,5 +620,54 @@ namespace UnrealGameSync
 			}
 			return NormalUserName.ToString();
 		}
+
+		public static IEnumerable<string> GetPerforcePaths()
+		{
+			string PathList = Environment.GetEnvironmentVariable("PATH");
+			if (!String.IsNullOrEmpty(PathList))
+			{
+				foreach (string PathEntry in PathList.Split(Path.PathSeparator))
+				{
+					string PerforcePath = null;
+					try
+					{
+						string TestPerforcePath = Path.Combine(PathEntry, "p4.exe");
+						if (File.Exists(TestPerforcePath))
+						{
+							PerforcePath = TestPerforcePath;
+						}
+					}
+					catch
+					{
+					}
+
+					if(PerforcePath != null)
+					{
+						yield return PerforcePath;
+					}
+				}
+			}
+		}
+
+		public static void SpawnP4VC(string Arguments)
+		{
+			string Executable = "p4vc.exe";
+
+			foreach (string PerforcePath in GetPerforcePaths())
+			{
+				string PerforceDir = Path.GetDirectoryName(PerforcePath);
+				if (File.Exists(Path.Combine(PerforceDir, "p4vc.bat")) && !File.Exists(Path.Combine(PerforceDir, "p4vc.exe")))
+				{
+					Executable = Path.Combine(PerforceDir, "p4v.exe");
+					Arguments = "-p4vc " + Arguments;
+					break;
+				}
+			}
+
+			if (!Utility.SpawnHiddenProcess(Executable, Arguments))
+			{
+				MessageBox.Show("Unable to spawn p4vc. Check you have P4V installed.");
+			}
+		}
 	}
 }
diff --git a/Engine/Source/Programs/UnrealGameSync/UnrealGameSync/Workspace.cs b/Engine/Source/Programs/UnrealGameSync/UnrealGameSync/Workspace.cs
--- a/Engine/Source/Programs/UnrealGameSync/UnrealGameSync/Workspace.cs
+++ b/Engine/Source/Programs/UnrealGameSync/UnrealGameSync/Workspace.cs
@@ -655,7 +655,13 @@ namespace UnrealGameSync
 								string RelativePath = FullName.Substring(LocalRootPrefix.Length).Replace('\\', '/');
 								if (Filter.Matches(RelativePath))
 								{
-									SyncTree.IncludeFile(PerforceUtils.EscapePath(RelativePath), SyncRecord.FileSize);
+									long FileSize = SyncRecord.FileSize;
+									if (SyncRecord.Action == "deleted" || SyncRecord.Action == "move/delete")
+									{
+										FileSize = 0;
+									}
+
+									SyncTree.IncludeFile(PerforceUtils.EscapePath(RelativePath), FileSize);
 									SyncDepotPaths.Add(SyncRecord.DepotPath);
 								}
 								else
