UGS: NET Core fixes.

- Fixed missing icon on Options/Filter buttons.
- Fixed errorrs reading responses from server, due to STJ case sensitivity defaults.

diff --git a/Engine/Source/Programs/UnrealGameSync/UnrealGameSync/Controls/WorkspaceControl.cs b/Engine/Source/Programs/UnrealGameSync/UnrealGameSync/Controls/WorkspaceControl.cs
--- a/Engine/Source/Programs/UnrealGameSync/UnrealGameSync/Controls/WorkspaceControl.cs
+++ b/Engine/Source/Programs/UnrealGameSync/UnrealGameSync/Controls/WorkspaceControl.cs
@@ -2988,7 +2988,7 @@ namespace UnrealGameSync
 					try
 					{
 						string Text = File.ReadAllText(ReceiptFileName);
-						ReceiptJsonObject Receipt = JsonSerializer.Deserialize<ReceiptJsonObject>(Text);
+						ReceiptJsonObject Receipt = JsonSerializer.Deserialize<ReceiptJsonObject>(Text, Program.DefaultJsonSerializerOptions);
 
 						string LaunchFileName = Receipt.Launch;
 						if (LaunchFileName != null)
diff --git a/Engine/Source/Programs/UnrealGameSync/UnrealGameSync/Program.cs b/Engine/Source/Programs/UnrealGameSync/UnrealGameSync/Program.cs
--- a/Engine/Source/Programs/UnrealGameSync/UnrealGameSync/Program.cs
+++ b/Engine/Source/Programs/UnrealGameSync/UnrealGameSync/Program.cs
@@ -10,6 +10,8 @@ using System.Net;
 using System.Reflection;
 using System.Runtime.InteropServices;
 using System.Text;
+using System.Text.Json;
+using System.Text.Json.Serialization;
 using System.Threading;
 using System.Threading.Tasks;
 using System.Windows.Forms;
@@ -18,6 +20,16 @@ namespace UnrealGameSync
 {
 	static class Program
 	{
+		static JsonSerializerOptions GetDefaultJsonSerializerOptions()
+		{
+			JsonSerializerOptions Options = new JsonSerializerOptions();
+			Options.PropertyNameCaseInsensitive = true;
+			Options.Converters.Add(new JsonStringEnumConverter());
+			return Options;
+		}
+
+		public static JsonSerializerOptions DefaultJsonSerializerOptions { get; } = GetDefaultJsonSerializerOptions();
+
 		public static string SyncVersion = null;
 
 		[STAThread]
diff --git a/Engine/Source/Programs/UnrealGameSync/UnrealGameSync/RESTApi.cs b/Engine/Source/Programs/UnrealGameSync/UnrealGameSync/RESTApi.cs
--- a/Engine/Source/Programs/UnrealGameSync/UnrealGameSync/RESTApi.cs
+++ b/Engine/Source/Programs/UnrealGameSync/UnrealGameSync/RESTApi.cs
@@ -88,7 +88,7 @@ namespace UnrealGameSync
 
 		public static T GET<T>(string URI, string Resource, params string[] QueryParams)
 		{
-			return JsonSerializer.Deserialize<T>(SendRequestInternal(URI, Resource, "GET", null, QueryParams));
+			return JsonSerializer.Deserialize<T>(SendRequestInternal(URI, Resource, "GET", null, QueryParams), Program.DefaultJsonSerializerOptions);
 		}
 
 		public static string PUT<T>(string URI, string Resource, T Object, params string[] QueryParams)
diff --git a/Engine/Source/Programs/UnrealGameSync/UnrealGameSync/Utility.cs b/Engine/Source/Programs/UnrealGameSync/UnrealGameSync/Utility.cs
--- a/Engine/Source/Programs/UnrealGameSync/UnrealGameSync/Utility.cs
+++ b/Engine/Source/Programs/UnrealGameSync/UnrealGameSync/Utility.cs
@@ -11,6 +11,7 @@ using System.IO;
 using System.Linq;	
 using System.Text;
 using System.Text.Json;
+using System.Text.Json.Serialization;
 using System.Text.RegularExpressions;
 using System.Windows.Forms;
 
@@ -260,7 +261,11 @@ namespace UnrealGameSync
 		{
 			try
 			{
-				ProjectJson Project = JsonSerializer.Deserialize<ProjectJson>(Text);
+				JsonSerializerOptions Options = new JsonSerializerOptions();
+				Options.PropertyNameCaseInsensitive = true;
+				Options.Converters.Add(new JsonStringEnumConverter());
+
+				ProjectJson Project = JsonSerializer.Deserialize<ProjectJson>(Text, Options);
 
 				return Project.Enterprise;
 			}
diff --git a/Engine/Source/Programs/UnrealGameSync/UnrealGameSync/Workspace.cs b/Engine/Source/Programs/UnrealGameSync/UnrealGameSync/Workspace.cs
--- a/Engine/Source/Programs/UnrealGameSync/UnrealGameSync/Workspace.cs
+++ b/Engine/Source/Programs/UnrealGameSync/UnrealGameSync/Workspace.cs
@@ -1593,7 +1593,7 @@ namespace UnrealGameSync
 
 		string UpdateBuildVersion(string Text, int Changelist, int CodeChangelist, string BranchOrStreamName, bool bIsLicenseeVersion)
 		{
-			Dictionary<string, object> Object = JsonSerializer.Deserialize<Dictionary<string, object>>(Text);
+			Dictionary<string, object> Object = JsonSerializer.Deserialize<Dictionary<string, object>>(Text, Program.DefaultJsonSerializerOptions);
 
 			object PrevCompatibleChangelistObj;
 			int PrevCompatibleChangelist = Object.TryGetValue("CompatibleChangelist", out PrevCompatibleChangelistObj) ? (int)Convert.ChangeType(PrevCompatibleChangelistObj.ToString(), typeof(int)) : 0;
