UGS: Include a tag in the version number indicating that it's the NET Core version.

#rnx

diff --git a/Engine/Source/Programs/UnrealGameSync/UnrealGameSync/Controls/WorkspaceControl.cs b/Engine/Source/Programs/UnrealGameSync/UnrealGameSync/Controls/WorkspaceControl.cs
--- a/Engine/Source/Programs/UnrealGameSync/UnrealGameSync/Controls/WorkspaceControl.cs
+++ b/Engine/Source/Programs/UnrealGameSync/UnrealGameSync/Controls/WorkspaceControl.cs
@@ -5218,7 +5218,7 @@ namespace UnrealGameSync
 		private void OptionsContextMenu_Diagnostics_Click(object sender, EventArgs e)
 		{
 			StringBuilder DiagnosticsText = new StringBuilder();
-			DiagnosticsText.AppendFormat("Application version: {0}\n", Assembly.GetExecutingAssembly().GetName().Version);
+			DiagnosticsText.AppendFormat("Application version: {0}\n", Program.GetVersionString());
 			DiagnosticsText.AppendFormat("Synced from: {0}\n", Program.SyncVersion ?? "(unknown)");
 			DiagnosticsText.AppendFormat("Selected file: {0}\n", (SelectedFileName == null) ? "(none)" : SelectedFileName);
 			if (Workspace != null)
diff --git a/Engine/Source/Programs/UnrealGameSync/UnrealGameSync/Forms/MainWindow.cs b/Engine/Source/Programs/UnrealGameSync/UnrealGameSync/Forms/MainWindow.cs
--- a/Engine/Source/Programs/UnrealGameSync/UnrealGameSync/Forms/MainWindow.cs
+++ b/Engine/Source/Programs/UnrealGameSync/UnrealGameSync/Forms/MainWindow.cs
@@ -177,7 +177,7 @@ namespace UnrealGameSync
 
 			if(bUnstable)
 			{
-				Text += String.Format(" (UNSTABLE BUILD {0})", Assembly.GetExecutingAssembly().GetName().Version);
+				Text += $"{Program.GetVersionString()} (UNSTABLE)";
 			}
 
 			AutomationLog = new TimestampLogWriter(new BoundedLogWriter(Path.Combine(DataFolder, "Automation.log")));
diff --git a/Engine/Source/Programs/UnrealGameSync/UnrealGameSync/Program.cs b/Engine/Source/Programs/UnrealGameSync/UnrealGameSync/Program.cs
--- a/Engine/Source/Programs/UnrealGameSync/UnrealGameSync/Program.cs
+++ b/Engine/Source/Programs/UnrealGameSync/UnrealGameSync/Program.cs
@@ -20,6 +20,12 @@ namespace UnrealGameSync
 {
 	static class Program
 	{
+		public static string GetVersionString()
+		{
+			Version Version = Assembly.GetExecutingAssembly().GetName().Version;
+			return $"{Version.Major}.{Version.Minor} (Core)";
+		}
+
 		static JsonSerializerOptions GetDefaultJsonSerializerOptions()
 		{
 			JsonSerializerOptions Options = new JsonSerializerOptions();
