UGS: Missing space in window title.

#rnx

diff --git a/Engine/Source/Programs/UnrealGameSync/UnrealGameSync/Program.cs b/Engine/Source/Programs/UnrealGameSync/UnrealGameSync/Program.cs
--- a/Engine/Source/Programs/UnrealGameSync/UnrealGameSync/Program.cs
+++ b/Engine/Source/Programs/UnrealGameSync/UnrealGameSync/Program.cs
@@ -23,7 +23,7 @@ namespace UnrealGameSync
 		public static string GetVersionString()
 		{
 			Version Version = Assembly.GetExecutingAssembly().GetName().Version;
-			return $"{Version.Major}.{Version.Minor} (Core)";
+			return $"{Version.Major}.{Version.Minor}-Core";
 		}
 
 		static JsonSerializerOptions GetDefaultJsonSerializerOptions()
