UGS: Fix rendering on custom build steps window.

diff --git a/Engine/Source/Programs/UnrealGameSync/UnrealGameSync/Forms/ModifyBuildStepsWindow.cs b/Engine/Source/Programs/UnrealGameSync/UnrealGameSync/Forms/ModifyBuildStepsWindow.cs
--- a/Engine/Source/Programs/UnrealGameSync/UnrealGameSync/Forms/ModifyBuildStepsWindow.cs
+++ b/Engine/Source/Programs/UnrealGameSync/UnrealGameSync/Forms/ModifyBuildStepsWindow.cs
@@ -157,7 +157,8 @@ namespace UnrealGameSync
 					bEnabled = Task.bShowAsTool;
 				}
 
-				e.Graphics.FillRectangle(e.ItemState.HasFlag(ListViewItemStates.Selected)? SystemBrushes.Highlight : SystemBrushes.Window, e.Bounds);
+				bool Selected = BuildStepList.SelectedItems.Contains(e.Item);
+				e.Graphics.FillRectangle(Selected? SystemBrushes.Highlight : SystemBrushes.Window, e.Bounds);
 
 				CheckBoxState State;
 				if(bEnabled)
