UGS: Catch ThreadInterruptedExceptions when cancelling workspace updates.

diff --git a/Engine/Source/Programs/UnrealGameSync/UnrealGameSync/Workspace.cs b/Engine/Source/Programs/UnrealGameSync/UnrealGameSync/Workspace.cs
--- a/Engine/Source/Programs/UnrealGameSync/UnrealGameSync/Workspace.cs
+++ b/Engine/Source/Programs/UnrealGameSync/UnrealGameSync/Workspace.cs
@@ -383,12 +383,17 @@ namespace UnrealGameSync
 					Log.WriteLine("{0}", StatusMessage);
 				}
 			}
-			catch(ThreadAbortException)
+			catch (ThreadAbortException)
 			{
 				StatusMessage = "Canceled.";
 				Log.WriteLine("Canceled.");
 			}
-			catch(Exception Ex)
+			catch (ThreadInterruptedException)
+			{
+				StatusMessage = "Canceled.";
+				Log.WriteLine("Canceled.");
+			}
+			catch (Exception Ex)
 			{
 				StatusMessage = "Failed with exception - " + Ex.ToString();
 				Log.WriteException(Ex, "Failed with exception");
