UGS: Use built-in NET zip archive class to avoid odd encoding errors with Ionic.

diff --git a/Engine/Source/Programs/UnrealGameSync/UnrealGameSync/Forms/DiagnosticsWindow.cs b/Engine/Source/Programs/UnrealGameSync/UnrealGameSync/Forms/DiagnosticsWindow.cs
--- a/Engine/Source/Programs/UnrealGameSync/UnrealGameSync/Forms/DiagnosticsWindow.cs
+++ b/Engine/Source/Programs/UnrealGameSync/UnrealGameSync/Forms/DiagnosticsWindow.cs
@@ -5,6 +5,7 @@ using System.Diagnostics;
 using System.Windows.Forms;
 using Ionic.Zip;
 using System.IO;
+using System.IO.Compression;
 
 namespace UnrealGameSync
 {
@@ -46,15 +47,23 @@ namespace UnrealGameSync
 				string ZipFileName = Dialog.FileName;
 				try
 				{
-					ZipFile Zip = new ZipFile();
-					foreach(string FileName in Directory.EnumerateFiles(DataFolder))
+					using (ZipArchive Zip = new ZipArchive(File.OpenWrite(ZipFileName), ZipArchiveMode.Create))
 					{
-						if(!FileName.EndsWith(".exe", StringComparison.InvariantCultureIgnoreCase) && !FileName.EndsWith(".dll", StringComparison.InvariantCultureIgnoreCase))
+						foreach (string FileName in Directory.EnumerateFiles(DataFolder))
 						{
-							Zip.AddFile(FileName, "");
+							if (!FileName.EndsWith(".exe", StringComparison.InvariantCultureIgnoreCase) && !FileName.EndsWith(".dll", StringComparison.InvariantCultureIgnoreCase))
+							{
+								using (FileStream InputStream = File.Open(FileName, FileMode.Open, FileAccess.Read, FileShare.ReadWrite))
+								{
+									ZipArchiveEntry Entry = Zip.CreateEntry(Path.GetFileName(FileName));
+									using (Stream OutputStream = Entry.Open())
+									{
+										InputStream.CopyTo(OutputStream);
+									}
+								}
+							}
 						}
 					}
-					Zip.Save(ZipFileName);
 				}
 				catch(Exception Ex)
 				{
