UGS: Fix solution not opening when clicking Visual Studio link in status bar.

diff --git a/Engine/Source/Programs/UnrealGameSync/UnrealGameSync/Controls/WorkspaceControl.cs b/Engine/Source/Programs/UnrealGameSync/UnrealGameSync/Controls/WorkspaceControl.cs
--- a/Engine/Source/Programs/UnrealGameSync/UnrealGameSync/Controls/WorkspaceControl.cs
+++ b/Engine/Source/Programs/UnrealGameSync/UnrealGameSync/Controls/WorkspaceControl.cs
@@ -4765,6 +4765,7 @@ namespace UnrealGameSync
 			else
 			{
 				ProcessStartInfo StartInfo = new ProcessStartInfo(SolutionFileName);
+				StartInfo.UseShellExecute = true;
 				StartInfo.WorkingDirectory = BranchDirectoryName;
 				SafeProcessStart(StartInfo);
 			}
