UGS: Fix elevation of UGS process when updating registry settings. Executing assembly is different to launch assembly in NET Core.

diff --git a/Engine/Source/Programs/UnrealGameSync/UnrealGameSync/UriHandlers/UriHandler.cs b/Engine/Source/Programs/UnrealGameSync/UnrealGameSync/UriHandlers/UriHandler.cs
--- a/Engine/Source/Programs/UnrealGameSync/UnrealGameSync/UriHandlers/UriHandler.cs
+++ b/Engine/Source/Programs/UnrealGameSync/UnrealGameSync/UriHandlers/UriHandler.cs
@@ -11,6 +11,7 @@ using System.Linq;
 using System.Diagnostics;
 using System.Windows.Forms;
 using System.Configuration;
+using System.IO;
 
 namespace UnrealGameSync
 {
@@ -392,7 +393,7 @@ namespace UnrealGameSync
 		{
 			using (Process Process = new Process())
 			{
-				Process.StartInfo.FileName = Assembly.GetExecutingAssembly().Location;
+				Process.StartInfo.FileName = Path.ChangeExtension(Assembly.GetExecutingAssembly().Location, ".exe");
 				Process.StartInfo.Arguments = Arguments;
 				Process.StartInfo.Verb = "runas";
 				Process.StartInfo.UseShellExecute = true;
