UGS: Fix URLs not opening correctly on NET Core.

diff --git a/Engine/Source/Programs/UnrealGameSync/UnrealGameSync/Controls/WorkspaceControl.cs b/Engine/Source/Programs/UnrealGameSync/UnrealGameSync/Controls/WorkspaceControl.cs
--- a/Engine/Source/Programs/UnrealGameSync/UnrealGameSync/Controls/WorkspaceControl.cs
+++ b/Engine/Source/Programs/UnrealGameSync/UnrealGameSync/Controls/WorkspaceControl.cs
@@ -2462,7 +2462,10 @@ namespace UnrealGameSync
 		{
 			try
 			{
-				Process.Start(Url);
+				ProcessStartInfo StartInfo = new ProcessStartInfo();
+				StartInfo.FileName = Url;
+				StartInfo.UseShellExecute = true;
+				using Process _ = Process.Start(StartInfo);
 			}
 			catch
 			{
@@ -2474,7 +2477,11 @@ namespace UnrealGameSync
 		{
 			try
 			{
-				Process.Start(Url, Arguments);
+				ProcessStartInfo StartInfo = new ProcessStartInfo();
+				StartInfo.FileName = Url;
+				StartInfo.Arguments = Arguments;
+				StartInfo.UseShellExecute = true;
+				using Process _ = Process.Start(StartInfo);
 			}
 			catch
 			{
@@ -2490,7 +2497,8 @@ namespace UnrealGameSync
 				StartInfo.FileName = Url;
 				StartInfo.Arguments = Arguments;
 				StartInfo.WorkingDirectory = WorkingDir;
-				Process.Start(StartInfo).Dispose();
+				StartInfo.UseShellExecute = true;
+				using Process _ = Process.Start(StartInfo);
 			}
 			catch
 			{
