UGS: Fix various places where URLs would not open correctly under NET Core.

diff --git a/Engine/Source/Programs/UnrealGameSync/UnrealGameSync/Controls/WorkspaceControl.cs b/Engine/Source/Programs/UnrealGameSync/UnrealGameSync/Controls/WorkspaceControl.cs
--- a/Engine/Source/Programs/UnrealGameSync/UnrealGameSync/Controls/WorkspaceControl.cs
+++ b/Engine/Source/Programs/UnrealGameSync/UnrealGameSync/Controls/WorkspaceControl.cs
@@ -4511,7 +4511,7 @@ namespace UnrealGameSync
 			string SwarmURL;
 			if (TryGetProjectSetting(PerforceMonitor.LatestProjectConfigFile, "SwarmURL", out SwarmURL))
 			{
-				Process.Start(String.Format("{0}/changes/{1}", SwarmURL, ContextMenuChange.Number));
+				Utility.OpenUrl(String.Format("{0}/changes/{1}", SwarmURL, ContextMenuChange.Number));
 			}
 			else
 			{
diff --git a/Engine/Source/Programs/UnrealGameSync/UnrealGameSync/Forms/IssueAlertWindow.cs b/Engine/Source/Programs/UnrealGameSync/UnrealGameSync/Forms/IssueAlertWindow.cs
--- a/Engine/Source/Programs/UnrealGameSync/UnrealGameSync/Forms/IssueAlertWindow.cs
+++ b/Engine/Source/Programs/UnrealGameSync/UnrealGameSync/Forms/IssueAlertWindow.cs
@@ -252,7 +252,7 @@ namespace UnrealGameSync
 			{
 				try
 				{
-					System.Diagnostics.Process.Start(Url);
+					Utility.OpenUrl(Url);
 				}
 				catch (Exception Ex)
 				{
diff --git a/Engine/Source/Programs/UnrealGameSync/UnrealGameSync/Forms/IssueDetailsWindow.cs b/Engine/Source/Programs/UnrealGameSync/UnrealGameSync/Forms/IssueDetailsWindow.cs
--- a/Engine/Source/Programs/UnrealGameSync/UnrealGameSync/Forms/IssueDetailsWindow.cs
+++ b/Engine/Source/Programs/UnrealGameSync/UnrealGameSync/Forms/IssueDetailsWindow.cs
@@ -3,6 +3,7 @@
 using System;
 using System.Collections.Generic;
 using System.Data;
+using System.Diagnostics;
 using System.Drawing;
 using System.Drawing.Drawing2D;
 using System.IO;
@@ -805,7 +806,7 @@ namespace UnrealGameSync
 				StatusLineListViewWidget BuildWidget = new StatusLineListViewWidget(BuildItem, StatusElementResources);
 				BuildWidget.HorizontalAlignment = HorizontalAlignment.Left;
 
-				BuildWidget.Line.AddLink(Range.BuildGroup.JobName, FontStyle.Underline, () => System.Diagnostics.Process.Start(Range.BuildGroup.JobUrl));
+				BuildWidget.Line.AddLink(Range.BuildGroup.JobName, FontStyle.Underline, () => Utility.OpenUrl(Range.BuildGroup.JobUrl));
 				BuildItem.SubItems.Add(new ListViewItem.ListViewSubItem(BuildItem, ""){ Tag = BuildWidget });
 				BuildItem.SubItems.Add(new ListViewItem.ListViewSubItem(BuildItem, ""){ Tag = BuildWidget });
 
@@ -849,7 +850,7 @@ namespace UnrealGameSync
 			foreach (IssueBuildData Build in BuildGroup.Builds.OrderBy(x => x.JobStepName))
 			{
 				ToolStripMenuItem MenuItem = new ToolStripMenuItem(String.Format("View Step: {0}", Build.JobStepName));
-				MenuItem.Click += (S, E) => System.Diagnostics.Process.Start(Build.JobStepUrl);
+				MenuItem.Click += (S, E) => Utility.OpenUrl(Build.JobStepUrl);
 				JobContextMenu.Items.Insert(MinIndex++, MenuItem);
 			}
 
@@ -1372,7 +1373,7 @@ namespace UnrealGameSync
 			IssueBuildData LastBuild = IssueBuilds.Where(x => x.Stream == SelectedStream).OrderByDescending(x => x.Change).ThenByDescending(x => x.ErrorUrl).FirstOrDefault();
 			if(LastBuild != null)
 			{
-				System.Diagnostics.Process.Start(LastBuild.ErrorUrl);
+				Utility.OpenUrl(LastBuild.ErrorUrl);
 			}
 		}
 
@@ -1382,7 +1383,7 @@ namespace UnrealGameSync
 			{
 				if(Build.ErrorUrl != null)
 				{
-					System.Diagnostics.Process.Start(Build.ErrorUrl);
+					Utility.OpenUrl(Build.ErrorUrl);
 					break;
 				}
 			}
@@ -1390,7 +1391,7 @@ namespace UnrealGameSync
 
 		private void JobContextMenu_ViewJob_Click(object sender, EventArgs e)
 		{
-			System.Diagnostics.Process.Start(SelectedBuildGroup.JobUrl);
+			Utility.OpenUrl(SelectedBuildGroup.JobUrl);
 		}
 
 		private void BuildListView_MouseUp(object sender, MouseEventArgs e)
@@ -1414,13 +1415,13 @@ namespace UnrealGameSync
 			IssueBuildData Build = IssueBuilds.FirstOrDefault(x => x.ErrorUrl != null);
 			if (Build != null)
 			{
-				System.Diagnostics.Process.Start(Build.ErrorUrl);
+				Utility.OpenUrl(Build.ErrorUrl);
 			}
 		}
 
 		private void DetailsTextBox_LinkClicked(object sender, LinkClickedEventArgs e)
 		{
-			System.Diagnostics.Process.Start(e.LinkText);
+			Utility.OpenUrl(e.LinkText);
 		}
 	}
 }
diff --git a/Engine/Source/Programs/UnrealGameSync/UnrealGameSync/Forms/NetCoreWindow.cs b/Engine/Source/Programs/UnrealGameSync/UnrealGameSync/Forms/NetCoreWindow.cs
--- a/Engine/Source/Programs/UnrealGameSync/UnrealGameSync/Forms/NetCoreWindow.cs
+++ b/Engine/Source/Programs/UnrealGameSync/UnrealGameSync/Forms/NetCoreWindow.cs
@@ -38,7 +38,7 @@ namespace UnrealGameSync.Forms
 
 		private void linkLabel1_LinkClicked(object sender, LinkLabelLinkClickedEventArgs e)
 		{
-			System.Diagnostics.Process.Start("https://dotnet.microsoft.com/download/dotnet-core/current/runtime");
+			Utility.OpenUrl("https://dotnet.microsoft.com/download/dotnet-core/current/runtime");
 		}
 	}
 }
diff --git a/Engine/Source/Programs/UnrealGameSync/UnrealGameSync/Utility.cs b/Engine/Source/Programs/UnrealGameSync/UnrealGameSync/Utility.cs
--- a/Engine/Source/Programs/UnrealGameSync/UnrealGameSync/Utility.cs
+++ b/Engine/Source/Programs/UnrealGameSync/UnrealGameSync/Utility.cs
@@ -674,5 +674,13 @@ namespace UnrealGameSync
 				MessageBox.Show("Unable to spawn p4vc. Check you have P4V installed.");
 			}
 		}
+
+		public static void OpenUrl(string Url)
+		{
+			ProcessStartInfo StartInfo = new ProcessStartInfo();
+			StartInfo.FileName = Url;
+			StartInfo.UseShellExecute = true;
+			using Process _ = Process.Start(StartInfo);
+		}
 	}
 }
