UGS: Merging custom tool changes from UE4.

#rnx

diff --git a/Engine/Source/Programs/UnrealGameSync/UnrealGameSync/Controls/WorkspaceControl.cs b/Engine/Source/Programs/UnrealGameSync/UnrealGameSync/Controls/WorkspaceControl.cs
--- a/Engine/Source/Programs/UnrealGameSync/UnrealGameSync/Controls/WorkspaceControl.cs
+++ b/Engine/Source/Programs/UnrealGameSync/UnrealGameSync/Controls/WorkspaceControl.cs
@@ -27,6 +27,8 @@ namespace UnrealGameSync
 
 	interface IWorkspaceControlOwner
 	{
+		ToolUpdateMonitor ToolUpdateMonitor { get; }
+
 		void EditSelectedProject(WorkspaceControl Workspace);
 		void RequestProjectChange(WorkspaceControl Workspace, UserSelectedProjectSettings Project, bool bModal);
 		void ShowAndActivate();
@@ -442,6 +444,8 @@ namespace UnrealGameSync
 
 				TintColor = GetTintColor();
 			}
+
+			Owner.ToolUpdateMonitor.OnChange += UpdateStatusPanel_CrossThread;
 		}
 
 		public void IssueMonitor_OnIssuesChangedAsync()
@@ -697,6 +701,8 @@ namespace UnrealGameSync
 		/// <param name="disposing">true if managed resources should be disposed; otherwise, false.</param>
 		protected override void Dispose(bool disposing)
 		{
+			Owner.ToolUpdateMonitor.OnChange -= UpdateStatusPanel_CrossThread;
+
 			bIsDisposing = true;
 
 			if (disposing && (components != null))
@@ -3222,16 +3228,24 @@ namespace UnrealGameSync
 			SafeProcessStart("p4v.exe", CommandLine.ToString());
 		}
 
-		private void OpenUshell()
+
+		void RunTool(ToolDefinition Tool, ToolLink Link)
 		{
-			StringBuilder CommandLine = new StringBuilder();
-			CommandLine.AppendFormat("/C \"\"{0}\"", Path.Combine(DataFolder, "Tools", "Ushell", "Current", "ushell.bat"));
-			if (SelectedFileName.EndsWith(".uproject"))
-			{
-				CommandLine.AppendFormat(" --project=\"{0}\"", SelectedFileName);
-			}
-			CommandLine.Append("\"");
-			SafeProcessStart(Environment.GetEnvironmentVariable("COMSPEC"), CommandLine.ToString(), BranchDirectoryName);
+			string ToolDir = Owner.ToolUpdateMonitor.GetToolPath(Tool.Name);
+
+			Dictionary<string, string> Variables = GetWorkspaceVariables(Workspace.CurrentChangeNumber);
+			Variables["ToolDir"] = ToolDir;
+
+			string FileName = Utility.ExpandVariables(Link.FileName, Variables);
+			string Arguments = Utility.ExpandVariables(Link.Arguments, Variables);
+			string WorkingDir = Utility.ExpandVariables(Link.WorkingDir ?? ToolDir, Variables);
+
+			SafeProcessStart(FileName, Arguments, WorkingDir);
+		}
+
+		private void UpdateStatusPanel_CrossThread()
+		{
+			MainThreadSynchronizationContext.Post((o) => { if (!IsDisposed) { UpdateStatusPanel(); } }, null);
 		}
 
 		private void UpdateStatusPanel()
@@ -3366,11 +3380,20 @@ namespace UnrealGameSync
 					ProgramsLine.AddLink("Visual Studio", FontStyle.Regular, () => { OpenSolution(); });
 					ProgramsLine.AddText("  |  ");
 				}
-				if (Settings.bEnableUshell)
+
+				List<ToolDefinition> Tools = Owner.ToolUpdateMonitor.Tools;
+				foreach(ToolDefinition Tool in Tools)
 				{
-					ProgramsLine.AddLink("Ushell", FontStyle.Regular, () => { OpenUshell(); });
-					ProgramsLine.AddText("  |  ");
+					if (Tool.Enabled)
+					{
+						foreach (ToolLink Link in Tool.StatusPanelLinks)
+						{
+							ProgramsLine.AddLink(Link.Label, FontStyle.Regular, () => RunTool(Tool, Link));
+							ProgramsLine.AddText("  |  ");
+						}
+					}
 				}
+
 				ProgramsLine.AddLink("Windows Explorer", FontStyle.Regular, () => { SafeProcessStart("explorer.exe", String.Format("\"{0}\"", Path.GetDirectoryName(SelectedFileName))); });
 
 				if (GetDefaultIssueFilter() != null)
diff --git a/Engine/Source/Programs/UnrealGameSync/UnrealGameSync/Forms/ApplicationSettingsWindow.cs b/Engine/Source/Programs/UnrealGameSync/UnrealGameSync/Forms/ApplicationSettingsWindow.cs
--- a/Engine/Source/Programs/UnrealGameSync/UnrealGameSync/Forms/ApplicationSettingsWindow.cs
+++ b/Engine/Source/Programs/UnrealGameSync/UnrealGameSync/Forms/ApplicationSettingsWindow.cs
@@ -59,6 +59,21 @@ namespace UnrealGameSync
 
 		ToolUpdateMonitor ToolUpdateMonitor;
 
+		class ToolItem
+		{
+			public ToolDefinition Definition { get; }
+
+			public ToolItem(ToolDefinition Definition)
+			{
+				this.Definition = Definition;
+			}
+
+			public override string ToString()
+			{
+				return Definition.Description;
+			}
+		}
+
 		private ApplicationSettingsWindow(string DefaultServerAndPort, string DefaultUserName, bool bUnstable, string OriginalExecutableFileName, UserSettings Settings, ToolUpdateMonitor ToolUpdateMonitor, TextWriter Log)
 		{
 			InitializeComponent();
@@ -119,8 +134,11 @@ namespace UnrealGameSync
 				this.EnableProtocolHandlerCheckBox.CheckState = CheckState.Indeterminate;
 			}
 
-			EnableP4VExtensionsCheckBox.Checked = Settings.bEnableP4VExtensions;
-			EnableUshellCheckBox.Checked = Settings.bEnableUshell;
+			List<ToolDefinition> Tools = ToolUpdateMonitor.Tools;
+			foreach (ToolDefinition Tool in Tools)
+			{
+				this.CustomToolsListBox.Items.Add(new ToolItem(Tool), Settings.EnabledTools.Contains(Tool.Id));
+			}
 		}
 
 		public static bool? ShowModal(IWin32Window Owner, PerforceConnection DefaultConnection, bool bUnstable, string OriginalExecutableFileName, UserSettings Settings, ToolUpdateMonitor ToolUpdateMonitor, TextWriter Log)
@@ -216,16 +234,14 @@ namespace UnrealGameSync
 				Settings.Save();
 			}
 
-			if (Settings.bEnableP4VExtensions != EnableP4VExtensionsCheckBox.Checked)
+			List<Guid> NewEnabledTools = new List<Guid>();
+			foreach (ToolItem Item in CustomToolsListBox.CheckedItems)
 			{
-				Settings.bEnableP4VExtensions = EnableP4VExtensionsCheckBox.Checked;
-				Settings.Save();
-				ToolUpdateMonitor.UpdateNow();
+				NewEnabledTools.Add(Item.Definition.Id);
 			}
-
-			if (Settings.bEnableUshell != EnableUshellCheckBox.Checked)
+			if (!NewEnabledTools.SequenceEqual(Settings.EnabledTools))
 			{
-				Settings.bEnableUshell = EnableUshellCheckBox.Checked;
+				Settings.EnabledTools = NewEnabledTools.ToArray();
 				Settings.Save();
 				ToolUpdateMonitor.UpdateNow();
 			}
diff --git a/Engine/Source/Programs/UnrealGameSync/UnrealGameSync/Forms/MainWindow.cs b/Engine/Source/Programs/UnrealGameSync/UnrealGameSync/Forms/MainWindow.cs
--- a/Engine/Source/Programs/UnrealGameSync/UnrealGameSync/Forms/MainWindow.cs
+++ b/Engine/Source/Programs/UnrealGameSync/UnrealGameSync/Forms/MainWindow.cs
@@ -106,7 +106,7 @@ namespace UnrealGameSync
 		Rectangle PrimaryWorkArea;
 		List<IssueAlertWindow> AlertWindows = new List<IssueAlertWindow>();
 
-		ToolUpdateMonitor ToolUpdateMonitor;
+		public ToolUpdateMonitor ToolUpdateMonitor { get; private set; }
 
 		NetCoreWindow NetCoreWindow;
 
@@ -126,7 +126,8 @@ namespace UnrealGameSync
 			OriginalExecutableFileName = InOriginalExecutableFileName;
 			bUnstable = bInUnstable;
 			DefaultConnection = InDefaultConnection;
-	
+			ToolUpdateMonitor = new ToolUpdateMonitor(DefaultConnection, DataFolder, InSettings);
+
 			Settings = InSettings;
 			OIDCTokenManager = InOidcTokenManager;
 
@@ -199,7 +200,7 @@ namespace UnrealGameSync
 				WorkspaceIssueMonitor.IssueMonitor.Start();
 			}
 
-			ToolUpdateMonitor = new ToolUpdateMonitor(DefaultConnection, DataFolder, Settings);
+			ToolUpdateMonitor.Start();
 		}
 
 		void PostAutomationRequest(AutomationRequest Request)
diff --git a/Engine/Source/Programs/UnrealGameSync/UnrealGameSync/ToolUpdateMonitor.cs b/Engine/Source/Programs/UnrealGameSync/UnrealGameSync/ToolUpdateMonitor.cs
--- a/Engine/Source/Programs/UnrealGameSync/UnrealGameSync/ToolUpdateMonitor.cs
+++ b/Engine/Source/Programs/UnrealGameSync/UnrealGameSync/ToolUpdateMonitor.cs
@@ -8,36 +8,47 @@ using System.Linq;
 using System.Text;
 using System.Threading;
 using System.Threading.Tasks;
+using System.Windows.Forms;
 
 namespace UnrealGameSync
 {
+	class ToolLink
+	{
+		public string Label;
+		public string FileName;
+		public string Arguments;
+		public string WorkingDir;
+	}
+
+	class ToolDefinition
+	{
+		public Guid Id;
+		public string Name;
+		public string Description;
+		public bool Enabled;
+		public Action<TextWriter> InstallAction;
+		public Action<TextWriter> UninstallAction;
+		public List<ToolLink> StatusPanelLinks = new List<ToolLink>();
+		public string ZipPath;
+		public int ZipChange;
+		public string ConfigPath;
+		public int ConfigChange;
+	}
+
 	class ToolUpdateMonitor : IDisposable
 	{
 		Thread WorkerThread;
 		AutoResetEvent WakeEvent;
 		bool bQuit;
 		string LogFile;
+		public List<ToolDefinition> Tools { get; private set; } = new List<ToolDefinition>();
+		int LastChange = -1;
 
 		public PerforceConnection Perforce { get; }
 		string ToolsDir { get; }
 		UserSettings Settings { get; }
 
-		class ToolDefinition
-		{
-			public string Name;
-			public Action<TextWriter> InstallAction;
-			public Action<TextWriter> UninstallAction;
-			public bool bEnabled;
-
-			public ToolDefinition(string Name, Action<TextWriter> InstallAction, Action<TextWriter> UninstallAction, bool bEnabled)
-			{
-				this.Name = Name;
-				this.InstallAction = InstallAction;
-				this.UninstallAction = UninstallAction;
-				this.bEnabled = bEnabled;
-			}
-		}
-
+		public event Action OnChange;
 
 		public ToolUpdateMonitor(PerforceConnection InPerforce, string DataDir, UserSettings Settings)
 		{
@@ -50,7 +61,10 @@ namespace UnrealGameSync
 			Directory.CreateDirectory(ToolsDir);
 
 			WakeEvent = new AutoResetEvent(false);
+		}
 
+		public void Start()
+		{
 			if (DeploymentSettings.ToolsDepotPath != null)
 			{
 				WorkerThread = new Thread(() => PollForUpdates());
@@ -122,47 +136,72 @@ namespace UnrealGameSync
 
 		void PollForUpdatesOnce(TextWriter Log)
 		{
-			List<ToolDefinition> Tools = new List<ToolDefinition>();
-			Tools.Add(new ToolDefinition("P4VUtils", x => RunP4VUtils("install", x), x => RunP4VUtils("uninstall", x), Settings.bEnableP4VExtensions));
-			Tools.Add(new ToolDefinition("Ushell", null, null, Settings.bEnableUshell));
+			List<PerforceChangeSummary> Changes;
+			if (!Perforce.FindChanges(DeploymentSettings.ToolsDepotPath + "/...", 1, out Changes, Log) || Changes.Count == 0 || Changes[0].Number == LastChange)
+			{
+				return;
+			}
+
+			List<PerforceFileRecord> FileRecords;
+			if (!Perforce.Stat(DeploymentSettings.ToolsDepotPath + "/...", out FileRecords, Log))
+			{
+				return;
+			}
 
-			if (Tools.Any(x => x.bEnabled))
+			// Update the tools list
+			List<ToolDefinition> NewTools = new List<ToolDefinition>();
+			foreach (PerforceFileRecord FileRecord in FileRecords)
 			{
-				List<PerforceFileRecord> FileRecords;
-				if (Perforce.Stat(DeploymentSettings.ToolsDepotPath + "/...", out FileRecords, Log))
+				if (FileRecord.DepotPath.EndsWith(".ini"))
 				{
-					foreach (ToolDefinition Tool in Tools.Where(x => x.bEnabled))
+					ToolDefinition Tool = Tools.FirstOrDefault(x => x.ConfigPath.Equals(FileRecord.DepotPath, StringComparison.Ordinal));
+					if (Tool == null || Tool.ConfigChange != FileRecord.HeadChange)
 					{
-						string Prefix = String.Format("{0}/{1}/", DeploymentSettings.ToolsDepotPath, Tool.Name);
-
-						List<PerforceFileRecord> ToolFileRecords = FileRecords.Where(x => x.DepotPath.StartsWith(Prefix, StringComparison.OrdinalIgnoreCase)).ToList();
-						if (ToolFileRecords.Count == 0)
-						{
-							continue;
-						}
-
-						int HeadChange = ToolFileRecords.Max(x => x.HeadChange);
-						if (HeadChange == GetToolChange(Tool.Name))
-						{
-							continue;
-						}
-
-						List<PerforceFileRecord> SyncFileRecords = ToolFileRecords.Where(x => x.Action != "delete").ToList();
-						try
-						{
-							UpdateTool(Tool.Name, HeadChange, SyncFileRecords, Tool.InstallAction, Log);
-						}
-						catch (Exception Ex)
-						{
-							Log.WriteLine("Exception while updating tool: {0}", Ex.ToString());
-						}
+						Tool = ReadToolDefinition(FileRecord.DepotPath, FileRecord.HeadChange, Log);
 					}
+					if (Tool != null)
+					{
+						NewTools.Add(Tool);
+					}
+				}
+			}
+			Tools = NewTools;
+
+			foreach (ToolDefinition Tool in Tools)
+			{
+				Tool.Enabled = Settings.EnabledTools.Contains(Tool.Id);
+
+				if(!Tool.Enabled)
+				{
+					continue;
+				}
+
+				List<PerforceFileRecord> ToolFileRecords = FileRecords.Where(x => x.DepotPath.Equals(Tool.ZipPath, StringComparison.OrdinalIgnoreCase)).ToList();
+				if (ToolFileRecords.Count == 0)
+				{
+					continue;
+				}
+
+				int HeadChange = ToolFileRecords.Max(x => x.HeadChange);
+				if (HeadChange == GetToolChange(Tool.Name))
+				{
+					continue;
+				}
+
+				List<PerforceFileRecord> SyncFileRecords = ToolFileRecords.Where(x => x.Action != "delete").ToList();
+				try
+				{
+					UpdateTool(Tool.Name, HeadChange, SyncFileRecords, Tool.InstallAction, Log);
+				}
+				catch (Exception Ex)
+				{
+					Log.WriteLine("Exception while updating tool: {0}", Ex.ToString());
 				}
 			}
 
-			foreach (ToolDefinition Tool in Tools.Where(x => !x.bEnabled))
+			foreach (ToolDefinition Tool in Tools)
 			{
-				if (GetToolChange(Tool.Name) != 0)
+				if (!Tool.Enabled && GetToolChange(Tool.Name) != 0)
 				{
 					try
 					{
@@ -174,13 +213,88 @@ namespace UnrealGameSync
 					}
 				}
 			}
+
+			OnChange();
 		}
 
-		void RunP4VUtils(string Command, TextWriter Log)
+		ToolDefinition ReadToolDefinition(string DepotPath, int Change, TextWriter Log)
 		{
-			string ToolPath = GetToolPathInternal("P4VUtils");
+			List<string> Lines;
+			if (!Perforce.Print(String.Format("{0}@{1}", DepotPath, Change), out Lines, Log))
+			{
+				return null;
+			}
+
+			int NameIdx = DepotPath.LastIndexOf('/') + 1;
+			int ExtensionIdx = DepotPath.LastIndexOf('.');
+
+			ConfigFile ConfigFile = new ConfigFile();
+			ConfigFile.Parse(Lines.ToArray());
+
+			ToolDefinition Tool = new ToolDefinition();
+
+			string Id = ConfigFile.GetValue("Settings.Id", null);
+			if (Id == null || !Guid.TryParse(Id, out Tool.Id))
+			{
+				return null;
+			}
+
+			Tool.Name = ConfigFile.GetValue("Settings.Name", DepotPath.Substring(NameIdx, ExtensionIdx - NameIdx));
+			Tool.Description = ConfigFile.GetValue("Settings.Description", Tool.Name);
+			Tool.ZipPath = DepotPath.Substring(0, ExtensionIdx) + ".zip";
+			Tool.ZipChange = GetToolChange(Tool.Name);
+			Tool.ConfigPath = DepotPath;
+			Tool.ConfigChange = Change;
+
+			string InstallCommand = ConfigFile.GetValue("Settings.InstallCommand", null);
+			if (!String.IsNullOrEmpty(InstallCommand))
+			{
+				Tool.InstallAction = NewLog => RunCommand(Tool.Name, InstallCommand, NewLog);
+			}
+
+			string UninstallCommand = ConfigFile.GetValue("Settings.UninstallCommand", null);
+			if (!String.IsNullOrEmpty(UninstallCommand))
+			{
+				Tool.UninstallAction = NewLog => RunCommand(Tool.Name, UninstallCommand, NewLog);
+			}
 
-			int ExitCode = Utility.ExecuteProcess(Path.Combine(ToolPath, "P4VUtils.exe"), ToolPath, Command, null, Log);
+			string[] StatusPanelLinks = ConfigFile.GetValues("Settings.StatusPanelLinks", new string[0]);
+			foreach (string StatusPanelLink in StatusPanelLinks)
+			{
+				ConfigObject Object = new ConfigObject(StatusPanelLink);
+
+				string Label = Object.GetValue("Label", null);
+				string FileName = Object.GetValue("FileName", null);
+
+				if (Label != null && FileName != null)
+				{
+					ToolLink Link = new ToolLink();
+					Link.Label = Label;
+					Link.FileName = FileName;
+					Link.Arguments = Object.GetValue("Arguments", null);
+					Link.WorkingDir = Object.GetValue("WorkingDir", null);
+					Tool.StatusPanelLinks.Add(Link);
+				}
+			}
+
+			return Tool;
+		}
+
+		void RunCommand(string ToolName, string Command, TextWriter Log)
+		{
+			string ToolPath = GetToolPathInternal(ToolName);
+
+			string CommandExe = Command;
+			string CommandArgs = string.Empty;
+
+			int SpaceIdx = Command.IndexOf(' ');
+			if (SpaceIdx != -1)
+			{
+				CommandExe = Command.Substring(0, SpaceIdx);
+				CommandArgs = Command.Substring(SpaceIdx + 1);
+			}
+
+			int ExitCode = Utility.ExecuteProcess(Path.Combine(ToolPath, CommandExe), ToolPath, CommandArgs, null, Log);
 			Log.WriteLine("(Exit: {0})", ExitCode);
 		}
 
@@ -197,9 +311,11 @@ namespace UnrealGameSync
 			SetToolChange(ToolName, null);
 
 			string ToolPath = GetToolPath(ToolName);
-			Log.WriteLine("Removing {0}", ToolPath);
-
-			TryDeleteDirectory(ToolPath);
+			if (ToolPath != null)
+			{
+				Log.WriteLine("Removing {0}", ToolPath);
+				TryDeleteDirectory(ToolPath);
+			}
 		}
 
 		static void ForceDeleteDirectory(string DirectoryName)
diff --git a/Engine/Source/Programs/UnrealGameSync/UnrealGameSync/UserSettings.cs b/Engine/Source/Programs/UnrealGameSync/UnrealGameSync/UserSettings.cs
--- a/Engine/Source/Programs/UnrealGameSync/UnrealGameSync/UserSettings.cs
+++ b/Engine/Source/Programs/UnrealGameSync/UnrealGameSync/UserSettings.cs
@@ -258,8 +258,7 @@ namespace UnrealGameSync
 		public bool bShowAutomatedChanges;
 		public bool bShowLocalTimes;
 		public bool bKeepInTray;
-		public bool bEnableP4VExtensions;
-		public bool bEnableUshell;
+		public Guid[] EnabledTools;
 		public bool bShowNetCoreInfo;
 		public int FilterIndex;
 		public UserSelectedProjectSettings LastProject;
@@ -359,8 +358,18 @@ namespace UnrealGameSync
 			bShowAutomatedChanges = ConfigFile.GetValue("General.ShowAutomated", false);
 			bShowLocalTimes = ConfigFile.GetValue("General.ShowLocalTimes", false);
 			bKeepInTray = ConfigFile.GetValue("General.KeepInTray", true);
-			bEnableP4VExtensions = ConfigFile.GetValue("General.EnableP4VExtensions", false);
-			bEnableUshell = ConfigFile.GetValue("General.EnableUshell", false);
+
+			List<Guid> EnabledTools = ConfigFile.GetGuidValues("General.EnabledTools", new Guid[0]).ToList();
+			if (ConfigFile.GetValue("General.EnableP4VExtensions", false))
+			{
+				EnabledTools.Add(new Guid("963850A0-BF63-4E0E-B903-1C5954C7DCF8"));
+			}
+			if (ConfigFile.GetValue("General.EnableUshell", false))
+			{
+				EnabledTools.Add(new Guid("922EED87-E732-464C-92DC-5A8F7ED955E2"));
+			}
+			this.EnabledTools = EnabledTools.ToArray();
+
 			bShowNetCoreInfo = ConfigFile.GetValue("General.ShowNetCoreInfo", true);
 			int.TryParse(ConfigFile.GetValue("General.FilterIndex", "0"), out FilterIndex);
 
@@ -723,8 +732,7 @@ namespace UnrealGameSync
 			}
 			GeneralSection.SetValues("OpenProjects", OpenProjects.Select(x => x.ToConfigEntry()).ToArray());
 			GeneralSection.SetValue("KeepInTray", bKeepInTray);
-			GeneralSection.SetValue("EnableP4VExtensions", bEnableP4VExtensions);
-			GeneralSection.SetValue("EnableUshell", bEnableUshell);
+			GeneralSection.SetValues("EnabledTools", EnabledTools);
 			GeneralSection.SetValue("ShowNetCoreInfo", bShowNetCoreInfo);
 			GeneralSection.SetValue("FilterIndex", FilterIndex);
 			GeneralSection.SetValues("RecentProjects", RecentProjects.Select(x => x.ToConfigEntry()).ToArray());
