UnrealGameSync: Check available free space before syncing from perforce

Add helper function to format bytes as a human readable string.

diff --git a/Engine/Source/Programs/Shared/EpicGames.Core/StringUtils.cs b/Engine/Source/Programs/Shared/EpicGames.Core/StringUtils.cs
--- a/Engine/Source/Programs/Shared/EpicGames.Core/StringUtils.cs
+++ b/Engine/Source/Programs/Shared/EpicGames.Core/StringUtils.cs
@@ -16,6 +16,11 @@ namespace EpicGames.Core
 		/// </summary>
 		static sbyte[] HexDigits;
 
+		/// <summary>
+		/// Array mapping human readable size of bytes, 1024^x. long max is within the range of Exabytes.
+		/// </summary>
+		static string[] ByteSizes = { "B", "KB", "MB", "GB", "TB", "PB", "EB" };
+
 		/// <summary>
 		/// Static constructor. Initializes the HexDigits array.
 		/// </summary>
@@ -408,5 +413,23 @@ namespace EpicGames.Core
 				return String;
 			}
 		}
+
+		/// <summary>
+		/// Formats bytes into a human readable string
+		/// </summary>
+		/// <param name="Bytes">The total number of bytes</param>
+		/// <param name="DecimalPlaces">The number of decimal places to round the resulting value</param>
+		/// <returns>Human readable string based on the value of Bytes</returns>
+		public static string FormatBytesString(long Bytes, int DecimalPlaces = 2)
+		{
+			if (Bytes == 0)
+			{
+				return $"0 {ByteSizes[0]}";
+			}
+			long BytesAbs = Math.Abs(Bytes);
+			int Power = Convert.ToInt32(Math.Floor(Math.Log(BytesAbs, 1024)));
+			double Value = Math.Round(BytesAbs / Math.Pow(1024, Power), DecimalPlaces);
+			return $"{(Math.Sign(Bytes) * Value)} {ByteSizes[Power]}";
+		}
 	}
 }
diff --git a/Engine/Source/Programs/UnrealGameSync/UnrealGameSync/Workspace.cs b/Engine/Source/Programs/UnrealGameSync/UnrealGameSync/Workspace.cs
--- a/Engine/Source/Programs/UnrealGameSync/UnrealGameSync/Workspace.cs
+++ b/Engine/Source/Programs/UnrealGameSync/UnrealGameSync/Workspace.cs
@@ -606,6 +606,9 @@ namespace UnrealGameSync
 					List<string> SyncDepotPaths = new List<string>();
 					using(RecordCounter Counter = new RecordCounter(Progress, "Filtering files..."))
 					{
+						// Track the total new bytes that will be required on disk when syncing. Add an extra 100MB for padding.
+						long RequiredFreeSpace = 100 * 1024 * 1024;
+
 						foreach(string SyncPath in SyncPaths)
 						{
 							List<PerforceFileRecord> SyncRecords = new List<PerforceFileRecord>();
@@ -632,6 +635,7 @@ namespace UnrealGameSync
 								{
 									BatchBuilder.Add(String.Format("{0}@{1}", SyncRecord.DepotPath, PendingChangeNumber), SyncRecord.FileSize);
 									SyncDepotPaths.Add(SyncRecord.DepotPath);
+									RequiredFreeSpace += SyncRecord.FileSize;
 									continue;
 								}
 
@@ -649,10 +653,12 @@ namespace UnrealGameSync
 								}
 
 								// Make sure it's under the current directory. Not sure why this would happen, just being safe.
+								// This occurs for files returned from GetOpenFiles as SyncRecord.ClientPath is the client workspace path not local path.
 								if (!FullName.StartsWith(LocalRootPrefix, StringComparison.OrdinalIgnoreCase))
 								{
 									BatchBuilder.Add(String.Format("{0}@{1}", SyncRecord.DepotPath, PendingChangeNumber), SyncRecord.FileSize);
 									SyncDepotPaths.Add(SyncRecord.DepotPath);
+									RequiredFreeSpace += SyncRecord.FileSize;
 									continue;
 								}
 
@@ -668,6 +674,14 @@ namespace UnrealGameSync
 
 									SyncTree.IncludeFile(PerforceUtils.EscapePath(RelativePath), FileSize);
 									SyncDepotPaths.Add(SyncRecord.DepotPath);
+									RequiredFreeSpace += FileSize;
+									FileInfo LocalFileInfo = new FileInfo(FullName);
+
+									// If the file exists the required free space can be reduced as those bytes will be replaced.
+									if (LocalFileInfo.Exists)
+									{
+										RequiredFreeSpace -= LocalFileInfo.Length;
+									}
 								}
 								else
 								{
@@ -675,6 +689,23 @@ namespace UnrealGameSync
 								}
 							}
 						}
+
+						try
+						{
+							DirectoryInfo LocalRootInfo = new DirectoryInfo(LocalRootPrefix);
+							DriveInfo Drive = new DriveInfo(LocalRootInfo.Root.FullName);
+
+							if (Drive.AvailableFreeSpace < RequiredFreeSpace)
+							{
+								Log.WriteLine("Syncing requires {0} which exceeds the {1} available free space on {2}.", EpicGames.Core.StringUtils.FormatBytesString(RequiredFreeSpace), EpicGames.Core.StringUtils.FormatBytesString(Drive.AvailableFreeSpace), Drive.Name);
+								StatusMessage = "Not enough available free space.";
+								return WorkspaceUpdateResult.FailedToSync;
+							}
+						}
+						catch (SystemException)
+						{
+							Log.WriteLine("Unable to check available free space for {0}.", LocalRootPrefix);
+						}
 					}
 					SyncTree.GetOptimizedSyncCommands(ClientRootPath, PendingChangeNumber, BatchBuilder);
 
