Merging last UGS changes from //UE4/Dev-Build.

#rnx

diff --git a/Engine/Source/Programs/UnrealGameSync/UnrealGameSync/Controls/WorkspaceControl.cs b/Engine/Source/Programs/UnrealGameSync/UnrealGameSync/Controls/WorkspaceControl.cs
--- a/Engine/Source/Programs/UnrealGameSync/UnrealGameSync/Controls/WorkspaceControl.cs
+++ b/Engine/Source/Programs/UnrealGameSync/UnrealGameSync/Controls/WorkspaceControl.cs
@@ -319,6 +319,12 @@ namespace UnrealGameSync
 		System.Threading.Timer StartupTimer;
 		List<WorkspaceStartupCallback> StartupCallbacks;
 
+		// When an author filter is applied this will be non-empty and != AuthorFilterPlaceholderText
+		string AuthorFilterText = "";
+		
+		// Placeholder text that is in the control and cleared when the user starts editing.
+		static string AuthorFilterPlaceholderText = "<username>";
+
 		public WorkspaceControl(IWorkspaceControlOwner InOwner, string InApiUrl, string InOriginalExecutableFileName, bool bInUnstable, DetectProjectSettingsTask DetectSettings, LineBasedTextWriter InLog, UserSettings InSettings, OIDCTokenManager InOidcTokenManager)
 		{
 			InitializeComponent();
@@ -1392,7 +1398,15 @@ namespace UnrealGameSync
 			Columns[ChangeColumn.Index] = String.Format("{0}", Change.Number);
 
 			Columns[TimeColumn.Index] = DisplayTime.ToShortTimeString();
-			Columns[AuthorColumn.Index] = FormatUserName(Change.User);
+			string UserName = FormatUserName(Change.User);
+
+			// If annotate robomerge is on, add a prefix to the user name
+			if (Settings.bAnnotateRobmergeChanges && IsRobomergeChange(Change))
+			{
+				UserName += " (robo)";
+			}
+
+			Columns[AuthorColumn.Index] = UserName;
 			Columns[DescriptionColumn.Index] = Change.Description.Replace('\n', ' ');
 
 			for (int ColumnIdx = 1; ColumnIdx < BuildList.Columns.Count; ColumnIdx++)
@@ -1416,6 +1430,18 @@ namespace UnrealGameSync
 			BuildList.Items.Add(Item);
 		}
 
+		/// <summary>
+		/// Returns true if this change was submitted on behalf of someone by robomerge
+		/// </summary>
+		/// <param name="Change"></param>
+		/// <returns></returns>
+		bool IsRobomergeChange(PerforceChangeSummary Change)
+		{
+			// If hiding robomerge, filter out based on workspace name. Note - don't look at the description because we
+			// *do* want to see conflicts that were manually merged by someone
+			return Change.Client.IndexOf("ROBOMERGE", StringComparison.OrdinalIgnoreCase) >= 0;
+		}
+
 		bool ShouldShowChange(PerforceChangeSummary Change, string[] ExcludeChanges)
 		{
 			if (ProjectSettings.FilterType != FilterType.None)
@@ -1434,6 +1460,33 @@ namespace UnrealGameSync
 					return false;
 				}
 			}
+
+			// if filtering by user, only show changes where the author contains the filter text
+			if (!string.IsNullOrEmpty(this.AuthorFilterText)
+				&& this.AuthorFilterText != AuthorFilterPlaceholderText
+				&& Change.User.IndexOf(this.AuthorFilterText, StringComparison.OrdinalIgnoreCase) < 0)
+			{
+				return false;
+			}
+			
+			// If this is a robomerge change, check if any filters will cause it to be hidden
+			if (IsRobomergeChange(Change))
+			{
+				if (Settings.ShowRobomerge == UserSettings.RobomergeShowChangesOption.None)
+				{
+					return false;
+				}
+				else if (Settings.ShowRobomerge == UserSettings.RobomergeShowChangesOption.Badged)
+				{
+					// if this change has any badges, we'll show it unless it's later filtered out
+					EventSummary Summary = EventMonitor.GetSummaryForChange(Change.Number);
+					if (Summary == null || !Summary.Badges.Any())
+					{
+						return false;
+					}
+				}
+			}
+
 			if (ProjectSettings.FilterBadges.Count > 0)
 			{
 				EventSummary Summary = EventMonitor.GetSummaryForChange(Change.Number);
@@ -5736,6 +5789,16 @@ namespace UnrealGameSync
 			FilterContextMenu_Badges.Enabled = FilterContextMenu_Badges.DropDownItems.Count > 0;
 
 			FilterContextMenu_ShowBuildMachineChanges.Checked = Settings.bShowAutomatedChanges;
+
+			// Set checks if any robomerge filters are applied
+			FilterContextMenu_Robomerge.Checked = Settings.ShowRobomerge != UserSettings.RobomergeShowChangesOption.All || Settings.bAnnotateRobmergeChanges;
+			FilterContextMenu_Robomerge_ShowAll.Checked = Settings.ShowRobomerge == UserSettings.RobomergeShowChangesOption.All;
+			FilterContextMenu_Robomerge_ShowBadged.Checked = Settings.ShowRobomerge == UserSettings.RobomergeShowChangesOption.Badged;
+			FilterContextMenu_Robomerge_Annotate.Checked = Settings.bAnnotateRobmergeChanges;
+
+			// Set checks if an author filter string is set
+			FilterContextMenu_Author.Checked = !string.IsNullOrEmpty(this.AuthorFilterText) && this.AuthorFilterText != AuthorFilterPlaceholderText;
+
 			FilterContextMenu.Show(FilterButton, new Point(0, FilterButton.Height));
 		}
 
@@ -5772,11 +5835,83 @@ namespace UnrealGameSync
 			UpdateBuildListFilter();
 		}
 
+		private void FilterContextMenu_Robomerge_ShowAll_Click(object sender, EventArgs e)
+		{
+			Settings.ShowRobomerge = UserSettings.RobomergeShowChangesOption.All;	
+			Settings.Save();
+			UpdateBuildListFilter();
+		}
+
+		private void FilterContextMenu_Robomerge_ShowBadged_Click(object sender, EventArgs e)
+		{
+			Settings.ShowRobomerge = UserSettings.RobomergeShowChangesOption.Badged;
+			Settings.Save();
+			UpdateBuildListFilter();
+		}
+
+		private void FilterContextMenu_Robomerge_ShowNone_Click(object sender, EventArgs e)
+		{
+			Settings.ShowRobomerge = UserSettings.RobomergeShowChangesOption.None;
+			Settings.Save();
+			UpdateBuildListFilter();
+		}
+
+		private void FilterContextMenu_Robomerge_Annotate_Click(object sender, EventArgs e)
+		{
+			Settings.bAnnotateRobmergeChanges ^= true;
+			Settings.Save();
+			BuildList.Items.Clear();	// need to reload to update user names
+
+			UpdateBuildListFilter();
+		}
+
+		/// <summary>
+		/// Handler for the edit box of the author filtering taking focus. This clears any
+		/// placeholder text and sets the foreground color to black so it shows as active
+		/// </summary>
+		/// <param name="sender"></param>
+		/// <param name="e"></param>
+		private void FilterContextMenu_Author_Name_GotFocus(object sender, EventArgs e)
+		{
+			if (this.FilterContextMenu_Author_Name.Text == AuthorFilterPlaceholderText)
+			{
+				this.FilterContextMenu_Author_Name.Text = "";
+				this.FilterContextMenu_Author_Name.TextBox.ForeColor = Color.Black;
+			}
+		}
+
+		/// <summary>
+		/// Handler for the edit box of the author filter losing focus. If the filter box
+		/// is empty we restore the placeholder text and state
+		/// </summary>
+		/// <param name="sender"></param>
+		/// <param name="e"></param>
+		private void FilterContextMenu_Author_Name_LostFocus(object sender, EventArgs e)
+		{
+			if (string.IsNullOrEmpty(this.FilterContextMenu_Author_Name.Text))
+			{
+				this.FilterContextMenu_Author_Name.Text = AuthorFilterPlaceholderText;
+				this.FilterContextMenu_Author_Name.TextBox.ForeColor = Color.DarkGray;
+			}
+		}
+
+		/// <summary>
+		/// Handler for the author name changing. We save off the text and update the 
+		/// filter string
+		/// </summary>
+		/// <param name="sender"></param>
+		/// <param name="e"></param>
+		private void FilterContextMenu_Author_Name_Changed(object sender, EventArgs e)
+		{
+			AuthorFilterText = this.FilterContextMenu_Author_Name.Text;
+			UpdateBuildListFilter();
+		}
+
+
 		private void FilterContextMenu_Type_ShowAll_Click(object sender, EventArgs e)
 		{
 			ProjectSettings.FilterType = FilterType.None;
 			Settings.Save();
-
 			UpdateBuildListFilter();
 		}
 
diff --git a/Engine/Source/Programs/UnrealGameSync/UnrealGameSync/UserSettings.cs b/Engine/Source/Programs/UnrealGameSync/UnrealGameSync/UserSettings.cs
--- a/Engine/Source/Programs/UnrealGameSync/UnrealGameSync/UserSettings.cs
+++ b/Engine/Source/Programs/UnrealGameSync/UnrealGameSync/UserSettings.cs
@@ -244,6 +244,16 @@ namespace UnrealGameSync
 
 	class UserSettings
 	{
+		/// <summary>
+		/// Enum that decribes which robomerge changes to show
+		/// </summary>
+		public enum RobomergeShowChangesOption
+		{
+			All,		// Show all changes from robomerge
+			Badged,		// Show only robomerge changes that have an associated badge
+			None		// Show no robomerge changes
+		};
+
 		string FileName;
 		ConfigFile ConfigFile = new ConfigFile();
 
@@ -256,6 +266,8 @@ namespace UnrealGameSync
 		public bool bAutoResolveConflicts;
 		public bool bShowUnreviewedChanges;
 		public bool bShowAutomatedChanges;
+		public RobomergeShowChangesOption ShowRobomerge;
+		public bool bAnnotateRobmergeChanges;
 		public bool bShowLocalTimes;
 		public bool bKeepInTray;
 		public Guid[] EnabledTools;
@@ -356,6 +368,12 @@ namespace UnrealGameSync
 			bAutoResolveConflicts = (ConfigFile.GetValue("General.AutoResolveConflicts", "1") != "0");
 			bShowUnreviewedChanges = ConfigFile.GetValue("General.ShowUnreviewed", true);
 			bShowAutomatedChanges = ConfigFile.GetValue("General.ShowAutomated", false);
+
+			// safely parse the filter enum
+			ShowRobomerge = RobomergeShowChangesOption.All;
+			Enum.TryParse(ConfigFile.GetValue("General.RobomergeFilter", ""), out ShowRobomerge);
+
+			bAnnotateRobmergeChanges = ConfigFile.GetValue("General.AnnotateRobomerge", true);
 			bShowLocalTimes = ConfigFile.GetValue("General.ShowLocalTimes", false);
 			bKeepInTray = ConfigFile.GetValue("General.KeepInTray", true);
 
@@ -725,6 +743,8 @@ namespace UnrealGameSync
 			GeneralSection.SetValue("AutoResolveConflicts", bAutoResolveConflicts);
 			GeneralSection.SetValue("ShowUnreviewed", bShowUnreviewedChanges);
 			GeneralSection.SetValue("ShowAutomated", bShowAutomatedChanges);
+			GeneralSection.SetValue("RobomergeFilter", ShowRobomerge.ToString());
+			GeneralSection.SetValue("AnnotateRobomerge", bAnnotateRobmergeChanges);
 			GeneralSection.SetValue("ShowLocalTimes", bShowLocalTimes);
 			if(LastProject != null)
 			{
