UGS: Ignore projects without a config file.

diff --git a/Engine/Source/Programs/UnrealGameSync/UnrealGameSync/OIDCTokenManager.cs b/Engine/Source/Programs/UnrealGameSync/UnrealGameSync/OIDCTokenManager.cs
--- a/Engine/Source/Programs/UnrealGameSync/UnrealGameSync/OIDCTokenManager.cs
+++ b/Engine/Source/Programs/UnrealGameSync/UnrealGameSync/OIDCTokenManager.cs
@@ -72,12 +72,18 @@ namespace UnrealGameSync
 			foreach (DetectProjectSettingsTask DetectProjectSettingsTask in ConfigFiles)
 			{
 				ConfigFile ConfigFile = DetectProjectSettingsTask.LatestProjectConfigFile;
+				if(ConfigFile == null)
+				{
+					continue;
+				}
+
 				ConfigSection ProviderSection = ConfigFile.FindSection("OIDCProvider");
 				if (ProviderSection == null)
+				{
 					continue;
+				}
 
 				string[] ProviderValues = ProviderSection.GetValues("Provider", (string[]) null);
-
 				foreach (ConfigObject Provider in ProviderValues.Select(s => new ConfigObject(s)).ToList())
 				{
 					string Identifier = Provider.GetValue("Identifier");
diff --git a/Engine/Source/Programs/UnrealGameSync/UnrealGameSync/Properties/AssemblyInfo.cs b/Engine/Source/Programs/UnrealGameSync/UnrealGameSync/Properties/AssemblyInfo.cs
--- a/Engine/Source/Programs/UnrealGameSync/UnrealGameSync/Properties/AssemblyInfo.cs
+++ b/Engine/Source/Programs/UnrealGameSync/UnrealGameSync/Properties/AssemblyInfo.cs
@@ -34,5 +34,5 @@ using System.Runtime.InteropServices;
 // You can specify all the values or you can default the Build and Revision Numbers 
 // by using the '*' as shown below:
 // [assembly: AssemblyVersion("1.0.*")]
-[assembly: AssemblyVersion("1.240.0.0")]
-[assembly: AssemblyFileVersion("1.240.0.0")]
+[assembly: AssemblyVersion("1.241.0.0")]
+[assembly: AssemblyFileVersion("1.241.0.0")]
