UGS: Ignore null project settings objects.

diff --git a/Engine/Source/Programs/UnrealGameSync/UnrealGameSync/OIDCTokenManager.cs b/Engine/Source/Programs/UnrealGameSync/UnrealGameSync/OIDCTokenManager.cs
--- a/Engine/Source/Programs/UnrealGameSync/UnrealGameSync/OIDCTokenManager.cs
+++ b/Engine/Source/Programs/UnrealGameSync/UnrealGameSync/OIDCTokenManager.cs
@@ -71,6 +71,11 @@ namespace UnrealGameSync
 			Dictionary<string, ProviderInfo> Providers = new Dictionary<string, ProviderInfo>();
 			foreach (DetectProjectSettingsTask DetectProjectSettingsTask in ConfigFiles)
 			{
+				if(DetectProjectSettingsTask == null)
+				{
+					continue;
+				}
+
 				ConfigFile ConfigFile = DetectProjectSettingsTask.LatestProjectConfigFile;
 				if(ConfigFile == null)
 				{
