UnrealGameSync: Fix scanning a workspace for projects to add if the root directory is a drive root

Issue was an additional path separator was being appened to the drive root (ex. D:\ -> D:\\)

diff --git a/Engine/Source/Programs/UnrealGameSync/UnrealGameSync/Forms/SelectProjectFromWorkspaceWindow.cs b/Engine/Source/Programs/UnrealGameSync/UnrealGameSync/Forms/SelectProjectFromWorkspaceWindow.cs
--- a/Engine/Source/Programs/UnrealGameSync/UnrealGameSync/Forms/SelectProjectFromWorkspaceWindow.cs
+++ b/Engine/Source/Programs/UnrealGameSync/UnrealGameSync/Forms/SelectProjectFromWorkspaceWindow.cs
@@ -65,7 +65,11 @@ namespace UnrealGameSync
 					FileRecords.AddRange(WildcardFileRecords);
 				}
 
-				string ClientPrefix = ClientRoot + Path.DirectorySeparatorChar;
+				string ClientPrefix = ClientRoot;
+				if (!ClientPrefix.EndsWith(Path.DirectorySeparatorChar))
+				{
+					ClientPrefix += Path.DirectorySeparatorChar;
+				}
 
 				Paths = new List<string>();
 				foreach(PerforceFileRecord FileRecord in FileRecords)
