UGS: Fix unhandled exception when aborting a sync.

diff --git a/Engine/Source/Programs/UnrealGameSync/UnrealGameSync/Workspace.cs b/Engine/Source/Programs/UnrealGameSync/UnrealGameSync/Workspace.cs
--- a/Engine/Source/Programs/UnrealGameSync/UnrealGameSync/Workspace.cs
+++ b/Engine/Source/Programs/UnrealGameSync/UnrealGameSync/Workspace.cs
@@ -1415,7 +1415,16 @@ namespace UnrealGameSync
 
 				// Sync the files
 				string StatusMessage;
-				WorkspaceUpdateResult Result = StaticSyncFileRevisions(Perforce, PendingChangeNumber, Context, SyncCommands, Record => SyncOutput(Record, ThreadLog), ThreadLog, out StatusMessage);
+				WorkspaceUpdateResult Result;
+				try
+				{
+					Result = StaticSyncFileRevisions(Perforce, PendingChangeNumber, Context, SyncCommands, Record => SyncOutput(Record, ThreadLog), ThreadLog, out StatusMessage);
+				}
+				catch (ThreadInterruptedException)
+				{
+					StatusMessage = "Sync cancelled";
+					Result = WorkspaceUpdateResult.Canceled;
+				}
 
 				// If it failed, try to set it on the state if nothing else has failed first
 				if (Result != WorkspaceUpdateResult.Success)
