UGS: Ignore ThreadInterruptedExceptions.

#rnx

diff --git a/Engine/Source/Programs/UnrealGameSync/UnrealGameSync/Forms/IssueDetailsWindow.cs b/Engine/Source/Programs/UnrealGameSync/UnrealGameSync/Forms/IssueDetailsWindow.cs
--- a/Engine/Source/Programs/UnrealGameSync/UnrealGameSync/Forms/IssueDetailsWindow.cs
+++ b/Engine/Source/Programs/UnrealGameSync/UnrealGameSync/Forms/IssueDetailsWindow.cs
@@ -175,7 +175,13 @@ namespace UnrealGameSync
 					}
 
 					// Wait for something to change
-					RefreshEvent.WaitOne();
+					try
+					{
+						RefreshEvent.WaitOne();
+					}
+					catch (ThreadInterruptedException)
+					{
+					}
 				}
 			}
 
diff --git a/Engine/Source/Programs/UnrealGameSync/UnrealGameSync/Forms/IssueFixedWindow.cs b/Engine/Source/Programs/UnrealGameSync/UnrealGameSync/Forms/IssueFixedWindow.cs
--- a/Engine/Source/Programs/UnrealGameSync/UnrealGameSync/Forms/IssueFixedWindow.cs
+++ b/Engine/Source/Programs/UnrealGameSync/UnrealGameSync/Forms/IssueFixedWindow.cs
@@ -73,7 +73,13 @@ namespace UnrealGameSync
 			{
 				for(;;)
 				{
-					WakeEvent.WaitOne();
+					try
+					{
+						WakeEvent.WaitOne();
+					}
+					catch (ThreadInterruptedException)
+					{
+					}
 
 					if(bStopRequested)
 					{
diff --git a/Engine/Source/Programs/UnrealGameSync/UnrealGameSync/Properties/AssemblyInfo.cs b/Engine/Source/Programs/UnrealGameSync/UnrealGameSync/Properties/AssemblyInfo.cs
--- a/Engine/Source/Programs/UnrealGameSync/UnrealGameSync/Properties/AssemblyInfo.cs
+++ b/Engine/Source/Programs/UnrealGameSync/UnrealGameSync/Properties/AssemblyInfo.cs
@@ -34,5 +34,5 @@ using System.Runtime.InteropServices;
 // You can specify all the values or you can default the Build and Revision Numbers 
 // by using the '*' as shown below:
 // [assembly: AssemblyVersion("1.0.*")]
-[assembly: AssemblyVersion("1.242.0.0")]
-[assembly: AssemblyFileVersion("1.242.0.0")]
+[assembly: AssemblyVersion("1.243.0.0")]
+[assembly: AssemblyFileVersion("1.243.0.0")]
