Changes application path gather to be compliant with .Net Core, and gets correct path to exe instead of dll.

diff --git a/Engine/Source/Programs/UnrealGameSync/UnrealGameSync/UriHandlers/UriHandler.cs b/Engine/Source/Programs/UnrealGameSync/UnrealGameSync/UriHandlers/UriHandler.cs
--- a/Engine/Source/Programs/UnrealGameSync/UnrealGameSync/UriHandlers/UriHandler.cs
+++ b/Engine/Source/Programs/UnrealGameSync/UnrealGameSync/UriHandlers/UriHandler.cs
@@ -295,9 +295,14 @@ namespace UnrealGameSync
 			}
 		}
 
+		public static string CurrentProcessFilePath
+		{
+			get { return Path.GetFullPath(Process.GetCurrentProcess().MainModule.FileName); }
+		}
+
 		static List<RegistrySetting> GetRegistrySettings()
 		{
-			string ApplicationPath = Assembly.GetExecutingAssembly().Location;
+			string ApplicationPath = CurrentProcessFilePath;
 
 			List<RegistrySetting> Keys = new List<RegistrySetting>();
 			Keys.Add(new RegistrySetting(Registry.ClassesRoot, "UGS", null, "URL:UGS Protocol"));
