UGS: Batch up delete commands separately to sync commands. When a large number of deletes/edits are mixed under a directory tree, previous behaviour could result in deletes being stored in separate batches.

#rnx

diff --git a/Engine/Source/Programs/UnrealGameSync/UnrealGameSync/Properties/AssemblyInfo.cs b/Engine/Source/Programs/UnrealGameSync/UnrealGameSync/Properties/AssemblyInfo.cs
--- a/Engine/Source/Programs/UnrealGameSync/UnrealGameSync/Properties/AssemblyInfo.cs
+++ b/Engine/Source/Programs/UnrealGameSync/UnrealGameSync/Properties/AssemblyInfo.cs
@@ -34,5 +34,5 @@ using System.Runtime.InteropServices;
 // You can specify all the values or you can default the Build and Revision Numbers 
 // by using the '*' as shown below:
 // [assembly: AssemblyVersion("1.0.*")]
-[assembly: AssemblyVersion("1.243.0.0")]
-[assembly: AssemblyFileVersion("1.243.0.0")]
+[assembly: AssemblyVersion("1.244.0.0")]
+[assembly: AssemblyFileVersion("1.244.0.0")]
diff --git a/Engine/Source/Programs/UnrealGameSync/UnrealGameSync/Workspace.cs b/Engine/Source/Programs/UnrealGameSync/UnrealGameSync/Workspace.cs
--- a/Engine/Source/Programs/UnrealGameSync/UnrealGameSync/Workspace.cs
+++ b/Engine/Source/Programs/UnrealGameSync/UnrealGameSync/Workspace.cs
@@ -201,6 +201,7 @@ namespace UnrealGameSync
 			public Queue<List<string>> Batches { get; }
 
 			List<string> Commands;
+			List<string> DeleteCommands;
 			long Size;
 
 			public SyncBatchBuilder(int MaxCommandsPerList, long MaxSizePerList)
@@ -212,15 +213,28 @@ namespace UnrealGameSync
 
 			public void Add(string NewCommand, long NewSize)
 			{
-				if (Commands == null || Commands.Count >= MaxCommandsPerList || Size + NewSize >= MaxSizePerList)
+				if (NewSize == 0)
 				{
-					Commands = new List<string>();
-					Batches.Enqueue(Commands);
-					Size = 0;
+					if (DeleteCommands == null || DeleteCommands.Count >= MaxCommandsPerList)
+					{
+						DeleteCommands = new List<string>();
+						Batches.Enqueue(DeleteCommands);
+					}
+
+					DeleteCommands.Add(NewCommand);
 				}
+				else
+				{
+					if (Commands == null || Commands.Count >= MaxCommandsPerList || Size + NewSize >= MaxSizePerList)
+					{
+						Commands = new List<string>();
+						Batches.Enqueue(Commands);
+						Size = 0;
+					}
 
-				Commands.Add(NewCommand);
-				Size += NewSize;
+					Commands.Add(NewCommand);
+					Size += NewSize;
+				}
 			}
 		}
 
