UGS: Use shell execute to open SDK files. Fixes hangs trying to launch a zip.

#rnx

diff --git a/Engine/Source/Programs/UnrealGameSync/UnrealGameSync/Forms/SdkInfoWindow.cs b/Engine/Source/Programs/UnrealGameSync/UnrealGameSync/Forms/SdkInfoWindow.cs
--- a/Engine/Source/Programs/UnrealGameSync/UnrealGameSync/Forms/SdkInfoWindow.cs
+++ b/Engine/Source/Programs/UnrealGameSync/UnrealGameSync/Forms/SdkInfoWindow.cs
@@ -245,7 +245,10 @@ namespace UnrealGameSync
 		{
 			try
 			{
-				Process.Start(FileName);
+				ProcessStartInfo StartInfo = new ProcessStartInfo();
+				StartInfo.FileName = FileName;
+				StartInfo.UseShellExecute = true;
+				Process.Start(StartInfo);
 			}
 			catch(Exception Ex)
 			{
diff --git a/Engine/Source/Programs/UnrealGameSync/UnrealGameSync/Properties/AssemblyInfo.cs b/Engine/Source/Programs/UnrealGameSync/UnrealGameSync/Properties/AssemblyInfo.cs
--- a/Engine/Source/Programs/UnrealGameSync/UnrealGameSync/Properties/AssemblyInfo.cs
+++ b/Engine/Source/Programs/UnrealGameSync/UnrealGameSync/Properties/AssemblyInfo.cs
@@ -34,5 +34,5 @@ using System.Runtime.InteropServices;
 // You can specify all the values or you can default the Build and Revision Numbers 
 // by using the '*' as shown below:
 // [assembly: AssemblyVersion("1.0.*")]
-[assembly: AssemblyVersion("1.244.0.0")]
-[assembly: AssemblyFileVersion("1.244.0.0")]
+[assembly: AssemblyVersion("1.245.0.0")]
+[assembly: AssemblyFileVersion("1.245.0.0")]
