UGS: Fix exception when closing the issue details window.

diff --git a/Engine/Source/Programs/UnrealGameSync/UnrealGameSync/Forms/IssueDetailsWindow.cs b/Engine/Source/Programs/UnrealGameSync/UnrealGameSync/Forms/IssueDetailsWindow.cs
--- a/Engine/Source/Programs/UnrealGameSync/UnrealGameSync/Forms/IssueDetailsWindow.cs
+++ b/Engine/Source/Programs/UnrealGameSync/UnrealGameSync/Forms/IssueDetailsWindow.cs
@@ -104,6 +104,17 @@ namespace UnrealGameSync
 			}
 
 			void Run()
+			{
+				try
+				{
+					RunInternal();
+				}
+				catch (ThreadInterruptedException)
+				{
+				}
+			}
+
+			void RunInternal()
 			{
 				List<PerforceChangeRange> CompletedRequests = new List<PerforceChangeRange>(); 
 				while(!bTerminate)
@@ -175,13 +186,7 @@ namespace UnrealGameSync
 					}
 
 					// Wait for something to change
-					try
-					{
-						RefreshEvent.WaitOne();
-					}
-					catch (ThreadInterruptedException)
-					{
-					}
+					RefreshEvent.WaitOne();
 				}
 			}
 
diff --git a/Engine/Source/Programs/UnrealGameSync/UnrealGameSync/Properties/AssemblyInfo.cs b/Engine/Source/Programs/UnrealGameSync/UnrealGameSync/Properties/AssemblyInfo.cs
--- a/Engine/Source/Programs/UnrealGameSync/UnrealGameSync/Properties/AssemblyInfo.cs
+++ b/Engine/Source/Programs/UnrealGameSync/UnrealGameSync/Properties/AssemblyInfo.cs
@@ -34,5 +34,5 @@ using System.Runtime.InteropServices;
 // You can specify all the values or you can default the Build and Revision Numbers 
 // by using the '*' as shown below:
 // [assembly: AssemblyVersion("1.0.*")]
-[assembly: AssemblyVersion("1.245.0.0")]
-[assembly: AssemblyFileVersion("1.245.0.0")]
+[assembly: AssemblyVersion("1.246.0.0")]
+[assembly: AssemblyFileVersion("1.246.0.0")]
