Fix for possible timing bug that causes a null LogWriter to get used.

diff --git a/Engine/Source/Programs/UnrealGameSync/UnrealGameSync/PerforceMonitor.cs b/Engine/Source/Programs/UnrealGameSync/UnrealGameSync/PerforceMonitor.cs
--- a/Engine/Source/Programs/UnrealGameSync/UnrealGameSync/PerforceMonitor.cs
+++ b/Engine/Source/Programs/UnrealGameSync/UnrealGameSync/PerforceMonitor.cs
@@ -280,7 +280,10 @@ namespace UnrealGameSync
 				}
 				catch (Exception Ex)
 				{
-					LogWriter.WriteException(Ex, "Unhandled exception in PollForUpdatesInner()");
+					if(!bDisposing)
+					{
+						LogWriter.WriteException(Ex, "Unhandled exception in PollForUpdatesInner()");
+					}
 				}
 			}
 		}
