UnrealGameSync: Add description badges for Jira, RoboMerge, & Reviewed
Ensure duplicate description badges aren't added

diff --git a/Engine/Source/Programs/UnrealGameSync/UnrealGameSync/Controls/WorkspaceControl.cs b/Engine/Source/Programs/UnrealGameSync/UnrealGameSync/Controls/WorkspaceControl.cs
--- a/Engine/Source/Programs/UnrealGameSync/UnrealGameSync/Controls/WorkspaceControl.cs
+++ b/Engine/Source/Programs/UnrealGameSync/UnrealGameSync/Controls/WorkspaceControl.cs
@@ -2454,7 +2454,7 @@ namespace UnrealGameSync
 				if (ProjectConfigFile != null)
 				{
 					string[] BadgeDefinitions = ProjectConfigFile.GetValues("Badges.DescriptionBadges", new string[0]);
-					foreach (string BadgeDefinition in BadgeDefinitions)
+					foreach (string BadgeDefinition in BadgeDefinitions.Distinct())
 					{
 						ConfigObject BadgeDefinitionObject = new ConfigObject(BadgeDefinition);
 						string Pattern = BadgeDefinitionObject.GetValue("Pattern", null);
diff --git a/Engine/Source/Programs/UnrealGameSync/UnrealGameSync/Properties/AssemblyInfo.cs b/Engine/Source/Programs/UnrealGameSync/UnrealGameSync/Properties/AssemblyInfo.cs
--- a/Engine/Source/Programs/UnrealGameSync/UnrealGameSync/Properties/AssemblyInfo.cs
+++ b/Engine/Source/Programs/UnrealGameSync/UnrealGameSync/Properties/AssemblyInfo.cs
@@ -34,5 +34,5 @@ using System.Runtime.InteropServices;
 // You can specify all the values or you can default the Build and Revision Numbers 
 // by using the '*' as shown below:
 // [assembly: AssemblyVersion("1.0.*")]
-[assembly: AssemblyVersion("1.246.0.0")]
-[assembly: AssemblyFileVersion("1.246.0.0")]
+[assembly: AssemblyVersion("1.247.0.0")]
+[assembly: AssemblyFileVersion("1.247.0.0")]
