UGS: Allow adding build steps that reference custom tools.

diff --git a/Engine/Source/Programs/UnrealGameSync/UnrealGameSync/BuildStep.cs b/Engine/Source/Programs/UnrealGameSync/UnrealGameSync/BuildStep.cs
--- a/Engine/Source/Programs/UnrealGameSync/UnrealGameSync/BuildStep.cs
+++ b/Engine/Source/Programs/UnrealGameSync/UnrealGameSync/BuildStep.cs
@@ -37,6 +37,7 @@ namespace UnrealGameSync
 		public bool bNormalSync;
 		public bool bScheduledSync;
 		public bool bShowAsTool;
+		public Guid ToolId;
 
 		public BuildStep(Guid InUniqueId, int InOrderIndex, string InDescription, string InStatusText, int InEstimatedDuration, string InFileName, string InArguments, string InWorkingDir, bool bInUseLogWindow)
 		{
@@ -116,6 +117,10 @@ namespace UnrealGameSync
 			{
 				bShowAsTool = false;
 			}
+			if (!Guid.TryParse(Object.GetValue("Tool", ""), out ToolId))
+			{
+				ToolId = Guid.Empty;
+			}
 		}
 
 		public bool IsValid()
@@ -180,6 +185,10 @@ namespace UnrealGameSync
 			Result["bNormalSync"] = bNormalSync.ToString();
 			Result["bScheduledSync"] = bScheduledSync.ToString();
 			Result["bShowAsTool"] = bShowAsTool.ToString();
+			if (ToolId != Guid.Empty)
+			{
+				Result["Tool"] = ToolId.ToString();
+			}
 			return Result;
 		}
 
diff --git a/Engine/Source/Programs/UnrealGameSync/UnrealGameSync/Controls/WorkspaceControl.cs b/Engine/Source/Programs/UnrealGameSync/UnrealGameSync/Controls/WorkspaceControl.cs
--- a/Engine/Source/Programs/UnrealGameSync/UnrealGameSync/Controls/WorkspaceControl.cs
+++ b/Engine/Source/Programs/UnrealGameSync/UnrealGameSync/Controls/WorkspaceControl.cs
@@ -5076,10 +5076,10 @@ namespace UnrealGameSync
 					List<BuildStep> UserSteps = GetUserBuildSteps(ProjectBuildStepObjects);
 					foreach (BuildStep Step in UserSteps)
 					{
-						if (Step.bShowAsTool)
+						if (Step.bShowAsTool && (Step.ToolId == Guid.Empty || Settings.EnabledTools.Contains(Step.ToolId)))
 						{
 							ToolStripMenuItem NewMenuItem = new ToolStripMenuItem(Step.Description.Replace("&", "&&"));
-							NewMenuItem.Click += new EventHandler((sender, e) => { RunCustomTool(Step.UniqueId); });
+							NewMenuItem.Click += new EventHandler((sender, e) => { RunCustomTool(Step); });
 							CustomToolMenuItems.Add(NewMenuItem);
 							MoreToolsContextMenu.Items.Insert(InsertIdx++, NewMenuItem);
 						}
@@ -5091,7 +5091,7 @@ namespace UnrealGameSync
 			MoreActionsContextMenu_CustomToolSeparator.Visible = (CustomToolMenuItems.Count > 0);
 		}
 
-		private void RunCustomTool(Guid UniqueId)
+		private void RunCustomTool(BuildStep Step)
 		{
 			if (Workspace != null)
 			{
@@ -5101,7 +5101,17 @@ namespace UnrealGameSync
 				}
 				else
 				{
-					WorkspaceUpdateContext Context = new WorkspaceUpdateContext(Workspace.CurrentChangeNumber, WorkspaceUpdateOptions.Build, null, GetDefaultBuildStepObjects(), ProjectSettings.BuildSteps, new HashSet<Guid> { UniqueId }, GetWorkspaceVariables(Workspace.CurrentChangeNumber));
+					Dictionary<string, string> Variables = GetWorkspaceVariables(Workspace.CurrentChangeNumber);
+					if (Step.ToolId != Guid.Empty)
+					{
+						string ToolName = Owner.ToolUpdateMonitor.GetToolName(Step.ToolId);
+						if (ToolName != null)
+						{
+							Variables["ToolDir"] = Owner.ToolUpdateMonitor.GetToolPath(ToolName);
+						}
+					}
+
+					WorkspaceUpdateContext Context = new WorkspaceUpdateContext(Workspace.CurrentChangeNumber, WorkspaceUpdateOptions.Build, null, GetDefaultBuildStepObjects(), ProjectSettings.BuildSteps, new HashSet<Guid> { Step.UniqueId }, Variables);
 					StartWorkspaceUpdate(Context, null);
 				}
 			}
diff --git a/Engine/Source/Programs/UnrealGameSync/UnrealGameSync/ToolUpdateMonitor.cs b/Engine/Source/Programs/UnrealGameSync/UnrealGameSync/ToolUpdateMonitor.cs
--- a/Engine/Source/Programs/UnrealGameSync/UnrealGameSync/ToolUpdateMonitor.cs
+++ b/Engine/Source/Programs/UnrealGameSync/UnrealGameSync/ToolUpdateMonitor.cs
@@ -98,6 +98,18 @@ namespace UnrealGameSync
 			return Path.Combine(ToolsDir, ToolName, "Current");
 		}
 
+		public string GetToolName(Guid ToolId)
+		{
+			foreach (ToolDefinition Tool in Tools)
+			{
+				if(Tool.Id == ToolId)
+				{
+					return Tool.Name;
+				}
+			}
+			return null;
+		}
+
 		public string GetToolPath(string ToolName)
 		{
 			if (GetToolChange(ToolName) != 0)
