Force InterchangeWorker to be build when building the editor

diff --git a/Engine/Plugins/Runtime/Interchange/Source/Dispatcher/Private/InterchangeWorkerHandler.cpp b/Engine/Plugins/Runtime/Interchange/Source/Dispatcher/Private/InterchangeWorkerHandler.cpp
--- a/Engine/Plugins/Runtime/Interchange/Source/Dispatcher/Private/InterchangeWorkerHandler.cpp
+++ b/Engine/Plugins/Runtime/Interchange/Source/Dispatcher/Private/InterchangeWorkerHandler.cpp
@@ -29,14 +29,25 @@ namespace UE
 				Path = FPaths::Combine(Path, TEXT("Linux/InterchangeWorker"));
 #elif PLATFORM_WINDOWS
 #if UE_BUILD_DEBUG
-				Path = FPaths::Combine(Path, TEXT("Win64/InterchangeWorker-Win64-Debug.exe"));
+				{
+					FString DebugPath = FPaths::Combine(Path, TEXT("Win64/InterchangeWorker-Win64-Debug.exe"));
+					if (!FPaths::FileExists(DebugPath))
+					{
+						//Try the development build if we cannot find the debug build
+						Path = FPaths::Combine(Path, TEXT("Win64/InterchangeWorker.exe"));
+					}
+					else
+					{
+						Path = MoveTemp(DebugPath);
+					}
+				}
 #else //#if UE_BUILD_DEBUG
 				Path = FPaths::Combine(Path, TEXT("Win64/InterchangeWorker.exe"));
 #endif //#else UE_BUILD_DEBUG
 #endif //#elif PLATFORM_WINDOWS
 				if (!FPaths::FileExists(Path))
 				{
-					UE_LOG(LogInterchangeDispatcher, Warning, TEXT("InterchangeWorker executable not found. Expected location: %s"), *FPaths::ConvertRelativePathToFull(ProcessorPath));
+					UE_LOG(LogInterchangeDispatcher, Warning, TEXT("InterchangeWorker executable not found. Expected location: %s"), *FPaths::ConvertRelativePathToFull(Path));
 				}
 				return Path;
 			}();
diff --git a/Engine/Source/Programs/UnrealBuildTool/ProjectFiles/KDevelop/KDevelopGenerator.cs b/Engine/Source/Programs/UnrealBuildTool/ProjectFiles/KDevelop/KDevelopGenerator.cs
--- a/Engine/Source/Programs/UnrealBuildTool/ProjectFiles/KDevelop/KDevelopGenerator.cs
+++ b/Engine/Source/Programs/UnrealBuildTool/ProjectFiles/KDevelop/KDevelopGenerator.cs
@@ -166,14 +166,14 @@ namespace UnrealBuildTool
 
 			FileContent.Append("Title=BuildMeFirst\n\n");
 			FileContent.Append("[CustomBuildSystem][BuildConfig0][ToolBuild]\n");
-			FileContent.Append("Arguments=-f Makefile UnrealEditor UnrealGame ShaderCompileWorker UnrealLightmass UnrealPak\n");
+			FileContent.Append("Arguments=-f Makefile UnrealEditor UnrealGame ShaderCompileWorker UnrealLightmass InterchangeWorker UnrealPak\n");
 			FileContent.Append("Enabled=true\n");
 			FileContent.Append("Environment=\n");
 			FileContent.Append("Executable=make\n");
 			FileContent.Append("Type=0\n\n");
 
 			FileContent.Append("[CustomBuildSystem][BuildConfig0][ToolClean]\n");
-			FileContent.Append("Arguments=-f Makefile UnrealEditor UnrealGame ShaderCompileWorker UnrealLightmass UnrealPak -clean\n");
+			FileContent.Append("Arguments=-f Makefile UnrealEditor UnrealGame ShaderCompileWorker UnrealLightmass InterchangeWorker UnrealPak -clean\n");
 			FileContent.Append("Enabled=true\n");
 			FileContent.Append("Environment=\n");
 			FileContent.Append("Executable=make\n");
diff --git a/Engine/Source/Programs/UnrealBuildTool/ProjectFiles/Make/MakefileGenerator.cs b/Engine/Source/Programs/UnrealBuildTool/ProjectFiles/Make/MakefileGenerator.cs
--- a/Engine/Source/Programs/UnrealBuildTool/ProjectFiles/Make/MakefileGenerator.cs
+++ b/Engine/Source/Programs/UnrealBuildTool/ProjectFiles/Make/MakefileGenerator.cs
@@ -128,7 +128,7 @@ namespace UnrealBuildTool
 
 			MakefileContent.Append("\n\n" + BuildCommand + ProjectBuildCommand + "\n" +
 				"all: StandardSet\n\n" +
-				"RequiredTools: CrashReportClient-Linux-Shipping CrashReportClientEditor-Linux-Shipping ShaderCompileWorker UnrealLightmass EpicWebHelper-Linux-Shipping\n\n" +
+				"RequiredTools: CrashReportClient-Linux-Shipping CrashReportClientEditor-Linux-Shipping ShaderCompileWorker UnrealLightmass InterchangeWorker EpicWebHelper-Linux-Shipping\n\n" +
 				"StandardSet: RequiredTools UnrealFrontend UnrealEditor UnrealInsights\n\n" +
 				"DebugSet: RequiredTools UnrealFrontend-Linux-Debug UnrealEditor-Linux-Debug\n\n"
 			);
diff --git a/Engine/Source/Programs/UnrealBuildTool/ProjectFiles/VisualStudio/VCProject.cs b/Engine/Source/Programs/UnrealBuildTool/ProjectFiles/VisualStudio/VCProject.cs
--- a/Engine/Source/Programs/UnrealBuildTool/ProjectFiles/VisualStudio/VCProject.cs
+++ b/Engine/Source/Programs/UnrealBuildTool/ProjectFiles/VisualStudio/VCProject.cs
@@ -422,7 +422,7 @@ namespace UnrealBuildTool
 					foreach (ProjectTarget ProjectTarget in ProjectTargets)
 					{
 						string TargetName = ProjectTargets[0].TargetRules.Name;
-						if (TargetName == "ShaderCompileWorker" || TargetName == "UnrealLightmass")
+						if (TargetName == "ShaderCompileWorker" || TargetName == "UnrealLightmass" || TargetName == "InterchangeWorker")
 						{
 							MatchingProjectTargets.Add(ProjectTarget);
 							break;
@@ -1660,6 +1660,7 @@ namespace UnrealBuildTool
 						if (TargetRulesObject.Type == TargetType.Editor && Settings.bEditorDependsOnShaderCompileWorker && !UnrealBuildTool.IsEngineInstalled())
 						{
 							ExtraTargets.Add("ShaderCompileWorker Win64 Development");
+							ExtraTargets.Add("InterchangeWorker Win64 Development");
 						}
 						if (TargetRulesObject.bWithLiveCoding && Settings.bBuildLiveCodingConsole && !UnrealBuildTool.IsEngineInstalled() && TargetRulesObject.Name != "LiveCodingConsole")
 						{
diff --git a/Engine/Source/Programs/UnrealGameSync/UnrealGameSync/Controls/WorkspaceControl.cs b/Engine/Source/Programs/UnrealGameSync/UnrealGameSync/Controls/WorkspaceControl.cs
--- a/Engine/Source/Programs/UnrealGameSync/UnrealGameSync/Controls/WorkspaceControl.cs
+++ b/Engine/Source/Programs/UnrealGameSync/UnrealGameSync/Controls/WorkspaceControl.cs
@@ -3122,6 +3122,8 @@ namespace UnrealGameSync
 			ProcessNames.Add("CrashReportClientEditor-Win64-Debug");
 			ProcessNames.Add("UnrealBuildTool");
 			ProcessNames.Add("AutomationTool");
+			ProcessNames.Add("InterchangeWorker");
+			ProcessNames.Add("InterchangeWorker-Win64-Debug");
 
 			List<string> ProcessFileNames = new List<string>();
 			try
@@ -5266,6 +5268,7 @@ namespace UnrealGameSync
 			DefaultBuildSteps.Add(new BuildStep(new Guid("{24FFD88C-7901-4899-9696-AE1066B4B6E8}"), 3, "Compile UnrealLightmass", "Compiling UnrealLightmass...", 1, "UnrealLightmass", "Win64", "Development", "", !ShouldSyncPrecompiledEditor));
 			DefaultBuildSteps.Add(new BuildStep(new Guid("{FFF20379-06BF-4205-8A3E-C53427736688}"), 4, "Compile CrashReportClient", "Compiling CrashReportClient...", 1, "CrashReportClient", "Win64", "Shipping", "", !ShouldSyncPrecompiledEditor && !bUseCrashReportClientEditor));
 			DefaultBuildSteps.Add(new BuildStep(new Guid("{7143D861-58D3-4F83-BADC-BC5DCB2079F6}"), 5, "Compile CrashReportClientEditor", "Compiling CrashReportClientEditor...", 1, "CrashReportClientEditor", "Win64", "Shipping", "", !ShouldSyncPrecompiledEditor && bUseCrashReportClientEditor));
+			DefaultBuildSteps.Add(new BuildStep(new Guid("{C434716F-D61F-4209-B900-B6D942751453}"), 6, "Compile InterchangeWorker", "Compiling InterchangeWorker...", 1, "InterchangeWorker", "Win64", "Development", "", !ShouldSyncPrecompiledEditor));
 
 			return DefaultBuildSteps.ToDictionary(x => x.UniqueId, x => x.ToConfigObject());
 		}
