UGS: Add script to build and version UGS.

diff --git a/Engine/Source/Programs/UnrealGameSync/BuildUGS.xml b/Engine/Source/Programs/UnrealGameSync/BuildUGS.xml
new file mode 100644
--- /dev/null
+++ b/Engine/Source/Programs/UnrealGameSync/BuildUGS.xml
@@ -0,0 +1,22 @@
+<?xml version='1.0' ?>
+<BuildGraph xmlns="http://www.epicgames.com/BuildGraph" xmlns:xsi="http://www.w3.org/2001/XMLSchema-instance" xsi:schemaLocation="http://www.epicgames.com/BuildGraph ../../../Build/Graph/Schema.xsd" >
+
+	<Option Name="PreflightChange" DefaultValue="" Description="Preflight changelist number"/>
+
+	<Property Name="UgsDir" Value="$(RootDir)/Engine/Source/Programs/UnrealGameSync"/>
+	<Property Name="UgsOutputDir" Value="$(RootDir)/LocalBuilds/UnrealGameSync"/>
+	<Property Name="WixDir" Value="$(RootDir)/Engine/Source/ThirdParty/WiX/3.8"/>
+
+	<Property Name="Version" Value="$(EngineMajorVersion).$(EngineMinorVersion).$(EnginePatchVersion)"/>
+	<Property Name="InformationalVersion" Value="$(Version)-$(Change)"/>
+	<Property Name="InformationalVersion" Value="$(InformationalVersion)-PF-$(PreflightChange)" If="'$(PreflightChange)' != ''"/>
+	<Property Name="VersionArguments" Value="/p:Version=$(Version).0 /p:InformationalVersion=$(InformationalVersion)"/>
+	
+	<Agent Name="UnrealGameSync Agent" Type="Win64">
+		<Node Name="Build UnrealGameSync">
+			<Delete Files="$(UgsOutputDir)/..."/>
+			<DotNet Arguments="publish &quot;$(UgsDir)/UnrealGameSync/UnrealGameSync.csproj&quot; --output=&quot;$(UgsOutputDir)&quot; --no-self-contained -r win-x64 $(VersionArguments)"/>
+		</Node>
+	</Agent>
+	
+</BuildGraph>
diff --git a/Engine/Source/Programs/UnrealGameSync/UnrealGameSync/Properties/AssemblyInfo.cs b/Engine/Source/Programs/UnrealGameSync/UnrealGameSync/Properties/AssemblyInfo.cs
deleted file mode 100644
diff --git a/Engine/Source/Programs/UnrealGameSync/UnrealGameSync/UnrealGameSync.csproj b/Engine/Source/Programs/UnrealGameSync/UnrealGameSync/UnrealGameSync.csproj
--- a/Engine/Source/Programs/UnrealGameSync/UnrealGameSync/UnrealGameSync.csproj
+++ b/Engine/Source/Programs/UnrealGameSync/UnrealGameSync/UnrealGameSync.csproj
@@ -3,8 +3,10 @@
     <OutputType>WinExe</OutputType>
     <TargetFramework>netcoreapp3.1</TargetFramework>
     <UseWindowsForms>true</UseWindowsForms>
-    <GenerateAssemblyInfo>false</GenerateAssemblyInfo>
     <Deterministic>true</Deterministic>
+    <Product>UnrealGameSync</Product>
+    <Company>Epic Games, Inc</Company>
+    <Copyright>Copyright Epic Games, Inc. All Rights Reserved.</Copyright>
     <RootNamespace>UnrealGameSync</RootNamespace>
     <AssemblyName>UnrealGameSync</AssemblyName>
     <StartupObject>UnrealGameSync.Program</StartupObject>
