UGS: Display the informational version string inside the application, and allow configuring the path to submit to from the build script.

diff --git a/Engine/Source/Programs/AutomationTool/BuildGraph/Tasks/SubmitTask.cs b/Engine/Source/Programs/AutomationTool/BuildGraph/Tasks/SubmitTask.cs
--- a/Engine/Source/Programs/AutomationTool/BuildGraph/Tasks/SubmitTask.cs
+++ b/Engine/Source/Programs/AutomationTool/BuildGraph/Tasks/SubmitTask.cs
@@ -129,7 +129,7 @@ namespace AutomationTool.Tasks
 					Client.Name = $"{Parameters.Workspace}_{Regex.Replace(Client.Host, "[^a-zA-Z0-9]", "-")}_{ContentHash.MD5((CommandUtils.P4Env.ServerAndPort ?? "").ToUpperInvariant())}";
 					Client.Options = P4ClientOption.NoAllWrite | P4ClientOption.Clobber | P4ClientOption.NoCompress | P4ClientOption.Unlocked | P4ClientOption.NoModTime | P4ClientOption.RmDir;
 					Client.LineEnd = P4LineEnd.Local;
-					if (Parameters.Branch != null)
+					if (!String.IsNullOrEmpty(Parameters.Branch))
 					{
 						Client.View.Add(new KeyValuePair<string, string>($"{Parameters.Branch}/...", $"/..."));
 					}
diff --git a/Engine/Source/Programs/UnrealGameSync/BuildUGS.xml b/Engine/Source/Programs/UnrealGameSync/BuildUGS.xml
--- a/Engine/Source/Programs/UnrealGameSync/BuildUGS.xml
+++ b/Engine/Source/Programs/UnrealGameSync/BuildUGS.xml
@@ -2,20 +2,33 @@
 <BuildGraph xmlns="http://www.epicgames.com/BuildGraph" xmlns:xsi="http://www.w3.org/2001/XMLSchema-instance" xsi:schemaLocation="http://www.epicgames.com/BuildGraph ../../../Build/Graph/Schema.xsd" >
 
 	<Option Name="PreflightChange" DefaultValue="" Description="Preflight changelist number"/>
+	
+	<Option Name="SubmitToStream" DefaultValue="" Description="Stream to submit binaries to"/>
+	<Option Name="SubmitToBranch" DefaultValue="" Description="Branch to submit binaries to"/>
+	<Option Name="SubmitPath" DefaultValue="" Description="Relative path within the stream/branch to submit binaries"/>
 
 	<Property Name="UgsDir" Value="$(RootDir)/Engine/Source/Programs/UnrealGameSync"/>
 	<Property Name="UgsOutputDir" Value="$(RootDir)/LocalBuilds/UnrealGameSync"/>
 	<Property Name="WixDir" Value="$(RootDir)/Engine/Source/ThirdParty/WiX/3.8"/>
 
 	<Property Name="Version" Value="$(EngineMajorVersion).$(EngineMinorVersion).$(EnginePatchVersion)"/>
-	<Property Name="InformationalVersion" Value="$(Version)-$(Change)"/>
+	<Property Name="InformationalVersion" Value="$(EscapedBranch)-CL-$(Change)"/>
 	<Property Name="InformationalVersion" Value="$(InformationalVersion)-PF-$(PreflightChange)" If="'$(PreflightChange)' != ''"/>
 	<Property Name="VersionArguments" Value="/p:Version=$(Version).0 /p:InformationalVersion=$(InformationalVersion)"/>
 	
 	<Agent Name="UnrealGameSync Agent" Type="Win64">
+		<Property Name="UgsPublishDir" Value="$(UgsOutputDir)"/>
+		<Property Name="UgsPublishDir" Value="$(UgsPublishDir)/$(SubmitPath)" If="'$(SubmitPath)' != ''"/>
+	
 		<Node Name="Build UnrealGameSync">
 			<Delete Files="$(UgsOutputDir)/..."/>
-			<DotNet Arguments="publish &quot;$(UgsDir)/UnrealGameSync/UnrealGameSync.csproj&quot; --output=&quot;$(UgsOutputDir)&quot; --no-self-contained -r win-x64 $(VersionArguments)"/>
+			<DotNet Arguments="publish &quot;$(UgsDir)/UnrealGameSync/UnrealGameSync.csproj&quot; --output=&quot;$(UgsPublishDir)&quot; --no-self-contained -r win-x64 $(VersionArguments)"/>
+		</Node>
+		<Node Name="Publish Unstable" Requires="Build UnrealGameSync">
+			<Submit Description="Updating UGS from CL $(Change)" Workspace="Temp" Files="$(UgsPublishDir)/..." RootDir="$(UgsOutputDir)" Stream="$(SubmitToStream)" Branch="$(SubmitToBranch)"/>
+		</Node>
+		<Node Name="Publish Live" Requires="Build UnrealGameSync">
+			<Submit Description="Updating UGS from CL $(Change)" Workspace="Temp" Files="$(UgsPublishDir)/..." RootDir="$(UgsOutputDir)" Stream="$(SubmitToStream)" Branch="$(SubmitToBranch)"/>
 		</Node>
 	</Agent>
 	
diff --git a/Engine/Source/Programs/UnrealGameSync/UnrealGameSync/Program.cs b/Engine/Source/Programs/UnrealGameSync/UnrealGameSync/Program.cs
--- a/Engine/Source/Programs/UnrealGameSync/UnrealGameSync/Program.cs
+++ b/Engine/Source/Programs/UnrealGameSync/UnrealGameSync/Program.cs
@@ -22,8 +22,8 @@ namespace UnrealGameSync
 	{
 		public static string GetVersionString()
 		{
-			Version Version = Assembly.GetExecutingAssembly().GetName().Version;
-			return $"{Version.Major}.{Version.Minor}-Core";
+			AssemblyInformationalVersionAttribute Version = Assembly.GetExecutingAssembly().GetCustomAttribute<AssemblyInformationalVersionAttribute>();
+			return Version?.InformationalVersion ?? "Unknown";
 		}
 
 		static JsonSerializerOptions GetDefaultJsonSerializerOptions()
