Reconcile file extensions to treat as code between UGS and Horde. Add .csproj & .sln.

diff --git a/Engine/Source/Programs/Horde/HordeServer/IssueHandlers/IIssueHandler.cs b/Engine/Source/Programs/Horde/HordeServer/IssueHandlers/IIssueHandler.cs
--- a/Engine/Source/Programs/Horde/HordeServer/IssueHandlers/IIssueHandler.cs
+++ b/Engine/Source/Programs/Horde/HordeServer/IssueHandlers/IIssueHandler.cs
@@ -28,11 +28,15 @@ namespace HordeServer.IssueHandlers
 			".mm",
 			".rc",
 			".cs",
+			".csproj",
 			".h",
 			".hpp",
 			".inl",
+			".usf",
+			".ush",
 			".uproject",
-			".uplugin"
+			".uplugin",
+			".sln"
 		};
 
 		/// <summary>
diff --git a/Engine/Source/Programs/Horde/HordeServer/Models/Change.cs b/Engine/Source/Programs/Horde/HordeServer/Models/Change.cs
--- a/Engine/Source/Programs/Horde/HordeServer/Models/Change.cs
+++ b/Engine/Source/Programs/Horde/HordeServer/Models/Change.cs
@@ -72,6 +72,29 @@ namespace HordeServer.Models
 	/// </summary>
 	public class ChangeDetails
 	{
+		/// <summary>
+		/// Set of extensions to treat as code
+		/// </summary>
+		static readonly HashSet<string> CodeExtensions = new HashSet<string>
+		{
+			".c",
+			".cc",
+			".cpp",
+			".m",
+			".mm",
+			".rc",
+			".cs",
+			".csproj",
+			".h",
+			".hpp",
+			".inl",
+			".usf",
+			".ush",
+			".uproject",
+			".uplugin",
+			".sln"
+		};
+
 		/// <summary>
 		/// The source changelist number
 		/// </summary>
@@ -126,7 +149,6 @@ namespace HordeServer.Models
 			ChangeContentFlags Scope = 0;
 
 			// Check whether the files are code or content
-			string[] CodeExtensions = { ".cs", ".h", ".cpp", ".inl", ".usf", ".ush", ".uproject", ".uplugin" };
 			foreach (string File in Files)
 			{
 				if (CodeExtensions.Any(Extension => File.EndsWith(Extension, StringComparison.OrdinalIgnoreCase)))
diff --git a/Engine/Source/Programs/UnrealGameSync/UnrealGameSync/PerforceMonitor.cs b/Engine/Source/Programs/UnrealGameSync/UnrealGameSync/PerforceMonitor.cs
--- a/Engine/Source/Programs/UnrealGameSync/UnrealGameSync/PerforceMonitor.cs
+++ b/Engine/Source/Programs/UnrealGameSync/UnrealGameSync/PerforceMonitor.cs
@@ -14,6 +14,29 @@ namespace UnrealGameSync
 {
 	class PerforceChangeDetails
 	{
+		/// <summary>
+		/// Set of extensions to treat as code
+		/// </summary>
+		static readonly HashSet<string> CodeExtensions = new HashSet<string>
+		{
+			".c",
+			".cc",
+			".cpp",
+			".m",
+			".mm",
+			".rc",
+			".cs",
+			".csproj",
+			".h",
+			".hpp",
+			".inl",
+			".usf",
+			".ush",
+			".uproject",
+			".uplugin",
+			".sln"
+		};
+
 		public string Description;
 		public bool bContainsCode;
 		public bool bContainsContent;
@@ -23,7 +46,6 @@ namespace UnrealGameSync
 			Description = DescribeRecord.Description;
 
 			// Check whether the files are code or content
-			string[] CodeExtensions = { ".cs", ".h", ".cpp", ".inl", ".usf", ".ush", ".uproject", ".uplugin" };
 			foreach(PerforceDescribeFileRecord File in DescribeRecord.Files)
 			{
 				if(CodeExtensions.Any(Extension => File.DepotFile.EndsWith(Extension, StringComparison.InvariantCultureIgnoreCase)))
