UGS: Add an URI handler that allows remoting into a build machine.

ugs://rdp?host=<HOSTNAME>

Adding generic credentials called UnrealGameSync:RDP to the Windows credential store allows setting defaults for the username/password to be created for machines accessed using this URI, allowing a single click remote into a previously unvisited machine from the Horde dashboard.

diff --git a/Engine/Source/Programs/UnrealGameSync/UnrealGameSync/UriHandlers/RdpHandler.cs b/Engine/Source/Programs/UnrealGameSync/UnrealGameSync/UriHandlers/RdpHandler.cs
new file mode 100644
--- /dev/null
+++ b/Engine/Source/Programs/UnrealGameSync/UnrealGameSync/UriHandlers/RdpHandler.cs
@@ -0,0 +1,136 @@
+// Copyright Epic Games, Inc. All Rights Reserved.
+
+using System;
+using System.Collections.Generic;
+using Microsoft.Win32;
+using System.Reflection;
+using System.Threading;
+using System.Collections.Specialized;
+using System.Web;
+using System.Linq;
+using System.Diagnostics;
+using System.Windows.Forms;
+using System.Configuration;
+using System.IO;
+using System.Runtime.InteropServices;
+using System.Runtime.InteropServices.ComTypes;
+using System.Text;
+using System.Security;
+
+namespace UnrealGameSync
+{
+	static class RdpHandler
+	{
+		const uint CRED_TYPE_GENERIC = 1;
+		const uint CRED_TYPE_DOMAIN_PASSWORD = 2;
+		const uint CRED_PERSIST_LOCAL_MACHINE = 2;
+
+		[StructLayout(LayoutKind.Sequential, CharSet = CharSet.Auto)]
+		struct CREDENTIAL_ATTRIBUTE
+		{
+			string Keyword;
+			uint Flags;
+			uint ValueSize;
+			IntPtr Value;
+		}
+
+		[StructLayout(LayoutKind.Sequential, CharSet = CharSet.Unicode)]
+		class CREDENTIAL
+		{
+			public uint Flags;
+			public uint Type;
+			public string TargetName;
+			public string Comment;
+			public FILETIME LastWritten;
+			public int CredentialBlobSize;
+			public IntPtr CredentialBlob;
+			public uint Persist;
+			public uint AttributeCount;
+			public IntPtr Attributes;
+			public string TargetAlias;
+			public string UserName;
+		}
+
+		[DllImport("Advapi32.dll", EntryPoint = "CredReadW", CharSet = CharSet.Unicode, SetLastError = true)]
+		static extern bool CredRead(string Target, uint Type, int ReservedFlag, out IntPtr Buffer);
+
+		[DllImport("Advapi32.dll", EntryPoint = "CredWriteW", CharSet = CharSet.Unicode, SetLastError = true)]
+		static extern bool CredWrite(CREDENTIAL userCredential, uint Flags);
+
+		[DllImport("advapi32.dll", SetLastError = true)]
+		static extern bool CredFree(IntPtr Buffer);
+
+		public static bool ReadCredential(string Name, uint Type, out string UserName, out string Password)
+		{
+			IntPtr Buffer = IntPtr.Zero;
+			try
+			{
+				if (CredRead(Name, Type, 0, out Buffer))
+				{
+					CREDENTIAL Credential = Marshal.PtrToStructure<CREDENTIAL>(Buffer);
+					UserName = Credential.UserName;
+					Password = Marshal.PtrToStringUni(Credential.CredentialBlob, Credential.CredentialBlobSize / sizeof(char));
+					return true;
+				}
+				else
+				{
+					UserName = null;
+					Password = null;
+					return false;
+				}
+			}
+			finally
+			{
+				if (Buffer != IntPtr.Zero)
+				{
+					CredFree(Buffer);
+				}
+			}
+		}
+
+		public static void WriteCredential(string Name, uint Type, string UserName, string Password)
+		{
+			CREDENTIAL Credential = new CREDENTIAL();
+			try
+			{
+				Credential.Type = Type;
+				Credential.Persist = CRED_PERSIST_LOCAL_MACHINE;
+				Credential.CredentialBlobSize = Password.Length * sizeof(char);
+				Credential.TargetName = Name;
+				Credential.CredentialBlob = Marshal.StringToCoTaskMemUni(Password);
+				Credential.UserName = UserName;
+				CredWrite(Credential, 0);
+			}
+			finally
+			{
+				if (Credential.CredentialBlob != IntPtr.Zero)
+				{
+					Marshal.FreeCoTaskMem(Credential.CredentialBlob);
+				}
+			}
+		}
+
+		[UriHandler(Terminate = true)]
+		public static UriResult Rdp(string Host)
+		{
+			// Copy the credentials from a generic Windows credential to 
+			if (!ReadCredential(Host, CRED_TYPE_DOMAIN_PASSWORD, out _, out _))
+			{
+				if (ReadCredential("UnrealGameSync:RDP", CRED_TYPE_GENERIC, out string UserName, out string Password))
+				{
+					WriteCredential(Host, CRED_TYPE_DOMAIN_PASSWORD, UserName, Password);
+				}
+			}
+
+			using (Process Process = new Process())
+			{
+				Process.StartInfo.FileName = Path.Combine(Environment.GetFolderPath(Environment.SpecialFolder.System), "mstsc.exe");
+				Process.StartInfo.ArgumentList.Add($"/v:{Host}");
+				Process.Start();
+			}
+
+			return new UriResult() { Success = true };
+		}
+	}
+}
+
