Fix paths for submitting UGS.

diff --git a/Engine/Source/Programs/UnrealGameSync/BuildUGS.xml b/Engine/Source/Programs/UnrealGameSync/BuildUGS.xml
--- a/Engine/Source/Programs/UnrealGameSync/BuildUGS.xml
+++ b/Engine/Source/Programs/UnrealGameSync/BuildUGS.xml
@@ -25,10 +25,10 @@
 			<DotNet Arguments="publish &quot;$(UgsDir)/UnrealGameSync/UnrealGameSync.csproj&quot; --output=&quot;$(UgsPublishDir)&quot; --no-self-contained -r win-x64 $(VersionArguments)"/>
 		</Node>
 		<Node Name="Publish Unstable" Requires="Build UnrealGameSync">
-			<Submit Description="Updating UGS from CL $(Change)" Workspace="Temp" Files="$(UgsPublishDir)/..." RootDir="$(UgsOutputDir)" Stream="$(SubmitToStream)" Branch="$(SubmitToBranch)"/>
+			<Submit Description="Updating UGS from CL $(Change)" Workspace="Temp" Files="$(UgsPublishDir)/..." RootDir="$(UgsOutputDir)" Stream="$(SubmitToStream)" Branch="$(SubmitToBranch)/UnstableRelease"/>
 		</Node>
 		<Node Name="Publish Live" Requires="Build UnrealGameSync">
-			<Submit Description="Updating UGS from CL $(Change)" Workspace="Temp" Files="$(UgsPublishDir)/..." RootDir="$(UgsOutputDir)" Stream="$(SubmitToStream)" Branch="$(SubmitToBranch)"/>
+			<Submit Description="Updating UGS from CL $(Change)" Workspace="Temp" Files="$(UgsPublishDir)/..." RootDir="$(UgsOutputDir)" Stream="$(SubmitToStream)" Branch="$(SubmitToBranch)/Release"/>
 		</Node>
 	</Agent>
 	
