UGS: Start remote desktop sessions in fullscreen by default.

diff --git a/Engine/Source/Programs/UnrealGameSync/UnrealGameSync/UriHandlers/RdpHandler.cs b/Engine/Source/Programs/UnrealGameSync/UnrealGameSync/UriHandlers/RdpHandler.cs
--- a/Engine/Source/Programs/UnrealGameSync/UnrealGameSync/UriHandlers/RdpHandler.cs
+++ b/Engine/Source/Programs/UnrealGameSync/UnrealGameSync/UriHandlers/RdpHandler.cs
@@ -126,6 +126,7 @@ namespace UnrealGameSync
 			{
 				Process.StartInfo.FileName = Path.Combine(Environment.GetFolderPath(Environment.SpecialFolder.System), "mstsc.exe");
 				Process.StartInfo.ArgumentList.Add($"/v:{Host}");
+				Process.StartInfo.ArgumentList.Add($"/f");
 				Process.Start();
 			}
 
