UGS: Fixed GUI state management:

- Build "Investigating" events are now processed correctly.
- Build "Resolved" events are now processed correctly.
- Build "Unstarred" events are now processed correctly.

See CL 16886006 for Horde code changes.

diff --git a/Engine/Source/Programs/UnrealGameSync/UnrealGameSync/EventMonitor.cs b/Engine/Source/Programs/UnrealGameSync/UnrealGameSync/EventMonitor.cs
--- a/Engine/Source/Programs/UnrealGameSync/UnrealGameSync/EventMonitor.cs
+++ b/Engine/Source/Programs/UnrealGameSync/UnrealGameSync/EventMonitor.cs
@@ -152,8 +152,8 @@ namespace UnrealGameSync
 		public long? SyncTime { get; set; }
 		public UgsUserVote? Vote { get; set; }
 		public string Comment { get; set; }
-		public bool Investigating { get; set; }
-		public bool Starred { get; set; }
+		public bool? Investigating { get; set; }
+		public bool? Starred { get; set; }
 	}
 
 	class GetBadgeDataResponseV2
@@ -231,7 +231,6 @@ namespace UnrealGameSync
 		// MetadataV2
 		string MetadataStream;
 		string MetadataProject;
-		long IncomingMetadataId;
 		ConcurrentQueue<GetMetadataResponseV2> IncomingMetadata = new ConcurrentQueue<GetMetadataResponseV2>();
 		int MinChange;
 		int NewMinChange;
@@ -584,6 +583,7 @@ namespace UnrealGameSync
 				foreach (GetBadgeDataResponseV2 BadgeData in Metadata.Badges)
 				{
 					BadgeData Badge = new BadgeData();
+					Badge.Id = ++LatestIds.LastBuildId;
 					Badge.ChangeNumber = Metadata.Change;
 					Badge.BuildType = BadgeData.Name;
 					Badge.Result = BadgeData.State;
@@ -606,25 +606,27 @@ namespace UnrealGameSync
 					EventType? Type = GetEventTypeFromVote(UserData.Vote);
 					if (Type != null)
 					{
-						EventData Event = new EventData { Id = ++IncomingMetadataId, Change = Metadata.Change, Project = Metadata.Project, UserName = UserData.User, Type = Type.Value };
+						EventData Event = new EventData { Id = ++LatestIds.LastEventId, Change = Metadata.Change, Project = Metadata.Project, UserName = UserData.User, Type = Type.Value };
 						IncomingEvents.Enqueue(Event);
 					}
 
-					if (UserData.Investigating)
+					if (UserData.Investigating != null)
 					{
-						EventData Event = new EventData { Id = ++IncomingMetadataId, Change = Metadata.Change, Project = Metadata.Project, UserName = UserData.User, Type = EventType.Starred };
+						EventType InvestigationEventType = (UserData.Investigating.Value ? EventType.Investigating : EventType.Resolved);
+						EventData Event = new EventData { Id = ++LatestIds.LastEventId, Change = Metadata.Change, Project = Metadata.Project, UserName = UserData.User, Type = InvestigationEventType };
 						IncomingEvents.Enqueue(Event);
 					}
 
-					if (UserData.Starred)
+					if (UserData.Starred != null)
 					{
-						EventData Event = new EventData { Id = ++IncomingMetadataId, Change = Metadata.Change, Project = Metadata.Project, UserName = UserData.User, Type = EventType.Starred };
+						EventType StarEventType = (UserData.Starred.Value ? EventType.Starred : EventType.Unstarred);
+						EventData Event = new EventData { Id = ++LatestIds.LastEventId, Change = Metadata.Change, Project = Metadata.Project, UserName = UserData.User, Type = StarEventType };
 						IncomingEvents.Enqueue(Event);
 					}
 
 					if (UserData.Comment != null)
 					{
-						CommentData Comment = new CommentData { Id = ++IncomingMetadataId, ChangeNumber = Metadata.Change, Project = Metadata.Project, UserName = UserData.User, Text = UserData.Comment };
+						CommentData Comment = new CommentData { Id = ++LatestIds.LastCommentId, ChangeNumber = Metadata.Change, Project = Metadata.Project, UserName = UserData.User, Text = UserData.Comment };
 						IncomingComments.Enqueue(Comment);
 					}
 				}
@@ -929,7 +931,7 @@ namespace UnrealGameSync
 				else
 				{
 				//////////////
-				/// Bulids
+				/// Builds
 				//////////////
 				List<BadgeData> Builds = RESTApi.GET<List<BadgeData>>(ApiUrl, "build", string.Format("project={0}", Project), string.Format("lastbuildid={0}", LatestIds.LastBuildId));
 				foreach (BadgeData Build in Builds)
@@ -987,6 +989,7 @@ namespace UnrealGameSync
 			if(ApiUrl != null)
 			{
 				EventData Event = new EventData();
+				Event.Id = ++LatestIds.LastEventId;
 				Event.Change = ChangeNumber;
 				Event.UserName = CurrentUserName;
 				Event.Type = Type;
@@ -1101,7 +1104,7 @@ namespace UnrealGameSync
 		{
 			if(ActiveInvestigations == null)
 			{
-				// Insert investigation events into the active list, sorted by change number. Remove 
+				// Insert investigation events into the active list, sorted by change number.
 				ActiveInvestigations = new List<EventData>();
 				foreach(EventData InvestigationEvent in InvestigationEvents)
 				{
