UGS: Fixed an issue with auto-updates failing for instances of UGS that are launched via a Windows taskbar pin.

This fix ensures that the CWD is always set to the launcher directory so that child processes (such as P4V) won't disrupt auto-updates by holding onto file handles in the UGS executable directory.

diff --git a/Engine/Source/Programs/UnrealGameSync/UnrealGameSync/Program.cs b/Engine/Source/Programs/UnrealGameSync/UnrealGameSync/Program.cs
--- a/Engine/Source/Programs/UnrealGameSync/UnrealGameSync/Program.cs
+++ b/Engine/Source/Programs/UnrealGameSync/UnrealGameSync/Program.cs
@@ -100,6 +100,12 @@ namespace UnrealGameSync
 			string UpdateConfigFile = Path.Combine(Path.GetDirectoryName(Assembly.GetExecutingAssembly().Location), "AutoUpdate.ini");
 			MergeUpdateSettings(UpdateConfigFile, ref UpdatePath, ref UpdateSpawn);
 
+			// Set the current working directory to the update directory to prevent child-process file handles from disrupting auto-updates
+			if (UpdateSpawn != null)
+			{
+				Directory.SetCurrentDirectory(Path.GetDirectoryName(UpdateSpawn));
+			}
+
 			string SyncVersionFile = Path.Combine(Path.GetDirectoryName(Assembly.GetExecutingAssembly().Location), "SyncVersion.txt");
 			if(File.Exists(SyncVersionFile))
 			{
