UGS can now detect when the editor is running under other Windows sessions, without requiring elevated permissions.

diff --git a/Engine/Source/Programs/UnrealGameSync/UnrealGameSync/Controls/WorkspaceControl.cs b/Engine/Source/Programs/UnrealGameSync/UnrealGameSync/Controls/WorkspaceControl.cs
--- a/Engine/Source/Programs/UnrealGameSync/UnrealGameSync/Controls/WorkspaceControl.cs
+++ b/Engine/Source/Programs/UnrealGameSync/UnrealGameSync/Controls/WorkspaceControl.cs
@@ -3102,53 +3102,50 @@ namespace UnrealGameSync
 		private string[] GetProcessesRunningInWorkspace()
 		{
 			HashSet<string> ProcessNames = new HashSet<string>(StringComparer.InvariantCultureIgnoreCase);
-			ProcessNames.Add("UE4Editor");
-			ProcessNames.Add("UE4Editor-Cmd");
-			ProcessNames.Add("UE4Editor-Win64-Debug");
-			ProcessNames.Add("UE4Editor-Win64-Debug-Cmd");
-			ProcessNames.Add("UE4Editor-Win64-DebugGame");
-			ProcessNames.Add("UE4Editor-Win64-DebugGame-Cmd");
-			ProcessNames.Add("UnrealEditor");
-			ProcessNames.Add("UnrealEditor-Cmd");
-			ProcessNames.Add("UnrealEditor-Win64-Debug");
-			ProcessNames.Add("UnrealEditor-Win64-Debug-Cmd");
-			ProcessNames.Add("UnrealEditor-Win64-DebugGame");
-			ProcessNames.Add("UnrealEditor-Win64-DebugGame-Cmd");
-			ProcessNames.Add("CrashReportClient");
-			ProcessNames.Add("CrashReportClient-Win64-Development");
-			ProcessNames.Add("CrashReportClient-Win64-Debug");
-			ProcessNames.Add("CrashReportClientEditor");
-			ProcessNames.Add("CrashReportClientEditor-Win64-Development");
-			ProcessNames.Add("CrashReportClientEditor-Win64-Debug");
-			ProcessNames.Add("UnrealBuildTool");
-			ProcessNames.Add("AutomationTool");
+			ProcessNames.Add("Win64\\UE4Editor.exe");
+			ProcessNames.Add("Win64\\UE4Editor-Cmd.exe");
+			ProcessNames.Add("Win64\\UE4Editor-Win64-Debug.exe");
+			ProcessNames.Add("Win64\\UE4Editor-Win64-Debug-Cmd.exe");
+			ProcessNames.Add("Win64\\UE4Editor-Win64-DebugGame.exe");
+			ProcessNames.Add("Win64\\UE4Editor-Win64-DebugGame-Cmd.exe");
+			ProcessNames.Add("Win64\\UnrealEditor.exe");
+			ProcessNames.Add("Win64\\UnrealEditor-Cmd.exe");
+			ProcessNames.Add("Win64\\UnrealEditor-Win64-Debug.exe");
+			ProcessNames.Add("Win64\\UnrealEditor-Win64-Debug-Cmd.exe");
+			ProcessNames.Add("Win64\\UnrealEditor-Win64-DebugGame.exe");
+			ProcessNames.Add("Win64\\UnrealEditor-Win64-DebugGame-Cmd.exe");
+			ProcessNames.Add("Win64\\CrashReportClient.exe");
+			ProcessNames.Add("Win64\\CrashReportClient-Win64-Development.exe");
+			ProcessNames.Add("Win64\\CrashReportClient-Win64-Debug.exe");
+			ProcessNames.Add("Win64\\CrashReportClientEditor.exe");
+			ProcessNames.Add("Win64\\CrashReportClientEditor-Win64-Development.exe");
+			ProcessNames.Add("Win64\\CrashReportClientEditor-Win64-Debug.exe");
+			ProcessNames.Add("DotNET\\UnrealBuildTool\\UnrealBuildTool.exe");
+			ProcessNames.Add("DotNET\\AutomationTool\\AutomationTool.exe");
 
 			List<string> ProcessFileNames = new List<string>();
 			try
 			{
-				string RootDirectoryName = Path.GetFullPath(Workspace.LocalRootPath) + Path.DirectorySeparatorChar;
-				foreach (Process ProcessInstance in Process.GetProcesses())
+				string BinariesRootPath = Path.Combine(Workspace.LocalRootPath, "Engine\\Binaries");
+
+				foreach (string ProcessName in ProcessNames)
 				{
 					try
 					{
-						if (ProcessNames.Contains(ProcessInstance.ProcessName))
+						string ProcessFilename = Path.Combine(BinariesRootPath, ProcessName);
+
+						if (File.Exists(ProcessFilename))
 						{
-							string ProcessFileName = Path.GetFullPath(ProcessInstance.MainModule.FileName);
-							if (ProcessFileName.StartsWith(RootDirectoryName, StringComparison.InvariantCultureIgnoreCase))
+							try
 							{
-								string DisplayTitle = ProcessFileName;
-								try
-								{
-									string MainWindowTitle = ProcessInstance.MainWindowTitle;
-									if (!String.IsNullOrEmpty(MainWindowTitle))
-									{
-										DisplayTitle = String.Format("{0} ({1})", ProcessInstance.MainWindowTitle, DisplayTitle);
-									}
-								}
-								catch
-								{
-								}
-								ProcessFileNames.Add(DisplayTitle);
+								// Try to open the file to determine whether the executable is running.
+								// We do this so that we can also find processes running under other user sessions, without needing PROCESS_QUERY_INFORMATION access rights.
+								using FileStream TestStream = File.OpenWrite(ProcessFilename);
+							}
+							catch (IOException)
+							{
+								// This file is in use, so add it to the list of running processes
+								ProcessFileNames.Add(ProcessFilename);
 							}
 						}
 					}
