Flags Uncontrolled Files in ClobberWindow using Uncontrolled Changelist Persistency file.

#changelist validated

diff --git a/Engine/Source/Programs/UnrealGameSync/UnrealGameSync/Controls/WorkspaceControl.cs b/Engine/Source/Programs/UnrealGameSync/UnrealGameSync/Controls/WorkspaceControl.cs
--- a/Engine/Source/Programs/UnrealGameSync/UnrealGameSync/Controls/WorkspaceControl.cs
+++ b/Engine/Source/Programs/UnrealGameSync/UnrealGameSync/Controls/WorkspaceControl.cs
@@ -18,7 +18,20 @@ using System.Threading;
 
 namespace UnrealGameSync
 {
-	public enum LatestChangeType
+	public class UncontrolledChangelist
+	{
+		public string GUID { get; set; }
+		public string Name { get; set; }
+		public List<string> Files { get; set; }
+	}
+
+	public class UncontrolledChangelistPersistency
+	{
+		public int Version { get; set; }
+		public List<UncontrolledChangelist> Changelists { get; set; }
+	}
+
+public enum LatestChangeType
 	{
 		Any,
 		Good,
@@ -1140,7 +1153,24 @@ namespace UnrealGameSync
 				DesiredTaskbarState = Tuple.Create(TaskbarState.Paused, 0.0f);
 				Owner.UpdateProgress();
 
-				ClobberWindow Window = new ClobberWindow(Context.ClobberFiles);
+				HashSet<string> UncontrolledFiles = new HashSet<string>();
+				
+				if (SelectedProject.LocalPath.EndsWith(".uprojectdirs", StringComparison.InvariantCultureIgnoreCase))
+				{
+					List<string> ProjectRoots = GetProjectRoots(SelectedProject.LocalPath);
+
+					foreach (string ProjectRoot in ProjectRoots)
+					{
+						ParseUncontrolledChangelistsPersistencyFile(ProjectRoot, UncontrolledFiles);
+					}
+				}
+				else
+				{
+					ParseUncontrolledChangelistsPersistencyFile(Path.GetDirectoryName(SelectedProject.LocalPath), UncontrolledFiles);
+				}
+
+				ClobberWindow Window = new ClobberWindow(Context.ClobberFiles, UncontrolledFiles);
+				
 				if (Window.ShowDialog(this) == DialogResult.OK)
 				{
 					StartWorkspaceUpdate(Context, UpdateCallback);
@@ -1764,6 +1794,76 @@ namespace UnrealGameSync
 			}
 		}
 
+		List<string> GetProjectRoots(string InUProjectDirsPath)
+		{
+			List<string> ProjectRoots = new List<string>();
+
+			if (!File.Exists(InUProjectDirsPath))
+			{
+				return ProjectRoots;
+			}
+
+			string UProjectDirsRoot = Path.GetDirectoryName(InUProjectDirsPath);
+
+			foreach (string Line in File.ReadLines(InUProjectDirsPath))
+			{
+				// Remove unnecessary whitespaces
+				string TrimmedLine = Line.Trim();
+
+				// Remove every comments
+				int CommentIndex = TrimmedLine.IndexOf(';');
+
+				if (CommentIndex >= 0)
+				{
+					TrimmedLine = TrimmedLine.Substring(0, CommentIndex);
+				}
+
+				if (TrimmedLine.Length <= 0)
+				{
+					continue;
+				}
+
+				string UProjectDirsEntry = Path.Combine(UProjectDirsRoot, TrimmedLine);
+
+				if (!Directory.Exists(UProjectDirsEntry))
+				{
+					continue;
+				}
+
+				List<string> EntryDirectories = new List<string>(Directory.EnumerateDirectories(UProjectDirsEntry));
+
+				ProjectRoots = ProjectRoots.Concat(EntryDirectories).ToList();
+			}
+
+			return ProjectRoots;
+		}
+
+		void ParseUncontrolledChangelistsPersistencyFile(string InProjectRoot, HashSet<string> OutUncontrolledFiles)
+		{
+			String UncontrolledChangelistPersistencyFilePath = Path.Combine(InProjectRoot, "Saved", "SourceControl", "UncontrolledChangelists.json");
+
+			if (File.Exists(UncontrolledChangelistPersistencyFilePath))
+			{
+				try
+				{
+					string JsonString = File.ReadAllText(UncontrolledChangelistPersistencyFilePath);
+					UncontrolledChangelistPersistency UCLPersistency = JsonSerializer.Deserialize<UncontrolledChangelistPersistency>(JsonString, Program.DefaultJsonSerializerOptions);
+
+					foreach (UncontrolledChangelist UCL in UCLPersistency.Changelists)
+					{
+						foreach (string UncontrolledFile in UCL.Files)
+						{
+							OutUncontrolledFiles.Add(UncontrolledFile);
+						}
+					}
+				}
+				catch (Exception Ex)
+				{
+					Log.WriteLine("Exception while parsing Uncontrolled Changelist persistency file: {0}", Ex.ToString());
+				}
+			}
+		}
+
 		void UpdateReviewsCallback()
 		{
 			if (!bUpdateReviewsPosted)
diff --git a/Engine/Source/Programs/UnrealGameSync/UnrealGameSync/Forms/ClobberWindow.cs b/Engine/Source/Programs/UnrealGameSync/UnrealGameSync/Forms/ClobberWindow.cs
--- a/Engine/Source/Programs/UnrealGameSync/UnrealGameSync/Forms/ClobberWindow.cs
+++ b/Engine/Source/Programs/UnrealGameSync/UnrealGameSync/Forms/ClobberWindow.cs
@@ -1,4 +1,4 @@
-﻿// Copyright Epic Games, Inc. All Rights Reserved.
+// Copyright Epic Games, Inc. All Rights Reserved.
 
 using System;
 using System.Collections.Generic;
@@ -17,8 +17,10 @@ namespace UnrealGameSync
 	{
 		Dictionary<string, bool> FilesToClobber;
 
-		public ClobberWindow(Dictionary<string, bool> InFilesToClobber)
+		public ClobberWindow(Dictionary<string, bool> InFilesToClobber, HashSet<string> InUncontrolledFiles)
 		{
+			bool bUncontrolledChangeFound = false;
+
 			InitializeComponent();
 
 			FilesToClobber = InFilesToClobber;
@@ -30,6 +32,20 @@ namespace UnrealGameSync
 				Item.Checked = FileToClobber.Value;
 				Item.SubItems.Add(Path.GetDirectoryName(FileToClobber.Key));
 				FileList.Items.Add(Item);
+
+				if (InUncontrolledFiles.Contains(FileToClobber.Key.Replace("\\", "/")))
+				{
+					bUncontrolledChangeFound = true;
+					Item.ForeColor = Color.Red;
+				}
+			}
+
+			if (bUncontrolledChangeFound)
+			{
+				// Updates the string to inform the user to take special care with Uncontrolled Changes
+				this.label1.Text = "The following files are writable in your workspace." + Environment.NewLine +
+	"Red files are Uncontrolled Changes and may contain modifications you made on purpose." + Environment.NewLine +
+	"Select which files you want to overwrite:";
 			}
 		}
 
