UGS: Fix launch on startup not working when launching through dotnet (or debugger). Use the .NET bootstraper rather than the assembly path.

diff --git a/Engine/Source/Programs/UnrealGameSync/UnrealGameSync/ProgramApplicationContext.cs b/Engine/Source/Programs/UnrealGameSync/UnrealGameSync/ProgramApplicationContext.cs
--- a/Engine/Source/Programs/UnrealGameSync/UnrealGameSync/ProgramApplicationContext.cs
+++ b/Engine/Source/Programs/UnrealGameSync/UnrealGameSync/ProgramApplicationContext.cs
@@ -248,9 +248,16 @@ namespace UnrealGameSync
 				OIDCLoginWindow LoginDialog = new OIDCLoginWindow(OIDCTokenManager);
 				LoginDialog.ShowDialog();
 			}
-				
+
+			// Get the application path
+			string OriginalExe = UpdateSpawn;
+			if (OriginalExe == null)
+			{
+				OriginalExe = Path.ChangeExtension(Assembly.GetExecutingAssembly().Location, ".exe");
+			}
+
 			// Create the main window
-			MainWindowInstance = new MainWindow(UpdateMonitor, ApiUrl, DataFolder, CacheFolder, bRestoreState, UpdateSpawn ?? Assembly.GetExecutingAssembly().Location, bUnstable, StartupProjects, DefaultConnection, Log, Settings, Uri, OIDCTokenManager);
+			MainWindowInstance = new MainWindow(UpdateMonitor, ApiUrl, DataFolder, CacheFolder, bRestoreState, OriginalExe, bUnstable, StartupProjects, DefaultConnection, Log, Settings, Uri, OIDCTokenManager);
 			if(bVisible)
 			{
 				MainWindowInstance.Show();
