UGS: Fix UGS closing anyway when cancelling due to a pending operation.

diff --git a/Engine/Source/Programs/UnrealGameSync/UnrealGameSync/Forms/MainWindow.cs b/Engine/Source/Programs/UnrealGameSync/UnrealGameSync/Forms/MainWindow.cs
--- a/Engine/Source/Programs/UnrealGameSync/UnrealGameSync/Forms/MainWindow.cs
+++ b/Engine/Source/Programs/UnrealGameSync/UnrealGameSync/Forms/MainWindow.cs
@@ -600,6 +600,19 @@ namespace UnrealGameSync
 			base.Dispose(disposing);
 		}
 
+		public bool ConfirmClose()
+		{
+			for (int Idx = 0; Idx < TabControl.GetTabCount(); Idx++)
+			{
+				IMainWindowTabPanel TabPanel = (IMainWindowTabPanel)TabControl.GetTabData(Idx);
+				if (!TabPanel.CanClose())
+				{
+					return false;
+				}
+			}
+			return true;
+		}
+
 		private void MainWindow_FormClosing(object Sender, FormClosingEventArgs EventArgs)
 		{
 			if(!bAllowClose && Settings.bKeepInTray)
@@ -609,13 +622,16 @@ namespace UnrealGameSync
 			}
 			else
 			{
-				for(int Idx = 0; Idx < TabControl.GetTabCount(); Idx++)
+				if (!bAllowClose)
 				{
-					IMainWindowTabPanel TabPanel = (IMainWindowTabPanel)TabControl.GetTabData(Idx);
-					if(!TabPanel.CanClose())
+					for (int Idx = 0; Idx < TabControl.GetTabCount(); Idx++)
 					{
-						EventArgs.Cancel = true;
-						return;
+						IMainWindowTabPanel TabPanel = (IMainWindowTabPanel)TabControl.GetTabData(Idx);
+						if (!TabPanel.CanClose())
+						{
+							EventArgs.Cancel = true;
+							return;
+						}
 					}
 				}
 
diff --git a/Engine/Source/Programs/UnrealGameSync/UnrealGameSync/ProgramApplicationContext.cs b/Engine/Source/Programs/UnrealGameSync/UnrealGameSync/ProgramApplicationContext.cs
--- a/Engine/Source/Programs/UnrealGameSync/UnrealGameSync/ProgramApplicationContext.cs
+++ b/Engine/Source/Programs/UnrealGameSync/UnrealGameSync/ProgramApplicationContext.cs
@@ -413,7 +413,12 @@ namespace UnrealGameSync
 
 		private void NotifyMenu_Exit_Click(object sender, EventArgs e)
 		{
-			if(DetectStartupProjectSettingsWindow != null)
+			if (MainWindowInstance != null && !MainWindowInstance.ConfirmClose())
+			{
+				return;
+			}
+
+			if (DetectStartupProjectSettingsWindow != null)
 			{
 				DetectStartupProjectSettingsWindow.Close();
 				DetectStartupProjectSettingsWindow = null;
