UGS: Fix error the second time a machine is RDP'd to.

diff --git a/Engine/Source/Programs/UnrealGameSync/UnrealGameSync/UriHandlers/RdpHandler.cs b/Engine/Source/Programs/UnrealGameSync/UnrealGameSync/UriHandlers/RdpHandler.cs
--- a/Engine/Source/Programs/UnrealGameSync/UnrealGameSync/UriHandlers/RdpHandler.cs
+++ b/Engine/Source/Programs/UnrealGameSync/UnrealGameSync/UriHandlers/RdpHandler.cs
@@ -65,19 +65,24 @@ namespace UnrealGameSync
 			IntPtr Buffer = IntPtr.Zero;
 			try
 			{
-				if (CredRead(Name, Type, 0, out Buffer))
+				if (!CredRead(Name, Type, 0, out Buffer))
 				{
-					CREDENTIAL Credential = Marshal.PtrToStructure<CREDENTIAL>(Buffer);
-					UserName = Credential.UserName;
-					Password = Marshal.PtrToStringUni(Credential.CredentialBlob, Credential.CredentialBlobSize / sizeof(char));
-					return true;
+					UserName = null;
+					Password = null;
+					return false;
 				}
-				else
+
+				CREDENTIAL Credential = Marshal.PtrToStructure<CREDENTIAL>(Buffer);
+				if (Credential.CredentialBlob == IntPtr.Zero)
 				{
 					UserName = null;
 					Password = null;
 					return false;
 				}
+
+				UserName = Credential.UserName;
+				Password = Marshal.PtrToStringUni(Credential.CredentialBlob, Credential.CredentialBlobSize / sizeof(char));
+				return true;
 			}
 			finally
 			{
