UGS: Prevent user credentials being overwritten for RDP URI handler.

diff --git a/Engine/Source/Programs/UnrealGameSync/UnrealGameSync/UriHandlers/RdpHandler.cs b/Engine/Source/Programs/UnrealGameSync/UnrealGameSync/UriHandlers/RdpHandler.cs
--- a/Engine/Source/Programs/UnrealGameSync/UnrealGameSync/UriHandlers/RdpHandler.cs
+++ b/Engine/Source/Programs/UnrealGameSync/UnrealGameSync/UriHandlers/RdpHandler.cs
@@ -73,15 +73,15 @@ namespace UnrealGameSync
 				}
 
 				CREDENTIAL Credential = Marshal.PtrToStructure<CREDENTIAL>(Buffer);
+				UserName = Credential.UserName;
 				if (Credential.CredentialBlob == IntPtr.Zero)
 				{
-					UserName = null;
 					Password = null;
-					return false;
 				}
-
-				UserName = Credential.UserName;
-				Password = Marshal.PtrToStringUni(Credential.CredentialBlob, Credential.CredentialBlobSize / sizeof(char));
+				else
+				{
+					Password = Marshal.PtrToStringUni(Credential.CredentialBlob, Credential.CredentialBlobSize / sizeof(char));
+				}
 				return true;
 			}
 			finally
