UGS: Add support for dependencies between build steps. Allows "Launch UnrealInsights" to depend on "Compile UnrealInsights" executing first.

diff --git a/Engine/Source/Programs/UnrealGameSync/UnrealGameSync/BuildStep.cs b/Engine/Source/Programs/UnrealGameSync/UnrealGameSync/BuildStep.cs
--- a/Engine/Source/Programs/UnrealGameSync/UnrealGameSync/BuildStep.cs
+++ b/Engine/Source/Programs/UnrealGameSync/UnrealGameSync/BuildStep.cs
@@ -37,6 +37,7 @@ namespace UnrealGameSync
 		public bool bNormalSync;
 		public bool bScheduledSync;
 		public bool bShowAsTool;
+		public Guid[] Requires;
 		public Guid ToolId;
 
 		public BuildStep(Guid InUniqueId, int InOrderIndex, string InDescription, string InStatusText, int InEstimatedDuration, string InFileName, string InArguments, string InWorkingDir, bool bInUseLogWindow)
@@ -117,6 +118,17 @@ namespace UnrealGameSync
 			{
 				bShowAsTool = false;
 			}
+
+			List<Guid> Requires = new List<Guid>();
+			foreach (string RequireString in Object.GetValue("Requires", String.Empty).Split(';', StringSplitOptions.RemoveEmptyEntries))
+			{
+				if (Guid.TryParse(RequireString, out Guid Require))
+				{
+					Requires.Add(Require);
+				}
+			}
+			this.Requires = Requires.ToArray();
+	
 			if (!Guid.TryParse(Object.GetValue("Tool", ""), out ToolId))
 			{
 				ToolId = Guid.Empty;
diff --git a/Engine/Source/Programs/UnrealGameSync/UnrealGameSync/Controls/WorkspaceControl.cs b/Engine/Source/Programs/UnrealGameSync/UnrealGameSync/Controls/WorkspaceControl.cs
--- a/Engine/Source/Programs/UnrealGameSync/UnrealGameSync/Controls/WorkspaceControl.cs
+++ b/Engine/Source/Programs/UnrealGameSync/UnrealGameSync/Controls/WorkspaceControl.cs
@@ -5171,12 +5171,38 @@ public enum LatestChangeType
 					int InsertIdx = 0;
 
 					List<BuildStep> UserSteps = GetUserBuildSteps(ProjectBuildStepObjects);
+
+					Dictionary<Guid, BuildStep> IdToStep = new Dictionary<Guid, BuildStep>();
+					foreach (BuildStep Step in UserSteps)
+					{
+						IdToStep[Step.UniqueId] = Step;
+					}
+
 					foreach (BuildStep Step in UserSteps)
 					{
 						if (Step.bShowAsTool && (Step.ToolId == Guid.Empty || Settings.EnabledTools.Contains(Step.ToolId)))
 						{
+							// Find all the dependent steps that also need executing
+							HashSet<Guid> StepSet = new HashSet<Guid> { Step.UniqueId };
+
+							Stack<Guid> Stack = new Stack<Guid>(StepSet);
+							while (Stack.Count > 0)
+							{
+								Guid Id = Stack.Pop();
+								if (IdToStep.TryGetValue(Id, out BuildStep StackStep))
+								{
+									foreach (Guid RequiresId in StackStep.Requires)
+									{
+										if (StepSet.Add(RequiresId))
+										{
+											Stack.Push(RequiresId);
+										}
+									}
+								}
+							}
+
 							ToolStripMenuItem NewMenuItem = new ToolStripMenuItem(Step.Description.Replace("&", "&&"));
-							NewMenuItem.Click += new EventHandler((sender, e) => { RunCustomTool(Step); });
+							NewMenuItem.Click += new EventHandler((sender, e) => { RunCustomTool(Step, StepSet); });
 							CustomToolMenuItems.Add(NewMenuItem);
 							MoreToolsContextMenu.Items.Insert(InsertIdx++, NewMenuItem);
 						}
@@ -5188,7 +5214,7 @@ public enum LatestChangeType
 			MoreActionsContextMenu_CustomToolSeparator.Visible = (CustomToolMenuItems.Count > 0);
 		}
 
-		private void RunCustomTool(BuildStep Step)
+		private void RunCustomTool(BuildStep Step, HashSet<Guid> AllStepSet)
 		{
 			if (Workspace != null)
 			{
@@ -5208,7 +5234,7 @@ public enum LatestChangeType
 						}
 					}
 
-					WorkspaceUpdateContext Context = new WorkspaceUpdateContext(Workspace.CurrentChangeNumber, WorkspaceUpdateOptions.Build, null, GetDefaultBuildStepObjects(), ProjectSettings.BuildSteps, new HashSet<Guid> { Step.UniqueId }, Variables);
+					WorkspaceUpdateContext Context = new WorkspaceUpdateContext(Workspace.CurrentChangeNumber, WorkspaceUpdateOptions.Build, null, GetDefaultBuildStepObjects(), ProjectSettings.BuildSteps, AllStepSet, Variables);
 					StartWorkspaceUpdate(Context, null);
 				}
 			}
