UGS: Show custom build steps without an explicit ordering after default steps.

diff --git a/Engine/Source/Programs/UnrealGameSync/UnrealGameSync/Controls/WorkspaceControl.cs b/Engine/Source/Programs/UnrealGameSync/UnrealGameSync/Controls/WorkspaceControl.cs
--- a/Engine/Source/Programs/UnrealGameSync/UnrealGameSync/Controls/WorkspaceControl.cs
+++ b/Engine/Source/Programs/UnrealGameSync/UnrealGameSync/Controls/WorkspaceControl.cs
@@ -5413,7 +5413,7 @@ public enum LatestChangeType
 			}
 
 			// Create the expanded task objects
-			return UserBuildStepObjects.Values.Select(x => new BuildStep(x)).OrderBy(x => x.OrderIndex).ToList();
+			return UserBuildStepObjects.Values.Select(x => new BuildStep(x)).OrderBy(x => (x.OrderIndex == -1)? 10000 : x.OrderIndex).ToList();
 		}
 
 		private void OptionsContextMenu_SyncPrecompiledBinaries_Click(object sender, EventArgs e)
diff --git a/Engine/Source/Programs/UnrealGameSync/UnrealGameSync/Workspace.cs b/Engine/Source/Programs/UnrealGameSync/UnrealGameSync/Workspace.cs
--- a/Engine/Source/Programs/UnrealGameSync/UnrealGameSync/Workspace.cs
+++ b/Engine/Source/Programs/UnrealGameSync/UnrealGameSync/Workspace.cs
@@ -1057,7 +1057,7 @@ namespace UnrealGameSync
 				BuildStep.MergeBuildStepObjects(BuildStepObjects, Context.UserBuildStepObjects);
 
 				// Construct build steps from them
-				List<BuildStep> BuildSteps = BuildStepObjects.Values.Select(x => new BuildStep(x)).OrderBy(x => x.OrderIndex).ToList();
+				List<BuildStep> BuildSteps = BuildStepObjects.Values.Select(x => new BuildStep(x)).OrderBy(x => (x.OrderIndex == -1) ? 10000 : x.OrderIndex).ToList();
 				if(Context.CustomBuildSteps != null && Context.CustomBuildSteps.Count > 0)
 				{
 					BuildSteps.RemoveAll(x => !Context.CustomBuildSteps.Contains(x.UniqueId));
