UGS: Remove references to NDA platforms from UGS code.

diff --git a/Engine/Source/Programs/UnrealGameSync/UnrealGameSync/Controls/WorkspaceControl.cs b/Engine/Source/Programs/UnrealGameSync/UnrealGameSync/Controls/WorkspaceControl.cs
--- a/Engine/Source/Programs/UnrealGameSync/UnrealGameSync/Controls/WorkspaceControl.cs
+++ b/Engine/Source/Programs/UnrealGameSync/UnrealGameSync/Controls/WorkspaceControl.cs
@@ -62,27 +62,6 @@ public enum LatestChangeType
 
 	partial class WorkspaceControl : UserControl, IMainWindowTabPanel
 	{
-		/// <summary>
-		/// Default sync categories.
-		/// NOTE: These category definitions are now overridden in UnrealGameSync.ini; this list exists solely for backwards compatibility, and will not normally be used.
-		/// </summary>
-		readonly WorkspaceSyncCategory[] DefaultSyncCategories =
-		{
-			new WorkspaceSyncCategory(new Guid("{6703E989-D912-451D-93AD-B48DE748D282}"), "Content", "*.uasset", "*.umap"),
-			new WorkspaceSyncCategory(new Guid("{6507C2FB-19DD-403A-AFA3-BBF898248D5A}"), "Documentation", "/Engine/Documentation/..."),
-			new WorkspaceSyncCategory(new Guid("{FD7C716E-4BAD-43AE-8FAE-8748EF9EE44D}"), "Platform Support: Android", "/Engine/Source/ThirdParty/.../Android/...", ".../Build/Android/PipelineCaches/..."),
-			new WorkspaceSyncCategory(new Guid("{176B2EB2-35F7-4E8E-B131-5F1C5F0959AF}"), "Platform Support: iOS", "/Engine/Source/ThirdParty/.../IOS/...", ".../Build/IOS/PipelineCaches/..."),
-			new WorkspaceSyncCategory(new Guid("{F44B2D25-CBC0-4A8F-B6B3-E4A8125533DD}"), "Platform Support: Linux", "/Engine/Source/ThirdParty/.../Linux/..."),
-			new WorkspaceSyncCategory(new Guid("{2AF45231-0D75-463B-BF9F-ABB3231091BB}"), "Platform Support: Mac", "/Engine/Source/ThirdParty/.../Mac/...", ".../Build/Mac/PipelineCaches/..."),
-			new WorkspaceSyncCategory(new Guid("{C8CB4934-ADE9-46C9-B6E3-61A659E1FAF5}"), "Platform Support: PS4", ".../PS4/..."),
-			new WorkspaceSyncCategory(new Guid("{F8AE5AC3-DA2D-4719-BABF-8A90D878379E}"), "Platform Support: Switch", ".../Switch/..."),
-			new WorkspaceSyncCategory(new Guid("{3788A0BC-188C-4A0D-950A-D68175F0D110}"), "Platform Support: tvOS", "/Engine/Source/ThirdParty/.../TVOS/..."),
-			new WorkspaceSyncCategory(new Guid("{1144E719-FCD7-491B-B0FC-8B4C3565BF79}"), "Platform Support: Win32", "/Engine/Source/ThirdParty/.../Win32/..."),
-			new WorkspaceSyncCategory(new Guid("{5206CCEE-9024-4E36-8B89-F5F5A7D288D2}"), "Platform Support: Win64", "/Engine/Source/ThirdParty/.../Win64/..."),
-			new WorkspaceSyncCategory(new Guid("{06887423-B094-4718-9B55-C7A21EE67EE4}"), "Platform Support: XboxOne", ".../XboxOne/..."),
-			new WorkspaceSyncCategory(new Guid("{CFEC942A-BB90-4F0C-ACCF-238ECAAD9430}"), "Source Code", "/Engine/Source/..."),
-		};
-
 		enum HorizontalAlignment
 		{
 			Left,
@@ -1740,12 +1719,6 @@ public enum LatestChangeType
 		{
 			Dictionary<Guid, WorkspaceSyncCategory> UniqueIdToCategory = new Dictionary<Guid, WorkspaceSyncCategory>();
 
-			// Add the default filters
-			foreach (WorkspaceSyncCategory DefaultSyncCategory in DefaultSyncCategories)
-			{
-				UniqueIdToCategory.Add(DefaultSyncCategory.UniqueId, DefaultSyncCategory);
-			}
-
 			// Add the custom filters
 			ConfigFile ProjectConfigFile = PerforceMonitor.LatestProjectConfigFile;
 			if (ProjectConfigFile != null)
diff --git a/Engine/Source/Programs/UnrealGameSync/UnrealGameSync/FileFilter.cs b/Engine/Source/Programs/UnrealGameSync/UnrealGameSync/FileFilter.cs
--- a/Engine/Source/Programs/UnrealGameSync/UnrealGameSync/FileFilter.cs
+++ b/Engine/Source/Programs/UnrealGameSync/UnrealGameSync/FileFilter.cs
@@ -1,4 +1,4 @@
-﻿// Copyright Epic Games, Inc. All Rights Reserved.
+// Copyright Epic Games, Inc. All Rights Reserved.
 
 using System;
 using System.Collections.Generic;
@@ -294,15 +294,6 @@ namespace UnrealGameSync
 			}
 		}
 
-		/// <summary>
-		/// Excludes all restricted platform folders from the filter
-		/// </summary>
-		public void ExcludeConfidentialPlatforms()
-		{
-			AddRule(".../PS4/...", FileFilterType.Exclude);
-			AddRule(".../XboxOne/...", FileFilterType.Exclude);
-		}
-
 		/// <summary>
 		/// Excludes all confidential folders from the filter
 		/// </summary>
