UGS: Allow adding arbitrary links to commands from the status panel (and rename "build steps" to "commands").

The "Status panel link" option in the command editor allows specifying the name of a link to be displayed in the status panel. Specifying a string of the form "Menu|Item Name" allows specifying items to be displayed in a submenu.

diff --git a/Engine/Source/Programs/UnrealGameSync/UnrealGameSync/BuildStep.cs b/Engine/Source/Programs/UnrealGameSync/UnrealGameSync/BuildStep.cs
--- a/Engine/Source/Programs/UnrealGameSync/UnrealGameSync/BuildStep.cs
+++ b/Engine/Source/Programs/UnrealGameSync/UnrealGameSync/BuildStep.cs
@@ -36,7 +36,7 @@ namespace UnrealGameSync
 		public bool bUseLogWindow;
 		public bool bNormalSync;
 		public bool bScheduledSync;
-		public bool bShowAsTool;
+		public string StatusPanelLink;
 		public Guid[] Requires;
 		public Guid ToolId;
 
@@ -114,9 +114,15 @@ namespace UnrealGameSync
 			{
 				bScheduledSync = bNormalSync;
 			}
-			if(!Boolean.TryParse(Object.GetValue("bShowAsTool", ""), out bShowAsTool))
+
+			StatusPanelLink = Object.GetValue("Link", null);
+			if (String.IsNullOrEmpty(StatusPanelLink))
 			{
-				bShowAsTool = false;
+				bool bShowAsTool;
+				if (Boolean.TryParse(Object.GetValue("bShowAsTool", ""), out bShowAsTool) && bShowAsTool)
+				{
+					StatusPanelLink = $"More...|{Description}";
+				}
 			}
 
 			List<Guid> Requires = new List<Guid>();
@@ -196,7 +202,10 @@ namespace UnrealGameSync
 			Result["OrderIndex"] = OrderIndex.ToString();
 			Result["bNormalSync"] = bNormalSync.ToString();
 			Result["bScheduledSync"] = bScheduledSync.ToString();
-			Result["bShowAsTool"] = bShowAsTool.ToString();
+			if (!String.IsNullOrEmpty(StatusPanelLink))
+			{
+				Result["Link"] = StatusPanelLink;
+			}
 			if (ToolId != Guid.Empty)
 			{
 				Result["Tool"] = ToolId.ToString();
diff --git a/Engine/Source/Programs/UnrealGameSync/UnrealGameSync/Controls/WorkspaceControl.cs b/Engine/Source/Programs/UnrealGameSync/UnrealGameSync/Controls/WorkspaceControl.cs
--- a/Engine/Source/Programs/UnrealGameSync/UnrealGameSync/Controls/WorkspaceControl.cs
+++ b/Engine/Source/Programs/UnrealGameSync/UnrealGameSync/Controls/WorkspaceControl.cs
@@ -263,7 +263,8 @@ public enum LatestChangeType
 		List<int> SortedChangeNumbers = new List<int>();
 		Dictionary<string, Dictionary<int, string>> ArchiveToChangeNumberToArchiveKey = new Dictionary<string, Dictionary<int, string>>();
 		Dictionary<int, ChangeLayoutInfo> ChangeNumberToLayoutInfo = new Dictionary<int, ChangeLayoutInfo>();
-		List<ToolStripMenuItem> CustomToolMenuItems = new List<ToolStripMenuItem>();
+		List<ContextMenuStrip> CustomStatusPanelMenus = new List<ContextMenuStrip>();
+		List<(string, Action<Point, Rectangle>)> CustomStatusPanelLinks = new List<(string, Action<Point, Rectangle>)>();
 		int NumChanges;
 		int PendingSelectedChangeNumber = -1;
 		bool bHasBuildSteps = false;
@@ -3517,6 +3518,12 @@ public enum LatestChangeType
 					}
 				}
 
+				foreach ((string Name, Action<Point, Rectangle> Action) in CustomStatusPanelLinks)
+				{
+					ProgramsLine.AddLink(Name, FontStyle.Regular, Action);
+					ProgramsLine.AddText("  |  ");
+				}
+
 				ProgramsLine.AddLink("Windows Explorer", FontStyle.Regular, () => { SafeProcessStart("explorer.exe", String.Format("\"{0}\"", Path.GetDirectoryName(SelectedFileName))); });
 
 				if (GetDefaultIssueFilter() != null)
@@ -3966,6 +3973,11 @@ public enum LatestChangeType
 			SyncLog.ScrollToEnd();
 		}
 
+		private void ShowToolContextMenu(Rectangle Bounds, ContextMenuStrip MenuStrip)
+		{
+			MenuStrip.Show(StatusPanel, new Point(Bounds.Left, Bounds.Bottom), ToolStripDropDownDirection.BelowRight);
+		}
+
 		private void ShowActionsMenu(Rectangle Bounds)
 		{
 			MoreToolsContextMenu.Show(StatusPanel, new Point(Bounds.Left, Bounds.Bottom), ToolStripDropDownDirection.BelowRight);
@@ -5127,12 +5139,20 @@ public enum LatestChangeType
 		{
 			bHasBuildSteps = false;
 
-			foreach (ToolStripMenuItem CustomToolMenuItem in CustomToolMenuItems)
+			int MoreToolsItemCount = MoreToolsContextMenu.Items.IndexOf(MoreActionsContextMenu_CustomToolSeparator);
+			while (MoreToolsItemCount > 0)
+			{
+				MoreToolsContextMenu.Items.RemoveAt(--MoreToolsItemCount);
+			}
+
+			for (int Idx = 0; Idx < CustomStatusPanelMenus.Count; Idx++)
 			{
-				MoreToolsContextMenu.Items.Remove(CustomToolMenuItem);
+				components.Remove(CustomStatusPanelMenus[Idx]);
+				CustomStatusPanelMenus[Idx].Dispose();
 			}
 
-			CustomToolMenuItems.Clear();
+			CustomStatusPanelLinks.Clear();
+			CustomStatusPanelMenus.Clear();
 
 			if (Workspace != null)
 			{
@@ -5140,54 +5160,55 @@ public enum LatestChangeType
 				if (ProjectConfigFile != null)
 				{
 					Dictionary<Guid, ConfigObject> ProjectBuildStepObjects = GetProjectBuildStepObjects(ProjectConfigFile);
-
-					int InsertIdx = 0;
-
 					List<BuildStep> UserSteps = GetUserBuildSteps(ProjectBuildStepObjects);
 
-					Dictionary<Guid, BuildStep> IdToStep = new Dictionary<Guid, BuildStep>();
-					foreach (BuildStep Step in UserSteps)
-					{
-						IdToStep[Step.UniqueId] = Step;
-					}
-
+					Dictionary<string, ContextMenuStrip> NameToMenu = new Dictionary<string, ContextMenuStrip>();
 					foreach (BuildStep Step in UserSteps)
 					{
-						if (Step.bShowAsTool && (Step.ToolId == Guid.Empty || Settings.EnabledTools.Contains(Step.ToolId)))
+						if (!String.IsNullOrEmpty(Step.StatusPanelLink))
 						{
-							// Find all the dependent steps that also need executing
-							HashSet<Guid> StepSet = new HashSet<Guid> { Step.UniqueId };
-
-							Stack<Guid> Stack = new Stack<Guid>(StepSet);
-							while (Stack.Count > 0)
+							int BaseMenuIdx = Step.StatusPanelLink.IndexOf('|');
+							if (BaseMenuIdx == -1)
+							{
+								CustomStatusPanelLinks.Add((Step.StatusPanelLink, (P, R) => RunCustomTool(Step, UserSteps)));
+							}
+							else
 							{
-								Guid Id = Stack.Pop();
-								if (IdToStep.TryGetValue(Id, out BuildStep StackStep))
+								string MenuName = Step.StatusPanelLink.Substring(0, BaseMenuIdx);
+								string ItemName = Step.StatusPanelLink.Substring(BaseMenuIdx + 1).Replace("&", "&&");
+
+								ToolStripMenuItem NewMenuItem = new ToolStripMenuItem(ItemName);
+								NewMenuItem.Click += new EventHandler((sender, e) => { RunCustomTool(Step, UserSteps); });
+
+								if (MenuName == "More...")
+								{
+									MoreToolsContextMenu.Items.Insert(MoreToolsItemCount++, NewMenuItem);
+								}
+								else
 								{
-									foreach (Guid RequiresId in StackStep.Requires)
+									ContextMenuStrip Menu;
+									if (!NameToMenu.TryGetValue(MenuName, out Menu))
 									{
-										if (StepSet.Add(RequiresId))
-										{
-											Stack.Push(RequiresId);
-										}
+										Menu = new ContextMenuStrip();
+										NameToMenu.Add(MenuName, Menu);
+										CustomStatusPanelLinks.Add(($"{MenuName} \u25BE", (P, R) => ShowToolContextMenu(R, Menu)));
+										CustomStatusPanelMenus.Add(Menu);
+										components.Add(Menu);
 									}
+									Menu.Items.Add(NewMenuItem);
 								}
 							}
-
-							ToolStripMenuItem NewMenuItem = new ToolStripMenuItem(Step.Description.Replace("&", "&&"));
-							NewMenuItem.Click += new EventHandler((sender, e) => { RunCustomTool(Step, StepSet); });
-							CustomToolMenuItems.Add(NewMenuItem);
-							MoreToolsContextMenu.Items.Insert(InsertIdx++, NewMenuItem);
 						}
+
 						bHasBuildSteps |= Step.bNormalSync;
 					}
 				}
 			}
 
-			MoreActionsContextMenu_CustomToolSeparator.Visible = (CustomToolMenuItems.Count > 0);
+			MoreActionsContextMenu_CustomToolSeparator.Visible = (MoreToolsItemCount > 0);
 		}
 
-		private void RunCustomTool(BuildStep Step, HashSet<Guid> AllStepSet)
+		private void RunCustomTool(BuildStep Step, List<BuildStep> AllSteps)
 		{
 			if (Workspace != null)
 			{
@@ -5197,6 +5218,29 @@ public enum LatestChangeType
 				}
 				else
 				{
+					HashSet<Guid> StepSet = new HashSet<Guid> { Step.UniqueId };
+
+					ConfigFile ProjectConfigFile = Workspace.ProjectConfigFile;
+					if (ProjectConfigFile != null && Step.Requires.Length > 0)
+					{
+						Stack<Guid> Stack = new Stack<Guid>(StepSet);
+						while (Stack.Count > 0)
+						{
+							Guid Id = Stack.Pop();
+							BuildStep NextStep = AllSteps.FirstOrDefault(x => x.UniqueId == Id);
+							if (NextStep != null)
+							{
+								foreach (Guid RequiresId in NextStep.Requires)
+								{
+									if (StepSet.Add(RequiresId))
+									{
+										Stack.Push(RequiresId);
+									}
+								}
+							}
+						}
+					}
+
 					Dictionary<string, string> Variables = GetWorkspaceVariables(Workspace.CurrentChangeNumber);
 					if (Step.ToolId != Guid.Empty)
 					{
@@ -5207,7 +5251,7 @@ public enum LatestChangeType
 						}
 					}
 
-					WorkspaceUpdateContext Context = new WorkspaceUpdateContext(Workspace.CurrentChangeNumber, WorkspaceUpdateOptions.Build, null, GetDefaultBuildStepObjects(), ProjectSettings.BuildSteps, AllStepSet, Variables);
+					WorkspaceUpdateContext Context = new WorkspaceUpdateContext(Workspace.CurrentChangeNumber, WorkspaceUpdateOptions.Build, null, GetDefaultBuildStepObjects(), ProjectSettings.BuildSteps, StepSet, Variables);
 					StartWorkspaceUpdate(Context, null);
 				}
 			}
@@ -5289,6 +5333,7 @@ public enum LatestChangeType
 
 				// Update the custom tools menu, because we might have changed it
 				UpdateBuildSteps();
+				UpdateStatusPanel();
 				UpdateSyncActionCheckboxes();
 			}
 		}
diff --git a/Engine/Source/Programs/UnrealGameSync/UnrealGameSync/Forms/BuildStepWindow.cs b/Engine/Source/Programs/UnrealGameSync/UnrealGameSync/Forms/BuildStepWindow.cs
--- a/Engine/Source/Programs/UnrealGameSync/UnrealGameSync/Forms/BuildStepWindow.cs
+++ b/Engine/Source/Programs/UnrealGameSync/UnrealGameSync/Forms/BuildStepWindow.cs
@@ -1,4 +1,4 @@
-﻿// Copyright Epic Games, Inc. All Rights Reserved.
+// Copyright Epic Games, Inc. All Rights Reserved.
 
 using System;
 using System.Collections.Generic;
@@ -39,6 +39,7 @@ namespace UnrealGameSync
 			string DefaultTargetName = TargetNames.FirstOrDefault(x => !x.EndsWith("Editor") && !x.EndsWith("Client") && !x.EndsWith("Server")) ?? ((TargetNames.Count > 0)? TargetNames[0] : "");
 			CompileTargetComboBox.Items.AddRange(TargetNames.ToArray());
 			CompileTargetComboBox.Text = String.IsNullOrEmpty(Step.Target)? DefaultTargetName : Step.Target;
+			StatusPanelLinkTextBox.Text = Step.StatusPanelLink;
 
 			DescriptionTextBox.Text = Step.Description;
 			StatusTextTextBox.Text = Step.StatusText;
@@ -118,8 +119,9 @@ namespace UnrealGameSync
 		{
 			Step.Description = DescriptionTextBox.Text;
 			Step.StatusText = StatusTextTextBox.Text;
+			Step.StatusPanelLink = StatusPanelLinkTextBox.Text;
 
-			if(!int.TryParse(DurationTextBox.Text, out Step.EstimatedDuration))
+			if (!int.TryParse(DurationTextBox.Text, out Step.EstimatedDuration))
 			{
 				Step.EstimatedDuration = 1;
 			}
diff --git a/Engine/Source/Programs/UnrealGameSync/UnrealGameSync/Forms/ModifyBuildStepsWindow.cs b/Engine/Source/Programs/UnrealGameSync/UnrealGameSync/Forms/ModifyBuildStepsWindow.cs
--- a/Engine/Source/Programs/UnrealGameSync/UnrealGameSync/Forms/ModifyBuildStepsWindow.cs
+++ b/Engine/Source/Programs/UnrealGameSync/UnrealGameSync/Forms/ModifyBuildStepsWindow.cs
@@ -38,8 +38,7 @@ namespace UnrealGameSync
 
 				NormalSyncColumn.Width = (int)(104 * DpiScaleX);
 				ScheduledSyncColumn.Width = (int)(104 * DpiScaleX);
-				ShowAsToolColumn.Width = (int)(104 * DpiScaleX);
-				DescriptionColumn.Width = BuildStepList.ClientSize.Width - NormalSyncColumn.Width - ScheduledSyncColumn.Width - ShowAsToolColumn.Width - 10;
+				DescriptionColumn.Width = BuildStepList.ClientSize.Width - NormalSyncColumn.Width - ScheduledSyncColumn.Width - 10;
 			}
 
 			BuildStepList.Font = SystemFonts.IconTitleFont;
@@ -60,7 +59,6 @@ namespace UnrealGameSync
 			Item.Tag = Task;
 			Item.SubItems.Add(new ListViewItem.ListViewSubItem());
 			Item.SubItems.Add(new ListViewItem.ListViewSubItem());
-			Item.SubItems.Add(new ListViewItem.ListViewSubItem());
 			BuildStepList.Items.Add(Item);
 		}
 
@@ -148,13 +146,9 @@ namespace UnrealGameSync
 				{
 					bEnabled = Task.bNormalSync;
 				}
-				else if(e.ColumnIndex == 2)
-				{
-					bEnabled = Task.bScheduledSync;
-				}
 				else
 				{
-					bEnabled = Task.bShowAsTool;
+					bEnabled = Task.bScheduledSync;
 				}
 
 				bool Selected = BuildStepList.SelectedItems.Contains(e.Item);
@@ -198,13 +192,9 @@ namespace UnrealGameSync
 					{
 						Task.bNormalSync ^= true;
 					}
-					else if(ColumnIndex == 2)
-					{
-						Task.bScheduledSync ^= true;
-					}
 					else
 					{
-						Task.bShowAsTool ^= true;
+						Task.bScheduledSync ^= true;
 					}
 					BuildStepList.Invalidate(HitTest.SubItem.Bounds);
 				}
