Branch snapshot for CL 17775965

[CL 17775965 in ue5-main branch]

diff --git a/Engine/Source/Programs/UnrealGameSync/MetadataServer/App_Start/WebApiConfig.cs b/Engine/Source/Programs/UnrealGameSync/MetadataServer/App_Start/WebApiConfig.cs
new file mode 100644
--- /dev/null
+++ b/Engine/Source/Programs/UnrealGameSync/MetadataServer/App_Start/WebApiConfig.cs
@@ -0,0 +1,45 @@
+﻿// Copyright Epic Games, Inc. All Rights Reserved.
+
+using System.Web.Http;
+
+namespace MetadataServer
+{
+	public static class WebApiConfig
+	{
+		public static void Register(HttpConfiguration config)
+		{
+			// Web API routes
+			config.MapHttpAttributeRoutes();
+
+			config.Routes.MapHttpRoute(
+				name: "IssueBuildsSubApi",
+				routeTemplate: "api/issues/{IssueId}/builds",
+				defaults: new { controller = "IssueBuildsSub" }
+			);
+
+			config.Routes.MapHttpRoute(
+				name: "IssueBuildsApi",
+				routeTemplate: "api/issuebuilds/{BuildId}",
+				defaults: new { controller = "IssueBuilds" }
+			);
+
+			config.Routes.MapHttpRoute(
+				name: "IssueDiagnosticsSubApi",
+				routeTemplate: "api/issues/{IssueId}/diagnostics",
+				defaults: new { controller = "IssueDiagnosticsSub" }
+			);
+
+			config.Routes.MapHttpRoute(
+				name: "IssueWatchersApi",
+				routeTemplate: "api/issues/{IssueId}/watchers",
+				defaults: new { controller = "IssueWatchers" }
+			);
+
+			config.Routes.MapHttpRoute(
+				name: "DefaultApi",
+				routeTemplate: "api/{controller}/{id}",
+				defaults: new { id = RouteParameter.Optional }
+			);
+		}
+	}
+}
diff --git a/Engine/Source/Programs/UnrealGameSync/MetadataServer/Connectors/SqlConnector.cs b/Engine/Source/Programs/UnrealGameSync/MetadataServer/Connectors/SqlConnector.cs
new file mode 100644
--- /dev/null
+++ b/Engine/Source/Programs/UnrealGameSync/MetadataServer/Connectors/SqlConnector.cs
@@ -0,0 +1,839 @@
+﻿// Copyright Epic Games, Inc. All Rights Reserved.
+
+using Newtonsoft.Json;
+using Newtonsoft.Json.Linq;
+using System;
+using System.Collections.Generic;
+using System.Data.SqlClient;
+using System.IO;
+using System.Linq;
+using System.Text;
+using System.Text.RegularExpressions;
+using MetadataServer.Models;
+using MySql.Data.MySqlClient;
+
+namespace MetadataServer.Connectors
+{
+	public static class SqlConnector
+	{
+		private static string ConnectionString = System.Configuration.ConfigurationManager.ConnectionStrings["ConnectionString"].ConnectionString;
+		public static LatestData GetLastIds(string Project = null)
+		{
+			// Get ids going back 432 builds for the project being asked for
+			// Do this by grouping by ChangeNumber to get unique entries, then take the 432nd id
+			long LastEventId = 0;
+			long LastCommentId = 0;
+			long LastBuildId = 0;
+			string ProjectLikeString = "%" + (Project == null ? String.Empty : GetProjectStream(Project)) + "%";
+			using (MySqlConnection Connection = new MySqlConnection(ConnectionString))
+			{
+				Connection.Open();
+				using (MySqlCommand Command = new MySqlCommand("WITH user_votes AS (SELECT UserVotes.Id, UserVotes.Changelist FROM ugs_db.UserVotes " +
+																 "INNER JOIN ugs_db.Projects ON Projects.Id = UserVotes.ProjectId " +
+																 "WHERE Projects.Name LIKE @param1 GROUP BY Changelist ORDER BY Changelist DESC LIMIT 100) " +
+																 "SELECT * FROM user_votes ORDER BY user_votes.Changelist ASC LIMIT 1", Connection))
+				{
+					Command.Parameters.AddWithValue("@param1", ProjectLikeString);
+					using (MySqlDataReader Reader = Command.ExecuteReader())
+					{
+						while (Reader.Read())
+						{
+							LastEventId = Reader.GetInt64(0);
+							break;
+						}
+					}	
+				}
+				
+				using (MySqlCommand Command = new MySqlCommand("WITH comments AS (SELECT Comments.Id, Comments.ChangeNumber FROM ugs_db.Comments " +
+																 "INNER JOIN ugs_db.Projects ON Projects.Id = Comments.ProjectId " +
+																 "WHERE Projects.Name LIKE @param1 GROUP BY ChangeNumber ORDER BY ChangeNumber DESC LIMIT 100) " +
+																 "SELECT * FROM comments ORDER BY comments.ChangeNumber ASC LIMIT 1", Connection))
+				{
+					Command.Parameters.AddWithValue("@param1", ProjectLikeString);
+					using (MySqlDataReader Reader = Command.ExecuteReader())
+					{
+						while (Reader.Read())
+						{
+							LastCommentId = Reader.GetInt32(0);
+							break;
+						}
+					}
+				}
+				using (MySqlCommand Command = new MySqlCommand("WITH badges AS (SELECT Badges.Id, Badges.ChangeNumber FROM ugs_db.Badges " +
+																 "INNER JOIN ugs_db.Projects ON Projects.Id = Badges.ProjectId " +
+																 "WHERE Projects.Name LIKE @param1 GROUP BY ChangeNumber ORDER BY ChangeNumber DESC LIMIT 100) " + 
+																 "SELECT * FROM badges ORDER BY badges.ChangeNumber ASC LIMIT 1", Connection))
+				{
+					//Command.Parameters.AddWithValue("@param1", ProjectId);
+					Command.Parameters.AddWithValue("@param1", ProjectLikeString);
+					using (MySqlDataReader Reader = Command.ExecuteReader())
+					{
+						while (Reader.Read())
+						{
+							LastBuildId = Math.Max(LastBuildId, Reader.GetInt32(0));
+							break;
+						}
+					}
+				}
+			}
+			return new LatestData { LastBuildId = LastBuildId, LastCommentId = LastCommentId, LastEventId = LastEventId };
+		}
+		public static List<EventData> GetUserVotes(string Project, long LastEventId)
+		{
+			List<EventData> ReturnedEvents = new List<EventData>();
+			string ProjectLikeString = "%" + (Project == null ? String.Empty : GetProjectStream(Project)) + "%";
+			using (MySqlConnection Connection = new MySqlConnection(ConnectionString))
+			{
+				Connection.Open();
+				using (MySqlCommand Command = new MySqlCommand("SELECT UserVotes.Id, UserVotes.Changelist, UserVotes.UserName, UserVotes.Verdict, UserVotes.Project FROM ugs_db.UserVotes " +
+															   "INNER JOIN ugs_db.Projects ON Projects.Id = UserVotes.ProjectId WHERE UserVotes.Id > @param1 AND Projects.Name LIKE @param2 ORDER BY UserVotes.Id", Connection))
+				{
+					Command.Parameters.AddWithValue("@param1", LastEventId);
+					Command.Parameters.AddWithValue("@param2", ProjectLikeString);
+					using (MySqlDataReader Reader = Command.ExecuteReader())
+					{
+						while (Reader.Read())
+						{
+							EventData Review = new EventData();
+							Review.Id = Reader.GetInt64(0);
+							Review.Change = Reader.GetInt32(1);
+							Review.UserName = Reader.GetString(2);
+							Review.Project = Reader.IsDBNull(4) ? null : Reader.GetString(4);
+							if (Enum.TryParse(Reader.GetString(3), out Review.Type))
+							{
+								if (Review.Project == null || String.Compare(Review.Project, Project, true) == 0)
+								{
+									ReturnedEvents.Add(Review);
+								}
+							}
+						}
+					}
+				}
+			}
+			return ReturnedEvents;
+		}
+		public static List<CommentData> GetComments(string Project, long LastCommentId)
+		{
+			List<CommentData> ReturnedComments = new List<CommentData>();
+			string ProjectLikeString = "%" + (Project == null ? String.Empty : GetProjectStream(Project)) + "%";
+			using (MySqlConnection Connection = new MySqlConnection(ConnectionString))
+			{
+				Connection.Open();
+				using (MySqlCommand Command = new MySqlCommand("SELECT Comments.Id, Comments.ChangeNumber, Comments.UserName, Comments.Text, Comments.Project FROM ugs_db.Comments " +
+															     "INNER JOIN ugs_db.Projects ON Projects.Id = Comments.ProjectId WHERE Comments.Id > @param1 AND Projects.Name LIKE @param2 ORDER BY Comments.Id", Connection))
+				{
+					Command.Parameters.AddWithValue("@param1", LastCommentId);
+					Command.Parameters.AddWithValue("@param2", ProjectLikeString);
+					using (MySqlDataReader Reader = Command.ExecuteReader())
+					{
+						while (Reader.Read())
+						{
+							CommentData Comment = new CommentData();
+							Comment.Id = Reader.GetInt32(0);
+							Comment.ChangeNumber = Reader.GetInt32(1);
+							Comment.UserName = Reader.GetString(2);
+							Comment.Text = Reader.GetString(3);
+							Comment.Project = Reader.GetString(4);
+							if (Comment.Project == null || String.Compare(Comment.Project, Project, true) == 0)
+							{
+								ReturnedComments.Add(Comment);
+							}
+						}
+					}
+				}
+			}
+			return ReturnedComments;
+		}
+
+		public static List<BuildData> GetBuilds(string Project, long LastBuildId)
+		{
+			List<BuildData> ReturnedBuilds = new List<BuildData>();
+			string ProjectLikeString = "%" + (Project == null ? String.Empty : GetProjectStream(Project)) + "%";
+			using (MySqlConnection Connection = new MySqlConnection(ConnectionString))
+			{
+				Connection.Open();
+				using (MySqlCommand Command = new MySqlCommand("SELECT Badges.Id, Badges.ChangeNumber, Badges.BuildType, Badges.Result, Badges.Url, Projects.Name, Badges.ArchivePath FROM ugs_db.Badges " +
+																 "INNER JOIN ugs_db.Projects ON Projects.Id = Badges.ProjectId WHERE Badges.Id > @param1 AND Projects.Name LIKE @param2 ORDER BY Badges.Id", Connection))
+				{ 
+					Command.Parameters.AddWithValue("@param1", LastBuildId);
+					Command.Parameters.AddWithValue("@param2", ProjectLikeString);
+					using (MySqlDataReader Reader = Command.ExecuteReader())
+					{
+						while (Reader.Read())
+						{
+							BuildData Build = new BuildData();
+							Build.Id = Reader.GetInt32(0);
+							Build.ChangeNumber = Reader.GetInt32(1);
+							Build.BuildType = Reader.GetString(2).TrimEnd();
+							if (Enum.TryParse(Reader.GetString(3).TrimEnd(), true, out Build.Result))
+							{
+								Build.Url = Reader.GetString(4);
+								Build.Project = Reader.IsDBNull(5) ? null : Reader.GetString(5);
+								if (Build.Project == null || String.Compare(Build.Project, Project, true) == 0 || MatchesWildcard(Build.Project, Project))
+								{
+									ReturnedBuilds.Add(Build);
+								}
+							}
+							LastBuildId = Math.Max(LastBuildId, Build.Id);
+						}
+					}
+				}
+			}
+			return ReturnedBuilds;
+		}
+
+		public static List<TelemetryErrorData> GetErrorData(int Records)
+		{
+			List<TelemetryErrorData> ReturnedErrors = new List<TelemetryErrorData>();
+			using (MySqlConnection Connection = new MySqlConnection(ConnectionString))
+			{
+				Connection.Open();
+				using (MySqlCommand Command = new MySqlCommand("SELECT Id, Type, Text, UserName, Project, Timestamp, Version, IpAddress FROM ugs_db.Errors ORDER BY Id DESC LIMIT @param1", Connection))
+				{
+					Command.Parameters.AddWithValue("@param1", Records);
+					using (MySqlDataReader Reader = Command.ExecuteReader())
+					{
+						while (Reader.Read())
+						{
+							TelemetryErrorData Error = new TelemetryErrorData();
+							Error.Id = Reader.GetInt32(0);
+							Enum.TryParse(Reader.GetString(1), true, out Error.Type);
+							Error.Text = Reader.GetString(2);
+							Error.UserName = Reader.GetString(3);
+							Error.Project = Reader.IsDBNull(4) ? null : Reader.GetString(4);
+							Error.Timestamp = Reader.GetDateTime(5);
+							Error.Version = Reader.GetString(6);
+							Error.IpAddress = Reader.GetString(7);
+							ReturnedErrors.Add(Error);
+						}
+					}
+				}
+			}
+			return ReturnedErrors;
+		}
+
+		private static long TryInsertAndGetProject(MySqlConnection Connection, string Project)
+		{
+			using (MySqlCommand Command = new MySqlCommand("INSERT IGNORE INTO ugs_db.Projects (Name) VALUES (@Project); SELECT Id FROM ugs_db.Projects WHERE Name = @Project", Connection))
+			{
+				Command.Parameters.AddWithValue("@Project", Project);
+				object ProjectId = Command.ExecuteScalar();
+				return Convert.ToInt64(ProjectId);
+			}
+		}
+		public static void PostBuild(BuildData Build)
+		{
+			using (MySqlConnection Connection = new MySqlConnection(ConnectionString))
+			{
+				Connection.Open();
+				long ProjectId = TryInsertAndGetProject(Connection, Build.Project);
+				using (MySqlCommand Command = new MySqlCommand("INSERT INTO ugs_db.Badges (ChangeNumber, BuildType, Result, URL, ArchivePath, ProjectId) VALUES (@ChangeNumber, @BuildType, @Result, @URL, @ArchivePath, @ProjectId)", Connection))
+				{
+					Command.Parameters.AddWithValue("@ChangeNumber", Build.ChangeNumber);
+					Command.Parameters.AddWithValue("@BuildType", Build.BuildType);
+					Command.Parameters.AddWithValue("@Result", Build.Result);
+					Command.Parameters.AddWithValue("@URL", Build.Url);
+					Command.Parameters.AddWithValue("@ArchivePath", Build.ArchivePath);
+					Command.Parameters.AddWithValue("@ProjectId", ProjectId);
+					Command.ExecuteNonQuery();
+				}
+			}
+		}
+
+		public static void PostEvent(EventData Event)
+		{
+			using (MySqlConnection Connection = new MySqlConnection(ConnectionString))
+			{
+				Connection.Open();
+				long ProjectId = TryInsertAndGetProject(Connection, Event.Project);
+				using (MySqlCommand Command = new MySqlCommand("INSERT INTO ugs_db.UserVotes (Changelist, UserName, Verdict, Project, ProjectId) VALUES (@Changelist, @UserName, @Verdict, @Project, @ProjectId)", Connection))
+				{
+					Command.Parameters.AddWithValue("@Changelist", Event.Change);
+					Command.Parameters.AddWithValue("@UserName", Event.UserName.ToString());
+					Command.Parameters.AddWithValue("@Verdict", Event.Type.ToString());
+					Command.Parameters.AddWithValue("@Project", Event.Project);
+					Command.Parameters.AddWithValue("@ProjectId", ProjectId);
+					Command.ExecuteNonQuery();
+				}
+			}
+		}
+
+		public static void PostComment(CommentData Comment)
+		{
+			using (MySqlConnection Connection = new MySqlConnection(ConnectionString))
+			{
+				Connection.Open();
+				long ProjectId = TryInsertAndGetProject(Connection, Comment.Project);
+				using (MySqlCommand Command = new MySqlCommand("INSERT INTO ugs_db.Comments (ChangeNumber, UserName, Text, Project, ProjectId) VALUES (@ChangeNumber, @UserName, @Text, @Project, @ProjectId)", Connection))
+				{
+					Command.Parameters.AddWithValue("@ChangeNumber", Comment.ChangeNumber);
+					Command.Parameters.AddWithValue("@UserName", Comment.UserName);
+					Command.Parameters.AddWithValue("@Text", Comment.Text);
+					Command.Parameters.AddWithValue("@Project", Comment.Project);
+					Command.Parameters.AddWithValue("@ProjectId", ProjectId);
+					Command.ExecuteNonQuery();
+				}
+			}
+		}
+
+		public static void PostTelemetryData(TelemetryTimingData Data, string Version, string IpAddress)
+		{
+			using (MySqlConnection Connection = new MySqlConnection(ConnectionString))
+			{
+				Connection.Open();
+				long ProjectId = TryInsertAndGetProject(Connection, Data.Project);
+				using (MySqlCommand Command = new MySqlCommand("INSERT INTO ugs_db.Telemetry_v2 (Action, Result, UserName, Project, Timestamp, Duration, Version, IpAddress, ProjectId) VALUES (@Action, @Result, @UserName, @Project, @Timestamp, @Duration, @Version, @IpAddress, @ProjectId)", Connection))
+				{
+					Command.Parameters.AddWithValue("@Action", Data.Action);
+					Command.Parameters.AddWithValue("@Result", Data.Result);
+					Command.Parameters.AddWithValue("@UserName", Data.UserName);
+					Command.Parameters.AddWithValue("@Project", Data.Project);
+					Command.Parameters.AddWithValue("@Timestamp", Data.Timestamp);
+					Command.Parameters.AddWithValue("@Duration", Data.Duration);
+					Command.Parameters.AddWithValue("@Version", Version);
+					Command.Parameters.AddWithValue("@IPAddress", IpAddress);
+					Command.Parameters.AddWithValue("@ProjectId", ProjectId);
+					Command.ExecuteNonQuery();
+				}
+			}
+		}
+
+		public static void PostErrorData(TelemetryErrorData Data, string Version, string IpAddress)
+		{
+			using (MySqlConnection Connection = new MySqlConnection(ConnectionString))
+			{
+				Connection.Open();
+				long? ProjectId = null;
+				if(Data.Project != null)
+				{
+					ProjectId = TryInsertAndGetProject(Connection, Data.Project);
+				}
+				using (MySqlCommand Command = new MySqlCommand("INSERT INTO ugs_db.Errors (Type, Text, UserName, Project, Timestamp, Version, IpAddress, ProjectId) VALUES (@Type, @Text, @UserName, @Project, @Timestamp, @Version, @IpAddress, @ProjectId)", Connection))
+				{
+					Command.Parameters.AddWithValue("@Type", Data.Type.ToString());
+					Command.Parameters.AddWithValue("@Text", Data.Text);
+					Command.Parameters.AddWithValue("@UserName", Data.UserName);
+					if (Data.Project == null)
+					{
+						Command.Parameters.AddWithValue("@Project", DBNull.Value);
+						Command.Parameters.AddWithValue("@ProjectId", DBNull.Value);
+					}
+					else
+					{
+						Command.Parameters.AddWithValue("@Project", Data.Project);
+						Command.Parameters.AddWithValue("@ProjectId", ProjectId.Value);
+					}
+					Command.Parameters.AddWithValue("@Timestamp", Data.Timestamp);
+					Command.Parameters.AddWithValue("@Version", Version);
+					Command.Parameters.AddWithValue("@IPAddress", IpAddress);
+					Command.ExecuteNonQuery();
+				}
+			}
+		}
+		private static string GetProjectStream(string Project)
+		{
+			// Get first two fragments of the p4 path.  If it doesn't work, just return back the project.
+			Regex StreamPattern = new Regex("(\\/\\/[a-zA-Z0-9\\.\\-_]{1,}\\/[a-zA-Z0-9\\.\\-_]{1,})");
+			Match StreamMatch = StreamPattern.Match(Project);
+			if(StreamMatch.Success)
+			{
+				return StreamMatch.Groups[1].Value;
+			}
+			return Project;
+		}
+		private static bool MatchesWildcard(string Wildcard, string Project)
+		{
+			return Wildcard.EndsWith("...") && Project.StartsWith(Wildcard.Substring(0, Wildcard.Length - 4), StringComparison.InvariantCultureIgnoreCase);
+		}
+
+		private static string NormalizeUserName(string UserName)
+		{
+			return UserName.ToUpperInvariant();
+		}
+
+		public static long FindOrAddUserId(string Name)
+		{
+			using (MySqlConnection Connection = new MySqlConnection(SqlConnector.ConnectionString))
+			{
+				Connection.Open();
+				return FindOrAddUserId(Name, Connection);
+			}
+		}
+
+		private static long FindOrAddUserId(string Name, MySqlConnection Connection)
+		{
+			if(Name.Length == 0)
+			{
+				return -1;
+			}
+
+			string NormalizedName = NormalizeUserName(Name);
+
+			using (MySqlCommand Command = new MySqlCommand("SELECT Id FROM ugs_db.Users WHERE Name = @Name", Connection))
+			{
+				Command.Parameters.AddWithValue("@Name", NormalizedName);
+				object UserId = Command.ExecuteScalar();
+				if(UserId != null)
+				{
+					return Convert.ToInt64(UserId);
+				}
+			}
+
+			using (MySqlCommand Command = new MySqlCommand("INSERT IGNORE INTO ugs_db.Users (Name) VALUES (@Name); SELECT Id FROM ugs_db.Users WHERE Name = @Name", Connection))
+			{
+				Command.Parameters.AddWithValue("@Name", NormalizedName);
+				object UserId = Command.ExecuteScalar();
+				return Convert.ToInt64(UserId);
+			}
+		}
+
+		const int IssueSummaryMaxLength = 200;
+
+		public static long AddIssue(IssueData Issue)
+		{
+			long IssueId;
+			using (MySqlConnection Connection = new MySqlConnection(SqlConnector.ConnectionString))
+			{
+				Connection.Open();
+
+				using (MySqlCommand Command = new MySqlCommand("INSERT INTO ugs_db.Issues (Project, Summary, OwnerId, CreatedAt, FixChange) VALUES (@Project, @Summary, @OwnerId, UTC_TIMESTAMP(), 0)", Connection))
+				{
+					Command.Parameters.AddWithValue("@Project", Issue.Project);
+					Command.Parameters.AddWithValue("@Summary", SanitizeText(Issue.Summary, IssueSummaryMaxLength));
+					if (Issue.Owner != null)
+					{
+						Command.Parameters.AddWithValue("OwnerId", FindOrAddUserId(Issue.Owner, Connection));
+					}
+					else
+					{
+						Command.Parameters.AddWithValue("OwnerId", null);
+					}
+					Command.ExecuteNonQuery();
+
+					IssueId = Command.LastInsertedId;
+				}
+			}
+			return IssueId;
+		}
+
+		public static IssueData GetIssue(long IssueId)
+		{
+			List<IssueData> Issues = GetIssuesInternal(IssueId, null, true, -1);
+			if(Issues.Count == 0)
+			{
+				return null;
+			}
+			else
+			{
+				return Issues[0];
+			}
+		}
+
+		public static List<IssueData> GetIssues(bool IncludeResolved, int NumResults)
+		{
+			return GetIssuesInternal(-1, null, IncludeResolved, NumResults);
+		}
+
+		public static List<IssueData> GetIssues(string UserName)
+		{
+			return GetIssuesInternal(-1, UserName, false, -1);
+		}
+
+		private static List<IssueData> GetIssuesInternal(long IssueId, string UserName, bool IncludeResolved, int NumResults)
+		{
+			List<IssueData> Issues = new List<IssueData>();
+			using (MySqlConnection Connection = new MySqlConnection(SqlConnector.ConnectionString))
+			{
+				Connection.Open();
+
+				long UserId = -1;
+				if(UserName != null)
+				{
+					UserId = FindOrAddUserId(UserName);
+				}
+
+				StringBuilder CommandBuilder = new StringBuilder();
+				CommandBuilder.Append("SELECT");
+				CommandBuilder.Append(" Issues.Id, Issues.CreatedAt, UTC_TIMESTAMP(), Issues.Project, Issues.Summary, OwnerUsers.Name, NominatedByUsers.Name, Issues.AcknowledgedAt, Issues.FixChange, Issues.ResolvedAt");
+				if(UserName != null)
+				{
+					CommandBuilder.Append(", IssueWatchers.UserId");
+				}
+				CommandBuilder.Append(" FROM ugs_db.Issues");
+				CommandBuilder.Append(" LEFT JOIN ugs_db.Users AS OwnerUsers ON OwnerUsers.Id = Issues.OwnerId");
+				CommandBuilder.Append(" LEFT JOIN ugs_db.Users AS NominatedByUsers ON NominatedByUsers.Id = Issues.NominatedById");
+				if(UserName != null)
+				{
+					CommandBuilder.Append(" LEFT JOIN ugs_db.IssueWatchers ON IssueWatchers.IssueId = Issues.Id AND IssueWatchers.UserId = @UserId");
+				}
+				if(IssueId != -1)
+				{
+					CommandBuilder.Append(" WHERE Issues.Id = @IssueId");
+				}
+				else if(!IncludeResolved)
+				{
+					CommandBuilder.Append(" WHERE Issues.ResolvedAt IS NULL");
+				}
+				if(NumResults > 0)
+				{
+					CommandBuilder.AppendFormat(" ORDER BY Issues.Id DESC LIMIT {0}", NumResults);
+				}
+
+				using (MySqlCommand Command = new MySqlCommand(CommandBuilder.ToString(), Connection))
+				{
+					if(IssueId != -1)
+					{
+						Command.Parameters.AddWithValue("@IssueId", IssueId);
+					}
+					if(UserName != null)
+					{
+						Command.Parameters.AddWithValue("@UserId", UserId);
+					}
+
+					using(MySqlDataReader Reader = Command.ExecuteReader())
+					{
+						while(Reader.Read())
+						{
+							IssueData Issue = new IssueData();
+							Issue.Id = Reader.GetInt64(0);
+							Issue.CreatedAt = Reader.GetDateTime(1);
+							Issue.RetrievedAt = Reader.GetDateTime(2);
+							Issue.Project = Reader.GetString(3);
+							Issue.Summary = Reader.GetString(4);
+							Issue.Owner = Reader.IsDBNull(5)? null : Reader.GetString(5);
+							Issue.NominatedBy = Reader.IsDBNull(6)? null : Reader.GetString(6);
+							Issue.AcknowledgedAt = Reader.IsDBNull(7)? (DateTime?)null : Reader.GetDateTime(7);
+							Issue.FixChange = Reader.GetInt32(8);
+							Issue.ResolvedAt = Reader.IsDBNull(9)? (DateTime?)null : Reader.GetDateTime(9);
+							if(UserName != null)
+							{
+								Issue.bNotify = !Reader.IsDBNull(10);
+							}
+							Issues.Add(Issue);
+						}
+					}
+				}
+			}
+			return Issues;
+		}
+
+		public static void UpdateIssue(long IssueId, IssueUpdateData Issue)
+		{
+			using (MySqlConnection Connection = new MySqlConnection(SqlConnector.ConnectionString))
+			{
+				Connection.Open();
+
+				using (MySqlCommand Command = Connection.CreateCommand())
+				{
+					List<string> Columns = new List<string>();
+					List<string> Values = new List<string>();
+					if(Issue.Summary != null)
+					{
+						Columns.Add("Summary");
+						Values.Add("@Summary");
+						Command.Parameters.AddWithValue("@Summary", SanitizeText(Issue.Summary, IssueSummaryMaxLength));
+					}
+					if (Issue.Owner != null)
+					{
+						Columns.Add("OwnerId");
+						Values.Add("@OwnerId");
+						Command.Parameters.AddWithValue("OwnerId", FindOrAddUserId(Issue.Owner, Connection));
+					}
+					if(Issue.NominatedBy != null)
+					{
+						Columns.Add("NominatedById");
+						Values.Add("@NominatedById");
+						Command.Parameters.AddWithValue("NominatedById", FindOrAddUserId(Issue.NominatedBy, Connection));
+					}
+					if(Issue.Acknowledged.HasValue)
+					{
+						Columns.Add("AcknowledgedAt");
+						Values.Add(Issue.Acknowledged.Value? "UTC_TIMESTAMP()" : "NULL");
+					}
+					if(Issue.FixChange.HasValue)
+					{
+						Columns.Add("FixChange");
+						Values.Add("@FixChange");
+						Command.Parameters.AddWithValue("FixChange", Issue.FixChange.Value);
+					}
+					if(Issue.Resolved.HasValue)
+					{
+						Columns.Add("ResolvedAt");
+						Values.Add(Issue.Resolved.Value? "UTC_TIMESTAMP()" : "NULL");
+					}
+
+					StringBuilder CommandText = new StringBuilder("UPDATE ugs_db.Issues SET ");
+					for(int idx = 0; idx < Columns.Count; idx++)
+					{
+						CommandText.Append(String.Format("{0}={1}", Columns[idx], Values[idx]));
+						if(idx != Columns.Count - 1)
+						{
+							CommandText.Append(",");
+						}
+					}
+					CommandText.Append(" WHERE Id = @IssueId");
+					Command.CommandText = CommandText.ToString();
+					Command.Parameters.AddWithValue("@IssueId", IssueId);
+					Command.ExecuteNonQuery();
+				}
+			}
+		}
+
+		public static string SanitizeText(string Text, int Length)
+		{
+			if(Text.Length > Length)
+			{
+				int NewlineIdx = Text.LastIndexOf('\n', Length);
+				if(NewlineIdx == -1)
+				{
+					Text = Text.Substring(0, Length - 3).TrimEnd() + "...";
+				}
+				else
+				{
+					Text = Text.Substring(0, NewlineIdx + 1) + "...";
+				}
+			}
+			return Text;
+		}
+
+		public static void DeleteIssue(long IssueId)
+		{
+			using (MySqlConnection Connection = new MySqlConnection(SqlConnector.ConnectionString))
+			{
+				Connection.Open();
+
+				using (MySqlTransaction Transaction = Connection.BeginTransaction())
+				{
+					using (MySqlCommand Command = Connection.CreateCommand())
+					{
+						Command.Transaction = Transaction;
+
+						Command.CommandText = "DELETE FROM ugs_db.IssueWatchers WHERE IssueId = @IssueId";
+						Command.Parameters.AddWithValue("@IssueId", IssueId);
+						Command.ExecuteNonQuery();
+
+						Command.CommandText = "DELETE FROM ugs_db.IssueBuilds WHERE IssueId = @IssueId";
+						Command.Parameters.AddWithValue("@IssueId", IssueId);
+						Command.ExecuteNonQuery();
+
+						Command.CommandText = "DELETE FROM ugs_db.Issues WHERE Id = @IssueId";
+						Command.Parameters.AddWithValue("@IssueId", IssueId);
+						Command.ExecuteNonQuery();
+
+						Transaction.Commit();
+					}
+				}
+			}
+		}
+
+		public static void AddDiagnostic(long IssueId, IssueDiagnosticData Diagnostic)
+		{
+			using (MySqlConnection Connection = new MySqlConnection(SqlConnector.ConnectionString))
+			{
+				Connection.Open();
+
+				using (MySqlCommand Command = new MySqlCommand("INSERT INTO ugs_db.IssueDiagnostics (IssueId, BuildId, Message, Url) VALUES (@IssueId, @BuildId, @Message, @Url)", Connection))
+				{
+					Command.Parameters.AddWithValue("@IssueId", IssueId);
+					Command.Parameters.AddWithValue("@BuildId", Diagnostic.BuildId);
+					Command.Parameters.AddWithValue("@Message", SanitizeText(Diagnostic.Message, 1000));
+					Command.Parameters.AddWithValue("@Url", Diagnostic.Url);
+					Command.ExecuteNonQuery();
+				}
+			}
+		}
+
+		public static List<IssueDiagnosticData> GetDiagnostics(long IssueId)
+		{
+			List<IssueDiagnosticData> Diagnostics = new List<IssueDiagnosticData>();
+			using (MySqlConnection Connection = new MySqlConnection(SqlConnector.ConnectionString))
+			{
+				Connection.Open();
+
+				StringBuilder CommandBuilder = new StringBuilder();
+				CommandBuilder.Append("SELECT BuildId, Message, Url FROM ugs_db.IssueDiagnostics");
+				CommandBuilder.Append(" WHERE IssueDiagnostics.IssueId = @IssueId");
+
+				using (MySqlCommand Command = new MySqlCommand(CommandBuilder.ToString(), Connection))
+				{
+					Command.Parameters.AddWithValue("@IssueId", IssueId);
+					using (MySqlDataReader Reader = Command.ExecuteReader())
+					{
+						while (Reader.Read())
+						{
+							IssueDiagnosticData Diagnostic = new IssueDiagnosticData();
+							Diagnostic.BuildId = Reader.IsDBNull(0)? (long?)null : (long?)Reader.GetInt64(0);
+							Diagnostic.Message = Reader.GetString(1);
+							Diagnostic.Url = Reader.GetString(2);
+							Diagnostics.Add(Diagnostic);
+						}
+					}
+				}
+			}
+			return Diagnostics;
+		}
+
+		public static void AddWatcher(long IssueId, string UserName)
+		{
+			using (MySqlConnection Connection = new MySqlConnection(SqlConnector.ConnectionString))
+			{
+				Connection.Open();
+
+				long UserId = FindOrAddUserId(UserName, Connection);
+
+				using(MySqlCommand Command = new MySqlCommand("INSERT IGNORE INTO ugs_db.IssueWatchers (IssueId, UserId) VALUES (@IssueId, @UserId)", Connection))
+				{
+					Command.Parameters.AddWithValue("@IssueId", IssueId);
+					Command.Parameters.AddWithValue("@UserId", UserId);
+					Command.ExecuteNonQuery();
+				}
+			}
+		}
+
+		public static List<string> GetWatchers(long IssueId)
+		{
+			List<string> Watchers = new List<string>();
+			using (MySqlConnection Connection = new MySqlConnection(SqlConnector.ConnectionString))
+			{
+				Connection.Open();
+
+				StringBuilder CommandBuilder = new StringBuilder();
+				CommandBuilder.Append("SELECT Users.Name FROM ugs_db.IssueWatchers");
+				CommandBuilder.Append(" LEFT JOIN ugs_db.Users ON IssueWatchers.UserId = Users.Id");
+				CommandBuilder.Append(" WHERE IssueWatchers.IssueId = @IssueId");
+
+				using(MySqlCommand Command = new MySqlCommand(CommandBuilder.ToString(), Connection))
+				{
+					Command.Parameters.AddWithValue("@IssueId", IssueId);
+					using(MySqlDataReader Reader = Command.ExecuteReader())
+					{
+						while(Reader.Read())
+						{
+							Watchers.Add(Reader.GetString(0));
+						}
+					}
+				}
+			}
+			return Watchers;
+		}
+
+		public static void RemoveWatcher(long IssueId, string UserName)
+		{
+			using (MySqlConnection Connection = new MySqlConnection(SqlConnector.ConnectionString))
+			{
+				Connection.Open();
+
+				long UserId = FindOrAddUserId(UserName, Connection);
+
+				using(MySqlCommand Command = new MySqlCommand("DELETE FROM ugs_db.IssueWatchers WHERE IssueId = @IssueId AND UserId = @UserId", Connection))
+				{
+					Command.Parameters.AddWithValue("@IssueId", IssueId);
+					Command.Parameters.AddWithValue("@UserId", UserId);
+					Command.ExecuteNonQuery();
+				}
+			}
+		}
+
+		public static long AddBuild(long IssueId, IssueBuildData Build)
+		{
+			using (MySqlConnection Connection = new MySqlConnection(SqlConnector.ConnectionString))
+			{
+				Connection.Open();
+
+				using (MySqlCommand Command = new MySqlCommand("INSERT INTO ugs_db.IssueBuilds (IssueId, Stream, `Change`, JobName, JobUrl, JobStepName, JobStepUrl, ErrorUrl, Outcome) VALUES (@IssueId, @Stream, @Change, @JobName, @JobUrl, @JobStepName, @JobStepUrl, @ErrorUrl, @Outcome)", Connection))
+				{
+					Command.Parameters.AddWithValue("@IssueId", IssueId);
+					Command.Parameters.AddWithValue("@Stream", Build.Stream);
+					Command.Parameters.AddWithValue("@Change", Build.Change);
+					Command.Parameters.AddWithValue("@JobName", Build.JobName);
+					Command.Parameters.AddWithValue("@JobUrl", Build.JobUrl);
+					Command.Parameters.AddWithValue("@JobStepName", Build.JobStepName);
+					Command.Parameters.AddWithValue("@JobStepUrl", Build.JobStepUrl);
+					Command.Parameters.AddWithValue("@ErrorUrl", Build.ErrorUrl);
+					Command.Parameters.AddWithValue("@Outcome", Build.Outcome);
+					Command.ExecuteNonQuery();
+
+					return Command.LastInsertedId;
+				}
+			}
+		}
+
+		public static List<IssueBuildData> GetBuilds(long IssueId)
+		{
+			List<IssueBuildData> Builds = new List<IssueBuildData>();
+			using (MySqlConnection Connection = new MySqlConnection(SqlConnector.ConnectionString))
+			{
+				Connection.Open();
+
+				using(MySqlCommand Command = new MySqlCommand("SELECT IssueBuilds.Id, IssueBuilds.Stream, IssueBuilds.Change, IssueBuilds.JobName, IssueBuilds.JobUrl, IssueBuilds.JobStepName, IssueBuilds.JobStepUrl, IssueBuilds.ErrorUrl, IssueBuilds.Outcome FROM ugs_db.IssueBuilds WHERE IssueBuilds.IssueId = @IssueId", Connection))
+				{
+					Command.Parameters.AddWithValue("@IssueId", IssueId);
+					using(MySqlDataReader Reader = Command.ExecuteReader())
+					{
+						while(Reader.Read())
+						{
+							long Id = Reader.GetInt64(0);
+							string Stream = Reader.GetString(1);
+							int Change = Reader.GetInt32(2);
+							string JobName = Reader.GetString(3);
+							string JobUrl = Reader.GetString(4);
+							string JobStepName = Reader.GetString(5);
+							string JobStepUrl = Reader.GetString(6);
+							string ErrorUrl = Reader.IsDBNull(7)? null : Reader.GetString(7);
+							int Outcome = Reader.GetInt32(8);
+							Builds.Add(new IssueBuildData { Id = Id, Stream = Stream, Change = Change, JobName = JobName, JobUrl = JobUrl, JobStepName = JobStepName, JobStepUrl = JobStepUrl, ErrorUrl = ErrorUrl, Outcome = Outcome });
+						}
+					}
+				}
+			}
+			return Builds;
+		}
+
+		public static IssueBuildData GetBuild(long BuildId)
+		{
+			IssueBuildData Build = null;
+			using (MySqlConnection Connection = new MySqlConnection(SqlConnector.ConnectionString))
+			{
+				Connection.Open();
+
+				using(MySqlCommand Command = new MySqlCommand("SELECT IssueBuilds.Id, IssueBuilds.Stream, IssueBuilds.Change, IssueBuilds.JobName, IssueBuilds.JobUrl, IssueBuilds.JobStepName, IssueBuilds.JobStepUrl, IssueBuilds.ErrorUrl, IssueBuilds.Outcome FROM ugs_db.IssueBuilds WHERE IssueBuilds.Id = @BuildId", Connection))
+				{
+					Command.Parameters.AddWithValue("@BuildId", BuildId);
+					using(MySqlDataReader Reader = Command.ExecuteReader())
+					{
+						while(Reader.Read())
+						{
+							long Id = Reader.GetInt64(0);
+							string Stream = Reader.GetString(1);
+							int Change = Reader.GetInt32(2);
+							string JobName = Reader.GetString(3);
+							string JobUrl = Reader.GetString(4);
+							string JobStepName = Reader.GetString(5);
+							string JobStepUrl = Reader.GetString(6);
+							string ErrorUrl = Reader.GetString(7);
+							int Outcome = Reader.GetInt32(8);
+
+							Build = new IssueBuildData { Id = Id, Stream = Stream, Change = Change, JobName = JobName, JobUrl = JobUrl, JobStepName = JobStepName, JobStepUrl = JobStepUrl, ErrorUrl = ErrorUrl, Outcome = Outcome };
+						}
+					}
+				}
+			}
+			return Build;
+		}
+
+		public static void UpdateBuild(long BuildId, int Outcome)
+		{
+			using (MySqlConnection Connection = new MySqlConnection(SqlConnector.ConnectionString))
+			{
+				Connection.Open();
+
+				using (MySqlCommand Command = new MySqlCommand("UPDATE ugs_db.IssueBuilds SET (Outcome) = (@Outcome) WHERE Id = @BuildId", Connection))
+				{
+					Command.Parameters.AddWithValue("@BuildId", BuildId);
+					Command.Parameters.AddWithValue("@Outcome", Outcome);
+					Command.ExecuteNonQuery();
+				}
+			}
+		}
+	}
+}
\ No newline at end of file
diff --git a/Engine/Source/Programs/UnrealGameSync/MetadataServer/Controllers/BuildController.cs b/Engine/Source/Programs/UnrealGameSync/MetadataServer/Controllers/BuildController.cs
new file mode 100644
--- /dev/null
+++ b/Engine/Source/Programs/UnrealGameSync/MetadataServer/Controllers/BuildController.cs
@@ -0,0 +1,21 @@
+﻿// Copyright Epic Games, Inc. All Rights Reserved.
+
+using System.Collections.Generic;
+using System.Web.Http;
+using MetadataServer.Connectors;
+using MetadataServer.Models;
+
+namespace MetadataServer.Controllers
+{
+    public class BuildController : ApiController
+    {
+		public List<BuildData> Get(string Project, long LastBuildId)
+		{
+			return SqlConnector.GetBuilds(Project, LastBuildId);
+		}
+		public void Post([FromBody]BuildData Build)
+		{
+			SqlConnector.PostBuild(Build);
+		}
+	}
+}
diff --git a/Engine/Source/Programs/UnrealGameSync/MetadataServer/Controllers/CISController.cs b/Engine/Source/Programs/UnrealGameSync/MetadataServer/Controllers/CISController.cs
new file mode 100644
--- /dev/null
+++ b/Engine/Source/Programs/UnrealGameSync/MetadataServer/Controllers/CISController.cs
@@ -0,0 +1,28 @@
+﻿// Copyright Epic Games, Inc. All Rights Reserved.
+
+using System;
+using System.Collections.Generic;
+using System.Web.Http;
+using MetadataServer.Connectors;
+using MetadataServer.Models;
+
+namespace MetadataServer.Controllers
+{
+	// Deprecated, use Latest or Build controllers instead
+    public class CISController : ApiController
+    {
+		public long[] Get(string Project = null)
+		{
+			LatestData LatestIDs = SqlConnector.GetLastIds(Project);
+			return new long[] { LatestIDs.LastEventId, LatestIDs.LastCommentId, LatestIDs.LastBuildId };
+		}
+		public List<BuildData> Get(string Project, long LastBuildId)
+		{
+			return SqlConnector.GetBuilds(Project, LastBuildId);
+		}
+		public void Post([FromBody]BuildData Build)
+		{
+			SqlConnector.PostBuild(Build);
+		}
+	}
+}
diff --git a/Engine/Source/Programs/UnrealGameSync/MetadataServer/Controllers/CommentController.cs b/Engine/Source/Programs/UnrealGameSync/MetadataServer/Controllers/CommentController.cs
new file mode 100644
--- /dev/null
+++ b/Engine/Source/Programs/UnrealGameSync/MetadataServer/Controllers/CommentController.cs
@@ -0,0 +1,21 @@
+﻿// Copyright Epic Games, Inc. All Rights Reserved.
+
+using System.Collections.Generic;
+using System.Web.Http;
+using MetadataServer.Connectors;
+using MetadataServer.Models;
+
+namespace MetadataServer.Controllers
+{
+	public class CommentController : ApiController
+	{
+		public List<CommentData> Get(string Project, long LastCommentId)
+		{
+			return SqlConnector.GetComments(Project, LastCommentId);
+		}
+		public void Post([FromBody] CommentData Comment)
+		{
+			SqlConnector.PostComment(Comment);
+		}
+	}
+}
\ No newline at end of file
diff --git a/Engine/Source/Programs/UnrealGameSync/MetadataServer/Controllers/ErrorController.cs b/Engine/Source/Programs/UnrealGameSync/MetadataServer/Controllers/ErrorController.cs
new file mode 100644
--- /dev/null
+++ b/Engine/Source/Programs/UnrealGameSync/MetadataServer/Controllers/ErrorController.cs
@@ -0,0 +1,21 @@
+﻿// Copyright Epic Games, Inc. All Rights Reserved.
+
+using System.Collections.Generic;
+using System.Web.Http;
+using MetadataServer.Connectors;
+using MetadataServer.Models;
+
+namespace MetadataServer.Controllers
+{
+    public class ErrorController : ApiController
+    {
+		public List<TelemetryErrorData> Get(int Records = 10)
+		{
+			return SqlConnector.GetErrorData(Records);
+		}
+		public void Post([FromBody] TelemetryErrorData Data, string Version, string IpAddress)
+		{
+			SqlConnector.PostErrorData(Data, Version, IpAddress);
+		}
+    }
+}
diff --git a/Engine/Source/Programs/UnrealGameSync/MetadataServer/Controllers/EventController.cs b/Engine/Source/Programs/UnrealGameSync/MetadataServer/Controllers/EventController.cs
new file mode 100644
--- /dev/null
+++ b/Engine/Source/Programs/UnrealGameSync/MetadataServer/Controllers/EventController.cs
@@ -0,0 +1,22 @@
+﻿// Copyright Epic Games, Inc. All Rights Reserved.
+
+using System.Collections.Generic;
+using System.Web.Http;
+using MetadataServer.Models;
+using MetadataServer.Connectors;
+
+namespace MetadataServer.Controllers
+{
+	public class EventController : ApiController
+	{
+		public List<EventData> Get(string Project, long LastEventId)
+		{
+			return SqlConnector.GetUserVotes(Project, LastEventId);
+		}
+
+		public void Post([FromBody] EventData Event)
+		{
+			SqlConnector.PostEvent(Event);
+		}
+	}
+}
\ No newline at end of file
diff --git a/Engine/Source/Programs/UnrealGameSync/MetadataServer/Controllers/IssuesController.cs b/Engine/Source/Programs/UnrealGameSync/MetadataServer/Controllers/IssuesController.cs
new file mode 100644
--- /dev/null
+++ b/Engine/Source/Programs/UnrealGameSync/MetadataServer/Controllers/IssuesController.cs
@@ -0,0 +1,122 @@
+// Copyright Epic Games, Inc. All Rights Reserved.
+
+using Newtonsoft.Json.Linq;
+using System;
+using System.Collections.Generic;
+using System.Dynamic;
+using System.Linq;
+using System.Net.Http;
+using System.Web.Http;
+using System.Web.Http.Results;
+using MetadataServer.Connectors;
+using MetadataServer.Models;
+
+namespace MetadataServer.Controllers
+{
+	public class IssuesController : ApiController
+	{
+		[HttpGet]
+		public List<IssueData> Get(bool IncludeResolved = false, int MaxResults = -1)
+		{
+			return SqlConnector.GetIssues(IncludeResolved, MaxResults);
+		}
+
+		[HttpGet]
+		public List<IssueData> Get(string User)
+		{
+			return SqlConnector.GetIssues(User);
+		}
+
+		[HttpGet]
+		public IssueData Get(long id)
+		{
+			return SqlConnector.GetIssue(id);
+		}
+
+		[HttpPut]
+		public void Put(long id, IssueUpdateData Issue)
+		{
+			SqlConnector.UpdateIssue(id, Issue);
+		}
+
+		[HttpPost]
+		public object Post(IssueData Issue)
+		{
+			long IssueId = SqlConnector.AddIssue(Issue);
+			return new { Id = IssueId };
+		}
+
+		[HttpDelete]
+		public void Delete(long id)
+		{
+			SqlConnector.DeleteIssue(id);
+		}
+	}
+
+	public class IssueBuildsSubController : ApiController
+	{
+		[HttpGet]
+		public List<IssueBuildData> Get(long IssueId)
+		{
+			return SqlConnector.GetBuilds(IssueId);
+		}
+
+		[HttpPost]
+		public object Post(long IssueId, [FromBody] IssueBuildData Data)
+		{
+			long BuildId = SqlConnector.AddBuild(IssueId, Data);
+			return new { Id = BuildId };
+		}
+	}
+
+	public class IssueDiagnosticsSubController : ApiController
+	{
+		[HttpGet]
+		public List<IssueDiagnosticData> Get(long IssueId)
+		{
+			return SqlConnector.GetDiagnostics(IssueId);
+		}
+
+		[HttpPost]
+		public void Post(long IssueId, [FromBody] IssueDiagnosticData Data)
+		{
+			SqlConnector.AddDiagnostic(IssueId, Data);
+		}
+	}
+
+	public class IssueBuildsController : ApiController
+	{
+		[HttpGet]
+		public IssueBuildData Get(long BuildId)
+		{
+			return SqlConnector.GetBuild(BuildId);
+		}
+
+		[HttpPut]
+		public void Put(long BuildId, [FromBody] IssueBuildUpdateData Data)
+		{
+			SqlConnector.UpdateBuild(BuildId, Data.Outcome);
+		}
+	}
+
+	public class IssueWatchersController : ApiController
+	{
+		[HttpGet]
+		public List<string> Get(long IssueId)
+		{
+			return SqlConnector.GetWatchers(IssueId);
+		}
+
+		[HttpPost]
+		public void Post(long IssueId, [FromBody] IssueWatcherData Data)
+		{
+			SqlConnector.AddWatcher(IssueId, Data.UserName);
+		}
+
+		[HttpDelete]
+		public void Delete(long IssueId, [FromBody] IssueWatcherData Data)
+		{
+			SqlConnector.RemoveWatcher(IssueId, Data.UserName);
+		}
+	}
+}
diff --git a/Engine/Source/Programs/UnrealGameSync/MetadataServer/Controllers/LatestController.cs b/Engine/Source/Programs/UnrealGameSync/MetadataServer/Controllers/LatestController.cs
new file mode 100644
--- /dev/null
+++ b/Engine/Source/Programs/UnrealGameSync/MetadataServer/Controllers/LatestController.cs
@@ -0,0 +1,21 @@
+﻿// Copyright Epic Games, Inc. All Rights Reserved.
+
+using System;
+using System.Collections.Generic;
+using System.Linq;
+using System.Net;
+using System.Net.Http;
+using System.Web.Http;
+using MetadataServer.Connectors;
+using MetadataServer.Models;
+
+namespace MetadataServer.Controllers
+{
+    public class LatestController : ApiController
+    {
+		public LatestData Get(string Project = null)
+		{
+			return SqlConnector.GetLastIds(Project);
+		}
+	}
+}
diff --git a/Engine/Source/Programs/UnrealGameSync/MetadataServer/Controllers/TelemetryController.cs b/Engine/Source/Programs/UnrealGameSync/MetadataServer/Controllers/TelemetryController.cs
new file mode 100644
--- /dev/null
+++ b/Engine/Source/Programs/UnrealGameSync/MetadataServer/Controllers/TelemetryController.cs
@@ -0,0 +1,16 @@
+﻿// Copyright Epic Games, Inc. All Rights Reserved.
+
+using System.Web.Http;
+using MetadataServer.Connectors;
+using MetadataServer.Models;
+
+namespace MetadataServer.Controllers
+{
+    public class TelemetryController : ApiController
+    {
+		public void Post([FromBody] TelemetryTimingData Data, string Version, string IpAddress)
+		{
+			SqlConnector.PostTelemetryData(Data, Version, IpAddress);
+		}
+    }
+}
diff --git a/Engine/Source/Programs/UnrealGameSync/MetadataServer/Controllers/UserController.cs b/Engine/Source/Programs/UnrealGameSync/MetadataServer/Controllers/UserController.cs
new file mode 100644
--- /dev/null
+++ b/Engine/Source/Programs/UnrealGameSync/MetadataServer/Controllers/UserController.cs
@@ -0,0 +1,19 @@
+// Copyright Epic Games, Inc. All Rights Reserved.
+
+using System;
+using System.Collections.Generic;
+using System.Web.Http;
+using MetadataServer.Connectors;
+using MetadataServer.Models;
+
+namespace MetadataServer.Controllers
+{
+	public class UserController : ApiController
+	{
+		public object Get(string Name)
+		{
+			long UserId = SqlConnector.FindOrAddUserId(Name);
+			return new { Id = UserId };
+		}
+	}
+}
diff --git a/Engine/Source/Programs/UnrealGameSync/MetadataServer/Global.asax.cs b/Engine/Source/Programs/UnrealGameSync/MetadataServer/Global.asax.cs
new file mode 100644
--- /dev/null
+++ b/Engine/Source/Programs/UnrealGameSync/MetadataServer/Global.asax.cs
@@ -0,0 +1,38 @@
+﻿// Copyright Epic Games, Inc. All Rights Reserved.
+using MySql.Data.MySqlClient;
+using System;
+using System.Collections.Generic;
+using System.IO;
+using System.Linq;
+using System.Text.RegularExpressions;
+using System.Web;
+using System.Web.Http;
+
+namespace MetadataServer
+{
+	public class WebApiApplication : System.Web.HttpApplication
+	{
+        protected void Application_Start()
+        {
+            try
+            {
+                string ConnectionString = System.Configuration.ConfigurationManager.ConnectionStrings["ConnectionString"].ConnectionString;
+				using (MySqlConnection Connection = new MySqlConnection(ConnectionString))
+				{
+					Connection.Open();
+					using (MySqlCommand Command = new MySqlCommand(Properties.Resources.Setup, Connection))
+					{
+						Command.ExecuteNonQuery();
+					}
+				}
+			}
+            catch (Exception)
+            {
+                HttpRuntime.UnloadAppDomain();
+                return;
+            }
+            GlobalConfiguration.Configure(WebApiConfig.Register);
+        }
+	}
+    
+}
diff --git a/Engine/Source/Programs/UnrealGameSync/MetadataServer/MetadataServer.csproj b/Engine/Source/Programs/UnrealGameSync/MetadataServer/MetadataServer.csproj
new file mode 100644
--- /dev/null
+++ b/Engine/Source/Programs/UnrealGameSync/MetadataServer/MetadataServer.csproj
@@ -0,0 +1,191 @@
+﻿<?xml version="1.0" encoding="utf-8"?>
+<Project ToolsVersion="14.0" DefaultTargets="Build" xmlns="http://schemas.microsoft.com/developer/msbuild/2003">
+  <Import Project="$(MSBuildExtensionsPath)\$(MSBuildToolsVersion)\Microsoft.Common.props" Condition="Exists('$(MSBuildExtensionsPath)\$(MSBuildToolsVersion)\Microsoft.Common.props')" />
+  <PropertyGroup>
+    <Configuration Condition=" '$(Configuration)' == '' ">Debug</Configuration>
+    <Platform Condition=" '$(Platform)' == '' ">AnyCPU</Platform>
+    <ProductVersion>
+    </ProductVersion>
+    <SchemaVersion>2.0</SchemaVersion>
+    <ProjectGuid>{C8D13E98-C8CD-42FB-9C79-F3EB11B0F3E8}</ProjectGuid>
+    <ProjectTypeGuids>{349c5851-65df-11da-9384-00065b846f21};{fae04ec0-301f-11d3-bf4b-00c04f79efbc}</ProjectTypeGuids>
+    <OutputType>Library</OutputType>
+    <AppDesignerFolder>Properties</AppDesignerFolder>
+    <RootNamespace>MetadataServer</RootNamespace>
+    <AssemblyName>MetadataServer</AssemblyName>
+    <TargetFrameworkVersion>v4.6.2</TargetFrameworkVersion>
+    <MvcBuildViews>false</MvcBuildViews>
+    <UseIISExpress>true</UseIISExpress>
+    <IISExpressSSLPort />
+    <IISExpressAnonymousAuthentication />
+    <IISExpressWindowsAuthentication />
+    <IISExpressUseClassicPipelineMode />
+    <UseGlobalApplicationHostFile />
+    <NuGetPackageImportStamp>
+    </NuGetPackageImportStamp>
+    <TargetFrameworkProfile />
+    <Use64BitIISExpress />
+  </PropertyGroup>
+  <PropertyGroup Condition=" '$(Configuration)|$(Platform)' == 'Debug|AnyCPU' ">
+    <TransformWebConfigEnabled>False</TransformWebConfigEnabled>
+    <DebugSymbols>true</DebugSymbols>
+    <DebugType>full</DebugType>
+    <Optimize>false</Optimize>
+    <OutputPath>bin\</OutputPath>
+    <DefineConstants>DEBUG;TRACE</DefineConstants>
+    <ErrorReport>prompt</ErrorReport>
+    <WarningLevel>4</WarningLevel>
+  </PropertyGroup>
+  <PropertyGroup Condition=" '$(Configuration)|$(Platform)' == 'Release|AnyCPU' ">
+    <TransformWebConfigEnabled>False</TransformWebConfigEnabled>
+    <DebugType>pdbonly</DebugType>
+    <Optimize>true</Optimize>
+    <OutputPath>bin\</OutputPath>
+    <DefineConstants>TRACE</DefineConstants>
+    <ErrorReport>prompt</ErrorReport>
+    <WarningLevel>4</WarningLevel>
+  </PropertyGroup>
+  <ItemGroup>
+    <Reference Include="BouncyCastle.Crypto, Version=1.8.3.0, Culture=neutral, PublicKeyToken=0e99375e54769942">
+      <HintPath>..\packages\BouncyCastle.1.8.3.1\lib\BouncyCastle.Crypto.dll</HintPath>
+    </Reference>
+    <Reference Include="Google.Protobuf, Version=3.6.1.0, Culture=neutral, PublicKeyToken=a7d26565bac4d604, processorArchitecture=MSIL">
+      <HintPath>..\packages\Google.Protobuf.3.6.1\lib\net45\Google.Protobuf.dll</HintPath>
+    </Reference>
+    <Reference Include="Microsoft.CSharp" />
+    <Reference Include="MySql.Data, Version=8.0.16.0, Culture=neutral, PublicKeyToken=c5687fc88969c44d, processorArchitecture=MSIL">
+      <HintPath>..\packages\MySql.Data.8.0.16\lib\net452\MySql.Data.dll</HintPath>
+    </Reference>
+    <Reference Include="Newtonsoft.Json, Version=12.0.0.0, Culture=neutral, PublicKeyToken=30ad4fe6b2a6aeed, processorArchitecture=MSIL">
+      <HintPath>..\packages\Newtonsoft.Json.12.0.2\lib\net45\Newtonsoft.Json.dll</HintPath>
+    </Reference>
+    <Reference Include="System" />
+    <Reference Include="System.Data" />
+    <Reference Include="System.Drawing" />
+    <Reference Include="System.Net.Http.Formatting, Version=5.2.7.0, Culture=neutral, PublicKeyToken=31bf3856ad364e35, processorArchitecture=MSIL">
+      <HintPath>..\packages\Microsoft.AspNet.WebApi.Client.5.2.7\lib\net45\System.Net.Http.Formatting.dll</HintPath>
+    </Reference>
+    <Reference Include="System.Web.DynamicData" />
+    <Reference Include="System.Web.Entity" />
+    <Reference Include="System.Web.ApplicationServices" />
+    <Reference Include="System.ComponentModel.DataAnnotations" />
+    <Reference Include="System.Web" />
+    <Reference Include="System.Web.Abstractions" />
+    <Reference Include="System.Web.Http, Version=5.2.7.0, Culture=neutral, PublicKeyToken=31bf3856ad364e35, processorArchitecture=MSIL">
+      <HintPath>..\packages\Microsoft.AspNet.WebApi.Core.5.2.7\lib\net45\System.Web.Http.dll</HintPath>
+    </Reference>
+    <Reference Include="System.Web.Http.WebHost, Version=5.2.7.0, Culture=neutral, PublicKeyToken=31bf3856ad364e35, processorArchitecture=MSIL">
+      <HintPath>..\packages\Microsoft.AspNet.WebApi.WebHost.5.2.7\lib\net45\System.Web.Http.WebHost.dll</HintPath>
+    </Reference>
+    <Reference Include="System.Web.Routing" />
+    <Reference Include="System.Xml" />
+    <Reference Include="System.Configuration" />
+    <Reference Include="System.Runtime.Serialization" />
+    <Reference Include="System.Net.Http">
+    </Reference>
+    <Reference Include="System.Net.Http.WebRequest">
+    </Reference>
+  </ItemGroup>
+  <ItemGroup>
+    <Folder Include="App_Data\" />
+  </ItemGroup>
+  <ItemGroup>
+    <Compile Include="App_Start\WebApiConfig.cs" />
+    <Compile Include="Controllers\IssuesController.cs" />
+    <Compile Include="Controllers\UserController.cs" />
+    <Compile Include="Models\IssueData.cs" />
+    <Compile Include="Properties\Resources.Designer.cs">
+      <AutoGen>True</AutoGen>
+      <DesignTime>True</DesignTime>
+      <DependentUpon>Resources.resx</DependentUpon>
+    </Compile>
+    <Content Include="Setup.sql" />
+    <Compile Include="Controllers\BuildController.cs" />
+    <Compile Include="Controllers\CISController.cs" />
+    <Compile Include="Controllers\CommentController.cs" />
+    <Compile Include="Controllers\ErrorController.cs" />
+    <Compile Include="Controllers\EventController.cs" />
+    <Compile Include="Controllers\LatestController.cs" />
+    <Compile Include="Controllers\TelemetryController.cs" />
+    <Compile Include="Global.asax.cs">
+      <DependentUpon>Global.asax</DependentUpon>
+    </Compile>
+    <Compile Include="Models\BuildData.cs" />
+    <Compile Include="Models\CommentData.cs" />
+    <Compile Include="Models\EventData.cs" />
+    <Compile Include="Models\EventSummary.cs" />
+    <Compile Include="Models\LatestData.cs" />
+    <Compile Include="Models\TelemetryErrorData.cs" />
+    <Compile Include="Models\TelemetryTimingData.cs" />
+    <Compile Include="Properties\AssemblyInfo.cs" />
+    <Compile Include="Connectors\SqlConnector.cs" />
+  </ItemGroup>
+  <ItemGroup>
+    <Content Include="Global.asax" />
+    <Content Include="Web.config">
+      <SubType>Designer</SubType>
+    </Content>
+    <None Include="Web.template.config">
+      <SubType>Designer</SubType>
+    </None>
+    <None Include="Web.Debug.config">
+      <DependentUpon>Web.template.config</DependentUpon>
+    </None>
+    <None Include="Web.Release.config">
+      <DependentUpon>Web.template.config</DependentUpon>
+    </None>
+  </ItemGroup>
+  <ItemGroup>
+    <Content Include="NotForLicensees\Web.config">
+      <SubType>Designer</SubType>
+    </Content>
+  </ItemGroup>
+  <ItemGroup>
+    <Content Include="packages.config" />
+  </ItemGroup>
+  <ItemGroup>
+    <EmbeddedResource Include="Properties\Resources.resx">
+      <Generator>ResXFileCodeGenerator</Generator>
+      <LastGenOutput>Resources.Designer.cs</LastGenOutput>
+      <SubType>Designer</SubType>
+    </EmbeddedResource>
+  </ItemGroup>
+  <PropertyGroup>
+    <VisualStudioVersion Condition="'$(VisualStudioVersion)' == ''">10.0</VisualStudioVersion>
+    <VSToolsPath Condition="'$(VSToolsPath)' == ''">$(MSBuildExtensionsPath32)\Microsoft\VisualStudio\v$(VisualStudioVersion)</VSToolsPath>
+  </PropertyGroup>
+  <Import Project="$(MSBuildBinPath)\Microsoft.CSharp.targets" />
+  <Import Project="$(VSToolsPath)\WebApplications\Microsoft.WebApplication.targets" Condition="'$(VSToolsPath)' != ''" />
+  <Import Project="$(MSBuildExtensionsPath32)\Microsoft\VisualStudio\v10.0\WebApplications\Microsoft.WebApplication.targets" Condition="false" />
+  <Target Name="MvcBuildViews" AfterTargets="AfterBuild" Condition="'$(MvcBuildViews)'=='true'">
+    <AspNetCompiler VirtualPath="temp" PhysicalPath="$(WebProjectOutputDir)" />
+  </Target>
+  <ProjectExtensions>
+    <VisualStudio>
+      <FlavorProperties GUID="{349c5851-65df-11da-9384-00065b846f21}">
+        <WebProjectProperties>
+          <UseIIS>True</UseIIS>
+          <AutoAssignPort>True</AutoAssignPort>
+          <DevelopmentServerPort>62486</DevelopmentServerPort>
+          <DevelopmentServerVPath>/</DevelopmentServerVPath>
+          <IISUrl>http://localhost:62487/</IISUrl>
+          <NTLMAuthentication>False</NTLMAuthentication>
+          <UseCustomServer>False</UseCustomServer>
+          <CustomServerUrl>
+          </CustomServerUrl>
+          <SaveServerSettingsInUserFile>False</SaveServerSettingsInUserFile>
+        </WebProjectProperties>
+      </FlavorProperties>
+    </VisualStudio>
+  </ProjectExtensions>
+  <!-- To modify your build process, add your task inside one of the targets below and uncomment it.
+       Other similar extension points exist, see Microsoft.Common.targets.
+  <Target Name="BeforeBuild">
+  </Target> -->
+  <Target Name="BeforeBuild">
+    <Message Importance="High" Text="Creating web.config from web.template.config" />
+    <TransformXml Source="Web.template.config" Transform="Web.$(Configuration).config" Destination="Web.config" />
+    <Message Importance="High" Text="Found NotForLicensees\Web.config, applying transform" Condition="Exists('NotForLicensees\Web.config')" />
+    <TransformXml Source="Web.config" Transform="NotForLicensees\Web.config" Destination="Web.config" Condition="Exists('NotForLicensees\Web.config')" />
+  </Target>
+</Project>
\ No newline at end of file
diff --git a/Engine/Source/Programs/UnrealGameSync/MetadataServer/Models/BuildData.cs b/Engine/Source/Programs/UnrealGameSync/MetadataServer/Models/BuildData.cs
new file mode 100644
--- /dev/null
+++ b/Engine/Source/Programs/UnrealGameSync/MetadataServer/Models/BuildData.cs
@@ -0,0 +1,39 @@
+﻿// Copyright Epic Games, Inc. All Rights Reserved.
+
+using System;
+using System.Collections.Generic;
+using System.Linq;
+using System.Web;
+
+namespace MetadataServer.Models
+{
+	public class BuildData
+	{
+		public enum BuildDataResult
+		{
+			Starting,
+			Failure,
+			Warning,
+			Success,
+			Skipped,
+		}
+
+		public long Id;
+		public int ChangeNumber;
+		public string BuildType;
+		public BuildDataResult Result;
+		public string Url;
+		public string Project;
+		public string ArchivePath;
+
+		public bool IsSuccess
+		{
+			get { return Result == BuildDataResult.Success || Result == BuildDataResult.Warning; }
+		}
+
+		public bool IsFailure
+		{
+			get { return Result == BuildDataResult.Failure; }
+		}
+	}
+}
\ No newline at end of file
diff --git a/Engine/Source/Programs/UnrealGameSync/MetadataServer/Models/CommentData.cs b/Engine/Source/Programs/UnrealGameSync/MetadataServer/Models/CommentData.cs
new file mode 100644
--- /dev/null
+++ b/Engine/Source/Programs/UnrealGameSync/MetadataServer/Models/CommentData.cs
@@ -0,0 +1,18 @@
+﻿// Copyright Epic Games, Inc. All Rights Reserved.
+
+using System;
+using System.Collections.Generic;
+using System.Linq;
+using System.Web;
+
+namespace MetadataServer.Models
+{
+	public class CommentData
+	{
+		public long Id;
+		public int ChangeNumber;
+		public string UserName;
+		public string Text;
+		public string Project;
+	}
+}
\ No newline at end of file
diff --git a/Engine/Source/Programs/UnrealGameSync/MetadataServer/Models/EventData.cs b/Engine/Source/Programs/UnrealGameSync/MetadataServer/Models/EventData.cs
new file mode 100644
--- /dev/null
+++ b/Engine/Source/Programs/UnrealGameSync/MetadataServer/Models/EventData.cs
@@ -0,0 +1,38 @@
+﻿// Copyright Epic Games, Inc. All Rights Reserved.
+
+using System;
+using System.Collections.Generic;
+using System.Linq;
+using System.Web;
+
+namespace MetadataServer.Models
+{
+	public class EventData
+	{
+		public enum EventType
+		{
+			Syncing,
+
+			// Reviews
+			Compiles,
+			DoesNotCompile,
+			Good,
+			Bad,
+			Unknown,
+
+			// Starred builds
+			Starred,
+			Unstarred,
+
+			// Investigating events
+			Investigating,
+			Resolved,
+		}
+		
+		public long Id;
+		public int Change;
+		public string UserName;
+		public EventType Type;
+		public string Project;
+	}
+}
\ No newline at end of file
diff --git a/Engine/Source/Programs/UnrealGameSync/MetadataServer/Models/EventSummary.cs b/Engine/Source/Programs/UnrealGameSync/MetadataServer/Models/EventSummary.cs
new file mode 100644
--- /dev/null
+++ b/Engine/Source/Programs/UnrealGameSync/MetadataServer/Models/EventSummary.cs
@@ -0,0 +1,29 @@
+﻿// Copyright Epic Games, Inc. All Rights Reserved.
+
+using System;
+using System.Collections.Generic;
+using System.Linq;
+using System.Web;
+
+namespace MetadataServer.Models
+{
+	public class EventSummary
+	{
+		public enum ReviewVerdict
+		{
+			Unknown,
+			Good,
+			Bad,
+			Mixed,
+		}
+
+		public int ChangeNumber;
+		public ReviewVerdict Verdict;
+		public List<EventData> SyncEvents = new List<EventData>();
+		public List<EventData> Reviews = new List<EventData>();
+		public List<string> CurrentUsers = new List<string>();
+		public EventData LastStarReview;
+		public List<BuildData> Builds = new List<BuildData>();
+		public List<CommentData> Comments = new List<CommentData>();
+	}
+}
\ No newline at end of file
diff --git a/Engine/Source/Programs/UnrealGameSync/MetadataServer/Models/IssueData.cs b/Engine/Source/Programs/UnrealGameSync/MetadataServer/Models/IssueData.cs
new file mode 100644
--- /dev/null
+++ b/Engine/Source/Programs/UnrealGameSync/MetadataServer/Models/IssueData.cs
@@ -0,0 +1,65 @@
+// Copyright Epic Games, Inc. All Rights Reserved.
+
+using Newtonsoft.Json.Linq;
+using System;
+using System.Collections.Generic;
+using System.Linq;
+using System.Web;
+
+namespace MetadataServer.Models
+{
+	public class IssueWatcherData
+	{
+		public string UserName;
+	}
+
+	public class IssueBuildData
+	{
+		public long Id;
+		public string Stream;
+		public int Change;
+		public string JobName;
+		public string JobUrl;
+		public string JobStepName;
+		public string JobStepUrl;
+		public string ErrorUrl;
+		public int Outcome;
+	}
+
+	public class IssueBuildUpdateData
+	{
+		public int Outcome;
+	}
+
+	public class IssueDiagnosticData
+	{
+		public long? BuildId;
+		public string Message;
+		public string Url;
+	}
+
+	public class IssueData
+	{
+		public long Id;
+		public DateTime CreatedAt;
+		public DateTime RetrievedAt;
+		public string Project;
+		public string Summary;
+		public string Owner;
+		public string NominatedBy;
+		public DateTime? AcknowledgedAt;
+		public int FixChange;
+		public DateTime? ResolvedAt;
+		public bool bNotify;
+	}
+
+	public class IssueUpdateData
+	{
+		public string Summary;
+		public string Owner;
+		public string NominatedBy;
+		public bool? Acknowledged;
+		public int? FixChange;
+		public bool? Resolved;
+	}
+}
diff --git a/Engine/Source/Programs/UnrealGameSync/MetadataServer/Models/LatestData.cs b/Engine/Source/Programs/UnrealGameSync/MetadataServer/Models/LatestData.cs
new file mode 100644
--- /dev/null
+++ b/Engine/Source/Programs/UnrealGameSync/MetadataServer/Models/LatestData.cs
@@ -0,0 +1,16 @@
+﻿// Copyright Epic Games, Inc. All Rights Reserved.
+
+using System;
+using System.Collections.Generic;
+using System.Linq;
+using System.Web;
+
+namespace MetadataServer.Models
+{
+	public class LatestData
+	{
+		public long LastEventId { get; set; }
+		public long LastCommentId { get; set; }
+		public long LastBuildId { get; set; }
+	}
+}
\ No newline at end of file
diff --git a/Engine/Source/Programs/UnrealGameSync/MetadataServer/Models/TelemetryErrorData.cs b/Engine/Source/Programs/UnrealGameSync/MetadataServer/Models/TelemetryErrorData.cs
new file mode 100644
--- /dev/null
+++ b/Engine/Source/Programs/UnrealGameSync/MetadataServer/Models/TelemetryErrorData.cs
@@ -0,0 +1,25 @@
+﻿// Copyright Epic Games, Inc. All Rights Reserved.
+
+using System;
+using System.Collections.Generic;
+using System.Linq;
+using System.Web;
+
+namespace MetadataServer.Models
+{
+	public class TelemetryErrorData
+	{
+		public enum TelemetryErrorType
+		{
+			Crash,
+		}
+		public int Id;
+		public TelemetryErrorType Type;
+		public string Text;
+		public string UserName;
+		public string Project;
+		public DateTime Timestamp;
+		public string Version;
+		public string IpAddress;
+	}
+}
\ No newline at end of file
diff --git a/Engine/Source/Programs/UnrealGameSync/MetadataServer/Models/TelemetryTimingData.cs b/Engine/Source/Programs/UnrealGameSync/MetadataServer/Models/TelemetryTimingData.cs
new file mode 100644
--- /dev/null
+++ b/Engine/Source/Programs/UnrealGameSync/MetadataServer/Models/TelemetryTimingData.cs
@@ -0,0 +1,19 @@
+﻿// Copyright Epic Games, Inc. All Rights Reserved.
+
+using System;
+using System.Collections.Generic;
+using System.Linq;
+using System.Web;
+
+namespace MetadataServer.Models
+{
+	public class TelemetryTimingData
+	{
+		public string Action;
+		public string Result;
+		public string UserName;
+		public string Project;
+		public DateTime Timestamp;
+		public float Duration;
+	}
+}
\ No newline at end of file
diff --git a/Engine/Source/Programs/UnrealGameSync/MetadataServer/Properties/AssemblyInfo.cs b/Engine/Source/Programs/UnrealGameSync/MetadataServer/Properties/AssemblyInfo.cs
new file mode 100644
--- /dev/null
+++ b/Engine/Source/Programs/UnrealGameSync/MetadataServer/Properties/AssemblyInfo.cs
@@ -0,0 +1,37 @@
+﻿// Copyright Epic Games, Inc. All Rights Reserved.
+
+using System.Reflection;
+using System.Runtime.CompilerServices;
+using System.Runtime.InteropServices;
+
+// General Information about an assembly is controlled through the following
+// set of attributes. Change these attribute values to modify the information
+// associated with an assembly.
+[assembly: AssemblyTitle("MetadataServer")]
+[assembly: AssemblyDescription("")]
+[assembly: AssemblyConfiguration("")]
+[assembly: AssemblyCompany("")]
+[assembly: AssemblyProduct("UnrealGameSync")]
+[assembly: AssemblyCopyright("Copyright Epic Games, Inc. All Rights Reserved.")]
+[assembly: AssemblyTrademark("")]
+[assembly: AssemblyCulture("")]
+
+// Setting ComVisible to false makes the types in this assembly not visible
+// to COM components.  If you need to access a type in this assembly from
+// COM, set the ComVisible attribute to true on that type.
+[assembly: ComVisible(false)]
+
+// The following GUID is for the ID of the typelib if this project is exposed to COM
+[assembly: Guid("f6044fb1-c977-4705-8cad-f17962adfb41")]
+
+// Version information for an assembly consists of the following four values:
+//
+//      Major Version
+//      Minor Version
+//      Build Number
+//      Revision
+//
+// You can specify all the values or you can default the Revision and Build Numbers
+// by using the '*' as shown below:
+[assembly: AssemblyVersion("1.0.0.0")]
+[assembly: AssemblyFileVersion("1.0.0.0")]
diff --git a/Engine/Source/Programs/UnrealGameSync/MetadataServer/Web.Debug.config b/Engine/Source/Programs/UnrealGameSync/MetadataServer/Web.Debug.config
new file mode 100644
--- /dev/null
+++ b/Engine/Source/Programs/UnrealGameSync/MetadataServer/Web.Debug.config
@@ -0,0 +1,30 @@
+﻿<?xml version="1.0"?>
+
+<!-- For more information on using Web.config transformation visit http://go.microsoft.com/fwlink/?LinkId=301874 -->
+
+<configuration xmlns:xdt="http://schemas.microsoft.com/XML-Document-Transform">
+  <!--
+    In the example below, the "SetAttributes" transform will change the value of
+    "connectionString" to use "ReleaseSQLServer" only when the "Match" locator
+    finds an attribute "name" that has a value of "MyDB".
+
+    <connectionStrings>
+      <add name="MyDB"
+        connectionString="Data Source=ReleaseSQLServer;Initial Catalog=MyReleaseDB;Integrated Security=True"
+        xdt:Transform="SetAttributes" xdt:Locator="Match(name)"/>
+    </connectionStrings>
+  -->
+  <system.web>
+    <!--
+      In the example below, the "Replace" transform will replace the entire
+      <customErrors> section of your Web.config file.
+      Note that because there is only one customErrors section under the
+      <system.web> node, there is no need to use the "xdt:Locator" attribute.
+
+      <customErrors defaultRedirect="GenericError.htm"
+        mode="RemoteOnly" xdt:Transform="Replace">
+        <error statusCode="500" redirect="InternalError.htm"/>
+      </customErrors>
+    -->
+  </system.web>
+</configuration>
diff --git a/Engine/Source/Programs/UnrealGameSync/MetadataServer/Web.Release.config b/Engine/Source/Programs/UnrealGameSync/MetadataServer/Web.Release.config
new file mode 100644
--- /dev/null
+++ b/Engine/Source/Programs/UnrealGameSync/MetadataServer/Web.Release.config
@@ -0,0 +1,31 @@
+﻿<?xml version="1.0"?>
+
+<!-- For more information on using Web.config transformation visit http://go.microsoft.com/fwlink/?LinkId=301874 -->
+
+<configuration xmlns:xdt="http://schemas.microsoft.com/XML-Document-Transform">
+  <!--
+    In the example below, the "SetAttributes" transform will change the value of
+    "connectionString" to use "ReleaseSQLServer" only when the "Match" locator
+    finds an attribute "name" that has a value of "MyDB".
+
+    <connectionStrings>
+      <add name="MyDB"
+        connectionString="Data Source=ReleaseSQLServer;Initial Catalog=MyReleaseDB;Integrated Security=True"
+        xdt:Transform="SetAttributes" xdt:Locator="Match(name)"/>
+    </connectionStrings>
+  -->
+  <system.web>
+    <compilation xdt:Transform="RemoveAttributes(debug)" />
+    <!--
+      In the example below, the "Replace" transform will replace the entire
+      <customErrors> section of your Web.config file.
+      Note that because there is only one customErrors section under the
+      <system.web> node, there is no need to use the "xdt:Locator" attribute.
+
+      <customErrors defaultRedirect="GenericError.htm"
+        mode="RemoteOnly" xdt:Transform="Replace">
+        <error statusCode="500" redirect="InternalError.htm"/>
+      </customErrors>
+    -->
+  </system.web>
+</configuration>
diff --git a/Engine/Source/Programs/UnrealGameSync/MetadataServer/Web.template.config b/Engine/Source/Programs/UnrealGameSync/MetadataServer/Web.template.config
new file mode 100644
--- /dev/null
+++ b/Engine/Source/Programs/UnrealGameSync/MetadataServer/Web.template.config
@@ -0,0 +1,42 @@
+﻿<?xml version="1.0" encoding="utf-8"?>
+<!--
+  For more information on how to configure your ASP.NET application, please visit
+  http://go.microsoft.com/fwlink/?LinkId=301879
+  -->
+<configuration>
+	<appSettings/>
+  <!--
+    For a description of web.config changes see http://go.microsoft.com/fwlink/?LinkId=235367.
+
+    The following attributes can be set on the <httpRuntime> tag.
+      <system.Web>
+        <httpRuntime targetFramework="4.5" />
+      </system.Web>
+  -->
+	<!-- Leave this connection string blank! It should go in either the .debug or .release transform config
+         There is a custom build step on BeforeBuild that generates a transformed Web.config on the fly. The generated Web.config should not be checked into source control, 
+         and this template.config should be checked in in it's place. -->
+	<connectionStrings>
+		<add name="ConnectionString" connectionString=""/>
+	</connectionStrings>
+	<system.web>
+    <compilation debug="true" targetFramework="4.5" />
+    <httpRuntime targetFramework="4.5" />
+  </system.web>
+  <system.webServer>
+    <handlers>
+      <remove name="ExtensionlessUrlHandler-Integrated-4.0" />
+      <remove name="OPTIONSVerbHandler" />
+      <remove name="TRACEVerbHandler" />
+      <add name="ExtensionlessUrlHandler-Integrated-4.0" path="*." verb="*" type="System.Web.Handlers.TransferRequestHandler" preCondition="integratedMode,runtimeVersionv4.0" />
+    </handlers>
+  </system.webServer>
+  <runtime>
+    <assemblyBinding xmlns="urn:schemas-microsoft-com:asm.v1">
+      <dependentAssembly>
+        <assemblyIdentity name="Newtonsoft.Json" publicKeyToken="30ad4fe6b2a6aeed" culture="neutral" />
+        <bindingRedirect oldVersion="0.0.0.0-12.0.0.0" newVersion="12.0.0.0" />
+      </dependentAssembly>
+    </assemblyBinding>
+  </runtime>
+</configuration>
\ No newline at end of file
diff --git a/Engine/Source/Programs/UnrealGameSync/MetadataServer/packages.config b/Engine/Source/Programs/UnrealGameSync/MetadataServer/packages.config
new file mode 100644
--- /dev/null
+++ b/Engine/Source/Programs/UnrealGameSync/MetadataServer/packages.config
@@ -0,0 +1,11 @@
+﻿<?xml version="1.0" encoding="utf-8"?>
+<packages>
+  <package id="BouncyCastle" version="1.8.3.1" targetFramework="net452" />
+  <package id="Google.Protobuf" version="3.6.1" targetFramework="net452" />
+  <package id="Microsoft.AspNet.WebApi" version="5.2.7" targetFramework="net45" />
+  <package id="Microsoft.AspNet.WebApi.Client" version="5.2.7" targetFramework="net45" />
+  <package id="Microsoft.AspNet.WebApi.Core" version="5.2.7" targetFramework="net45" />
+  <package id="Microsoft.AspNet.WebApi.WebHost" version="5.2.7" targetFramework="net45" />
+  <package id="MySql.Data" version="8.0.16" targetFramework="net452" />
+  <package id="Newtonsoft.Json" version="12.0.2" targetFramework="net45" />
+</packages>
\ No newline at end of file
diff --git a/Engine/Source/Programs/UnrealGameSync/PostBadgeStatus/App.config b/Engine/Source/Programs/UnrealGameSync/PostBadgeStatus/App.config
new file mode 100644
--- /dev/null
+++ b/Engine/Source/Programs/UnrealGameSync/PostBadgeStatus/App.config
@@ -0,0 +1,6 @@
+﻿<?xml version="1.0" encoding="utf-8" ?>
+<configuration>
+    <startup> 
+        <supportedRuntime version="v4.0" sku=".NETFramework,Version=v4.5" />
+    </startup>
+</configuration>
\ No newline at end of file
diff --git a/Engine/Source/Programs/UnrealGameSync/PostBadgeStatus/PostBadgeStatus.csproj b/Engine/Source/Programs/UnrealGameSync/PostBadgeStatus/PostBadgeStatus.csproj
new file mode 100644
--- /dev/null
+++ b/Engine/Source/Programs/UnrealGameSync/PostBadgeStatus/PostBadgeStatus.csproj
@@ -0,0 +1,67 @@
+﻿<?xml version="1.0" encoding="utf-8"?>
+<Project ToolsVersion="12.0" DefaultTargets="Build" xmlns="http://schemas.microsoft.com/developer/msbuild/2003">
+  <Import Project="$(MSBuildExtensionsPath)\$(MSBuildToolsVersion)\Microsoft.Common.props" Condition="Exists('$(MSBuildExtensionsPath)\$(MSBuildToolsVersion)\Microsoft.Common.props')" />
+  <PropertyGroup>
+    <Configuration Condition=" '$(Configuration)' == '' ">Debug</Configuration>
+    <Platform Condition=" '$(Platform)' == '' ">AnyCPU</Platform>
+    <ProjectGuid>{DCBE3437-5DBF-41B4-B854-967BC0C93C48}</ProjectGuid>
+    <OutputType>Exe</OutputType>
+    <AppDesignerFolder>Properties</AppDesignerFolder>
+    <RootNamespace>PostBadgeStatus</RootNamespace>
+    <AssemblyName>PostBadgeStatus</AssemblyName>
+    <TargetFrameworkVersion>v4.5</TargetFrameworkVersion>
+    <FileAlignment>512</FileAlignment>
+  </PropertyGroup>
+  <PropertyGroup Condition="'$(Configuration)|$(Platform)' == 'Debug|AnyCPU'">
+    <DebugSymbols>true</DebugSymbols>
+    <OutputPath>bin\Debug\</OutputPath>
+    <DefineConstants>DEBUG;TRACE</DefineConstants>
+    <DebugType>full</DebugType>
+    <PlatformTarget>AnyCPU</PlatformTarget>
+    <ErrorReport>prompt</ErrorReport>
+    <CodeAnalysisRuleSet>MinimumRecommendedRules.ruleset</CodeAnalysisRuleSet>
+    <Prefer32Bit>true</Prefer32Bit>
+  </PropertyGroup>
+  <PropertyGroup Condition="'$(Configuration)|$(Platform)' == 'Release|AnyCPU'">
+    <OutputPath>bin\Release\</OutputPath>
+    <DefineConstants>TRACE</DefineConstants>
+    <Optimize>true</Optimize>
+    <DebugType>pdbonly</DebugType>
+    <PlatformTarget>AnyCPU</PlatformTarget>
+    <ErrorReport>prompt</ErrorReport>
+    <CodeAnalysisRuleSet>MinimumRecommendedRules.ruleset</CodeAnalysisRuleSet>
+    <Prefer32Bit>true</Prefer32Bit>
+  </PropertyGroup>
+  <ItemGroup>
+    <Reference Include="System" />
+    <Reference Include="System.Core" />
+    <Reference Include="System.Web.Extensions" />
+    <Reference Include="System.Xml.Linq" />
+    <Reference Include="System.Data.DataSetExtensions" />
+    <Reference Include="Microsoft.CSharp" />
+    <Reference Include="System.Data" />
+    <Reference Include="System.Xml" />
+  </ItemGroup>
+  <ItemGroup>
+    <Compile Include="Program.cs" />
+    <Compile Include="Properties\AssemblyInfo.cs" />
+  </ItemGroup>
+  <ItemGroup>
+    <None Include="App.config" />
+  </ItemGroup>
+  <Choose>
+    <When Condition="Exists('NotForLicensees\ProgramSettings.cs')">
+      <ItemGroup>
+        <Compile Include="NotForLicensees\ProgramSettings.cs" />
+      </ItemGroup>
+    </When>
+  </Choose>
+  <Import Project="$(MSBuildToolsPath)\Microsoft.CSharp.targets" />
+  <!-- To modify your build process, add your task inside one of the targets below and uncomment it. 
+       Other similar extension points exist, see Microsoft.Common.targets.
+  <Target Name="BeforeBuild">
+  </Target>
+  <Target Name="AfterBuild">
+  </Target>
+  -->
+</Project>
\ No newline at end of file
diff --git a/Engine/Source/Programs/UnrealGameSync/PostBadgeStatus/Program.cs b/Engine/Source/Programs/UnrealGameSync/PostBadgeStatus/Program.cs
new file mode 100644
--- /dev/null
+++ b/Engine/Source/Programs/UnrealGameSync/PostBadgeStatus/Program.cs
@@ -0,0 +1,186 @@
+// Copyright Epic Games, Inc. All Rights Reserved.
+
+using System;
+using System.Collections.Generic;
+using System.Data.SqlClient;
+using System.IO;
+using System.Linq;
+using System.Net;
+using System.Text;
+using System.Threading;
+using System.Threading.Tasks;
+using System.Web.Script.Serialization;
+
+namespace WriteBadgeStatus
+{
+	class Program
+	{
+		static int Main(string[] Args)
+		{
+			// Parse all the parameters
+			List<string> Arguments = new List<string>(Args);
+			string Name = ParseParam(Arguments, "Name");
+			string Change = ParseParam(Arguments, "Change");
+			string Project = ParseParam(Arguments, "Project");
+			string RestUrl = ParseParam(Arguments, "RestUrl");
+			string Status = ParseParam(Arguments, "Status");
+			string Url = ParseParam(Arguments, "Url");
+
+			// Check we've got all the arguments we need (and no more)
+			if (Arguments.Count > 0 || Name == null || Change == null || Project == null || RestUrl == null || Status == null || Url == null)
+			{
+				Console.WriteLine("Syntax:");
+				Console.WriteLine("  PostBadgeStatus -Name=<Name> -Change=<CL> -Project=<DepotPath> -RestUrl=<Url> -Status=<Status> -Url=<Url>");
+				return 1;
+			}
+
+			BuildData Build = new BuildData
+			{
+				BuildType = Name,
+				Url = Url,
+				Project = Project,
+				ArchivePath = ""
+			};
+			if (!int.TryParse(Change, out Build.ChangeNumber))
+			{
+				Console.WriteLine("Change must be an integer!");
+				return 1;
+			}
+			if (!Enum.TryParse<BuildData.BuildDataResult>(Status, true, out Build.Result))
+			{
+				Console.WriteLine("Change must be Starting, Failure, Warning, Success, or Skipped!");
+				return 1;
+			}
+			int NumRetries = 0;
+			while (true)
+			{
+				try
+				{
+					return SendRequest(RestUrl, "Build", "POST", new JavaScriptSerializer().Serialize(Build));
+				}
+				catch (Exception ex)
+				{
+					if (++NumRetries <= 3)
+					{
+						Console.WriteLine(string.Format("An exception was thrown attempting to send the request: {0}, retrying...", ex.Message));
+						// Wait 5 seconds and retry;
+						Thread.Sleep(5000);
+					}
+					else
+					{
+						Console.WriteLine("Failed to POST due to multiple failures. Aborting.");
+						return 1;
+					}
+				}
+			}
+		}
+
+		static string ParseParam(List<string> Arguments, string ParamName)
+		{
+			string ParamPrefix = String.Format("-{0}=", ParamName);
+			for (int Idx = 0; Idx < Arguments.Count; Idx++)
+			{
+				if (Arguments[Idx].StartsWith(ParamPrefix, StringComparison.InvariantCultureIgnoreCase))
+				{
+					string ParamValue = Arguments[Idx].Substring(ParamPrefix.Length);
+					Arguments.RemoveAt(Idx);
+					return ParamValue;
+				}
+			}
+			return null;
+		}
+
+		class BuildData
+		{
+			public enum BuildDataResult
+			{
+				Starting,
+				Failure,
+				Warning,
+				Success,
+				Skipped,
+			}
+
+			public int ChangeNumber;
+			public string BuildType;
+			public BuildDataResult Result;
+			public string Url;
+			public string Project;
+			public string ArchivePath;
+
+			public bool IsSuccess
+			{
+				get { return Result == BuildDataResult.Success || Result == BuildDataResult.Warning; }
+			}
+
+			public bool IsFailure
+			{
+				get { return Result == BuildDataResult.Failure; }
+			}
+		}
+
+		static int SendRequest(string URI, string Resource, string Method, string RequestBody = null, params string[] QueryParams)
+		{
+			// set up the query string
+			StringBuilder TargetURI = new StringBuilder(string.Format("{0}/api/{1}", URI, Resource));
+			if (QueryParams.Length != 0)
+			{
+				TargetURI.Append("?");
+				for (int Idx = 0; Idx < QueryParams.Length; Idx++)
+				{
+					TargetURI.Append(QueryParams[Idx]);
+					if (Idx != QueryParams.Length - 1)
+					{
+						TargetURI.Append("&");
+					}
+				}
+			}
+			HttpWebRequest Request = (HttpWebRequest)WebRequest.Create(TargetURI.ToString());
+			Request.ContentType = "application/json";
+			Request.Method = Method;
+
+
+			// Add json to request body
+			if (!string.IsNullOrEmpty(RequestBody))
+			{
+				if (Method == "POST")
+				{
+					byte[] bytes = Encoding.UTF8.GetBytes(RequestBody);
+					using (Stream RequestStream = Request.GetRequestStream())
+					{
+						RequestStream.Write(bytes, 0, bytes.Length);
+					}
+				}
+			}
+			try
+			{
+				using (HttpWebResponse Response = (HttpWebResponse)Request.GetResponse())
+				{
+					string ResponseContent = null;
+					using (StreamReader ResponseReader = new System.IO.StreamReader(Response.GetResponseStream(), Encoding.UTF8))
+					{
+						ResponseContent = ResponseReader.ReadToEnd();
+						Console.WriteLine(ResponseContent);
+						return ((int)Response.StatusCode >= 200 && (int)Response.StatusCode <= 299) ? 0 : 1;
+					}
+				}
+				
+			}
+			catch (WebException ex)
+			{
+				if (ex.Response != null)
+				{
+					throw new Exception(string.Format("Request returned status: {0}, message: {1}", ((HttpWebResponse)ex.Response).StatusCode, ex.Message));
+				}
+				else
+				{
+					throw new Exception(string.Format("Request returned message: {0}", ex.InnerException.Message));
+				}
+			}
+			catch (Exception ex)
+			{
+				throw new Exception(string.Format("Couldn't complete the request, error: {0}", ex.Message));
+			}
+		}
+	}
+}
diff --git a/Engine/Source/Programs/UnrealGameSync/PostBadgeStatus/Properties/AssemblyInfo.cs b/Engine/Source/Programs/UnrealGameSync/PostBadgeStatus/Properties/AssemblyInfo.cs
new file mode 100644
--- /dev/null
+++ b/Engine/Source/Programs/UnrealGameSync/PostBadgeStatus/Properties/AssemblyInfo.cs
@@ -0,0 +1,38 @@
+﻿// Copyright Epic Games, Inc. All Rights Reserved.
+
+using System.Reflection;
+using System.Runtime.CompilerServices;
+using System.Runtime.InteropServices;
+
+// General Information about an assembly is controlled through the following 
+// set of attributes. Change these attribute values to modify the information
+// associated with an assembly.
+[assembly: AssemblyTitle("WriteBadgeStatus")]
+[assembly: AssemblyDescription("")]
+[assembly: AssemblyConfiguration("")]
+[assembly: AssemblyCompany("")]
+[assembly: AssemblyProduct("WriteBadgeStatus")]
+[assembly: AssemblyCopyright("Copyright Epic Games, Inc. All Rights Reserved.")]
+[assembly: AssemblyTrademark("")]
+[assembly: AssemblyCulture("")]
+
+// Setting ComVisible to false makes the types in this assembly not visible 
+// to COM components.  If you need to access a type in this assembly from 
+// COM, set the ComVisible attribute to true on that type.
+[assembly: ComVisible(false)]
+
+// The following GUID is for the ID of the typelib if this project is exposed to COM
+[assembly: Guid("fb419b69-e06d-420c-bb0a-321faf3f99a1")]
+
+// Version information for an assembly consists of the following four values:
+//
+//      Major Version
+//      Minor Version 
+//      Build Number
+//      Revision
+//
+// You can specify all the values or you can default the Build and Revision Numbers 
+// by using the '*' as shown below:
+// [assembly: AssemblyVersion("1.0.*")]
+[assembly: AssemblyVersion("1.0.0.0")]
+[assembly: AssemblyFileVersion("1.0.0.0")]
diff --git a/Engine/Source/Programs/UnrealGameSync/PostBadgeStatus/bin/Release/PostBadgeStatus.exe.config b/Engine/Source/Programs/UnrealGameSync/PostBadgeStatus/bin/Release/PostBadgeStatus.exe.config
new file mode 100644
--- /dev/null
+++ b/Engine/Source/Programs/UnrealGameSync/PostBadgeStatus/bin/Release/PostBadgeStatus.exe.config
@@ -0,0 +1,6 @@
+﻿<?xml version="1.0" encoding="utf-8" ?>
+<configuration>
+    <startup> 
+        <supportedRuntime version="v4.0" sku=".NETFramework,Version=v4.5" />
+    </startup>
+</configuration>
\ No newline at end of file
