UGS: Allow launching the editor for a content-only project with a target file from the engine directory rather than the project directory.

diff --git a/Engine/Source/Programs/UnrealGameSync/UnrealGameSync/Controls/WorkspaceControl.cs b/Engine/Source/Programs/UnrealGameSync/UnrealGameSync/Controls/WorkspaceControl.cs
--- a/Engine/Source/Programs/UnrealGameSync/UnrealGameSync/Controls/WorkspaceControl.cs
+++ b/Engine/Source/Programs/UnrealGameSync/UnrealGameSync/Controls/WorkspaceControl.cs
@@ -3253,6 +3253,7 @@ public enum LatestChangeType
 				else
 				{
 					PossiblePaths.Add(Path.Combine(Path.GetDirectoryName(SelectedFileName), "Binaries", "Win64", String.Format("{0}{1}.target", GetDefaultEditorTargetName(), ConfigSuffix)));
+					PossiblePaths.Add(Path.Combine(Workspace.LocalRootPath, "Engine", "Binaries", "Win64", String.Format("{0}{1}.target", GetDefaultEditorTargetName(), ConfigSuffix)));
 				}
 			}
 			else
