Merging UGS changes from //UE5/Release-5.0.

diff --git a/Engine/Source/Programs/UnrealGameSync/BuildUGS.xml b/Engine/Source/Programs/UnrealGameSync/BuildUGS.xml
--- a/Engine/Source/Programs/UnrealGameSync/BuildUGS.xml
+++ b/Engine/Source/Programs/UnrealGameSync/BuildUGS.xml
@@ -32,4 +32,19 @@
 		</Node>
 	</Agent>
 	
+	<Agent Name="UnrealGameSync Installer" Type="Win64">
+		<Node Name="Build HordeInstaller">
+			<Property Name="WixDir" Value="$(RootDir)/Engine/Source/ThirdParty/WiX/3.8"/>
+			<Property Name="SourceDir" Value="$(RootDir)/Engine/Source/Programs/UnrealGameSync/Installer"/>
+			<Property Name="UgsLauncherPublishDir" Value="$(UgsDir)/UnrealGameSyncLauncher/bin/Release/netcoreapp3.1/win-x64/publish"/>
+
+			<Delete Files="$(UgsLauncherPublishDir)/..."/>
+			<DotNet Arguments="publish &quot;$(UgsDir)/UnrealGameSyncLauncher/UnrealGameSyncLauncher.csproj&quot; --output=&quot;$(UgsLauncherPublishDir)&quot; --no-self-contained -c Release -r win-x64 $(VersionArguments)"/>
+
+			<Spawn Exe="$(WixDir)/heat.exe" Arguments="dir &quot;$(UgsLauncherPublishDir)&quot; -cg UGSLauncher_Project -dr INSTALLFOLDER -scom -sreg -srd -var var.BasePath -gg -sfrag -out obj/Release/UGSLauncher.wxs" WorkingDir="$(SourceDir)"/>
+			<Spawn Exe="$(WixDir)/candle.exe" Arguments="-dBasePath=&quot;$(UgsLauncherPublishDir)&quot; -out obj\Release\ -dConfiguration=Release -dPlatform=x64 -arch x86 -ext ..\..\..\ThirdParty\WiX\3.8\WixUtilExtension.dll Product.wxs obj/Release/UGSLauncher.wxs" WorkingDir="$(SourceDir)"/>
+			<Spawn Exe="$(WixDir)/light.exe" Arguments="-out &quot;bin/Release/UnrealGameSync.msi&quot; -pdbout &quot;bin/Release/UnrealGameSync.wixpdb&quot; -cultures:null -ext ..\..\..\ThirdParty\WiX\3.8\WixUtilExtension.dll -sice:ICE69 &quot;obj\Release\Product.wixobj&quot; &quot;obj\Release\UGSLauncher.wixobj&quot;" WorkingDir="$(SourceDir)"/>
+		</Node>
+	</Agent>
+
 </BuildGraph>
diff --git a/Engine/Source/Programs/UnrealGameSync/Installer/Build.bat b/Engine/Source/Programs/UnrealGameSync/Installer/Build.bat
new file mode 100755
--- /dev/null
+++ b/Engine/Source/Programs/UnrealGameSync/Installer/Build.bat
@@ -0,0 +1,18 @@
+@echo off
+setlocal
+
+pushd %~dp0
+set WIXDIR=..\..\..\ThirdParty\WiX\3.8
+set SRCDIR=..\UnrealGameSyncLauncher\bin\Release\netcoreapp3.1
+
+%WIXDIR%\heat.exe dir "%SRCDIR%" -cg UGSLauncher_Project -dr INSTALLFOLDER -scom -sreg -srd -var var.BasePath -gg -sfrag -out obj/Release/UGSLauncher.wxs -nologo 
+if errorlevel 1 goto :eof
+
+%WIXDIR%\candle.exe -nologo -dBasePath="%SRCDIR%" -out obj\Release\ -dConfiguration=Release -dPlatform=x64 -arch x86 -ext ..\..\..\ThirdParty\WiX\3.8\WixUtilExtension.dll Product.wxs obj/Release/UGSLauncher.wxs
+if errorlevel 1 goto :eof
+
+set OUTPUT=%~dp0bin\Release\UnrealGameSync.msi
+echo Linking %OUTPUT%...
+%WIXDIR%\light.exe -nologo -out "bin/Release/UnrealGameSync.msi" -pdbout "bin/Release/UnrealGameSync.wixpdb" -cultures:null -ext ..\..\..\ThirdParty\WiX\3.8\WixUtilExtension.dll -sice:ICE69 obj\Release\Product.wixobj obj\Release\UGSLauncher.wixobj
+if errorlevel 1 goto :eof
+if not exist "%OUTPUT%" echo Unable to find MSI: %OUTPUT% & goto :eof
diff --git a/Engine/Source/Programs/UnrealGameSync/InstallerCustomAction/CustomAction.config b/Engine/Source/Programs/UnrealGameSync/InstallerCustomAction/CustomAction.config
deleted file mode 100644
diff --git a/Engine/Source/Programs/UnrealGameSync/InstallerCustomAction/CustomAction.cs b/Engine/Source/Programs/UnrealGameSync/InstallerCustomAction/CustomAction.cs
deleted file mode 100644
diff --git a/Engine/Source/Programs/UnrealGameSync/InstallerCustomAction/InstallerCustomAction.csproj b/Engine/Source/Programs/UnrealGameSync/InstallerCustomAction/InstallerCustomAction.csproj
deleted file mode 100644
diff --git a/Engine/Source/Programs/UnrealGameSync/InstallerCustomAction/Properties/AssemblyInfo.cs b/Engine/Source/Programs/UnrealGameSync/InstallerCustomAction/Properties/AssemblyInfo.cs
deleted file mode 100644
