UnrealGameSync: Add advanced options for max batch size and file buffer

Cleanup advance settings to use numeric boxes instead of text

Update sync status before syncing so if a long sync is occuring it shows 0/X progress intead of staying on the filtering progress

diff --git a/Engine/Source/Programs/UnrealGameSync/UnrealGameSync/Forms/PerforceSyncSettingsWindow.cs b/Engine/Source/Programs/UnrealGameSync/UnrealGameSync/Forms/PerforceSyncSettingsWindow.cs
--- a/Engine/Source/Programs/UnrealGameSync/UnrealGameSync/Forms/PerforceSyncSettingsWindow.cs
+++ b/Engine/Source/Programs/UnrealGameSync/UnrealGameSync/Forms/PerforceSyncSettingsWindow.cs
@@ -1,13 +1,6 @@
 // Copyright Epic Games, Inc. All Rights Reserved.
 
 using System;
-using System.Collections.Generic;
-using System.ComponentModel;
-using System.Data;
-using System.Drawing;
-using System.Linq;
-using System.Text;
-using System.Threading.Tasks;
 using System.Windows.Forms;
 
 namespace UnrealGameSync
@@ -25,28 +18,20 @@ namespace UnrealGameSync
 		private void PerforceSettingsWindow_Load(object sender, EventArgs e)
 		{
 			PerforceSyncOptions SyncOptions = Settings.SyncOptions;
-			NumRetriesTextBox.Text = (SyncOptions.NumRetries > 0)? SyncOptions.NumRetries.ToString() : "";
-			TcpBufferSizeText.Text = (SyncOptions.TcpBufferSize > 0)? SyncOptions.TcpBufferSize.ToString() : "";
+			numericUpDownNumRetries.Value = (SyncOptions.NumRetries > 0) ? SyncOptions.NumRetries : PerforceSyncOptions.DefaultNumRetries;
+			numericUpDownTcpBufferSize.Value = (SyncOptions.TcpBufferSize > 0) ? SyncOptions.TcpBufferSize / 1024 : PerforceSyncOptions.DefaultTcpBufferSize / 1024;
+			numericUpDownFileBufferSize.Value = (SyncOptions.FileBufferSize > 0) ? SyncOptions.FileBufferSize / 1024 : PerforceSyncOptions.DefaultFileBufferSize / 1024;
+			numericUpDownMaxCommandsPerBatch.Value = (SyncOptions.MaxCommandsPerBatch > 0) ? SyncOptions.MaxCommandsPerBatch : PerforceSyncOptions.DefaultMaxCommandsPerBatch;
+			numericUpDownMaxSizePerBatch.Value = (SyncOptions.MaxSizePerBatch > 0) ? SyncOptions.MaxSizePerBatch / 1024 / 1024 : PerforceSyncOptions.DefaultMaxSizePerBatch / 1024 / 1024;
 		}
 
 		private void OkButton_Click(object sender, EventArgs e)
 		{
-			int NewNumRetries = 0;
-			if(NumRetriesTextBox.Text.Length > 0 && !int.TryParse(NumRetriesTextBox.Text, out NewNumRetries))
-			{
-				MessageBox.Show("Invalid value for number of retries");
-				return;
-			}
-
-			int NewTcpBufferSize = 0;
-			if(TcpBufferSizeText.Text.Length > 0 && !int.TryParse(TcpBufferSizeText.Text, out NewTcpBufferSize))
-			{
-				MessageBox.Show("Invalid value for TCP buffer size");
-				return;
-			}
-
-			Settings.SyncOptions.NumRetries = NewNumRetries;
-			Settings.SyncOptions.TcpBufferSize = NewTcpBufferSize;
+			Settings.SyncOptions.NumRetries = (int)numericUpDownNumRetries.Value;
+			Settings.SyncOptions.TcpBufferSize = (int)numericUpDownTcpBufferSize.Value * 1024;
+			Settings.SyncOptions.FileBufferSize = (int)numericUpDownFileBufferSize.Value * 1024;
+			Settings.SyncOptions.MaxCommandsPerBatch = (int)numericUpDownMaxCommandsPerBatch.Value;
+			Settings.SyncOptions.MaxSizePerBatch = (int)numericUpDownMaxSizePerBatch.Value * 1024 * 1024;
 			Settings.Save();
 
 			DialogResult = System.Windows.Forms.DialogResult.OK;
@@ -58,5 +43,15 @@ namespace UnrealGameSync
 			DialogResult = System.Windows.Forms.DialogResult.Cancel;
 			Close();
 		}
+
+		private void ResetButton_Click(object sender, EventArgs e)
+		{
+			PerforceSyncOptions SyncOptions = Settings.SyncOptions;
+			numericUpDownNumRetries.Value = PerforceSyncOptions.DefaultNumRetries;
+			numericUpDownTcpBufferSize.Value = PerforceSyncOptions.DefaultTcpBufferSize / 1024;
+			numericUpDownFileBufferSize.Value = PerforceSyncOptions.DefaultFileBufferSize / 1024;
+			numericUpDownMaxCommandsPerBatch.Value = PerforceSyncOptions.DefaultMaxCommandsPerBatch;
+			numericUpDownMaxSizePerBatch.Value = PerforceSyncOptions.DefaultMaxSizePerBatch / 1024 / 1024;
+		}
 	}
 }
diff --git a/Engine/Source/Programs/UnrealGameSync/UnrealGameSync/Perforce.cs b/Engine/Source/Programs/UnrealGameSync/UnrealGameSync/Perforce.cs
--- a/Engine/Source/Programs/UnrealGameSync/UnrealGameSync/Perforce.cs
+++ b/Engine/Source/Programs/UnrealGameSync/UnrealGameSync/Perforce.cs
@@ -350,11 +350,22 @@ namespace UnrealGameSync
 
 	public class PerforceSyncOptions
 	{
+		public const int DefaultNumRetries = 0;
 		public const int DefaultNumThreads = 2;
+		public const int DefaultTcpBufferSize = 0;
+		public const int DefaultFileBufferSize = 0;
 
-		public int NumRetries;
+		public const int DefaultMaxCommandsPerBatch = 200;
+		public const int DefaultMaxSizePerBatch = 128 * 1024 * 1024;
+
+		public int NumRetries = DefaultNumRetries;
 		public int NumThreads = DefaultNumThreads;
-		public int TcpBufferSize;
+		public int TcpBufferSize = DefaultTcpBufferSize;
+		public int FileBufferSize = DefaultFileBufferSize;
+
+		public int MaxCommandsPerBatch = DefaultMaxCommandsPerBatch;
+		public int MaxSizePerBatch = DefaultMaxSizePerBatch;
+
 
 		public PerforceSyncOptions Clone()
 		{
@@ -1290,6 +1301,10 @@ namespace UnrealGameSync
 					{
 						CommandLine.AppendFormat(" -v net.tcpsize={0}", Options.TcpBufferSize);
 					}
+					if (Options != null && Options.FileBufferSize > 0)
+					{
+						CommandLine.AppendFormat(" -v filesys.bufsize={0}", Options.FileBufferSize);
+					}
 					CommandLine.Append(" sync --parallel=0");
 					if(bForce)
 					{
diff --git a/Engine/Source/Programs/UnrealGameSync/UnrealGameSync/UserSettings.cs b/Engine/Source/Programs/UnrealGameSync/UnrealGameSync/UserSettings.cs
--- a/Engine/Source/Programs/UnrealGameSync/UnrealGameSync/UserSettings.cs
+++ b/Engine/Source/Programs/UnrealGameSync/UnrealGameSync/UserSettings.cs
@@ -522,10 +522,10 @@ namespace UnrealGameSync
 				}
 			}
 
-			if(!int.TryParse(ConfigFile.GetValue("Perforce.TcpBufferSize", "0"), out SyncOptions.TcpBufferSize))
-			{
-				SyncOptions.TcpBufferSize = 0;
-			}
+			SyncOptions.TcpBufferSize = ConfigFile.GetValue("Perforce.TcpBufferSize", PerforceSyncOptions.DefaultTcpBufferSize);
+			SyncOptions.FileBufferSize = ConfigFile.GetValue("Perforce.FileBufferSize", PerforceSyncOptions.DefaultFileBufferSize);
+			SyncOptions.MaxCommandsPerBatch = ConfigFile.GetValue("Perforce.MaxCommandsPerBatch", PerforceSyncOptions.DefaultMaxCommandsPerBatch);
+			SyncOptions.MaxSizePerBatch = ConfigFile.GetValue("Perforce.MaxSizePerBatch", PerforceSyncOptions.DefaultMaxSizePerBatch);
 		}
 
 		static Dictionary<Guid, bool> GetCategorySettings(ConfigSection Section, string IncludedKey, string ExcludedKey)
@@ -902,18 +902,30 @@ namespace UnrealGameSync
 			// Perforce settings
 			ConfigSection PerforceSection = ConfigFile.FindOrAddSection("Perforce");
 			PerforceSection.Clear();
-			if(SyncOptions.NumRetries > 0)
+			if(SyncOptions.NumRetries > 0 && SyncOptions.NumRetries != PerforceSyncOptions.DefaultNumRetries)
 			{
 				PerforceSection.SetValue("NumRetries", SyncOptions.NumRetries);
 			}
-			if(SyncOptions.NumThreads != PerforceSyncOptions.DefaultNumThreads)
+			if(SyncOptions.NumThreads > 0 && SyncOptions.NumThreads != PerforceSyncOptions.DefaultNumThreads)
 			{
 				PerforceSection.SetValue("NumThreads", SyncOptions.NumThreads);
 			}
-			if(SyncOptions.TcpBufferSize > 0)
+			if(SyncOptions.TcpBufferSize > 0 && SyncOptions.TcpBufferSize != PerforceSyncOptions.DefaultTcpBufferSize)
 			{
 				PerforceSection.SetValue("TcpBufferSize", SyncOptions.TcpBufferSize);
 			}
+			if (SyncOptions.FileBufferSize > 0 && SyncOptions.FileBufferSize != PerforceSyncOptions.DefaultFileBufferSize)
+			{
+				PerforceSection.SetValue("FileBufferSize", SyncOptions.FileBufferSize);
+			}
+			if (SyncOptions.MaxCommandsPerBatch > 0 && SyncOptions.MaxCommandsPerBatch != PerforceSyncOptions.DefaultMaxCommandsPerBatch)
+			{
+				PerforceSection.SetValue("MaxCommandsPerBatch", SyncOptions.MaxCommandsPerBatch);
+			}
+			if (SyncOptions.MaxSizePerBatch > 0 && SyncOptions.MaxSizePerBatch != PerforceSyncOptions.DefaultMaxSizePerBatch)
+			{
+				PerforceSection.SetValue("MaxSizePerBatch", SyncOptions.MaxSizePerBatch);
+			}
 
 			// Save the file
 			ConfigFile.Save(FileName);
diff --git a/Engine/Source/Programs/UnrealGameSync/UnrealGameSync/Workspace.cs b/Engine/Source/Programs/UnrealGameSync/UnrealGameSync/Workspace.cs
--- a/Engine/Source/Programs/UnrealGameSync/UnrealGameSync/Workspace.cs
+++ b/Engine/Source/Programs/UnrealGameSync/UnrealGameSync/Workspace.cs
@@ -364,6 +364,7 @@ namespace UnrealGameSync
 
 			// Spawn the new thread
 			WorkerThread = new Thread(x => UpdateWorkspace(Context));
+			WorkerThread.Name = "Workspace Update";
 			WorkerThread.Start();
 		}
 
@@ -616,7 +617,7 @@ namespace UnrealGameSync
 					}
 
 					// Find all the server changes, and anything that's opened for edit locally. We need to sync files we have open to schedule a resolve.
-					SyncBatchBuilder BatchBuilder = new SyncBatchBuilder(200, 100 * 1024 * 1024);
+					SyncBatchBuilder BatchBuilder = new SyncBatchBuilder(Context.PerforceSyncOptions.MaxCommandsPerBatch, Context.PerforceSyncOptions.MaxSizePerBatch);
 					List<string> SyncDepotPaths = new List<string>();
 					using(RecordCounter Counter = new RecordCounter(Progress, "Filtering files..."))
 					{
@@ -1356,6 +1357,9 @@ namespace UnrealGameSync
 				// Wrapper writer around the log class to prevent multiple threads writing to it at once
 				ThreadSafeTextWriter LogWrapper = new ThreadSafeTextWriter(Log);
 
+				// Initialize Sync Progress
+				UpdateSyncState(Prefix, State, Progress);
+
 				// Delegate for updating the sync state after a file has been synced
 				Action<PerforceFileRecord, LineBasedTextWriter> SyncOutput = (Record, LocalLog) => { UpdateSyncState(Prefix, Record, State, Progress, LocalLog); };
 
@@ -1364,6 +1368,7 @@ namespace UnrealGameSync
 				{
 					int ThreadNumber = ThreadIdx + 2;
 					Thread ChildThread = new Thread(() => StaticSyncWorker(ThreadNumber, Perforce, PendingChangeNumber, Context, State, SyncOutput, LogWrapper));
+					ChildThread.Name = "Sync File Revisions";
 					ChildThreads.Add(ChildThread);
 					ChildThread.Start();
 				}
@@ -1394,6 +1399,16 @@ namespace UnrealGameSync
 			}
 		}
 
+		static void UpdateSyncState(string Prefix, SyncState State, ProgressValue Progress)
+		{
+			lock (State)
+			{
+				string Message = String.Format("{0} ({1}/{2})", Prefix, State.TotalDepotPaths - State.RemainingDepotPaths.Count, State.TotalDepotPaths);
+				float Fraction = Math.Min((float)(State.TotalDepotPaths - State.RemainingDepotPaths.Count) / (float)State.TotalDepotPaths, 1.0f);
+				Progress.Set(Message, Fraction);
+			}
+		}
+
 		static void UpdateSyncState(string Prefix, PerforceFileRecord Record, SyncState State, ProgressValue Progress, TextWriter Log)
 		{
 			lock (State)
