UGS: Store settings for particular clients, projects and workspaces under the workspace root, in a hidden ".ugs" folder. This will allow state and settings to be read implicitly for command line tools.

diff --git a/Engine/Source/Programs/UnrealGameSync/UnrealGameSync/Controls/WorkspaceControl.cs b/Engine/Source/Programs/UnrealGameSync/UnrealGameSync/Controls/WorkspaceControl.cs
--- a/Engine/Source/Programs/UnrealGameSync/UnrealGameSync/Controls/WorkspaceControl.cs
+++ b/Engine/Source/Programs/UnrealGameSync/UnrealGameSync/Controls/WorkspaceControl.cs
@@ -15,6 +15,7 @@ using System.Text.RegularExpressions;
 using System.Windows.Forms;
 using System.Windows.Forms.VisualStyles;
 using System.Threading;
+using EpicGames.Core;
 
 namespace UnrealGameSync
 {
@@ -209,6 +210,7 @@ public enum LatestChangeType
 		LineBasedTextWriter Log;
 
 		UserSettings Settings;
+		UserWorkspaceState WorkspaceState;
 		UserWorkspaceSettings WorkspaceSettings;
 		UserProjectSettings ProjectSettings;
 
@@ -331,8 +333,11 @@ public enum LatestChangeType
 			bUnstable = bInUnstable;
 			Log = InLog;
 			Settings = InSettings;
-			WorkspaceSettings = InSettings.FindOrAddWorkspace(DetectSettings.BranchClientPath);
-			ProjectSettings = InSettings.FindOrAddProject(DetectSettings.NewSelectedClientFileName);
+			WorkspaceState = InSettings.FindOrAddWorkspaceState(new DirectoryReference(DetectSettings.BranchDirectoryName), DetectSettings.NewSelectedClientFileName);
+			WorkspaceSettings = InSettings.FindOrAddWorkspaceSettings(new DirectoryReference(DetectSettings.BranchDirectoryName), DetectSettings.BranchClientPath);
+
+			string ProjectPath = Regex.Replace(DetectSettings.NewSelectedClientFileName, "^//[^/]+/", "");
+			ProjectSettings = InSettings.FindOrAddProjectSettings(new DirectoryReference(DetectSettings.BranchDirectoryName), ProjectPath);
 
 			DesiredTaskbarState = Tuple.Create(TaskbarState.NoProgress, 0.0f);
 
@@ -376,10 +381,10 @@ public enum LatestChangeType
 			// Check if we've the project we've got open in this workspace is the one we're actually synced to
 			int CurrentChangeNumber = -1;
 			string CurrentSyncFilterHash = null;
-			if (String.Compare(WorkspaceSettings.CurrentProjectIdentifier, SelectedProjectIdentifier, true) == 0)
+			if (String.Compare(WorkspaceState.CurrentProjectIdentifier, SelectedProjectIdentifier, true) == 0)
 			{
-				CurrentChangeNumber = WorkspaceSettings.CurrentChangeNumber;
-				CurrentSyncFilterHash = WorkspaceSettings.CurrentSyncFilterHash;
+				CurrentChangeNumber = WorkspaceState.CurrentChangeNumber;
+				CurrentSyncFilterHash = WorkspaceState.CurrentSyncFilterHash;
 			}
 
 			// Figure out which API server to use
@@ -391,7 +396,7 @@ public enum LatestChangeType
 
 			string TelemetryProjectIdentifier = PerforceUtils.GetClientOrDepotDirectoryName(DetectSettings.NewSelectedProjectIdentifier);
 
-			Workspace = new Workspace(PerforceClient, BranchDirectoryName, SelectedFileName, DetectSettings.BranchClientPath, DetectSettings.NewSelectedClientFileName, CurrentChangeNumber, CurrentSyncFilterHash, WorkspaceSettings.LastBuiltChangeNumber, TelemetryProjectIdentifier, DetectSettings.bIsEnterpriseProject, new LogControlTextWriter(SyncLog));
+			Workspace = new Workspace(PerforceClient, BranchDirectoryName, SelectedFileName, DetectSettings.BranchClientPath, DetectSettings.NewSelectedClientFileName, CurrentChangeNumber, CurrentSyncFilterHash, WorkspaceState.LastBuiltChangeNumber, TelemetryProjectIdentifier, DetectSettings.bIsEnterpriseProject, new LogControlTextWriter(SyncLog));
 			Workspace.OnUpdateComplete += UpdateCompleteCallback;
 
 			string ProjectLogBaseName = Path.Combine(DataFolder, String.Format("{0}@{1}", PerforceClient.ClientName, DetectSettings.BranchClientPath.Replace("//" + PerforceClient.ClientName + "/", "").Trim('/').Replace("/", "$")));
@@ -844,14 +849,14 @@ public enum LatestChangeType
 
 		private void UpdateSyncConfig(int ChangeNumber, string SyncFilterHash)
 		{
-			WorkspaceSettings.CurrentProjectIdentifier = SelectedProjectIdentifier;
-			WorkspaceSettings.CurrentChangeNumber = ChangeNumber;
-			WorkspaceSettings.CurrentSyncFilterHash = SyncFilterHash;
-			if (ChangeNumber == -1 || ChangeNumber != WorkspaceSettings.CurrentChangeNumber)
+			WorkspaceState.CurrentProjectIdentifier = SelectedProjectIdentifier;
+			WorkspaceState.CurrentChangeNumber = ChangeNumber;
+			WorkspaceState.CurrentSyncFilterHash = SyncFilterHash;
+			if (ChangeNumber == -1 || ChangeNumber != WorkspaceState.CurrentChangeNumber)
 			{
-				WorkspaceSettings.AdditionalChangeNumbers.Clear();
+				WorkspaceState.AdditionalChangeNumbers.Clear();
 			}
-			Settings.Save();
+			WorkspaceState.Save();
 		}
 
 		private void BuildList_OnScroll()
@@ -951,7 +956,7 @@ public enum LatestChangeType
 				}
 			}
 
-			string[] CombinedSyncFilter = UserSettings.GetCombinedSyncFilter(GetSyncCategories(), Settings.SyncView, Settings.SyncCategories, WorkspaceSettings.SyncView, WorkspaceSettings.SyncCategories);
+			string[] CombinedSyncFilter = UserSettings.GetCombinedSyncFilter(GetSyncCategories(), Settings.SyncView, Settings.SyncCategories, WorkspaceSettings.SyncView, WorkspaceSettings.SyncCategoriesDict);
 
 			ConfigSection PerforceSection = PerforceMonitor.LatestProjectConfigFile.FindSection("Perforce");
 			if (PerforceSection != null)
@@ -1056,12 +1061,12 @@ public enum LatestChangeType
 		{
 			if (Workspace.IsBusy() && MessageBox.Show("Are you sure you want to cancel the current operation?", "Cancel operation", MessageBoxButtons.YesNo) == DialogResult.Yes)
 			{
-				WorkspaceSettings.LastSyncChangeNumber = Workspace.PendingChangeNumber;
-				WorkspaceSettings.LastSyncResult = WorkspaceUpdateResult.Canceled;
-				WorkspaceSettings.LastSyncResultMessage = null;
-				WorkspaceSettings.LastSyncTime = null;
-				WorkspaceSettings.LastSyncDurationSeconds = 0;
-				Settings.Save();
+				WorkspaceState.LastSyncChangeNumber = Workspace.PendingChangeNumber;
+				WorkspaceState.LastSyncResult = WorkspaceUpdateResult.Canceled;
+				WorkspaceState.LastSyncResultMessage = null;
+				WorkspaceState.LastSyncTime = null;
+				WorkspaceState.LastSyncDurationSeconds = 0;
+				WorkspaceState.Save();
 
 				Workspace.CancelUpdate();
 
@@ -1095,22 +1100,21 @@ public enum LatestChangeType
 
 			if (Result == WorkspaceUpdateResult.Success && Context.Options.HasFlag(WorkspaceUpdateOptions.SyncSingleChange))
 			{
-				WorkspaceSettings.AdditionalChangeNumbers.Add(Context.ChangeNumber);
-				Settings.Save();
+				WorkspaceState.AdditionalChangeNumbers.Add(Context.ChangeNumber);
 			}
 
-			if (Result == WorkspaceUpdateResult.Success && WorkspaceSettings.ExpandedArchiveTypes != null)
+			if (Result == WorkspaceUpdateResult.Success && WorkspaceState.ExpandedArchiveTypes != null)
 			{
-				WorkspaceSettings.ExpandedArchiveTypes = WorkspaceSettings.ExpandedArchiveTypes.Except(Context.ArchiveTypeToArchive.Where(x => x.Value == null).Select(x => x.Key)).ToArray();
+				WorkspaceState.ExpandedArchiveTypes = WorkspaceState.ExpandedArchiveTypes.Except(Context.ArchiveTypeToArchive.Where(x => x.Value == null).Select(x => x.Key)).ToArray();
 			}
 
-			WorkspaceSettings.LastSyncChangeNumber = Context.ChangeNumber;
-			WorkspaceSettings.LastSyncResult = Result;
-			WorkspaceSettings.LastSyncResultMessage = ResultMessage;
-			WorkspaceSettings.LastSyncTime = DateTime.UtcNow;
-			WorkspaceSettings.LastSyncDurationSeconds = (int)(WorkspaceSettings.LastSyncTime.Value - Context.StartTime).TotalSeconds;
-			WorkspaceSettings.LastBuiltChangeNumber = Workspace.LastBuiltChangeNumber;
-			Settings.Save();
+			WorkspaceState.LastSyncChangeNumber = Context.ChangeNumber;
+			WorkspaceState.LastSyncResult = Result;
+			WorkspaceState.LastSyncResultMessage = ResultMessage;
+			WorkspaceState.LastSyncTime = DateTime.UtcNow;
+			WorkspaceState.LastSyncDurationSeconds = (int)(WorkspaceState.LastSyncTime.Value - Context.StartTime).TotalSeconds;
+			WorkspaceState.LastBuiltChangeNumber = Workspace.LastBuiltChangeNumber;
+			WorkspaceState.Save();
 
 			if (Result == WorkspaceUpdateResult.FilesToResolve)
 			{
@@ -1520,7 +1524,7 @@ public enum LatestChangeType
 					return false;
 				}
 			}
-			if (IsBisectModeEnabled() && !WorkspaceSettings.ChangeNumberToBisectState.ContainsKey(Change.Number))
+			if (IsBisectModeEnabled() && !WorkspaceState.BisectChanges.Any(x => x.Change == Change.Number))
 			{
 				return false;
 			}
@@ -2155,20 +2159,20 @@ public enum LatestChangeType
 		private void GetRemainingBisectRange(out int OutPassChangeNumber, out int OutFailChangeNumber)
 		{
 			int PassChangeNumber = -1;
-			foreach (KeyValuePair<int, BisectState> Pair in WorkspaceSettings.ChangeNumberToBisectState)
+			foreach (BisectEntry Entry in WorkspaceState.BisectChanges)
 			{
-				if (Pair.Value == BisectState.Pass && (Pair.Key > PassChangeNumber || PassChangeNumber == -1))
+				if (Entry.State == BisectState.Pass && (Entry.Change > PassChangeNumber || PassChangeNumber == -1))
 				{
-					PassChangeNumber = Pair.Key;
+					PassChangeNumber = Entry.Change;
 				}
 			}
 
 			int FailChangeNumber = -1;
-			foreach (KeyValuePair<int, BisectState> Pair in WorkspaceSettings.ChangeNumberToBisectState)
+			foreach (BisectEntry Entry in WorkspaceState.BisectChanges)
 			{
-				if (Pair.Value == BisectState.Fail && Pair.Key > PassChangeNumber && (Pair.Key < FailChangeNumber || FailChangeNumber == -1))
+				if (Entry.State == BisectState.Fail && Entry.Change > PassChangeNumber && (Entry.Change < FailChangeNumber || FailChangeNumber == -1))
 				{
-					FailChangeNumber = Pair.Key;
+					FailChangeNumber = Entry.Change;
 				}
 			}
 
@@ -2198,7 +2202,7 @@ public enum LatestChangeType
 
 			bool bAllowSync = CanSyncChange(Change.Number);
 			int BadgeAlpha = bAllowSync ? 255 : 128;
-			Color TextColor = (bAllowSync || Change.Number == Workspace.PendingChangeNumber || Change.Number == Workspace.CurrentChangeNumber || (WorkspaceSettings != null && WorkspaceSettings.AdditionalChangeNumbers.Contains(Change.Number))) ? SystemColors.WindowText : Blend(SystemColors.Window, SystemColors.WindowText, 0.25f);
+			Color TextColor = (bAllowSync || Change.Number == Workspace.PendingChangeNumber || Change.Number == Workspace.CurrentChangeNumber || (WorkspaceState != null && WorkspaceState.AdditionalChangeNumbers.Contains(Change.Number))) ? SystemColors.WindowText : Blend(SystemColors.Window, SystemColors.WindowText, 0.25f);
 
 			const int FadeRange = 6;
 			if (e.ItemIndex >= BuildList.Items.Count - FadeRange && NumChanges >= PerforceMonitor.CurrentMaxChanges && !IsBisectModeEnabled())
@@ -2217,7 +2221,7 @@ public enum LatestChangeType
 				{
 					e.Graphics.DrawImage(Properties.Resources.Icons, MinX * DpiScaleX, IconY, PreviousSyncIcon, GraphicsUnit.Pixel);
 				}
-				else if (WorkspaceSettings != null && WorkspaceSettings.AdditionalChangeNumbers.Contains(Change.Number))
+				else if (WorkspaceSettings != null && WorkspaceState.AdditionalChangeNumbers.Contains(Change.Number))
 				{
 					e.Graphics.DrawImage(Properties.Resources.Icons, MinX * DpiScaleX, IconY, AdditionalSyncIcon, GraphicsUnit.Pixel);
 				}
@@ -2236,16 +2240,16 @@ public enum LatestChangeType
 						int FirstPass, FirstFail;
 						GetRemainingBisectRange(out FirstPass, out FirstFail);
 
-						BisectState State;
-						if (!WorkspaceSettings.ChangeNumberToBisectState.TryGetValue(Change.Number, out State) || State == BisectState.Exclude)
+						BisectEntry Entry = WorkspaceState.BisectChanges.FirstOrDefault(x => x.Change == Change.Number);
+						if (Entry == null || Entry.State == BisectState.Exclude)
 						{
 							QualityIcon = new Rectangle(0, 0, 0, 0);
 						}
-						if (State == BisectState.Pass)
+						if (Entry.State == BisectState.Pass)
 						{
 							QualityIcon = BisectPassIcon;
 						}
-						else if (State == BisectState.Fail)
+						else if (Entry.State == BisectState.Fail)
 						{
 							QualityIcon = BisectFailIcon;
 						}
@@ -2382,7 +2386,7 @@ public enum LatestChangeType
 					}
 
 					string SummaryText;
-					if (WorkspaceSettings.LastSyncChangeNumber == -1 || WorkspaceSettings.LastSyncChangeNumber != Change.Number || !GetLastUpdateMessage(WorkspaceSettings.LastSyncResult, WorkspaceSettings.LastSyncResultMessage, out SummaryText))
+					if (WorkspaceState.LastSyncChangeNumber == -1 || WorkspaceState.LastSyncChangeNumber != Change.Number || !GetLastUpdateMessage(WorkspaceState.LastSyncResult, WorkspaceState.LastSyncResultMessage, out SummaryText))
 					{
 						StringBuilder SummaryTextBuilder = new StringBuilder();
 
@@ -3557,19 +3561,19 @@ public enum LatestChangeType
 				Lines.Add(ProgramsLine);
 
 				// Get the summary of the last sync
-				if (WorkspaceSettings.LastSyncChangeNumber > 0)
+				if (WorkspaceState.LastSyncChangeNumber > 0)
 				{
 					string SummaryText;
-					if (WorkspaceSettings.LastSyncChangeNumber == Workspace.CurrentChangeNumber && WorkspaceSettings.LastSyncResult == WorkspaceUpdateResult.Success && WorkspaceSettings.LastSyncTime.HasValue)
+					if (WorkspaceState.LastSyncChangeNumber == Workspace.CurrentChangeNumber && WorkspaceState.LastSyncResult == WorkspaceUpdateResult.Success && WorkspaceState.LastSyncTime.HasValue)
 					{
 						Lines.Add(new StatusLine() { LineHeight = 0.5f });
 
 						StatusLine SuccessLine = new StatusLine();
 						SuccessLine.AddIcon(Properties.Resources.StatusIcons, new Size(16, 16), 0);
-						SuccessLine.AddText(String.Format("  Sync took {0}{1}s, completed at {2}.", (WorkspaceSettings.LastSyncDurationSeconds >= 60) ? String.Format("{0}m ", WorkspaceSettings.LastSyncDurationSeconds / 60) : "", WorkspaceSettings.LastSyncDurationSeconds % 60, WorkspaceSettings.LastSyncTime.Value.ToLocalTime().ToString("h\\:mmtt").ToLowerInvariant()));
+						SuccessLine.AddText(String.Format("  Sync took {0}{1}s, completed at {2}.", (WorkspaceState.LastSyncDurationSeconds >= 60) ? String.Format("{0}m ", WorkspaceState.LastSyncDurationSeconds / 60) : "", WorkspaceState.LastSyncDurationSeconds % 60, WorkspaceState.LastSyncTime.Value.ToLocalTime().ToString("h\\:mmtt").ToLowerInvariant()));
 						Lines.Add(SuccessLine);
 					}
-					else if (GetLastUpdateMessage(WorkspaceSettings.LastSyncResult, WorkspaceSettings.LastSyncResultMessage, out SummaryText))
+					else if (GetLastUpdateMessage(WorkspaceState.LastSyncResult, WorkspaceState.LastSyncResultMessage, out SummaryText))
 					{
 						Lines.Add(new StatusLine() { LineHeight = 0.5f });
 
@@ -3954,10 +3958,10 @@ public enum LatestChangeType
 		private void ViewLastSyncStatus()
 		{
 			string SummaryText;
-			if (GetLastUpdateMessage(WorkspaceSettings.LastSyncResult, WorkspaceSettings.LastSyncResultMessage, out SummaryText))
+			if (GetLastUpdateMessage(WorkspaceState.LastSyncResult, WorkspaceState.LastSyncResultMessage, out SummaryText))
 			{
 				string CaptionText;
-				if (!GetGenericLastUpdateMessage(WorkspaceSettings.LastSyncResult, out CaptionText))
+				if (!GetGenericLastUpdateMessage(WorkspaceState.LastSyncResult, out CaptionText))
 				{
 					CaptionText = "Sync error";
 				}
@@ -4183,8 +4187,7 @@ public enum LatestChangeType
 						BuildListContextMenu_OpenVisualStudio.Visible = !bIsBusy && bIsCurrentChange;
 						BuildListContextMenu_Cancel.Visible = bIsBusy;
 
-						BisectState State;
-						WorkspaceSettings.ChangeNumberToBisectState.TryGetValue(ContextMenuChange.Number, out State);
+						BisectState State = WorkspaceState.BisectChanges.FirstOrDefault(x => x.Change == ContextMenuChange.Number)?.State ?? BisectState.Include;
 						bool bIsBisectMode = IsBisectModeEnabled();
 						BuildListContextMenu_Bisect_Pass.Visible = bIsBisectMode && State != BisectState.Pass;
 						BuildListContextMenu_Bisect_Fail.Visible = bIsBisectMode && State != BisectState.Fail;
@@ -5130,7 +5133,7 @@ public enum LatestChangeType
 				ExtraSafeToDeleteExtensions = "";
 			}
 
-			string[] CombinedSyncFilter = UserSettings.GetCombinedSyncFilter(GetSyncCategories(), Settings.SyncView, Settings.SyncCategories, WorkspaceSettings.SyncView, WorkspaceSettings.SyncCategories);
+			string[] CombinedSyncFilter = UserSettings.GetCombinedSyncFilter(GetSyncCategories(), Settings.SyncView, Settings.SyncCategories, WorkspaceSettings.SyncView, WorkspaceSettings.SyncCategoriesDict);
 			List<string> SyncPaths = Workspace.GetSyncPaths(WorkspaceSettings.bSyncAllProjects ?? Settings.bSyncAllProjects, CombinedSyncFilter);
 
 			CleanWorkspaceWindow.DoClean(ParentForm, Workspace.Perforce, BranchDirectoryName, Workspace.ClientRootPath, SyncPaths, ExtraSafeToDeleteFolders.Split('\n'), ExtraSafeToDeleteExtensions.Split('\n'), Log);
@@ -5330,7 +5333,7 @@ public enum LatestChangeType
 
 				// Save the settings
 				ProjectSettings.BuildSteps = ModifiedBuildSteps;
-				Settings.Save();
+				ProjectSettings.Save();
 
 				// Update the custom tools menu, because we might have changed it
 				UpdateBuildSteps();
@@ -5468,24 +5471,26 @@ public enum LatestChangeType
 			DiagnosticsText.AppendFormat("Event monitor: {0}\n", (EventMonitor == null) ? "(inactive)" : EventMonitor.LastStatusMessage);
 			DiagnosticsText.AppendFormat("Issue monitor: {0}\n", (IssueMonitor == null) ? "(inactive)" : IssueMonitor.LastStatusMessage);
 
-			DiagnosticsWindow Diagnostics = new DiagnosticsWindow(DataFolder, DiagnosticsText.ToString());
+			DiagnosticsWindow Diagnostics = new DiagnosticsWindow(DataFolder, DiagnosticsText.ToString(), Settings.GetCachedFilePaths());
 			Diagnostics.ShowDialog(this);
 		}
 
 		private void OptionsContextMenu_SyncFilter_Click(object sender, EventArgs e)
 		{
-			SyncFilter Filter = new SyncFilter(GetSyncCategories(), Settings.SyncView, Settings.SyncCategories, Settings.bSyncAllProjects, Settings.bIncludeAllProjectsInSolution, WorkspaceSettings.SyncView, WorkspaceSettings.SyncCategories, WorkspaceSettings.bSyncAllProjects, WorkspaceSettings.bIncludeAllProjectsInSolution);
+			SyncFilter Filter = new SyncFilter(GetSyncCategories(), Settings.SyncView, Settings.SyncCategories, Settings.bSyncAllProjects, Settings.bIncludeAllProjectsInSolution, WorkspaceSettings.SyncView, WorkspaceSettings.SyncCategoriesDict, WorkspaceSettings.bSyncAllProjects, WorkspaceSettings.bIncludeAllProjectsInSolution);
 			if (Filter.ShowDialog() == DialogResult.OK)
 			{
 				Settings.SyncCategories = Filter.GlobalSyncCategories;
 				Settings.SyncView = Filter.GlobalView;
 				Settings.bSyncAllProjects = Filter.bGlobalSyncAllProjects;
 				Settings.bIncludeAllProjectsInSolution = Filter.bGlobalIncludeAllProjectsInSolution;
-				WorkspaceSettings.SyncCategories = Filter.WorkspaceSyncCategories;
+				Settings.Save();
+
+				WorkspaceSettings.SyncCategoriesDict = Filter.WorkspaceSyncCategories;
 				WorkspaceSettings.SyncView = Filter.WorkspaceView;
 				WorkspaceSettings.bSyncAllProjects = Filter.bWorkspaceSyncAllProjects;
 				WorkspaceSettings.bIncludeAllProjectsInSolution = Filter.bWorkspaceIncludeAllProjectsInSolution;
-				Settings.Save();
+				WorkspaceSettings.Save();
 			}
 		}
 
@@ -5790,7 +5795,7 @@ public enum LatestChangeType
 
 		private bool IsBisectModeEnabled()
 		{
-			return WorkspaceSettings.ChangeNumberToBisectState.Count >= 2;
+			return WorkspaceState.BisectChanges.Count >= 2;
 		}
 
 		private void EnableBisectMode()
@@ -5807,8 +5812,8 @@ public enum LatestChangeType
 				ChangeNumberToBisectState[ChangeNumberToBisectState.Keys.Min()] = BisectState.Pass;
 				ChangeNumberToBisectState[ChangeNumberToBisectState.Keys.Max()] = BisectState.Fail;
 
-				WorkspaceSettings.ChangeNumberToBisectState = ChangeNumberToBisectState;
-				Settings.Save();
+				WorkspaceState.BisectChanges = ChangeNumberToBisectState.Select(x => new BisectEntry {  Change = x.Key, State = x.Value }).ToList();
+				WorkspaceState.Save();
 
 				UpdateBuildList();
 				UpdateStatusPanel();
@@ -5817,8 +5822,8 @@ public enum LatestChangeType
 
 		private void CancelBisectMode()
 		{
-			WorkspaceSettings.ChangeNumberToBisectState.Clear();
-			Settings.Save();
+			WorkspaceState.BisectChanges.Clear();
+			WorkspaceState.Save();
 
 			UpdateBuildList();
 			UpdateStatusPanel();
@@ -5831,11 +5836,11 @@ public enum LatestChangeType
 				PerforceChangeSummary Change = (PerforceChangeSummary)SelectedItem.Tag;
 				if (Change != null)
 				{
-					WorkspaceSettings.ChangeNumberToBisectState[Change.Number] = State;
+					WorkspaceState.SetBisectState(Change.Number, State);
 				}
 			}
 
-			Settings.Save();
+			WorkspaceState.Save();
 
 			ChangeNumberToLayoutInfo.Clear();
 			BuildList.Invalidate();
@@ -5850,11 +5855,11 @@ public enum LatestChangeType
 			GetRemainingBisectRange(out PassChangeNumber, out FailChangeNumber);
 
 			List<int> ChangeNumbers = new List<int>();
-			foreach (KeyValuePair<int, BisectState> Pair in WorkspaceSettings.ChangeNumberToBisectState)
+			foreach (BisectEntry Entry in WorkspaceState.BisectChanges)
 			{
-				if (Pair.Value == BisectState.Include && Pair.Key > PassChangeNumber && Pair.Key < FailChangeNumber)
+				if (Entry.State == BisectState.Include && Entry.Change > PassChangeNumber && Entry.Change < FailChangeNumber)
 				{
-					ChangeNumbers.Add(Pair.Key);
+					ChangeNumbers.Add(Entry.Change);
 				}
 			}
 
@@ -5964,6 +5969,7 @@ public enum LatestChangeType
 			{
 				ProjectSettings.FilterBadges.Add(BadgeName);
 			}
+			ProjectSettings.Save();
 
 			UpdateBuildListFilter();
 		}
@@ -5972,6 +5978,7 @@ public enum LatestChangeType
 		{
 			ProjectSettings.FilterBadges.Clear();
 			ProjectSettings.FilterType = FilterType.None;
+			ProjectSettings.Save();
 
 			Settings.bShowAutomatedChanges = false;
 			Settings.Save();
@@ -6063,14 +6070,15 @@ public enum LatestChangeType
 		private void FilterContextMenu_Type_ShowAll_Click(object sender, EventArgs e)
 		{
 			ProjectSettings.FilterType = FilterType.None;
-			Settings.Save();
+			ProjectSettings.Save();
+
 			UpdateBuildListFilter();
 		}
 
 		private void FilterContextMenu_Type_Code_Click(object sender, EventArgs e)
 		{
 			ProjectSettings.FilterType = FilterType.Code;
-			Settings.Save();
+			ProjectSettings.Save();
 
 			UpdateBuildListFilter();
 		}
@@ -6078,7 +6086,7 @@ public enum LatestChangeType
 		private void FilterContextMenu_Type_Content_Click(object sender, EventArgs e)
 		{
 			ProjectSettings.FilterType = FilterType.Content;
-			Settings.Save();
+			ProjectSettings.Save();
 
 			UpdateBuildListFilter();
 		}
diff --git a/Engine/Source/Programs/UnrealGameSync/UnrealGameSync/Forms/DiagnosticsWindow.cs b/Engine/Source/Programs/UnrealGameSync/UnrealGameSync/Forms/DiagnosticsWindow.cs
--- a/Engine/Source/Programs/UnrealGameSync/UnrealGameSync/Forms/DiagnosticsWindow.cs
+++ b/Engine/Source/Programs/UnrealGameSync/UnrealGameSync/Forms/DiagnosticsWindow.cs
@@ -1,7 +1,10 @@
 // Copyright Epic Games, Inc. All Rights Reserved.
 
+using EpicGames.Core;
 using System;
+using System.Collections.Generic;
 using System.Diagnostics;
+using System.Linq;
 using System.Windows.Forms;
 using System.IO;
 using System.IO.Compression;
@@ -11,12 +14,14 @@ namespace UnrealGameSync
 	public partial class DiagnosticsWindow : Form
 	{
 		string DataFolder;
+		List<FileReference> ExtraFiles;
 
-		public DiagnosticsWindow(string InDataFolder, string InDiagnosticsText)
+		public DiagnosticsWindow(string InDataFolder, string InDiagnosticsText, IEnumerable<FileReference> InExtraFiles)
 		{
 			InitializeComponent();
 			DataFolder = InDataFolder;
 			DiagnosticsTextBox.Text = InDiagnosticsText.Replace("\n", "\r\n");
+			ExtraFiles = InExtraFiles.ToList();
 		}
 
 		private void ViewLogsButton_Click(object sender, EventArgs e)
@@ -62,6 +67,20 @@ namespace UnrealGameSync
 								}
 							}
 						}
+						foreach (FileReference ExtraFile in ExtraFiles)
+						{
+							if(FileReference.Exists(ExtraFile))
+							{
+								using (FileStream InputStream = FileReference.Open(ExtraFile, FileMode.Open, FileAccess.Read, FileShare.ReadWrite))
+								{
+									ZipArchiveEntry Entry = Zip.CreateEntry(ExtraFile.FullName.Replace(":", "").Replace('\\', '/'));
+									using (Stream OutputStream = Entry.Open())
+									{
+										InputStream.CopyTo(OutputStream);
+									}
+								}
+							}
+						}
 					}
 				}
 				catch(Exception Ex)
diff --git a/Engine/Source/Programs/UnrealGameSync/UnrealGameSync/UserSettings.cs b/Engine/Source/Programs/UnrealGameSync/UnrealGameSync/UserSettings.cs
--- a/Engine/Source/Programs/UnrealGameSync/UnrealGameSync/UserSettings.cs
+++ b/Engine/Source/Programs/UnrealGameSync/UnrealGameSync/UserSettings.cs
@@ -1,11 +1,15 @@
 // Copyright Epic Games, Inc. All Rights Reserved.
 
+using EpicGames.Core;
 using System;
 using System.Collections.Generic;
 using System.Drawing;
 using System.IO;
 using System.Linq;
 using System.Text;
+using System.Text.Json;
+using System.Text.Json.Serialization;
+using System.Text.RegularExpressions;
 using System.Threading.Tasks;
 using System.Windows.Forms;
 
@@ -204,42 +208,94 @@ namespace UnrealGameSync
 		}
 	}
 
+	class SyncCategory
+	{
+		public Guid Id { get; set; }
+		public bool Enable { get; set; }
+	}
+
+	class BisectEntry
+	{
+		public int Change { get; set; }
+		public BisectState State { get; set; }
+	}
+
 	class UserWorkspaceSettings
 	{
+		[JsonIgnore]
+		public FileReference File { get; set; }
+
+		// Workspace specific SyncFilters
+		public string[] SyncView { get; set; }
+		public List<SyncCategory> SyncCategories { get; set; } = new List<SyncCategory>();
+		public bool? bSyncAllProjects { get; set; }
+		public bool? bIncludeAllProjectsInSolution { get; set; }
+
+		[JsonIgnore]
+		public Dictionary<Guid, bool> SyncCategoriesDict
+		{
+			get => SyncCategories.ToDictionary(x => x.Id, x => x.Enable);
+			set => SyncCategories = value.Select(x => new SyncCategory { Id = x.Key, Enable = x.Value }).ToList();
+		}
+
+		public void Save() => Utility.SaveJson(File, this);
+	}
+
+	class UserWorkspaceState
+	{
+		[JsonIgnore]
+		public FileReference File { get; set; }
+
+		// Client path to the current project
+		public string SelectedClientFileName { get; set; }
+
 		// Settings for the currently synced project in this workspace. CurrentChangeNumber is only valid for this workspace if CurrentProjectPath is the current project.
-		public string CurrentProjectIdentifier;
-		public int CurrentChangeNumber;
-		public string CurrentSyncFilterHash;
-		public List<int> AdditionalChangeNumbers = new List<int>();
+		public string CurrentProjectIdentifier { get; set; }
+		public int CurrentChangeNumber { get; set; }
+		public string CurrentSyncFilterHash { get; set; }
+		public List<int> AdditionalChangeNumbers { get; set; } = new List<int>();
 
 		// Settings for the last attempted sync. These values are set to persist error messages between runs.
-		public int LastSyncChangeNumber;
-		public WorkspaceUpdateResult LastSyncResult;
-		public string LastSyncResultMessage;
-		public DateTime? LastSyncTime;
-		public int LastSyncDurationSeconds;
+		public int LastSyncChangeNumber { get; set; }
+		public WorkspaceUpdateResult LastSyncResult { get; set; }
+		public string LastSyncResultMessage { get; set; }
+		public DateTime? LastSyncTime { get; set; }
+		public int LastSyncDurationSeconds { get; set; }
 
 		// The last successful build, regardless of whether a failed sync has happened in the meantime. Used to determine whether to force a clean due to entries in the project config file.
-		public int LastBuiltChangeNumber;
+		public int LastBuiltChangeNumber { get; set; }
 
 		// Expanded archives in the workspace
-		public string[] ExpandedArchiveTypes;
+		public string[] ExpandedArchiveTypes { get; set; }
 
 		// The changes that we're regressing at the moment
-		public Dictionary<int, BisectState> ChangeNumberToBisectState = new Dictionary<int, BisectState>();
+		public List<BisectEntry> BisectChanges { get; set; } = new List<BisectEntry>();
 
-		// Workspace specific SyncFilters
-		public string[] SyncView;
-		public Dictionary<Guid, bool> SyncCategories;
-		public bool? bSyncAllProjects;
-		public bool? bIncludeAllProjectsInSolution;
+		public void SetBisectState(int Change, BisectState State)
+		{
+			BisectEntry Entry = BisectChanges.FirstOrDefault(x => x.Change == Change);
+			if (Entry == null)
+			{
+				Entry = new BisectEntry();
+				Entry.Change = Change;
+				BisectChanges.Add(Entry);
+			}
+			Entry.State = State;
+		}
+
+		public void Save() => Utility.SaveJson(File, this);
 	}
 
 	class UserProjectSettings
 	{
-		public List<ConfigObject> BuildSteps = new List<ConfigObject>();
-		public FilterType FilterType;
-		public HashSet<string> FilterBadges = new HashSet<string>(StringComparer.OrdinalIgnoreCase);
+		[JsonIgnore]
+		public FileReference File { get; set; }
+
+		public List<ConfigObject> BuildSteps { get; set; } = new List<ConfigObject>();
+		public FilterType FilterType { get; set; }
+		public HashSet<string> FilterBadges { get; set; } = new HashSet<string>(StringComparer.OrdinalIgnoreCase);
+
+		public void Save() => Utility.SaveJson(File, this);
 	}
 
 	class UserSettings
@@ -313,8 +369,9 @@ namespace UnrealGameSync
 		public int NotifyUnresolvedMinutes;
 
 		// Project settings
-		Dictionary<string, UserWorkspaceSettings> WorkspaceKeyToSettings = new Dictionary<string,UserWorkspaceSettings>();
-		Dictionary<string, UserProjectSettings> ProjectKeyToSettings = new Dictionary<string,UserProjectSettings>();
+		Dictionary<FileReference, UserWorkspaceState> WorkspaceKeyToState = new Dictionary<FileReference, UserWorkspaceState>();
+		Dictionary<FileReference, UserWorkspaceSettings> WorkspaceKeyToSettings = new Dictionary<FileReference, UserWorkspaceSettings>();
+		Dictionary<FileReference, UserProjectSettings> ProjectKeyToSettings = new Dictionary<FileReference, UserProjectSettings>();
 
 		// Perforce settings
 		public PerforceSyncOptions SyncOptions = new PerforceSyncOptions();
@@ -591,142 +648,228 @@ namespace UnrealGameSync
 			return Object.ToString();
 		}
 
-		public UserWorkspaceSettings FindOrAddWorkspace(string ClientBranchPath)
+		public DirectoryReference GetConfigDir(DirectoryReference WorkspaceDir)
+		{
+			DirectoryReference ConfigDir = DirectoryReference.Combine(WorkspaceDir, ".ugs");
+
+			DirectoryInfo ConfigDirInfo = ConfigDir.ToDirectoryInfo();
+			if (!ConfigDirInfo.Exists)
+			{
+				ConfigDirInfo.Create();
+				ConfigDirInfo.Attributes = FileAttributes.Directory | FileAttributes.Hidden;
+			}
+
+			return ConfigDir;
+		}
+
+		public UserWorkspaceState FindOrAddWorkspaceState(DirectoryReference RootDir, string SelectedClientFileName)
 		{
-			// Update the current workspace
-			string CurrentWorkspaceKey = ClientBranchPath.Trim('/');
+			FileReference ConfigFile = FileReference.Combine(GetConfigDir(RootDir), "global.json");
+			if (!WorkspaceKeyToState.TryGetValue(ConfigFile, out UserWorkspaceState State))
+			{
+				if (!Utility.TryLoadJson(ConfigFile, out State))
+				{
+					string ClientBranchPath = Regex.Replace(SelectedClientFileName, "^(//[^/]+)/.*$", "$1");
+					State = ImportWorkspaceState(ClientBranchPath);
+					State.SelectedClientFileName = SelectedClientFileName;
+					Utility.SaveJson(ConfigFile, State);
+				}
 
-			UserWorkspaceSettings CurrentWorkspace;
-			if(!WorkspaceKeyToSettings.TryGetValue(CurrentWorkspaceKey, out CurrentWorkspace))
+				State.File = ConfigFile;
+				WorkspaceKeyToState.Add(ConfigFile, State);
+			}
+			if (!String.Equals(State.SelectedClientFileName, SelectedClientFileName))
 			{
-				// Create a new workspace settings object
-				CurrentWorkspace = new UserWorkspaceSettings();
-				WorkspaceKeyToSettings.Add(CurrentWorkspaceKey, CurrentWorkspace);
+				State = new UserWorkspaceState();
+				State.File = ConfigFile;
+				State.SelectedClientFileName = SelectedClientFileName;
+				WorkspaceKeyToState[ConfigFile] = State;
+			}
+			return State;
+		}
 
-				// Read the workspace settings
-				ConfigSection WorkspaceSection = ConfigFile.FindSection(CurrentWorkspaceKey);
-				if(WorkspaceSection == null)
+		public UserWorkspaceSettings FindOrAddWorkspaceSettings(DirectoryReference LocalRoot, string ClientBranchPath)
+		{
+			FileReference ConfigFile = FileReference.Combine(GetConfigDir(LocalRoot), $"client_{ClientBranchPath.Trim('/').Replace('/', '+')}.json");
+			if (!WorkspaceKeyToSettings.TryGetValue(ConfigFile, out UserWorkspaceSettings Settings))
+			{
+				if (!Utility.TryLoadJson(ConfigFile, out Settings))
 				{
-					string LegacyBranchAndClientKey = ClientBranchPath.Trim('/');
+					Settings = ImportWorkspaceSettings(ClientBranchPath);
+					Utility.SaveJson(ConfigFile, Settings);
+				}
 
-					int SlashIdx = LegacyBranchAndClientKey.IndexOf('/');
-					if(SlashIdx != -1)
-					{
-						LegacyBranchAndClientKey = LegacyBranchAndClientKey.Substring(0, SlashIdx) + "$" + LegacyBranchAndClientKey.Substring(SlashIdx + 1);
-					}
+				Settings.File = ConfigFile;
+				WorkspaceKeyToSettings.Add(ConfigFile, Settings);
+			}
+			return Settings;
+		}
+
+		UserWorkspaceState ImportWorkspaceState(string ClientBranchPath)
+		{
+			UserWorkspaceState CurrentWorkspace = new UserWorkspaceState();
+
+			// Read the workspace settings
+			ConfigSection WorkspaceSection = ConfigFile.FindSection(ClientBranchPath.Trim('/'));
+			if(WorkspaceSection == null)
+			{
+				string LegacyBranchAndClientKey = ClientBranchPath.Trim('/');
+
+				int SlashIdx = LegacyBranchAndClientKey.IndexOf('/');
+				if(SlashIdx != -1)
+				{
+					LegacyBranchAndClientKey = LegacyBranchAndClientKey.Substring(0, SlashIdx) + "$" + LegacyBranchAndClientKey.Substring(SlashIdx + 1);
+				}
 
-					string CurrentSync = ConfigFile.GetValue("Clients." + LegacyBranchAndClientKey, null);
-					if(CurrentSync != null)
+				string CurrentSync = ConfigFile.GetValue("Clients." + LegacyBranchAndClientKey, null);
+				if(CurrentSync != null)
+				{
+					int AtIdx = CurrentSync.LastIndexOf('@');
+					if(AtIdx != -1)
 					{
-						int AtIdx = CurrentSync.LastIndexOf('@');
-						if(AtIdx != -1)
+						int ChangeNumber;
+						if(int.TryParse(CurrentSync.Substring(AtIdx + 1), out ChangeNumber))
 						{
-							int ChangeNumber;
-							if(int.TryParse(CurrentSync.Substring(AtIdx + 1), out ChangeNumber))
-							{
-								CurrentWorkspace.CurrentProjectIdentifier = CurrentSync.Substring(0, AtIdx);
-								CurrentWorkspace.CurrentChangeNumber = ChangeNumber;
-							}
+							CurrentWorkspace.CurrentProjectIdentifier = CurrentSync.Substring(0, AtIdx);
+							CurrentWorkspace.CurrentChangeNumber = ChangeNumber;
 						}
 					}
+				}
 
-					string LastUpdateResultText = ConfigFile.GetValue("Clients." + LegacyBranchAndClientKey + "$LastUpdate", null);
-					if(LastUpdateResultText != null)
+				string LastUpdateResultText = ConfigFile.GetValue("Clients." + LegacyBranchAndClientKey + "$LastUpdate", null);
+				if(LastUpdateResultText != null)
+				{
+					int ColonIdx = LastUpdateResultText.LastIndexOf(':');
+					if(ColonIdx != -1)
 					{
-						int ColonIdx = LastUpdateResultText.LastIndexOf(':');
-						if(ColonIdx != -1)
+						int ChangeNumber;
+						if(int.TryParse(LastUpdateResultText.Substring(0, ColonIdx), out ChangeNumber))
 						{
-							int ChangeNumber;
-							if(int.TryParse(LastUpdateResultText.Substring(0, ColonIdx), out ChangeNumber))
+							WorkspaceUpdateResult Result;
+							if(Enum.TryParse(LastUpdateResultText.Substring(ColonIdx + 1), out Result))
 							{
-								WorkspaceUpdateResult Result;
-								if(Enum.TryParse(LastUpdateResultText.Substring(ColonIdx + 1), out Result))
-								{
-									CurrentWorkspace.LastSyncChangeNumber = ChangeNumber;
-									CurrentWorkspace.LastSyncResult = Result;
-								}
+								CurrentWorkspace.LastSyncChangeNumber = ChangeNumber;
+								CurrentWorkspace.LastSyncResult = Result;
 							}
 						}
 					}
-
-					CurrentWorkspace.SyncView = new string[0];
-					CurrentWorkspace.SyncCategories = new Dictionary<Guid, bool>();
-					CurrentWorkspace.bSyncAllProjects = null;
-					CurrentWorkspace.bIncludeAllProjectsInSolution = null;
 				}
-				else
+			}
+			else
+			{
+				CurrentWorkspace.CurrentProjectIdentifier = WorkspaceSection.GetValue("CurrentProjectPath");
+				CurrentWorkspace.CurrentChangeNumber = WorkspaceSection.GetValue("CurrentChangeNumber", -1);
+				CurrentWorkspace.CurrentSyncFilterHash = WorkspaceSection.GetValue("CurrentSyncFilterHash", null);
+				foreach(string AdditionalChangeNumberString in WorkspaceSection.GetValues("AdditionalChangeNumbers", new string[0]))
 				{
-					CurrentWorkspace.CurrentProjectIdentifier = WorkspaceSection.GetValue("CurrentProjectPath");
-					CurrentWorkspace.CurrentChangeNumber = WorkspaceSection.GetValue("CurrentChangeNumber", -1);
-					CurrentWorkspace.CurrentSyncFilterHash = WorkspaceSection.GetValue("CurrentSyncFilterHash", null);
-					foreach(string AdditionalChangeNumberString in WorkspaceSection.GetValues("AdditionalChangeNumbers", new string[0]))
+					int AdditionalChangeNumber;
+					if(int.TryParse(AdditionalChangeNumberString, out AdditionalChangeNumber))
 					{
-						int AdditionalChangeNumber;
-						if(int.TryParse(AdditionalChangeNumberString, out AdditionalChangeNumber))
-						{
-							CurrentWorkspace.AdditionalChangeNumbers.Add(AdditionalChangeNumber);
-						}
+						CurrentWorkspace.AdditionalChangeNumbers.Add(AdditionalChangeNumber);
 					}
-					Enum.TryParse(WorkspaceSection.GetValue("LastSyncResult", ""), out CurrentWorkspace.LastSyncResult);
-					CurrentWorkspace.LastSyncResultMessage = UnescapeText(WorkspaceSection.GetValue("LastSyncResultMessage"));
-					CurrentWorkspace.LastSyncChangeNumber = WorkspaceSection.GetValue("LastSyncChangeNumber", -1);
+				}
 
-					DateTime LastSyncTime;
-					if(DateTime.TryParse(WorkspaceSection.GetValue("LastSyncTime", ""), out LastSyncTime))
-					{
-						CurrentWorkspace.LastSyncTime = LastSyncTime;
-					}
+				WorkspaceUpdateResult LastSyncResult;
+				Enum.TryParse(WorkspaceSection.GetValue("LastSyncResult", ""), out LastSyncResult);
+				CurrentWorkspace.LastSyncResult = LastSyncResult;
 
-					CurrentWorkspace.LastSyncDurationSeconds = WorkspaceSection.GetValue("LastSyncDuration", 0);
-					CurrentWorkspace.LastBuiltChangeNumber = WorkspaceSection.GetValue("LastBuiltChangeNumber", 0);
-					CurrentWorkspace.ExpandedArchiveTypes = WorkspaceSection.GetValues("ExpandedArchiveName", new string[0]);
+				CurrentWorkspace.LastSyncResultMessage = UnescapeText(WorkspaceSection.GetValue("LastSyncResultMessage"));
+				CurrentWorkspace.LastSyncChangeNumber = WorkspaceSection.GetValue("LastSyncChangeNumber", -1);
 
-					CurrentWorkspace.SyncView = WorkspaceSection.GetValues("SyncFilter", new string[0]);
-					CurrentWorkspace.SyncCategories = GetCategorySettings(WorkspaceSection, "SyncIncludedCategories", "SyncExcludedCategories");
+				DateTime LastSyncTime;
+				if(DateTime.TryParse(WorkspaceSection.GetValue("LastSyncTime", ""), out LastSyncTime))
+				{
+					CurrentWorkspace.LastSyncTime = LastSyncTime;
+				}
 
-					int SyncAllProjects = WorkspaceSection.GetValue("SyncAllProjects", -1);
-					CurrentWorkspace.bSyncAllProjects = (SyncAllProjects == 0)? (bool?)false : (SyncAllProjects == 1)? (bool?)true : (bool?)null;
+				CurrentWorkspace.LastSyncDurationSeconds = WorkspaceSection.GetValue("LastSyncDuration", 0);
+				CurrentWorkspace.LastBuiltChangeNumber = WorkspaceSection.GetValue("LastBuiltChangeNumber", 0);
+				CurrentWorkspace.ExpandedArchiveTypes = WorkspaceSection.GetValues("ExpandedArchiveName", new string[0]);
 
-					int IncludeAllProjectsInSolution = WorkspaceSection.GetValue("IncludeAllProjectsInSolution", -1);
-					CurrentWorkspace.bIncludeAllProjectsInSolution = (IncludeAllProjectsInSolution == 0)? (bool?)false : (IncludeAllProjectsInSolution == 1)? (bool?)true : (bool?)null;
+				string[] BisectEntries = WorkspaceSection.GetValues("Bisect", new string[0]);
+				foreach(string BisectEntry in BisectEntries)
+				{
+					ConfigObject BisectEntryObject = new ConfigObject(BisectEntry);
 
-					string[] BisectEntries = WorkspaceSection.GetValues("Bisect", new string[0]);
-					foreach(string BisectEntry in BisectEntries)
+					int ChangeNumber = BisectEntryObject.GetValue("Change", -1);
+					if(ChangeNumber != -1)
 					{
-						ConfigObject BisectEntryObject = new ConfigObject(BisectEntry);
-
-						int ChangeNumber = BisectEntryObject.GetValue("Change", -1);
-						if(ChangeNumber != -1)
+						BisectState State;
+						if(Enum.TryParse(BisectEntryObject.GetValue("State", ""), out State))
 						{
-							BisectState State;
-							if(Enum.TryParse(BisectEntryObject.GetValue("State", ""), out State))
-							{
-								CurrentWorkspace.ChangeNumberToBisectState[ChangeNumber] = State;
-							}
+							BisectEntry Entry = new BisectEntry();
+							Entry.Change = ChangeNumber;
+							Entry.State = State;
+							CurrentWorkspace.BisectChanges.Add(Entry);
 						}
 					}
 				}
 			}
+
 			return CurrentWorkspace;
 		}
 
-		public UserProjectSettings FindOrAddProject(string ClientProjectFileName)
+		UserWorkspaceSettings ImportWorkspaceSettings(string CurrentWorkspaceKey)
 		{
-			// Read the project settings
-			UserProjectSettings CurrentProject;
-			if(!ProjectKeyToSettings.TryGetValue(ClientProjectFileName, out CurrentProject))
+			UserWorkspaceSettings CurrentWorkspace = new UserWorkspaceSettings();
+
+			ConfigSection WorkspaceSection = ConfigFile.FindSection(CurrentWorkspaceKey);
+			if (WorkspaceSection == null)
 			{
-				CurrentProject = new UserProjectSettings();
-				ProjectKeyToSettings.Add(ClientProjectFileName, CurrentProject);
-	
-				ConfigSection ProjectSection = ConfigFile.FindOrAddSection(ClientProjectFileName);
-				CurrentProject.BuildSteps.AddRange(ProjectSection.GetValues("BuildStep", new string[0]).Select(x => new ConfigObject(x)));
-				if(!Enum.TryParse(ProjectSection.GetValue("FilterType", ""), true, out CurrentProject.FilterType))
+				CurrentWorkspace.SyncView = new string[0];
+				CurrentWorkspace.SyncCategories = new List<SyncCategory>();
+				CurrentWorkspace.bSyncAllProjects = null;
+				CurrentWorkspace.bIncludeAllProjectsInSolution = null;
+			}
+			else
+			{
+				CurrentWorkspace.SyncView = WorkspaceSection.GetValues("SyncFilter", new string[0]);
+				CurrentWorkspace.SyncCategories = GetCategorySettings(WorkspaceSection, "SyncIncludedCategories", "SyncExcludedCategories").Select(x => new SyncCategory { Id = x.Key, Enable = x.Value }).ToList();
+
+				int SyncAllProjects = WorkspaceSection.GetValue("SyncAllProjects", -1);
+				CurrentWorkspace.bSyncAllProjects = (SyncAllProjects == 0) ? (bool?)false : (SyncAllProjects == 1) ? (bool?)true : (bool?)null;
+
+				int IncludeAllProjectsInSolution = WorkspaceSection.GetValue("IncludeAllProjectsInSolution", -1);
+				CurrentWorkspace.bIncludeAllProjectsInSolution = (IncludeAllProjectsInSolution == 0) ? (bool?)false : (IncludeAllProjectsInSolution == 1) ? (bool?)true : (bool?)null;
+			}
+
+			return CurrentWorkspace;
+		}
+
+		public UserProjectSettings FindOrAddProjectSettings(DirectoryReference LocalRoot, string ClientProjectFileName)
+		{
+			FileReference ConfigFile = FileReference.Combine(GetConfigDir(LocalRoot), $"project_{ClientProjectFileName.Trim('/').Replace('/', '+')}.json");
+			if (!ProjectKeyToSettings.TryGetValue(ConfigFile, out UserProjectSettings Settings))
+			{
+				if (!Utility.TryLoadJson(ConfigFile, out Settings))
 				{
-					CurrentProject.FilterType = FilterType.None;
+					Settings = ImportProjectSettings(ClientProjectFileName);
+					Utility.SaveJson(ConfigFile, Settings);
 				}
-				CurrentProject.FilterBadges.UnionWith(ProjectSection.GetValues("FilterBadges", new string[0]));
+
+				Settings.File = ConfigFile;
+				ProjectKeyToSettings.Add(ConfigFile, Settings);
+			}
+			return Settings;
+		}
+
+		public UserProjectSettings ImportProjectSettings(string ClientProjectFileName)
+		{
+			// Read the project settings
+			UserProjectSettings CurrentProject = new UserProjectSettings();
+	
+			ConfigSection ProjectSection = ConfigFile.FindOrAddSection(ClientProjectFileName);
+			CurrentProject.BuildSteps.AddRange(ProjectSection.GetValues("BuildStep", new string[0]).Select(x => new ConfigObject(x)));
+
+			FilterType FilterType;
+			if (!Enum.TryParse(ProjectSection.GetValue("FilterType", ""), true, out FilterType))
+			{
+				FilterType = FilterType.None;
 			}
+
+			CurrentProject.FilterType = FilterType;
+			CurrentProject.FilterBadges.UnionWith(ProjectSection.GetValues("FilterBadges", new string[0]));
+
 			return CurrentProject;
 		}
 
@@ -833,72 +976,6 @@ namespace UnrealGameSync
 				NotificationSection.SetValue("NotifyUnresolvedMinutes", NotifyUnresolvedMinutes);
 			}
 
-			// Current workspace settings
-			foreach(KeyValuePair<string, UserWorkspaceSettings> Pair in WorkspaceKeyToSettings)
-			{
-				string CurrentWorkspaceKey = Pair.Key;
-				UserWorkspaceSettings CurrentWorkspace = Pair.Value;
-
-				ConfigSection WorkspaceSection = ConfigFile.FindOrAddSection(CurrentWorkspaceKey);
-				WorkspaceSection.Clear();
-				WorkspaceSection.SetValue("CurrentProjectPath", CurrentWorkspace.CurrentProjectIdentifier);
-				WorkspaceSection.SetValue("CurrentChangeNumber", CurrentWorkspace.CurrentChangeNumber);
-				if(CurrentWorkspace.CurrentSyncFilterHash != null)
-				{
-					WorkspaceSection.SetValue("CurrentSyncFilterHash", CurrentWorkspace.CurrentSyncFilterHash);
-				}
-				WorkspaceSection.SetValues("AdditionalChangeNumbers", CurrentWorkspace.AdditionalChangeNumbers.Select(x => x.ToString()).ToArray());
-				WorkspaceSection.SetValue("LastSyncResult", CurrentWorkspace.LastSyncResult.ToString());
-				WorkspaceSection.SetValue("LastSyncResultMessage", EscapeText(CurrentWorkspace.LastSyncResultMessage));
-				WorkspaceSection.SetValue("LastSyncChangeNumber", CurrentWorkspace.LastSyncChangeNumber);
-				if(CurrentWorkspace.LastSyncTime.HasValue)
-				{
-					WorkspaceSection.SetValue("LastSyncTime", CurrentWorkspace.LastSyncTime.ToString());
-				}
-				if(CurrentWorkspace.LastSyncDurationSeconds > 0)
-				{
-					WorkspaceSection.SetValue("LastSyncDuration", CurrentWorkspace.LastSyncDurationSeconds);
-				}
-				WorkspaceSection.SetValue("LastBuiltChangeNumber", CurrentWorkspace.LastBuiltChangeNumber);
-				WorkspaceSection.SetValues("ExpandedArchiveName", CurrentWorkspace.ExpandedArchiveTypes);
-				WorkspaceSection.SetValues("SyncFilter", CurrentWorkspace.SyncView);
-				SetCategorySettings(WorkspaceSection, "SyncIncludedCategories", "SyncExcludedCategories", CurrentWorkspace.SyncCategories);
-				if(CurrentWorkspace.bSyncAllProjects.HasValue)
-				{
-					WorkspaceSection.SetValue("SyncAllProjects", CurrentWorkspace.bSyncAllProjects.Value);
-				}
-				if(CurrentWorkspace.bIncludeAllProjectsInSolution.HasValue)
-				{
-					WorkspaceSection.SetValue("IncludeAllProjectsInSolution", CurrentWorkspace.bIncludeAllProjectsInSolution.Value);
-				}
-
-				List<ConfigObject> BisectEntryObjects = new List<ConfigObject>();
-				foreach(KeyValuePair<int, BisectState> BisectPair in CurrentWorkspace.ChangeNumberToBisectState)
-				{
-					ConfigObject BisectEntryObject = new ConfigObject();
-					BisectEntryObject.SetValue("Change", BisectPair.Key);
-					BisectEntryObject.SetValue("State", BisectPair.Value.ToString());
-					BisectEntryObjects.Add(BisectEntryObject);
-				}
-				WorkspaceSection.SetValues("Bisect", BisectEntryObjects.Select(x => x.ToString()).ToArray());
-			}
-
-			// Current project settings
-			foreach(KeyValuePair<string, UserProjectSettings> Pair in ProjectKeyToSettings)
-			{
-				string CurrentProjectKey = Pair.Key;
-				UserProjectSettings CurrentProject = Pair.Value;
-
-				ConfigSection ProjectSection = ConfigFile.FindOrAddSection(CurrentProjectKey);
-				ProjectSection.Clear();
-				ProjectSection.SetValues("BuildStep", CurrentProject.BuildSteps.Select(x => x.ToString()).ToArray());
-				if(CurrentProject.FilterType != FilterType.None)
-				{
-					ProjectSection.SetValue("FilterType", CurrentProject.FilterType.ToString());
-				}
-				ProjectSection.SetValues("FilterBadges", CurrentProject.FilterBadges.ToArray());
-			}
-
 			// Perforce settings
 			ConfigSection PerforceSection = ConfigFile.FindOrAddSection("Perforce");
 			PerforceSection.Clear();
@@ -1067,5 +1144,10 @@ namespace UnrealGameSync
 			}
 			return Result.ToString();
 		}
+
+		public IEnumerable<FileReference> GetCachedFilePaths()
+		{
+			return WorkspaceKeyToState.Keys.Concat(WorkspaceKeyToSettings.Keys).Concat(ProjectKeyToSettings.Keys);
+		}
 	}
 }
diff --git a/Engine/Source/Programs/UnrealGameSync/UnrealGameSync/Utility.cs b/Engine/Source/Programs/UnrealGameSync/UnrealGameSync/Utility.cs
--- a/Engine/Source/Programs/UnrealGameSync/UnrealGameSync/Utility.cs
+++ b/Engine/Source/Programs/UnrealGameSync/UnrealGameSync/Utility.cs
@@ -2,6 +2,7 @@
 
 #define USE_NEW_PROCESS_JOBS
 
+using EpicGames.Core;
 using Microsoft.Win32;
 using System;
 using System.Collections.Generic;
@@ -19,6 +20,49 @@ namespace UnrealGameSync
 {
 	static class Utility
 	{
+		public static bool TryLoadJson<T>(FileReference File, out T Object) where T : class
+		{
+			if (!FileReference.Exists(File))
+			{
+				Object = null;
+				return false;
+			}
+
+			try
+			{
+				Object = LoadJson<T>(File);
+				return true;
+			}
+			catch
+			{
+				Object = null;
+				return false;
+			}
+		}
+
+		public static T LoadJson<T>(FileReference File)
+		{
+			JsonSerializerOptions Options = new JsonSerializerOptions { AllowTrailingCommas = true, ReadCommentHandling = JsonCommentHandling.Skip };
+			Options.Converters.Add(new JsonStringEnumConverter());
+
+			byte[] Data = FileReference.ReadAllBytes(File);
+			return JsonSerializer.Deserialize<T>(Data, Options);
+		}
+
+		public static void SaveJson<T>(FileReference File, T Object)
+		{
+			JsonSerializerOptions Options = new JsonSerializerOptions { IgnoreNullValues = true, WriteIndented = true };
+			Options.Converters.Add(new JsonStringEnumConverter());
+
+			using (Stream Stream = FileReference.Open(File, FileMode.Create, FileAccess.Write, FileShare.Read))
+			{
+				using (Utf8JsonWriter Writer = new Utf8JsonWriter(Stream, new JsonWriterOptions { Indented = true }))
+				{
+					JsonSerializer.Serialize(Writer, Object, Options);
+				}
+			}
+		}
+
 		public static string GetPathWithCorrectCase(FileInfo Info)
 		{
 			DirectoryInfo ParentInfo = Info.Directory;
diff --git a/Engine/Source/Programs/UnrealGameSync/UnrealGameSync/Workspace.cs b/Engine/Source/Programs/UnrealGameSync/UnrealGameSync/Workspace.cs
--- a/Engine/Source/Programs/UnrealGameSync/UnrealGameSync/Workspace.cs
+++ b/Engine/Source/Programs/UnrealGameSync/UnrealGameSync/Workspace.cs
@@ -1,5 +1,6 @@
 // Copyright Epic Games, Inc. All Rights Reserved.
 
+using EpicGames.Core;
 using System;
 using System.Collections.Concurrent;
 using System.Collections.Generic;
@@ -440,7 +441,7 @@ namespace UnrealGameSync
 			{
 				using(TelemetryStopwatch SyncTelemetryStopwatch = new TelemetryStopwatch("Workspace_Sync", TelemetryProjectPath))
 				{
-					Log.WriteLine("Syncing to {0}...", PendingChangeNumber);
+					Log.WriteLine("Syncing to {0}...", Context.ChangeNumber);
 
 					// Make sure we're logged in
 					bool bLoggedIn;
@@ -627,7 +628,7 @@ namespace UnrealGameSync
 						foreach(string SyncPath in SyncPaths)
 						{
 							List<PerforceFileRecord> SyncRecords = new List<PerforceFileRecord>();
-							if(!Perforce.SyncPreview(SyncPath, PendingChangeNumber, !Context.Options.HasFlag(WorkspaceUpdateOptions.Sync), Record => { SyncRecords.Add(Record); Counter.Increment(); }, Log))
+							if(!Perforce.SyncPreview(SyncPath, Context.ChangeNumber, !Context.Options.HasFlag(WorkspaceUpdateOptions.Sync), Record => { SyncRecords.Add(Record); Counter.Increment(); }, Log))
 							{
 								StatusMessage = String.Format("Couldn't enumerate changes matching {0}.", SyncPath);
 								return WorkspaceUpdateResult.FailedToSync;
@@ -648,7 +649,7 @@ namespace UnrealGameSync
 								// If it doesn't exist locally, just add a sync command for it
 								if (String.IsNullOrEmpty(SyncRecord.ClientPath))
 								{
-									BatchBuilder.Add(String.Format("{0}@{1}", SyncRecord.DepotPath, PendingChangeNumber), SyncRecord.FileSize);
+									BatchBuilder.Add(String.Format("{0}@{1}", SyncRecord.DepotPath, Context.ChangeNumber), SyncRecord.FileSize);
 									SyncDepotPaths.Add(SyncRecord.DepotPath);
 									RequiredFreeSpace += SyncRecord.FileSize;
 									continue;
@@ -671,7 +672,7 @@ namespace UnrealGameSync
 								// This occurs for files returned from GetOpenFiles as SyncRecord.ClientPath is the client workspace path not local path.
 								if (!FullName.StartsWith(LocalRootPrefix, StringComparison.OrdinalIgnoreCase))
 								{
-									BatchBuilder.Add(String.Format("{0}@{1}", SyncRecord.DepotPath, PendingChangeNumber), SyncRecord.FileSize);
+									BatchBuilder.Add(String.Format("{0}@{1}", SyncRecord.DepotPath, Context.ChangeNumber), SyncRecord.FileSize);
 									SyncDepotPaths.Add(SyncRecord.DepotPath);
 									RequiredFreeSpace += SyncRecord.FileSize;
 									continue;
