UGS: Remove internal state from Workspace update.

* State object is now serialized to JSON file in the workspace root before making any destructive changes.
* Most functionality in Workspace class moved into static functions, in prepartion for moving into a command line tool.
* Also convert file/directory strings into FileReference/DirectoryReference objects.

diff --git a/Engine/Source/Programs/UnrealGameSync/UnrealGameSync/ArchiveManifest.cs b/Engine/Source/Programs/UnrealGameSync/UnrealGameSync/ArchiveManifest.cs
--- a/Engine/Source/Programs/UnrealGameSync/UnrealGameSync/ArchiveManifest.cs
+++ b/Engine/Source/Programs/UnrealGameSync/UnrealGameSync/ArchiveManifest.cs
@@ -1,5 +1,6 @@
 // Copyright Epic Games, Inc. All Rights Reserved.
 
+using EpicGames.Core;
 using System;
 using System.Collections.Generic;
 using System.Diagnostics;
@@ -79,14 +80,14 @@ namespace UnrealGameSync
 
 	static class ArchiveUtils
 	{
-		public static void ExtractFiles(string ArchiveFileName, string BaseDirectoryName, string ManifestFileName, ProgressValue Progress, TextWriter LogWriter)
+		public static void ExtractFiles(FileReference ArchiveFileName, DirectoryReference BaseDirectoryName, FileReference ManifestFileName, ProgressValue Progress, TextWriter LogWriter)
 		{
 			DateTime TimeStamp = DateTime.UtcNow;
-			using (ZipArchive Zip = new ZipArchive(File.OpenRead(ArchiveFileName)))
+			using (ZipArchive Zip = new ZipArchive(File.OpenRead(ArchiveFileName.FullName)))
 			{
 				if (ManifestFileName != null)
 				{
-					File.Delete(ManifestFileName);
+					FileReference.Delete(ManifestFileName);
 
 					// Create the manifest
 					ArchiveManifest Manifest = new ArchiveManifest();
@@ -99,12 +100,12 @@ namespace UnrealGameSync
 					}
 
 					// Write it out to a temporary file, then move it into place
-					string TempManifestFileName = ManifestFileName + ".tmp";
-					using (FileStream OutputStream = File.Open(TempManifestFileName, FileMode.Create, FileAccess.Write))
+					FileReference TempManifestFileName = ManifestFileName + ".tmp";
+					using (FileStream OutputStream = FileReference.Open(TempManifestFileName, FileMode.Create, FileAccess.Write))
 					{
 						Manifest.Write(OutputStream);
 					}
-					File.Move(TempManifestFileName, ManifestFileName);
+					FileReference.Move(TempManifestFileName, ManifestFileName);
 				}
 
 				// Extract all the files
@@ -113,27 +114,27 @@ namespace UnrealGameSync
 				{
 					if(!Entry.FullName.EndsWith("/") && !Entry.FullName.EndsWith("\\"))
 					{
-						string FileName = Path.Combine(BaseDirectoryName, Entry.FullName);
-						Directory.CreateDirectory(Path.GetDirectoryName(FileName));
+						FileReference FileName = FileReference.Combine(BaseDirectoryName, Entry.FullName);
+						DirectoryReference.CreateDirectory(FileName.Directory);
 
 						lock(LogWriter)
 						{
 							LogWriter.WriteLine("Writing {0}", FileName);
 						}
 
-						Entry.ExtractToFile(FileName, true);
-						File.SetLastWriteTimeUtc(FileName, TimeStamp);
+						Entry.ExtractToFile(FileName.FullName, true);
+						FileReference.SetLastWriteTimeUtc(FileName, TimeStamp);
 					}
 					Progress.Set((float)++EntryIdx / (float)Zip.Entries.Count);
 				}
 			}
 		}
 
-		public static void RemoveExtractedFiles(string BaseDirectoryName, string ManifestFileName, ProgressValue Progress, TextWriter LogWriter)
+		public static void RemoveExtractedFiles(DirectoryReference BaseDirectoryName, FileReference ManifestFileName, ProgressValue Progress, TextWriter LogWriter)
 		{
 			// Read the manifest in
 			ArchiveManifest Manifest;
-			using(FileStream InputStream = File.Open(ManifestFileName, FileMode.Open, FileAccess.Read))
+			using(FileStream InputStream = FileReference.Open(ManifestFileName, FileMode.Open, FileAccess.Read))
 			{
 				Manifest = new ArchiveManifest(InputStream);
 			}
@@ -141,7 +142,7 @@ namespace UnrealGameSync
 			// Remove all the files that haven't been modified match
 			for(int Idx = 0; Idx < Manifest.Files.Count; Idx++)
 			{
-				FileInfo File = new FileInfo(Path.Combine(BaseDirectoryName, Manifest.Files[Idx].FileName));
+				FileInfo File = FileReference.Combine(BaseDirectoryName, Manifest.Files[Idx].FileName).ToFileInfo();
 				if(File.Exists)
 				{
 					if(File.Length != Manifest.Files[Idx].Length)
diff --git a/Engine/Source/Programs/UnrealGameSync/UnrealGameSync/ConfigFile.cs b/Engine/Source/Programs/UnrealGameSync/UnrealGameSync/ConfigFile.cs
--- a/Engine/Source/Programs/UnrealGameSync/UnrealGameSync/ConfigFile.cs
+++ b/Engine/Source/Programs/UnrealGameSync/UnrealGameSync/ConfigFile.cs
@@ -1,5 +1,6 @@
 // Copyright Epic Games, Inc. All Rights Reserved.
 
+using EpicGames.Core;
 using System;
 using System.Collections.Generic;
 using System.Diagnostics;
@@ -431,31 +432,33 @@ namespace UnrealGameSync
 		{
 		}
 
-		string GetTempFileName(string FileName)
+		FileReference GetTempFileName(FileReference FileName)
 		{
 			return FileName + ".tmp";
 		}
 
-		public void Load(string FileName)
+		public void Load(FileReference FileName)
 		{
-			Parse(File.ReadAllLines(FileName));
+			Parse(FileReference.ReadAllLines(FileName));
 		}
 
-		public bool TryLoad(string FileName, TextWriter Log)
+		public bool TryLoad(FileReference FileName, TextWriter Log)
 		{
-			FileInfo FileInfo = new FileInfo(FileName);
+			FileInfo FileInfo = FileName.ToFileInfo();
 			if (FileInfo.Exists)
 			{
 				Log.WriteLine("Loading config file from {0} ({1} bytes)", FileInfo.FullName, FileInfo.Length);
-				Load(FileInfo.FullName);
+				Load(FileName);
 				return true;
 			}
 
-			FileInfo TempFileInfo = new FileInfo(GetTempFileName(FileName));
+			FileReference TempFileName = GetTempFileName(FileName);
+
+			FileInfo TempFileInfo = TempFileName.ToFileInfo();
 			if (TempFileInfo.Exists)
 			{
 				Log.WriteLine("Loading temporary config file from {0} ({1} bytes)", TempFileInfo.FullName, TempFileInfo.Length);
-				Load(TempFileInfo.FullName);
+				Load(TempFileName);
 				return true;
 			}
 
@@ -496,12 +499,12 @@ namespace UnrealGameSync
 			}
 		}
 
-		public void Save(string FileName)
+		public void Save(FileReference FileName)
 		{
-			string TempFileName = GetTempFileName(FileName);
-			File.Delete(TempFileName);
+			FileReference TempFileName = GetTempFileName(FileName);
+			FileReference.Delete(TempFileName);
 
-			using (StreamWriter Writer = new StreamWriter(TempFileName))
+			using (StreamWriter Writer = new StreamWriter(TempFileName.FullName))
 			{
 				for(int Idx = 0; Idx < Sections.Count; Idx++)
 				{
@@ -527,8 +530,8 @@ namespace UnrealGameSync
 				}
 			}
 
-			File.Delete(FileName);
-			File.Move(TempFileName, FileName);
+			FileReference.Delete(FileName);
+			FileReference.Move(TempFileName, FileName);
 		}
 
 		public ConfigSection FindSection(string Name)
diff --git a/Engine/Source/Programs/UnrealGameSync/UnrealGameSync/Controls/LogControl.cs b/Engine/Source/Programs/UnrealGameSync/UnrealGameSync/Controls/LogControl.cs
--- a/Engine/Source/Programs/UnrealGameSync/UnrealGameSync/Controls/LogControl.cs
+++ b/Engine/Source/Programs/UnrealGameSync/UnrealGameSync/Controls/LogControl.cs
@@ -13,6 +13,7 @@ using System.Runtime.InteropServices;
 using System.Collections.Concurrent;
 using System.IO;
 using System.Text.RegularExpressions;
+using EpicGames.Core;
 
 namespace UnrealGameSync
 {
@@ -178,13 +179,13 @@ namespace UnrealGameSync
 			CloseFile();
 		}
 
-		public bool OpenFile(string NewLogFileName)
+		public bool OpenFile(FileReference NewLogFileName)
 		{
 			CloseFile();
 			Clear();
 			try
 			{
-				LogFileStream = File.Open(NewLogFileName, FileMode.OpenOrCreate, FileAccess.ReadWrite, FileShare.Read);
+				LogFileStream = FileReference.Open(NewLogFileName, FileMode.OpenOrCreate, FileAccess.ReadWrite, FileShare.Read);
 			}
 			catch(Exception)
 			{
diff --git a/Engine/Source/Programs/UnrealGameSync/UnrealGameSync/Controls/WorkspaceControl.cs b/Engine/Source/Programs/UnrealGameSync/UnrealGameSync/Controls/WorkspaceControl.cs
--- a/Engine/Source/Programs/UnrealGameSync/UnrealGameSync/Controls/WorkspaceControl.cs
+++ b/Engine/Source/Programs/UnrealGameSync/UnrealGameSync/Controls/WorkspaceControl.cs
@@ -206,7 +206,7 @@ public enum LatestChangeType
 
 		IWorkspaceControlOwner Owner;
 		string ApiUrl;
-		string DataFolder;
+		DirectoryReference DataFolder;
 		LineBasedTextWriter Log;
 
 		UserSettings Settings;
@@ -222,7 +222,7 @@ public enum LatestChangeType
 			private set;
 		}
 
-		public string SelectedFileName
+		public FileReference SelectedFileName
 		{
 			get;
 			private set;
@@ -230,7 +230,7 @@ public enum LatestChangeType
 
 		string SelectedProjectIdentifier;
 
-		public string BranchDirectoryName
+		public DirectoryReference BranchDirectoryName
 		{
 			get;
 			private set;
@@ -333,11 +333,11 @@ public enum LatestChangeType
 			bUnstable = bInUnstable;
 			Log = InLog;
 			Settings = InSettings;
-			WorkspaceState = InSettings.FindOrAddWorkspaceState(new DirectoryReference(DetectSettings.BranchDirectoryName), DetectSettings.NewSelectedClientFileName);
-			WorkspaceSettings = InSettings.FindOrAddWorkspaceSettings(new DirectoryReference(DetectSettings.BranchDirectoryName), DetectSettings.BranchClientPath);
+			WorkspaceState = InSettings.FindOrAddWorkspaceState(DetectSettings.BranchDirectoryName, DetectSettings.NewSelectedClientFileName, DetectSettings.NewSelectedProjectIdentifier);
+			WorkspaceSettings = InSettings.FindOrAddWorkspaceSettings(DetectSettings.BranchDirectoryName, DetectSettings.BranchClientPath);
 
 			string ProjectPath = Regex.Replace(DetectSettings.NewSelectedClientFileName, "^//[^/]+/", "");
-			ProjectSettings = InSettings.FindOrAddProjectSettings(new DirectoryReference(DetectSettings.BranchDirectoryName), ProjectPath);
+			ProjectSettings = InSettings.FindOrAddProjectSettings(DetectSettings.BranchDirectoryName, ProjectPath);
 
 			DesiredTaskbarState = Tuple.Create(TaskbarState.NoProgress, 0.0f);
 
@@ -396,10 +396,12 @@ public enum LatestChangeType
 
 			string TelemetryProjectIdentifier = PerforceUtils.GetClientOrDepotDirectoryName(DetectSettings.NewSelectedProjectIdentifier);
 
-			Workspace = new Workspace(PerforceClient, BranchDirectoryName, SelectedFileName, DetectSettings.BranchClientPath, DetectSettings.NewSelectedClientFileName, CurrentChangeNumber, CurrentSyncFilterHash, WorkspaceState.LastBuiltChangeNumber, TelemetryProjectIdentifier, DetectSettings.bIsEnterpriseProject, new LogControlTextWriter(SyncLog));
+			ProjectInfo Project = new ProjectInfo(BranchDirectoryName, SelectedFileName, DetectSettings.BranchClientPath, DetectSettings.NewSelectedClientFileName, TelemetryProjectIdentifier, DetectSettings.bIsEnterpriseProject);
+
+			Workspace = new Workspace(PerforceClient, Project, WorkspaceState, new LogControlTextWriter(SyncLog));
 			Workspace.OnUpdateComplete += UpdateCompleteCallback;
 
-			string ProjectLogBaseName = Path.Combine(DataFolder, String.Format("{0}@{1}", PerforceClient.ClientName, DetectSettings.BranchClientPath.Replace("//" + PerforceClient.ClientName + "/", "").Trim('/').Replace("/", "$")));
+			FileReference ProjectLogBaseName = FileReference.Combine(DataFolder, String.Format("{0}@{1}", PerforceClient.ClientName, DetectSettings.BranchClientPath.Replace("//" + PerforceClient.ClientName + "/", "").Trim('/').Replace("/", "$")));
 
 			PerforceMonitor = new PerforceMonitor(PerforceClient, DetectSettings.BranchClientPath, DetectSettings.NewSelectedClientFileName, DetectSettings.NewSelectedProjectIdentifier, ProjectLogBaseName + ".p4.log", DetectSettings.bIsEnterpriseProject, DetectSettings.LatestProjectConfigFile, DetectSettings.CacheFolder, DetectSettings.LocalConfigFiles);
 			PerforceMonitor.OnUpdate += UpdateBuildListCallback;
@@ -414,8 +416,7 @@ public enum LatestChangeType
 
 			UpdateColumnSettings();
 
-			string LogFileName = Path.Combine(DataFolder, ProjectLogBaseName + ".sync.log");
-			SyncLog.OpenFile(LogFileName);
+			SyncLog.OpenFile(ProjectLogBaseName + ".sync.log");
 
 			Splitter.SetLogVisibility(Settings.bShowLogWindow);
 
@@ -436,14 +437,14 @@ public enum LatestChangeType
 			IssueMonitor = InOwner.CreateIssueMonitor(IssuesApiUrl, DetectSettings.PerforceClient.UserName);
 			IssueMonitor.OnIssuesChanged += IssueMonitor_OnIssuesChangedAsync;
 
-			if (SelectedFileName.EndsWith(".uproject", StringComparison.OrdinalIgnoreCase))
+			if (SelectedFileName.HasExtension(".uproject"))
 			{
-				string TintColorFileName = GetTintColorFileName();
+				FileReference TintColorFileName = GetTintColorFileName();
 
-				Directory.CreateDirectory(Path.GetDirectoryName(TintColorFileName));
+				DirectoryReference.CreateDirectory(TintColorFileName.Directory);
 
-				EditorConfigWatcher.Path = Path.GetDirectoryName(TintColorFileName);
-				EditorConfigWatcher.Filter = Path.GetFileName(TintColorFileName);
+				EditorConfigWatcher.Path = TintColorFileName.Directory.FullName;
+				EditorConfigWatcher.Filter = TintColorFileName.GetFileName();
 				EditorConfigWatcher.EnableRaisingEvents = true;
 
 				TintColor = GetTintColor();
@@ -847,18 +848,6 @@ public enum LatestChangeType
 			MainThreadSynchronizationContext.Post((o) => LoginExpired(), null);
 		}
 
-		private void UpdateSyncConfig(int ChangeNumber, string SyncFilterHash)
-		{
-			WorkspaceState.CurrentProjectIdentifier = SelectedProjectIdentifier;
-			WorkspaceState.CurrentChangeNumber = ChangeNumber;
-			WorkspaceState.CurrentSyncFilterHash = SyncFilterHash;
-			if (ChangeNumber == -1 || ChangeNumber != WorkspaceState.CurrentChangeNumber)
-			{
-				WorkspaceState.AdditionalChangeNumbers.Clear();
-			}
-			WorkspaceState.Save();
-		}
-
 		private void BuildList_OnScroll()
 		{
 			PendingSelectedChangeNumber = -1;
@@ -1028,7 +1017,6 @@ public enum LatestChangeType
 
 			if (Context.Options.HasFlag(WorkspaceUpdateOptions.Sync))
 			{
-				UpdateSyncConfig(-1, Workspace.CurrentSyncFilterHash);
 				EventMonitor.PostEvent(Context.ChangeNumber, EventType.Syncing);
 			}
 
@@ -1038,12 +1026,12 @@ public enum LatestChangeType
 				{
 					foreach (BuildConfig Config in Enum.GetValues(typeof(BuildConfig)))
 					{
-						List<string> EditorReceiptPaths = GetEditorReceiptPaths(Config);
-						foreach (string EditorReceiptPath in EditorReceiptPaths)
+						List<FileReference> EditorReceiptPaths = GetEditorReceiptPaths(Config);
+						foreach (FileReference EditorReceiptPath in EditorReceiptPaths)
 						{
-							if (File.Exists(EditorReceiptPath))
+							if (FileReference.Exists(EditorReceiptPath))
 							{
-								try { File.Delete(EditorReceiptPath); } catch (Exception) { }
+								try { FileReference.Delete(EditorReceiptPath); } catch (Exception) { }
 							}
 						}
 					}
@@ -1080,7 +1068,6 @@ public enum LatestChangeType
 
 				UpdateSyncActionCheckboxes();
 				Refresh();
-				UpdateSyncConfig(Workspace.CurrentChangeNumber, Workspace.CurrentSyncFilterHash);
 				UpdateStatusPanel();
 				DesiredTaskbarState = Tuple.Create(TaskbarState.NoProgress, 0.0f);
 				Owner.UpdateProgress();
@@ -1096,26 +1083,6 @@ public enum LatestChangeType
 		{
 			UpdateTimer.Stop();
 
-			UpdateSyncConfig(Workspace.CurrentChangeNumber, Workspace.CurrentSyncFilterHash);
-
-			if (Result == WorkspaceUpdateResult.Success && Context.Options.HasFlag(WorkspaceUpdateOptions.SyncSingleChange))
-			{
-				WorkspaceState.AdditionalChangeNumbers.Add(Context.ChangeNumber);
-			}
-
-			if (Result == WorkspaceUpdateResult.Success && WorkspaceState.ExpandedArchiveTypes != null)
-			{
-				WorkspaceState.ExpandedArchiveTypes = WorkspaceState.ExpandedArchiveTypes.Except(Context.ArchiveTypeToArchive.Where(x => x.Value == null).Select(x => x.Key)).ToArray();
-			}
-
-			WorkspaceState.LastSyncChangeNumber = Context.ChangeNumber;
-			WorkspaceState.LastSyncResult = Result;
-			WorkspaceState.LastSyncResultMessage = ResultMessage;
-			WorkspaceState.LastSyncTime = DateTime.UtcNow;
-			WorkspaceState.LastSyncDurationSeconds = (int)(WorkspaceState.LastSyncTime.Value - Context.StartTime).TotalSeconds;
-			WorkspaceState.LastBuiltChangeNumber = Workspace.LastBuiltChangeNumber;
-			WorkspaceState.Save();
-
 			if (Result == WorkspaceUpdateResult.FilesToResolve)
 			{
 				MessageBox.Show("You have files to resolve after syncing your workspace. Please check P4.");
@@ -3050,10 +3017,10 @@ public enum LatestChangeType
 			{
 				BuildConfig EditorBuildConfig = GetEditorBuildConfig();
 
-				List<string> ReceiptPaths = GetEditorReceiptPaths(EditorBuildConfig);
+				List<FileReference> ReceiptPaths = GetEditorReceiptPaths(EditorBuildConfig);
 
-				string EditorExe = GetEditorExePath(EditorBuildConfig);
-				if (ReceiptPaths.Any(x => File.Exists(x)) && File.Exists(EditorExe))
+				FileReference EditorExe = GetEditorExePath(EditorBuildConfig);
+				if (ReceiptPaths.Any(x => FileReference.Exists(x)) && FileReference.Exists(EditorExe))
 				{
 					if (Settings.bEditorArgumentsPrompt && !ModifyEditorArguments())
 					{
@@ -3061,7 +3028,7 @@ public enum LatestChangeType
 					}
 
 					StringBuilder LaunchArguments = new StringBuilder();
-					if (SelectedFileName.EndsWith(".uproject", StringComparison.InvariantCultureIgnoreCase))
+					if (SelectedFileName.HasExtension(".uproject"))
 					{
 						LaunchArguments.AppendFormat("\"{0}\"", SelectedFileName);
 					}
@@ -3077,7 +3044,7 @@ public enum LatestChangeType
 						LaunchArguments.Append(" -debug");
 					}
 
-					if (!Utility.SpawnProcess(EditorExe, LaunchArguments.ToString()))
+					if (!Utility.SpawnProcess(EditorExe.FullName, LaunchArguments.ToString()))
 					{
 						ShowErrorDialog("Unable to spawn {0} {1}", EditorExe, LaunchArguments.ToString());
 					}
@@ -3121,26 +3088,26 @@ public enum LatestChangeType
 			public string Launch { get; set; }
 		}
 
-		private string GetEditorExePath(BuildConfig Config)
+		private FileReference GetEditorExePath(BuildConfig Config)
 		{
 			// Try to read the executable path from the target receipt
-			List<string> ReceiptFileNames = GetEditorReceiptPaths(Config);
-			foreach (string ReceiptFileName in ReceiptFileNames)
+			List<FileReference> ReceiptFileNames = GetEditorReceiptPaths(Config);
+			foreach (FileReference ReceiptFileName in ReceiptFileNames)
 			{
-				if (File.Exists(ReceiptFileName))
+				if (FileReference.Exists(ReceiptFileName))
 				{
 					Log.WriteLine("Reading {0}", ReceiptFileName);
 					try
 					{
-						string Text = File.ReadAllText(ReceiptFileName);
+						string Text = FileReference.ReadAllText(ReceiptFileName);
 						ReceiptJsonObject Receipt = JsonSerializer.Deserialize<ReceiptJsonObject>(Text, Program.DefaultJsonSerializerOptions);
 
 						string LaunchFileName = Receipt.Launch;
 						if (LaunchFileName != null)
 						{
-							LaunchFileName = LaunchFileName.Replace("$(EngineDir)", Path.Combine(BranchDirectoryName, "Engine"));
-							LaunchFileName = LaunchFileName.Replace("$(ProjectDir)", Path.GetDirectoryName(SelectedFileName));
-							return Path.GetFullPath(LaunchFileName);
+							LaunchFileName = LaunchFileName.Replace("$(EngineDir)", DirectoryReference.Combine(BranchDirectoryName, "Engine").FullName);
+							LaunchFileName = LaunchFileName.Replace("$(ProjectDir)", SelectedFileName.Directory.FullName);
+							return new FileReference(LaunchFileName);
 						}
 					}
 					catch (Exception Ex)
@@ -3160,12 +3127,12 @@ public enum LatestChangeType
 			{
 				ExeFileName = String.Format("{0}-Win64-{1}.exe", TargetName, Config.ToString());
 			}
-			return Path.Combine(BranchDirectoryName, "Engine", "Binaries", "Win64", ExeFileName);
+			return FileReference.Combine(BranchDirectoryName, "Engine", "Binaries", "Win64", ExeFileName);
 		}
 
 		private bool WaitForProgramsToFinish()
 		{
-			string[] ProcessFileNames = GetProcessesRunningInWorkspace();
+			FileReference[] ProcessFileNames = GetProcessesRunningInWorkspace();
 			if (ProcessFileNames.Length > 0)
 			{
 				ProgramsRunningWindow ProgramsRunning = new ProgramsRunningWindow(GetProcessesRunningInWorkspace, ProcessFileNames);
@@ -3177,7 +3144,7 @@ public enum LatestChangeType
 			return true;
 		}
 
-		private string[] GetProcessesRunningInWorkspace()
+		private FileReference[] GetProcessesRunningInWorkspace()
 		{
 			HashSet<string> ProcessNames = new HashSet<string>(StringComparer.InvariantCultureIgnoreCase);
 			ProcessNames.Add("Win64\\UE4Editor.exe");
@@ -3201,24 +3168,24 @@ public enum LatestChangeType
 			ProcessNames.Add("DotNET\\UnrealBuildTool\\UnrealBuildTool.exe");
 			ProcessNames.Add("DotNET\\AutomationTool\\AutomationTool.exe");
 
-			List<string> ProcessFileNames = new List<string>();
+			List<FileReference> ProcessFileNames = new List<FileReference>();
 			try
 			{
-				string BinariesRootPath = Path.Combine(Workspace.LocalRootPath, "Engine\\Binaries");
+				DirectoryReference BinariesRootPath = DirectoryReference.Combine(Workspace.Project.LocalRootPath, "Engine\\Binaries");
 
 				foreach (string ProcessName in ProcessNames)
 				{
 					try
 					{
-						string ProcessFilename = Path.Combine(BinariesRootPath, ProcessName);
+						FileReference ProcessFilename = FileReference.Combine(BinariesRootPath, ProcessName);
 
-						if (File.Exists(ProcessFilename))
+						if (FileReference.Exists(ProcessFilename))
 						{
 							try
 							{
 								// Try to open the file to determine whether the executable is running.
 								// We do this so that we can also find processes running under other user sessions, without needing PROCESS_QUERY_INFORMATION access rights.
-								using FileStream TestStream = File.OpenWrite(ProcessFilename);
+								using FileStream TestStream = File.OpenWrite(ProcessFilename.FullName);
 							}
 							catch (IOException)
 							{
@@ -3238,31 +3205,31 @@ public enum LatestChangeType
 			return ProcessFileNames.ToArray();
 		}
 
-		private List<string> GetEditorReceiptPaths(BuildConfig Config)
+		private List<FileReference> GetEditorReceiptPaths(BuildConfig Config)
 		{
 			string ConfigSuffix = (Config == BuildConfig.Development) ? "" : String.Format("-Win64-{0}", Config.ToString());
 
-			List<string> PossiblePaths = new List<string>();
+			List<FileReference> PossiblePaths = new List<FileReference>();
 			string ActualEditorName = TryGetEditorNameFromArchive();
 			if (ActualEditorName != null)
 			{
-				PossiblePaths.Add(Path.Combine(Path.GetDirectoryName(SelectedFileName), "Binaries", "Win64", String.Format("{0}{1}.target", ActualEditorName, ConfigSuffix)));
+				PossiblePaths.Add(FileReference.Combine(SelectedFileName.Directory, "Binaries", "Win64", String.Format("{0}{1}.target", ActualEditorName, ConfigSuffix)));
 			}
-			else if (SelectedFileName.EndsWith(".uproject", StringComparison.InvariantCultureIgnoreCase))
+			else if (SelectedFileName.HasExtension(".uproject"))
 			{
 				if (bIsEnterpriseProject)
 				{
-					PossiblePaths.Add(Path.Combine(Path.GetDirectoryName(SelectedFileName), "Binaries", "Win64", String.Format("StudioEditor{0}.target", ConfigSuffix)));
+					PossiblePaths.Add(FileReference.Combine(SelectedFileName.Directory, "Binaries", "Win64", String.Format("StudioEditor{0}.target", ConfigSuffix)));
 				}
 				else
 				{
-					PossiblePaths.Add(Path.Combine(Path.GetDirectoryName(SelectedFileName), "Binaries", "Win64", String.Format("{0}{1}.target", GetDefaultEditorTargetName(), ConfigSuffix)));
-					PossiblePaths.Add(Path.Combine(Workspace.LocalRootPath, "Engine", "Binaries", "Win64", String.Format("{0}{1}.target", GetDefaultEditorTargetName(), ConfigSuffix)));
+					PossiblePaths.Add(FileReference.Combine(SelectedFileName.Directory, "Binaries", "Win64", String.Format("{0}{1}.target", GetDefaultEditorTargetName(), ConfigSuffix)));
+					PossiblePaths.Add(FileReference.Combine(Workspace.Project.LocalRootPath, "Engine", "Binaries", "Win64", String.Format("{0}{1}.target", GetDefaultEditorTargetName(), ConfigSuffix)));
 				}
 			}
 			else
 			{
-				PossiblePaths.Add(Path.Combine(Path.GetDirectoryName(SelectedFileName), "Engine", "Binaries", "Win64", String.Format("{0}{1}.target", GetDefaultEditorTargetName(), ConfigSuffix)));
+				PossiblePaths.Add(FileReference.Combine(SelectedFileName.Directory, "Engine", "Binaries", "Win64", String.Format("{0}{1}.target", GetDefaultEditorTargetName(), ConfigSuffix)));
 			}
 			return PossiblePaths;
 		}
@@ -3360,14 +3327,14 @@ public enum LatestChangeType
 
 		void RunTool(ToolDefinition Tool, ToolLink Link)
 		{
-			string ToolDir = Owner.ToolUpdateMonitor.GetToolPath(Tool.Name);
+			DirectoryReference ToolDir = Owner.ToolUpdateMonitor.GetToolPath(Tool.Name);
 
 			Dictionary<string, string> Variables = GetWorkspaceVariables(Workspace.CurrentChangeNumber);
-			Variables["ToolDir"] = ToolDir;
+			Variables["ToolDir"] = ToolDir.FullName;
 
 			string FileName = Utility.ExpandVariables(Link.FileName, Variables);
 			string Arguments = Utility.ExpandVariables(Link.Arguments, Variables);
-			string WorkingDir = Utility.ExpandVariables(Link.WorkingDir ?? ToolDir, Variables);
+			string WorkingDir = Utility.ExpandVariables(Link.WorkingDir ?? ToolDir.FullName, Variables);
 
 			SafeProcessStart(FileName, Arguments, WorkingDir);
 		}
@@ -3414,7 +3381,7 @@ public enum LatestChangeType
 				// Project
 				StatusLine ProjectLine = new StatusLine();
 				ProjectLine.AddText(String.Format("Opened "));
-				ProjectLine.AddLink(SelectedFileName + " \u25BE", FontStyle.Regular, (P, R) => { SelectRecentProject(R); });
+				ProjectLine.AddLink(SelectedFileName.FullName + " \u25BE", FontStyle.Regular, (P, R) => { SelectRecentProject(R); });
 				ProjectLine.AddText("  |  ");
 				ProjectLine.AddLink("Settings...", FontStyle.Regular, (P, R) => { Owner.EditSelectedProject(this); });
 				Lines.Add(ProjectLine);
@@ -3529,7 +3496,7 @@ public enum LatestChangeType
 					ProgramsLine.AddText("  |  ");
 				}
 
-				ProgramsLine.AddLink("Windows Explorer", FontStyle.Regular, () => { SafeProcessStart("explorer.exe", String.Format("\"{0}\"", Path.GetDirectoryName(SelectedFileName))); });
+				ProgramsLine.AddLink("Windows Explorer", FontStyle.Regular, () => { SafeProcessStart("explorer.exe", String.Format("\"{0}\"", SelectedFileName.Directory.FullName)); });
 
 				if (GetDefaultIssueFilter() != null)
 				{
@@ -4909,17 +4876,17 @@ public enum LatestChangeType
 		private void OpenSolution()
 		{
 			string MasterProjectName = "UE4";
-			if (File.Exists(Path.Combine(BranchDirectoryName, "UE5.sln")))
+			if (FileReference.Exists(FileReference.Combine(BranchDirectoryName, "UE5.sln")))
 			{
 				MasterProjectName = "UE5";
 			}
 
-			string MasterProjectNameFileName = Path.Combine(BranchDirectoryName, "Engine", "Intermediate", "ProjectFiles", "MasterProjectName.txt");
-			if (File.Exists(MasterProjectNameFileName))
+			FileReference MasterProjectNameFileName = FileReference.Combine(BranchDirectoryName, "Engine", "Intermediate", "ProjectFiles", "MasterProjectName.txt");
+			if (FileReference.Exists(MasterProjectNameFileName))
 			{
 				try
 				{
-					MasterProjectName = File.ReadAllText(MasterProjectNameFileName).Trim();
+					MasterProjectName = FileReference.ReadAllText(MasterProjectNameFileName).Trim();
 				}
 				catch (Exception Ex)
 				{
@@ -4927,16 +4894,16 @@ public enum LatestChangeType
 				}
 			}
 
-			string SolutionFileName = Path.Combine(BranchDirectoryName, MasterProjectName + ".sln");
-			if (!File.Exists(SolutionFileName))
+			FileReference SolutionFileName = FileReference.Combine(BranchDirectoryName, MasterProjectName + ".sln");
+			if (!FileReference.Exists(SolutionFileName))
 			{
 				MessageBox.Show(String.Format("Couldn't find solution at {0}", SolutionFileName));
 			}
 			else
 			{
-				ProcessStartInfo StartInfo = new ProcessStartInfo(SolutionFileName);
+				ProcessStartInfo StartInfo = new ProcessStartInfo(SolutionFileName.FullName);
 				StartInfo.UseShellExecute = true;
-				StartInfo.WorkingDirectory = BranchDirectoryName;
+				StartInfo.WorkingDirectory = BranchDirectoryName.FullName;
 				SafeProcessStart(StartInfo);
 			}
 		}
@@ -5134,9 +5101,9 @@ public enum LatestChangeType
 			}
 
 			string[] CombinedSyncFilter = UserSettings.GetCombinedSyncFilter(GetSyncCategories(), Settings.SyncView, Settings.SyncCategories, WorkspaceSettings.SyncView, WorkspaceSettings.SyncCategoriesDict);
-			List<string> SyncPaths = Workspace.GetSyncPaths(WorkspaceSettings.bSyncAllProjects ?? Settings.bSyncAllProjects, CombinedSyncFilter);
+			List<string> SyncPaths = Workspace.GetSyncPaths(Workspace.Project, WorkspaceSettings.bSyncAllProjects ?? Settings.bSyncAllProjects, CombinedSyncFilter);
 
-			CleanWorkspaceWindow.DoClean(ParentForm, Workspace.Perforce, BranchDirectoryName, Workspace.ClientRootPath, SyncPaths, ExtraSafeToDeleteFolders.Split('\n'), ExtraSafeToDeleteExtensions.Split('\n'), Log);
+			CleanWorkspaceWindow.DoClean(ParentForm, Workspace.Perforce, BranchDirectoryName, Workspace.Project.ClientRootPath, SyncPaths, ExtraSafeToDeleteFolders.Split('\n'), ExtraSafeToDeleteExtensions.Split('\n'), Log);
 		}
 
 		private void UpdateBuildSteps()
@@ -5251,7 +5218,7 @@ public enum LatestChangeType
 						string ToolName = Owner.ToolUpdateMonitor.GetToolName(Step.ToolId);
 						if (ToolName != null)
 						{
-							Variables["ToolDir"] = Owner.ToolUpdateMonitor.GetToolPath(ToolName);
+							Variables["ToolDir"] = Owner.ToolUpdateMonitor.GetToolPath(ToolName).FullName;
 						}
 					}
 
@@ -5272,11 +5239,11 @@ public enum LatestChangeType
 			Variables.Add("Stream", StreamName);
 			Variables.Add("Change", ChangeNumber.ToString());
 			Variables.Add("ClientName", ClientName);
-			Variables.Add("BranchDir", BranchDirectoryName);
-			Variables.Add("ProjectDir", Path.GetDirectoryName(SelectedFileName));
-			Variables.Add("ProjectFile", SelectedFileName);
-			Variables.Add("UE4EditorExe", GetEditorExePath(EditorBuildConfig));
-			Variables.Add("UE4EditorCmdExe", GetEditorExePath(EditorBuildConfig).Replace(".exe", "-Cmd.exe"));
+			Variables.Add("BranchDir", BranchDirectoryName.FullName);
+			Variables.Add("ProjectDir", SelectedFileName.Directory.FullName);
+			Variables.Add("ProjectFile", SelectedFileName.FullName);
+			Variables.Add("UE4EditorExe", GetEditorExePath(EditorBuildConfig).FullName);
+			Variables.Add("UE4EditorCmdExe", GetEditorExePath(EditorBuildConfig).FullName.Replace(".exe", "-Cmd.exe"));
 			Variables.Add("UE4EditorConfig", EditorBuildConfig.ToString());
 			Variables.Add("UE4EditorDebugArg", (EditorBuildConfig == BuildConfig.Debug || EditorBuildConfig == BuildConfig.DebugGame) ? " -debug" : "");
 			Variables.Add("UseIncrementalBuilds", "1");
@@ -5294,9 +5261,9 @@ public enum LatestChangeType
 			{
 				// Find all the target names for this project
 				List<string> TargetNames = new List<string>();
-				if (!String.IsNullOrEmpty(SelectedFileName) && SelectedFileName.EndsWith(".uproject", StringComparison.InvariantCultureIgnoreCase))
+				if (SelectedFileName != null && SelectedFileName.HasExtension(".uproject"))
 				{
-					DirectoryInfo SourceDirectory = new DirectoryInfo(Path.Combine(Path.GetDirectoryName(SelectedFileName), "Source"));
+					DirectoryInfo SourceDirectory = DirectoryReference.Combine(SelectedFileName.Directory, "Source").ToDirectoryInfo();
 					if (SourceDirectory.Exists)
 					{
 						foreach (FileInfo TargetFile in SourceDirectory.EnumerateFiles("*.target.cs", SearchOption.TopDirectoryOnly))
@@ -5383,7 +5350,7 @@ public enum LatestChangeType
 		private Dictionary<Guid, ConfigObject> GetDefaultBuildStepObjects()
 		{
 			string ProjectArgument = "";
-			if (SelectedFileName.EndsWith(".uproject", StringComparison.InvariantCultureIgnoreCase))
+			if (SelectedFileName.HasExtension(".uproject"))
 			{
 				ProjectArgument = String.Format("\"{0}\"", SelectedFileName);
 			}
@@ -5393,7 +5360,7 @@ public enum LatestChangeType
 			{
 				ActualEditorTargetName = EditorTargetName;
 			}
-			else if (SelectedFileName.EndsWith(".uproject", StringComparison.InvariantCultureIgnoreCase) && bIsEnterpriseProject)
+			else if (SelectedFileName.HasExtension(".uproject") && bIsEnterpriseProject)
 			{
 				ActualEditorTargetName = "StudioEditor";
 			}
@@ -5460,7 +5427,7 @@ public enum LatestChangeType
 			StringBuilder DiagnosticsText = new StringBuilder();
 			DiagnosticsText.AppendFormat("Application version: {0}\n", Program.GetVersionString());
 			DiagnosticsText.AppendFormat("Synced from: {0}\n", Program.SyncVersion ?? "(unknown)");
-			DiagnosticsText.AppendFormat("Selected file: {0}\n", (SelectedFileName == null) ? "(none)" : SelectedFileName);
+			DiagnosticsText.AppendFormat("Selected file: {0}\n", (SelectedFileName == null) ? "(none)" : SelectedFileName.FullName);
 			if (Workspace != null)
 			{
 				DiagnosticsText.AppendFormat("P4 server: {0}\n", Workspace.Perforce.ServerAndPort ?? "(default)");
@@ -6113,17 +6080,17 @@ public enum LatestChangeType
 			UpdateTintColor();
 		}
 
-		private string GetTintColorFileName()
+		private FileReference GetTintColorFileName()
 		{
-			return Path.Combine(Path.GetDirectoryName(SelectedFileName), "Saved", "Config", "Windows", "EditorPerProjectUserSettings.ini");
+			return FileReference.Combine(SelectedFileName.Directory, "Saved", "Config", "Windows", "EditorPerProjectUserSettings.ini");
 		}
 
 		private Color? GetTintColor()
 		{
 			try
 			{
-				string FileName = GetTintColorFileName();
-				if (!File.Exists(FileName))
+				FileReference FileName = GetTintColorFileName();
+				if (!FileReference.Exists(FileName))
 				{
 					return null;
 				}
diff --git a/Engine/Source/Programs/UnrealGameSync/UnrealGameSync/DetectProjectSettingsTask.cs b/Engine/Source/Programs/UnrealGameSync/UnrealGameSync/DetectProjectSettingsTask.cs
--- a/Engine/Source/Programs/UnrealGameSync/UnrealGameSync/DetectProjectSettingsTask.cs
+++ b/Engine/Source/Programs/UnrealGameSync/UnrealGameSync/DetectProjectSettingsTask.cs
@@ -1,5 +1,6 @@
 // Copyright Epic Games, Inc. All Rights Reserved.
 
+using EpicGames.Core;
 using System;
 using System.Collections.Generic;
 using System.Drawing;
@@ -14,22 +15,22 @@ namespace UnrealGameSync
 	{
 		public UserSelectedProjectSettings SelectedProject;
 		public PerforceConnection PerforceClient;
-		public string NewSelectedFileName;
+		public FileReference NewSelectedFileName;
 		public string NewSelectedProjectIdentifier;
 		public string NewProjectEditorTarget;
 		public string BranchClientPath;
-		public string BranchDirectoryName;
+		public DirectoryReference BranchDirectoryName;
 		public string NewSelectedClientFileName;
 		public string StreamName;
 		public Image ProjectLogo;
-		public string DataFolder;
-		public string CacheFolder;
+		public DirectoryReference DataFolder;
+		public DirectoryReference CacheFolder;
 		public bool bIsEnterpriseProject;
 		public ConfigFile LatestProjectConfigFile;
-		public List<KeyValuePair<string, DateTime>> LocalConfigFiles;
+		public List<KeyValuePair<FileReference, DateTime>> LocalConfigFiles;
 		TextWriter Log;
 
-		public DetectProjectSettingsTask(UserSelectedProjectSettings SelectedProject, string InDataFolder, string InCacheFolder, TextWriter InLog)
+		public DetectProjectSettingsTask(UserSelectedProjectSettings SelectedProject, DirectoryReference InDataFolder, DirectoryReference InCacheFolder, TextWriter InLog)
 		{
 			this.SelectedProject = SelectedProject;
 			DataFolder = InDataFolder;
@@ -85,18 +86,16 @@ namespace UnrealGameSync
 				PerforceClient = new PerforceConnection(Perforce.UserName, ClientName, Perforce.ServerAndPort);
 
 				// Figure out the path on the client. Use the cached location if it's valid.
-				if(SelectedProject.LocalPath != null && File.Exists(SelectedProject.LocalPath))
+				string LocalPath = SelectedProject.LocalPath;
+				if (LocalPath == null || !File.Exists(LocalPath))
 				{
-					NewSelectedFileName = SelectedProject.LocalPath;
-				}
-				else
-				{
-					if(!PerforceClient.ConvertToLocalPath(NewSelectedClientFileName, out NewSelectedFileName, Log))
+					if (!PerforceClient.ConvertToLocalPath(NewSelectedClientFileName, out LocalPath, Log))
 					{
 						ErrorMessage = String.Format("Couldn't get client path for {0}", NewSelectedClientFileName);
 						return false;
 					}
 				}
+				NewSelectedFileName = new FileReference(LocalPath);
 			}
 			else
 			{
@@ -109,10 +108,10 @@ namespace UnrealGameSync
 				}
 
 				// Use the path as the selected filename
-				NewSelectedFileName = SelectedProject.LocalPath;
+				NewSelectedFileName = new FileReference(SelectedProject.LocalPath);
 
 				// Make sure the project exists
-				if(!File.Exists(SelectedProject.LocalPath))
+				if (!FileReference.Exists(NewSelectedFileName))
 				{
 					ErrorMessage = String.Format("{0} does not exist.", SelectedProject.LocalPath);
 					return false;
@@ -153,7 +152,7 @@ namespace UnrealGameSync
 				// Check there's only one client
 				if(CandidateClients.Count > 1)
 				{
-					ErrorMessage = String.Format("Found multiple workspaces containing {0}:\n\n{1}\n\nCannot determine which to use.", Path.GetFileName(NewSelectedFileName), String.Join("\n", CandidateClients.Select(x => x.ClientName)));
+					ErrorMessage = String.Format("Found multiple workspaces containing {0}:\n\n{1}\n\nCannot determine which to use.", Path.GetFileName(NewSelectedFileName.GetFileName()), String.Join("\n", CandidateClients.Select(x => x.ClientName)));
 					return false;
 				}
 
@@ -161,21 +160,18 @@ namespace UnrealGameSync
 				PerforceClient = CandidateClients[0];
 
 				// Get the client path for the project file
-				if(!PerforceClient.ConvertToClientPath(NewSelectedFileName, out NewSelectedClientFileName, Log))
+				if(!PerforceClient.ConvertToClientPath(NewSelectedFileName.FullName, out NewSelectedClientFileName, Log))
 				{
 					ErrorMessage = String.Format("Couldn't get client path for {0}", NewSelectedFileName);
 					return false;
 				}
 			}
 
-			// Normalize the filename
-			NewSelectedFileName = Path.GetFullPath(NewSelectedFileName).Replace('/', Path.DirectorySeparatorChar);
-
 			// Make sure the path case is correct. This can cause UBT intermediates to be out of date if the case mismatches.
-			NewSelectedFileName = Utility.GetPathWithCorrectCase(new FileInfo(NewSelectedFileName));
+			NewSelectedFileName = FileReference.FindCorrectCase(NewSelectedFileName);
 
 			// Update the selected project with all the data we've found
-			SelectedProject = new UserSelectedProjectSettings(SelectedProject.ServerAndPort, SelectedProject.UserName, SelectedProject.Type, NewSelectedClientFileName, NewSelectedFileName);
+			SelectedProject = new UserSelectedProjectSettings(SelectedProject.ServerAndPort, SelectedProject.UserName, SelectedProject.Type, NewSelectedClientFileName, NewSelectedFileName.FullName);
 
 			// Figure out where the engine is in relation to it
 			int EndIdx = NewSelectedClientFileName.Length - 1;
@@ -202,7 +198,7 @@ namespace UnrealGameSync
 						}
 
 						BranchClientPath = NewSelectedClientFileName.Substring(0, EndIdx);
-						BranchDirectoryName = Path.GetFullPath(Path.Combine(Path.GetDirectoryName(FileRecords[0].ClientPath), "..", ".."));
+						BranchDirectoryName = new DirectoryReference(Path.Combine(Path.GetDirectoryName(FileRecords[0].ClientPath), "..", ".."));
 						break;
 					}
 				}
@@ -210,7 +206,7 @@ namespace UnrealGameSync
 			Log.WriteLine("Found branch root at {0}", BranchClientPath);
 
 			// Find the editor target for this project
-			if(NewSelectedFileName.EndsWith(".uproject", StringComparison.InvariantCultureIgnoreCase))
+			if(NewSelectedFileName.HasExtension(".uproject"))
 			{
 				List<PerforceFileRecord> Files;
 				if(PerforceClient.FindFiles(PerforceUtils.GetClientOrDepotDirectoryName(NewSelectedClientFileName) + "/Source/*Editor.Target.cs", out Files, Log) && Files.Count >= 1)
@@ -263,15 +259,15 @@ namespace UnrealGameSync
 			}
 
 			// Read the project logo
-			if (NewSelectedFileName.EndsWith(".uproject", StringComparison.InvariantCultureIgnoreCase))
+			if (NewSelectedFileName.HasExtension(".uproject"))
 			{
-				string LogoFileName = Path.Combine(Path.GetDirectoryName(NewSelectedFileName), "Build", "UnrealGameSync.png");
-				if(File.Exists(LogoFileName))
+				FileReference LogoFileName = FileReference.Combine(NewSelectedFileName.Directory, "Build", "UnrealGameSync.png");
+				if(FileReference.Exists(LogoFileName))
 				{
 					try
 					{
 						// Duplicate the image, otherwise we'll leave the file locked
-						using(Image Image = Image.FromFile(LogoFileName))
+						using(Image Image = Image.FromFile(LogoFileName.FullName))
 						{
 							ProjectLogo = new Bitmap(Image);
 						}
@@ -287,9 +283,9 @@ namespace UnrealGameSync
 			if(NewSelectedClientFileName.EndsWith(".uproject", StringComparison.InvariantCultureIgnoreCase))
 			{
 				string Text;
-				if(File.Exists(NewSelectedFileName))
+				if(FileReference.Exists(NewSelectedFileName))
 				{
-					Text = File.ReadAllText(NewSelectedFileName);
+					Text = FileReference.ReadAllText(NewSelectedFileName);
 				}
 				else
 				{
@@ -305,7 +301,7 @@ namespace UnrealGameSync
 			}
 
 			// Make sure the drive containing the project exists, to prevent other errors down the line
-			string PathRoot = Path.GetPathRoot(NewSelectedFileName);
+			string PathRoot = Path.GetPathRoot(NewSelectedFileName.FullName);
 			if (!Directory.Exists(PathRoot))
 			{
 				ErrorMessage = String.Format("Path '{0}' is invalid", NewSelectedFileName);
@@ -313,7 +309,7 @@ namespace UnrealGameSync
 			}
 
 			// Read the initial config file
-			LocalConfigFiles = new List<KeyValuePair<string, DateTime>>();
+			LocalConfigFiles = new List<KeyValuePair<FileReference, DateTime>>();
 			LatestProjectConfigFile = PerforceMonitor.ReadProjectConfigFile(PerforceClient, BranchClientPath, NewSelectedClientFileName, CacheFolder, LocalConfigFiles, Log);
 
 			// Run any event hooks
@@ -343,7 +339,7 @@ namespace UnrealGameSync
 					// Require either a username or host name match
 					if((String.IsNullOrEmpty(Client.Host) || String.Compare(Client.Host, HostName, StringComparison.InvariantCultureIgnoreCase) == 0) && (String.IsNullOrEmpty(Client.Owner) || String.Compare(Client.Owner, UserName, StringComparison.InvariantCultureIgnoreCase) == 0))
 					{
-						if(!Utility.SafeIsFileUnderDirectory(NewSelectedFileName, Client.Root))
+						if(!Utility.SafeIsFileUnderDirectory(NewSelectedFileName.FullName, Client.Root))
 						{
 							Log.WriteLine("Rejecting {0} due to root mismatch ({1})", Client.Name, Client.Root);
 							continue;
@@ -352,14 +348,14 @@ namespace UnrealGameSync
 						PerforceConnection CandidateClient = new PerforceConnection(UserName, Client.Name, ServerAndPort);
 
 						bool bFileExists;
-						if(!CandidateClient.FileExists(NewSelectedFileName, out bFileExists, Log) || !bFileExists)
+						if(!CandidateClient.FileExists(NewSelectedFileName.FullName, out bFileExists, Log) || !bFileExists)
 						{
 							Log.WriteLine("Rejecting {0} due to file not existing in workspace", Client.Name);
 							continue;
 						}
 
 						List<PerforceFileRecord> Records;
-						if(!CandidateClient.Stat(NewSelectedFileName, out Records, Log))
+						if(!CandidateClient.Stat(NewSelectedFileName.FullName, out Records, Log))
 						{
 							Log.WriteLine("Rejecting {0} due to {1} not in depot", Client.Name, NewSelectedFileName);
 							continue;
diff --git a/Engine/Source/Programs/UnrealGameSync/UnrealGameSync/EventMonitor.cs b/Engine/Source/Programs/UnrealGameSync/UnrealGameSync/EventMonitor.cs
--- a/Engine/Source/Programs/UnrealGameSync/UnrealGameSync/EventMonitor.cs
+++ b/Engine/Source/Programs/UnrealGameSync/UnrealGameSync/EventMonitor.cs
@@ -9,6 +9,7 @@ using System.IO;
 using System.Linq;
 using System.Text.Json;
 using System.Threading;
+using EpicGames.Core;
 
 namespace UnrealGameSync
 {
@@ -238,7 +239,7 @@ namespace UnrealGameSync
 
 		public event Action OnUpdatesReady;
 
-		public EventMonitor(string InApiUrl, string InProject, string InCurrentUserName, string InLogFileName)
+		public EventMonitor(string InApiUrl, string InProject, string InCurrentUserName, FileReference InLogFileName)
 		{
 			ApiUrl = InApiUrl;
 			Project = InProject;
diff --git a/Engine/Source/Programs/UnrealGameSync/UnrealGameSync/Forms/BuildStepWindow.cs b/Engine/Source/Programs/UnrealGameSync/UnrealGameSync/Forms/BuildStepWindow.cs
--- a/Engine/Source/Programs/UnrealGameSync/UnrealGameSync/Forms/BuildStepWindow.cs
+++ b/Engine/Source/Programs/UnrealGameSync/UnrealGameSync/Forms/BuildStepWindow.cs
@@ -1,5 +1,6 @@
 // Copyright Epic Games, Inc. All Rights Reserved.
 
+using EpicGames.Core;
 using System;
 using System.Collections.Generic;
 using System.ComponentModel;
@@ -17,15 +18,15 @@ namespace UnrealGameSync
 	{
 		BuildStep Step;
 		List<string> TargetNames;
-		string BaseDirectoryPrefix;
+		DirectoryReference BaseDirectory;
 		IReadOnlyDictionary<string, string> Variables;
 		VariablesWindow VariablesWindow;
 
-		public BuildStepWindow(BuildStep InTask, List<string> InTargetNames, string InBaseDirectory, IReadOnlyDictionary<string, string> InVariables)
+		public BuildStepWindow(BuildStep InTask, List<string> InTargetNames, DirectoryReference InBaseDirectory, IReadOnlyDictionary<string, string> InVariables)
 		{
 			Step = InTask;
 			TargetNames = InTargetNames;
-			BaseDirectoryPrefix = Path.GetFullPath(InBaseDirectory) + Path.DirectorySeparatorChar;
+			BaseDirectory = InBaseDirectory;
 			Variables = InVariables;
 
 			InitializeComponent();
@@ -212,20 +213,20 @@ namespace UnrealGameSync
 			}
 			else
 			{
-				return Path.Combine(BaseDirectoryPrefix, FileName);
+				return FileReference.Combine(BaseDirectory, FileName).FullName;
 			}
 		}
 
 		private string RemoveBaseDirectory(string FileName)
 		{
-			string FullFileName = Path.GetFullPath(FileName);
-			if(FullFileName.StartsWith(BaseDirectoryPrefix, StringComparison.InvariantCultureIgnoreCase))
+			FileReference FullFileName = new FileReference(FileName);
+			if(FullFileName.IsUnderDirectory(BaseDirectory))
 			{
-				return FullFileName.Substring(BaseDirectoryPrefix.Length);
+				return FullFileName.MakeRelativeTo(BaseDirectory);
 			}
 			else
 			{
-				return FullFileName;
+				return FileName;
 			}
 		}
 
diff --git a/Engine/Source/Programs/UnrealGameSync/UnrealGameSync/Forms/CleanWorkspaceWindow.cs b/Engine/Source/Programs/UnrealGameSync/UnrealGameSync/Forms/CleanWorkspaceWindow.cs
--- a/Engine/Source/Programs/UnrealGameSync/UnrealGameSync/Forms/CleanWorkspaceWindow.cs
+++ b/Engine/Source/Programs/UnrealGameSync/UnrealGameSync/Forms/CleanWorkspaceWindow.cs
@@ -1,5 +1,6 @@
 // Copyright Epic Games, Inc. All Rights Reserved.
 
+using EpicGames.Core;
 using System;
 using System.Collections.Generic;
 using System.ComponentModel;
@@ -114,10 +115,10 @@ namespace UnrealGameSync
 			InitializeComponent();
 		}
 
-		public static void DoClean(IWin32Window Owner, PerforceConnection PerforceClient, string LocalRootPath, string ClientRootPath, IReadOnlyList<string> SyncPaths, string[] ExtraSafeToDeleteFolders, string[] ExtraSafeToDeleteExtensions, TextWriter Log)
+		public static void DoClean(IWin32Window Owner, PerforceConnection PerforceClient, DirectoryReference LocalRootPath, string ClientRootPath, IReadOnlyList<string> SyncPaths, string[] ExtraSafeToDeleteFolders, string[] ExtraSafeToDeleteExtensions, TextWriter Log)
 		{
 			// Figure out which folders to clean
-			FolderToClean RootFolderToClean = new FolderToClean(new DirectoryInfo(LocalRootPath));
+			FolderToClean RootFolderToClean = new FolderToClean(LocalRootPath.ToDirectoryInfo());
 			using(FindFoldersToCleanTask QueryWorkspace = new FindFoldersToCleanTask(PerforceClient, RootFolderToClean, ClientRootPath, SyncPaths, Log))
 			{
 				string ErrorMessage;
diff --git a/Engine/Source/Programs/UnrealGameSync/UnrealGameSync/Forms/DiagnosticsWindow.cs b/Engine/Source/Programs/UnrealGameSync/UnrealGameSync/Forms/DiagnosticsWindow.cs
--- a/Engine/Source/Programs/UnrealGameSync/UnrealGameSync/Forms/DiagnosticsWindow.cs
+++ b/Engine/Source/Programs/UnrealGameSync/UnrealGameSync/Forms/DiagnosticsWindow.cs
@@ -13,10 +13,10 @@ namespace UnrealGameSync
 {
 	public partial class DiagnosticsWindow : Form
 	{
-		string DataFolder;
+		DirectoryReference DataFolder;
 		List<FileReference> ExtraFiles;
 
-		public DiagnosticsWindow(string InDataFolder, string InDiagnosticsText, IEnumerable<FileReference> InExtraFiles)
+		public DiagnosticsWindow(DirectoryReference InDataFolder, string InDiagnosticsText, IEnumerable<FileReference> InExtraFiles)
 		{
 			InitializeComponent();
 			DataFolder = InDataFolder;
@@ -26,7 +26,7 @@ namespace UnrealGameSync
 
 		private void ViewLogsButton_Click(object sender, EventArgs e)
 		{
-			Process.Start("explorer.exe", DataFolder);
+			Process.Start("explorer.exe", DataFolder.FullName);
 		}
 
 		private void SaveButton_Click(object sender, EventArgs e)
@@ -37,10 +37,10 @@ namespace UnrealGameSync
 			Dialog.FileName = Path.Combine(Dialog.InitialDirectory, "UGS-Diagnostics.zip");
 			if(Dialog.ShowDialog() == DialogResult.OK)
 			{
-				string DiagnosticsFileName = Path.Combine(DataFolder, "Diagnostics.txt");
+				FileReference DiagnosticsFileName = FileReference.Combine(DataFolder, "Diagnostics.txt");
 				try
 				{
-					File.WriteAllLines(DiagnosticsFileName, DiagnosticsTextBox.Lines);
+					FileReference.WriteAllLines(DiagnosticsFileName, DiagnosticsTextBox.Lines);
 				}
 				catch(Exception Ex)
 				{
@@ -53,13 +53,13 @@ namespace UnrealGameSync
 				{
 					using (ZipArchive Zip = new ZipArchive(File.OpenWrite(ZipFileName), ZipArchiveMode.Create))
 					{
-						foreach (string FileName in Directory.EnumerateFiles(DataFolder))
+						foreach (FileReference FileName in DirectoryReference.EnumerateFiles(DataFolder))
 						{
-							if (!FileName.EndsWith(".exe", StringComparison.InvariantCultureIgnoreCase) && !FileName.EndsWith(".dll", StringComparison.InvariantCultureIgnoreCase))
+							if (!FileName.HasExtension(".exe") && !FileName.HasExtension(".dll"))
 							{
-								using (FileStream InputStream = File.Open(FileName, FileMode.Open, FileAccess.Read, FileShare.ReadWrite))
+								using (FileStream InputStream = FileReference.Open(FileName, FileMode.Open, FileAccess.Read, FileShare.ReadWrite))
 								{
-									ZipArchiveEntry Entry = Zip.CreateEntry(Path.GetFileName(FileName));
+									ZipArchiveEntry Entry = Zip.CreateEntry(FileName.MakeRelativeTo(DataFolder).Replace('\\', '/'));
 									using (Stream OutputStream = Entry.Open())
 									{
 										InputStream.CopyTo(OutputStream);
diff --git a/Engine/Source/Programs/UnrealGameSync/UnrealGameSync/Forms/MainWindow.cs b/Engine/Source/Programs/UnrealGameSync/UnrealGameSync/Forms/MainWindow.cs
--- a/Engine/Source/Programs/UnrealGameSync/UnrealGameSync/Forms/MainWindow.cs
+++ b/Engine/Source/Programs/UnrealGameSync/UnrealGameSync/Forms/MainWindow.cs
@@ -14,6 +14,7 @@ using System.Text;
 using System.Diagnostics;
 using UnrealGameSync.Forms;
 using System.Text.RegularExpressions;
+using EpicGames.Core;
 
 namespace UnrealGameSync
 {
@@ -76,8 +77,8 @@ namespace UnrealGameSync
 		List<WorkspaceIssueMonitor> WorkspaceIssueMonitors = new List<WorkspaceIssueMonitor>();
 
 		string ApiUrl;
-		string DataFolder;
-		string CacheFolder;
+		DirectoryReference DataFolder;
+		DirectoryReference CacheFolder;
 		PerforceConnection DefaultConnection;
 		LineBasedTextWriter Log;
 		UserSettings Settings;
@@ -110,7 +111,7 @@ namespace UnrealGameSync
 
 		NetCoreWindow NetCoreWindow;
 
-		public MainWindow(UpdateMonitor InUpdateMonitor, string InApiUrl, string InDataFolder, string InCacheFolder, bool bInRestoreStateOnLoad, string InOriginalExecutableFileName, bool bInUnstable, DetectProjectSettingsResult[] StartupProjects, PerforceConnection InDefaultConnection, LineBasedTextWriter InLog, UserSettings InSettings, string InUri, OIDCTokenManager InOidcTokenManager)
+		public MainWindow(UpdateMonitor InUpdateMonitor, string InApiUrl, DirectoryReference InDataFolder, DirectoryReference InCacheFolder, bool bInRestoreStateOnLoad, string InOriginalExecutableFileName, bool bInUnstable, DetectProjectSettingsResult[] StartupProjects, PerforceConnection InDefaultConnection, LineBasedTextWriter InLog, UserSettings InSettings, string InUri, OIDCTokenManager InOidcTokenManager)
 		{
 			Log = InLog;
 			Log.WriteLine("Opening Main Window for {0} projects. Last Project {1}", StartupProjects.Length, InSettings.LastProject);
@@ -181,7 +182,7 @@ namespace UnrealGameSync
 				Text += $" {Program.GetVersionString()} (UNSTABLE)";
 			}
 
-			AutomationLog = new TimestampLogWriter(new BoundedLogWriter(Path.Combine(DataFolder, "Automation.log")));
+			AutomationLog = new TimestampLogWriter(new BoundedLogWriter(FileReference.Combine(DataFolder, "Automation.log")));
 			AutomationServer = new AutomationServer(Request => { MainThreadSynchronizationContext.Post(Obj => PostAutomationRequest(Request), null); }, AutomationLog, InUri);
 
 			// Allow creating controls from now on
@@ -298,7 +299,7 @@ namespace UnrealGameSync
 			if (Result == WorkspaceUpdateResult.Success && Command != null)
 			{
 				string CmdExe = Environment.GetEnvironmentVariable("COMSPEC");
-				Workspace.ExecCommand("Run build command", "Running build command", CmdExe, String.Format("/c {0}", Command), Workspace.BranchDirectoryName, true);
+				Workspace.ExecCommand("Run build command", "Running build command", CmdExe, String.Format("/c {0}", Command), Workspace.BranchDirectoryName.FullName, true);
 			}
 		}
 
@@ -324,7 +325,7 @@ namespace UnrealGameSync
 
 			if(!bForceSync && Workspace.CanLaunchEditor())
 			{
-				return new AutomationRequestOutput(AutomationRequestResult.Ok, Encoding.UTF8.GetBytes(Workspace.SelectedFileName));
+				return new AutomationRequestOutput(AutomationRequestResult.Ok, Encoding.UTF8.GetBytes(Workspace.SelectedFileName.FullName));
 			}
 
 			Workspace.AddStartupCallback((Control, bCancel) => StartAutomatedSyncAfterStartup(Control, bCancel, Request));
@@ -391,11 +392,11 @@ namespace UnrealGameSync
 			}
 		}
 
-		void CompleteAutomatedSync(WorkspaceUpdateResult Result, string SelectedFileName, AutomationRequest Request)
+		void CompleteAutomatedSync(WorkspaceUpdateResult Result, FileReference SelectedFileName, AutomationRequest Request)
 		{
 			if(Result == WorkspaceUpdateResult.Success)
 			{
-				Request.SetOutput(new AutomationRequestOutput(AutomationRequestResult.Ok, Encoding.UTF8.GetBytes(SelectedFileName)));
+				Request.SetOutput(new AutomationRequestOutput(AutomationRequestResult.Ok, Encoding.UTF8.GetBytes(SelectedFileName.FullName)));
 			}
 			else if(Result == WorkspaceUpdateResult.Canceled)
 			{
@@ -427,7 +428,7 @@ namespace UnrealGameSync
 							string ExistingProjectPath = ClientPath.Substring(SlashIdx);
 							if(String.Compare(ExistingProjectPath, ProjectPath, StringComparison.OrdinalIgnoreCase) == 0)
 							{
-								return new AutomationRequestOutput(AutomationRequestResult.Ok, Encoding.UTF8.GetBytes(ExistingWorkspace.SelectedFileName));
+								return new AutomationRequestOutput(AutomationRequestResult.Ok, Encoding.UTF8.GetBytes(ExistingWorkspace.SelectedFileName.FullName));
 							}
 						}
 					}
@@ -1030,7 +1031,7 @@ namespace UnrealGameSync
 
 		void UpdateRecentProjectsList(DetectProjectSettingsTask DetectedProjectSettings)
 		{
-			Settings.RecentProjects.RemoveAll(x => x.LocalPath == DetectedProjectSettings.NewSelectedFileName);
+			Settings.RecentProjects.RemoveAll(x => x.LocalPath == DetectedProjectSettings.NewSelectedFileName.FullName);
 			Settings.RecentProjects.Insert(0, DetectedProjectSettings.SelectedProject);
 
 			const int MaxRecentProjects = 10;
@@ -1130,7 +1131,7 @@ namespace UnrealGameSync
 					WorkspaceControl Workspace = TabControl.GetTabData(TabIdx) as WorkspaceControl;
 					if(Workspace != null)
 					{
-						if(Workspace.SelectedFileName.Equals(ProjectSettings.NewSelectedFileName, StringComparison.InvariantCultureIgnoreCase))
+						if(Workspace.SelectedFileName == ProjectSettings.NewSelectedFileName)
 						{
 							Log.WriteLine("  Already open in tab {0}", TabIdx);
 							if((Options & OpenProjectOptions.Quiet) == 0)
@@ -1139,9 +1140,9 @@ namespace UnrealGameSync
 							}
 							return TabIdx;
 						}
-						else if(ProjectSettings.NewSelectedFileName.StartsWith(Workspace.BranchDirectoryName + Path.DirectorySeparatorChar, StringComparison.InvariantCultureIgnoreCase))
+						else if(ProjectSettings.NewSelectedFileName.IsUnderDirectory(Workspace.BranchDirectoryName))
 						{
-							if((Options & OpenProjectOptions.Quiet) == 0 && MessageBox.Show(String.Format("{0} is already open under {1}.\n\nWould you like to close it?", Path.GetFileNameWithoutExtension(Workspace.SelectedFileName), Workspace.BranchDirectoryName, Path.GetFileNameWithoutExtension(ProjectSettings.NewSelectedFileName)), "Branch already open", MessageBoxButtons.YesNo) == System.Windows.Forms.DialogResult.Yes)
+							if((Options & OpenProjectOptions.Quiet) == 0 && MessageBox.Show(String.Format("{0} is already open under {1}.\n\nWould you like to close it?", Workspace.SelectedFileName.GetFileNameWithoutExtension(), Workspace.BranchDirectoryName, ProjectSettings.NewSelectedFileName.GetFileNameWithoutExtension()), "Branch already open", MessageBoxButtons.YesNo) == System.Windows.Forms.DialogResult.Yes)
 							{
 								Log.WriteLine("  Another project already open in this workspace, tab {0}. Replacing.", TabIdx);
 								TabControl.RemoveTab(TabIdx);
@@ -1251,13 +1252,13 @@ namespace UnrealGameSync
 					TabName = Workspace.StreamName;
 					break;
 				case TabLabels.ProjectFile:
-					TabName = Workspace.SelectedFileName;
+					TabName = Workspace.SelectedFileName.FullName;
 					break;
 				case TabLabels.WorkspaceName:
 					TabName = Workspace.ClientName;
 					break;
 				case TabLabels.WorkspaceRoot:
-					TabName = Workspace.BranchDirectoryName;
+					TabName = Workspace.BranchDirectoryName.FullName;
 					break;
 				default:
 					break;
@@ -1640,7 +1641,7 @@ namespace UnrealGameSync
 				string ServerId = ApiUrl != null ? Regex.Replace(ApiUrl, @"^.*://", "") : "noserver";
 				ServerId = Regex.Replace(ServerId, "[^a-zA-Z.]", "+");
 
-				string LogFileName = Path.Combine(DataFolder, String.Format("IssueMonitor-{0}-{1}.log", ServerId, UserName));
+				FileReference LogFileName = FileReference.Combine(DataFolder, String.Format("IssueMonitor-{0}-{1}.log", ServerId, UserName));
 
 				WorkspaceIssueMonitor = new WorkspaceIssueMonitor();
 				WorkspaceIssueMonitor.IssueMonitor = new IssueMonitor(ApiUrl, UserName, TimeSpan.FromSeconds(60.0), LogFileName);
diff --git a/Engine/Source/Programs/UnrealGameSync/UnrealGameSync/Forms/ModifyBuildStepsWindow.cs b/Engine/Source/Programs/UnrealGameSync/UnrealGameSync/Forms/ModifyBuildStepsWindow.cs
--- a/Engine/Source/Programs/UnrealGameSync/UnrealGameSync/Forms/ModifyBuildStepsWindow.cs
+++ b/Engine/Source/Programs/UnrealGameSync/UnrealGameSync/Forms/ModifyBuildStepsWindow.cs
@@ -1,5 +1,6 @@
 // Copyright Epic Games, Inc. All Rights Reserved.
 
+using EpicGames.Core;
 using System;
 using System.Collections.Generic;
 using System.ComponentModel;
@@ -18,11 +19,11 @@ namespace UnrealGameSync
 		List<string> TargetNames;
 		List<BuildStep> Steps;
 		HashSet<Guid> ProjectSteps;
-		string BaseDirectory;
+		DirectoryReference BaseDirectory;
 		ListViewItem.ListViewSubItem MouseDownSubItem = null;
 		IReadOnlyDictionary<string, string> Variables;
 
-		public ModifyBuildStepsWindow(List<string> InTargetNames, List<BuildStep> InSteps, HashSet<Guid> InProjectSteps, string InBaseDirectory, IReadOnlyDictionary<string, string> InVariables)
+		public ModifyBuildStepsWindow(List<string> InTargetNames, List<BuildStep> InSteps, HashSet<Guid> InProjectSteps, DirectoryReference InBaseDirectory, IReadOnlyDictionary<string, string> InVariables)
 		{
 			TargetNames = InTargetNames;
 			Steps = InSteps;
diff --git a/Engine/Source/Programs/UnrealGameSync/UnrealGameSync/Forms/OpenProjectWindow.cs b/Engine/Source/Programs/UnrealGameSync/UnrealGameSync/Forms/OpenProjectWindow.cs
--- a/Engine/Source/Programs/UnrealGameSync/UnrealGameSync/Forms/OpenProjectWindow.cs
+++ b/Engine/Source/Programs/UnrealGameSync/UnrealGameSync/Forms/OpenProjectWindow.cs
@@ -1,5 +1,6 @@
 // Copyright Epic Games, Inc. All Rights Reserved.
 
+using EpicGames.Core;
 using System;
 using System.Collections.Generic;
 using System.ComponentModel;
@@ -18,13 +19,13 @@ namespace UnrealGameSync
 		string ServerAndPortOverride;
 		string UserNameOverride;
 		DetectProjectSettingsTask DetectedProjectSettings;
-		string DataFolder;
-		string CacheFolder;
+		DirectoryReference DataFolder;
+		DirectoryReference CacheFolder;
 		PerforceConnection DefaultConnection;
 		TextWriter Log;
 		UserSettings Settings;
 
-		private OpenProjectWindow(UserSelectedProjectSettings Project, UserSettings Settings, string DataFolder, string CacheFolder, PerforceConnection DefaultConnection, TextWriter Log)
+		private OpenProjectWindow(UserSelectedProjectSettings Project, UserSettings Settings, DirectoryReference DataFolder, DirectoryReference CacheFolder, PerforceConnection DefaultConnection, TextWriter Log)
 		{
 			InitializeComponent();
 
@@ -86,7 +87,7 @@ namespace UnrealGameSync
 			get { return Utility.OverridePerforceSettings(DefaultConnection, ServerAndPortOverride, UserNameOverride); }
 		}
 
-		public static bool ShowModal(IWin32Window Owner, UserSelectedProjectSettings Project, out DetectProjectSettingsTask NewDetectedProjectSettings, UserSettings Settings, string DataFolder, string CacheFolder, PerforceConnection DefaultConnection, TextWriter Log)
+		public static bool ShowModal(IWin32Window Owner, UserSelectedProjectSettings Project, out DetectProjectSettingsTask NewDetectedProjectSettings, UserSettings Settings, DirectoryReference DataFolder, DirectoryReference CacheFolder, PerforceConnection DefaultConnection, TextWriter Log)
 		{
 			OpenProjectWindow Window = new OpenProjectWindow(Project, Settings, DataFolder, CacheFolder, DefaultConnection, Log);
 			if(Window.ShowDialog(Owner) == DialogResult.OK)
diff --git a/Engine/Source/Programs/UnrealGameSync/UnrealGameSync/Forms/ProgramsRunningWindow.cs b/Engine/Source/Programs/UnrealGameSync/UnrealGameSync/Forms/ProgramsRunningWindow.cs
--- a/Engine/Source/Programs/UnrealGameSync/UnrealGameSync/Forms/ProgramsRunningWindow.cs
+++ b/Engine/Source/Programs/UnrealGameSync/UnrealGameSync/Forms/ProgramsRunningWindow.cs
@@ -1,5 +1,6 @@
 ﻿// Copyright Epic Games, Inc. All Rights Reserved.
 
+using EpicGames.Core;
 using System;
 using System.Collections.Generic;
 using System.ComponentModel;
@@ -16,12 +17,12 @@ namespace UnrealGameSync
 	public partial class ProgramsRunningWindow : Form
 	{
 		object SyncObject = new object();
-		string[] Programs;
-		Func<string[]> EnumeratePrograms;
+		FileReference[] Programs;
+		Func<FileReference[]> EnumeratePrograms;
 		ManualResetEvent TerminateEvent;
 		Thread BackgroundThread;
 
-		public ProgramsRunningWindow(Func<string[]> EnumeratePrograms, string[] Programs)
+		public ProgramsRunningWindow(Func<FileReference[]> EnumeratePrograms, FileReference[] Programs)
 		{
 			InitializeComponent();
 
@@ -47,7 +48,7 @@ namespace UnrealGameSync
 					break;
 				}
 
-				string[] NewPrograms = EnumeratePrograms().OrderBy(x => x).ToArray();
+				FileReference[] NewPrograms = EnumeratePrograms().OrderBy(x => x).ToArray();
 				lock(SyncObject)
 				{
 					if(!Enumerable.SequenceEqual(Programs, NewPrograms))
diff --git a/Engine/Source/Programs/UnrealGameSync/UnrealGameSync/IssueMonitor.cs b/Engine/Source/Programs/UnrealGameSync/UnrealGameSync/IssueMonitor.cs
--- a/Engine/Source/Programs/UnrealGameSync/UnrealGameSync/IssueMonitor.cs
+++ b/Engine/Source/Programs/UnrealGameSync/UnrealGameSync/IssueMonitor.cs
@@ -1,6 +1,7 @@
 // Copyright Epic Games, Inc. All Rights Reserved.
 
 using EnvDTE;
+using EpicGames.Core;
 using System;
 using System.Collections.Concurrent;
 using System.Collections.Generic;
@@ -149,7 +150,7 @@ namespace UnrealGameSync
 		// Only used by MainWindow, but easier to just store here
 		public Dictionary<long, IssueAlertReason> IssueIdToAlertReason = new Dictionary<long, IssueAlertReason>();
 
-		public IssueMonitor(string ApiUrl, string UserName, TimeSpan UpdateInterval, string LogFileName)
+		public IssueMonitor(string ApiUrl, string UserName, TimeSpan UpdateInterval, FileReference LogFileName)
 		{
 			this.ApiUrl = ApiUrl;
 			this.UserName = UserName;
diff --git a/Engine/Source/Programs/UnrealGameSync/UnrealGameSync/Jupiter/JupiterMonitor.cs b/Engine/Source/Programs/UnrealGameSync/UnrealGameSync/Jupiter/JupiterMonitor.cs
--- a/Engine/Source/Programs/UnrealGameSync/UnrealGameSync/Jupiter/JupiterMonitor.cs
+++ b/Engine/Source/Programs/UnrealGameSync/UnrealGameSync/Jupiter/JupiterMonitor.cs
@@ -39,7 +39,7 @@ namespace UnrealGameSync
 			}
 		}
 
-		private JupiterMonitor(OIDCTokenManager InTokenManager, string InLogPath, string InNamespace, string InUrl,
+		private JupiterMonitor(OIDCTokenManager InTokenManager, FileReference InLogPath, string InNamespace, string InUrl,
 			string InProviderIdentifier, string InExpectedBranch)
 		{
 			TokenManager = InTokenManager;
@@ -75,7 +75,7 @@ namespace UnrealGameSync
 			}
 		}
 
-		public static JupiterMonitor CreateFromConfigFile(OIDCTokenManager TokenManager, string LogPath, ConfigFile ConfigFile, string SelectedProjectIdentifier)
+		public static JupiterMonitor CreateFromConfigFile(OIDCTokenManager TokenManager, FileReference LogPath, ConfigFile ConfigFile, string SelectedProjectIdentifier)
 		{
 			ConfigSection JupiterConfigSection = ConfigFile.FindSection("Jupiter");
 			if (JupiterConfigSection == null)
@@ -224,7 +224,7 @@ namespace UnrealGameSync
 				return ChangeToKey.TryGetValue(ChangeNumber, out ArchiveKey);
 			}
 
-			public bool DownloadArchive(string ArchiveKey, string LocalRootPath, string ManifestFileName, TextWriter Log, ProgressValue Progress)
+			public bool DownloadArchive(string ArchiveKey, DirectoryReference LocalRootPath, FileReference ManifestFileName, TextWriter Log, ProgressValue Progress)
 			{
 				try
 				{
@@ -239,10 +239,10 @@ namespace UnrealGameSync
 					});
 
 					// place the manifest for the Jupiter download next to the UGS manifest
-					FileReference UGSManifestFileReference = new FileReference(ManifestFileName);
+					FileReference UGSManifestFileReference = ManifestFileName;
 					FileReference JupiterManifestFileReference = FileReference.Combine(UGSManifestFileReference.Directory, "Jupiter-Manifest.json");
 
-					DirectoryReference RootDirectory = new DirectoryReference(LocalRootPath);
+					DirectoryReference RootDirectory = LocalRootPath;
 					JupiterFileTree FileTree = new JupiterFileTree(RootDirectory, InDeferReadingFiles: true);
 					Task<List<FileReference>> DownloadTask = FileTree.DownloadFromJupiter(JupiterManifestFileReference, JupiterUrl, JupiterNamespace, ArchiveKey, ProgressCallback);
 					DownloadTask.Wait();
@@ -255,12 +255,12 @@ namespace UnrealGameSync
 					}
 
 					// Write it out to a temporary file, then move it into place
-					string TempManifestFileName = ManifestFileName + ".tmp";
-					using (FileStream OutputStream = File.Open(TempManifestFileName, FileMode.Create, FileAccess.Write))
+					FileReference TempManifestFileName = ManifestFileName + ".tmp";
+					using (FileStream OutputStream = FileReference.Open(TempManifestFileName, FileMode.Create, FileAccess.Write))
 					{
 						ArchiveManifest.Write(OutputStream);
 					}
-					File.Move(TempManifestFileName, ManifestFileName);
+					FileReference.Move(TempManifestFileName, ManifestFileName);
 
 					return true;
 				}
diff --git a/Engine/Source/Programs/UnrealGameSync/UnrealGameSync/OutputAdapters.cs b/Engine/Source/Programs/UnrealGameSync/UnrealGameSync/OutputAdapters.cs
--- a/Engine/Source/Programs/UnrealGameSync/UnrealGameSync/OutputAdapters.cs
+++ b/Engine/Source/Programs/UnrealGameSync/UnrealGameSync/OutputAdapters.cs
@@ -1,5 +1,6 @@
 // Copyright Epic Games, Inc. All Rights Reserved.
 
+using EpicGames.Core;
 using System;
 using System.Collections.Generic;
 using System.IO;
@@ -134,12 +135,12 @@ namespace UnrealGameSync
 	
 	class BoundedLogWriter : LineBasedTextWriter
 	{
-		string FileName;
-		string BackupFileName;
+		FileReference FileName;
+		FileReference BackupFileName;
 		int MaxSize;
 		StreamWriter Inner;
 
-		public BoundedLogWriter(string InFileName, int InMaxSize = 128 * 1024)
+		public BoundedLogWriter(FileReference InFileName, int InMaxSize = 128 * 1024)
 		{
 			FileName = InFileName;
 			BackupFileName = FileName + ".bak";
@@ -173,16 +174,16 @@ namespace UnrealGameSync
 			CloseInner();
 			try
 			{
-				if(File.Exists(BackupFileName))
+				if(FileReference.Exists(BackupFileName))
 				{
-					File.Delete(BackupFileName);
+					FileReference.Delete(BackupFileName);
 				}
-				if(File.Exists(FileName))
+				if(FileReference.Exists(FileName))
 				{
-					File.Move(FileName, BackupFileName);
+					FileReference.Move(FileName, BackupFileName);
 				}
 
-				Inner = new StreamWriter(FileName);
+				Inner = new StreamWriter(FileName.FullName);
 				Inner.AutoFlush = true;
 			}
 			catch(Exception)
diff --git a/Engine/Source/Programs/UnrealGameSync/UnrealGameSync/PerforceMonitor.cs b/Engine/Source/Programs/UnrealGameSync/UnrealGameSync/PerforceMonitor.cs
--- a/Engine/Source/Programs/UnrealGameSync/UnrealGameSync/PerforceMonitor.cs
+++ b/Engine/Source/Programs/UnrealGameSync/UnrealGameSync/PerforceMonitor.cs
@@ -1,5 +1,6 @@
 // Copyright Epic Games, Inc. All Rights Reserved.
 
+using EpicGames.Core;
 using System;
 using System.Collections.Concurrent;
 using System.Collections.Generic;
@@ -78,7 +79,7 @@ namespace UnrealGameSync
 		string Target { get;  }
 		bool Exists();
 		bool TryGetArchiveKeyForChangeNumber(int ChangeNumber, out string ArchiveKey);
-		bool DownloadArchive(string ArchiveKey, string LocalRootPath, string ManifestFileName, TextWriter Log, ProgressValue Progress);
+		bool DownloadArchive(string ArchiveKey, DirectoryReference LocalRootPath, FileReference ManifestFileName, TextWriter Log, ProgressValue Progress);
 	}
 
 	class PerforceMonitor : IDisposable, IArchiveInfoSource
@@ -155,12 +156,12 @@ namespace UnrealGameSync
 				return ChangeNumberToFileRevision.TryGetValue(ChangeNumber, out ArchiveKey);
 			}
 
-			public bool DownloadArchive(string ArchiveKey, string LocalRootPath, string ManifestFileName, TextWriter Log, ProgressValue Progress)
+			public bool DownloadArchive(string ArchiveKey, DirectoryReference LocalRootPath, FileReference ManifestFileName, TextWriter Log, ProgressValue Progress)
 			{
-				string TempZipFileName = Path.GetTempFileName();
+				FileReference TempZipFileName = new FileReference(Path.GetTempFileName());
 				try
 				{
-					if (!Perforce.PrintToFile(ArchiveKey, TempZipFileName, Log) || new FileInfo(TempZipFileName).Length == 0)
+					if (!Perforce.PrintToFile(ArchiveKey, TempZipFileName.FullName, Log) || TempZipFileName.ToFileInfo().Length == 0)
 					{
 						return false;
 					}
@@ -170,8 +171,8 @@ namespace UnrealGameSync
 				}
 				finally
 				{
-					File.SetAttributes(TempZipFileName, FileAttributes.Normal);
-					File.Delete(TempZipFileName);
+					FileReference.SetAttributes(TempZipFileName, FileAttributes.Normal);
+					FileReference.Delete(TempZipFileName);
 				}
 
 				return true;
@@ -207,8 +208,8 @@ namespace UnrealGameSync
 		BoundedLogWriter LogWriter;
 		bool bIsEnterpriseProject;
 		bool bDisposing;
-		string CacheFolder;
-		List<KeyValuePair<string, DateTime>> LocalConfigFiles;
+		DirectoryReference CacheFolder;
+		List<KeyValuePair<FileReference, DateTime>> LocalConfigFiles;
 
 		public event Action OnUpdate;
 		public event Action OnUpdateMetadata;
@@ -221,7 +222,7 @@ namespace UnrealGameSync
 			private set;
 		}
 
-		public PerforceMonitor(PerforceConnection InPerforce, string InBranchClientPath, string InSelectedClientFileName, string InSelectedProjectIdentifier, string InLogPath, bool bInIsEnterpriseProject, ConfigFile InProjectConfigFile, string InCacheFolder, List<KeyValuePair<string, DateTime>> InLocalConfigFiles)
+		public PerforceMonitor(PerforceConnection InPerforce, string InBranchClientPath, string InSelectedClientFileName, string InSelectedProjectIdentifier, FileReference InLogPath, bool bInIsEnterpriseProject, ConfigFile InProjectConfigFile, DirectoryReference InCacheFolder, List<KeyValuePair<FileReference, DateTime>> InLocalConfigFiles)
 		{
 			Perforce = InPerforce;
 			BranchClientPath = InBranchClientPath;
@@ -609,7 +610,7 @@ namespace UnrealGameSync
 				}
 			}
 
-			if(LocalConfigFiles.Any(x => File.GetLastWriteTimeUtc(x.Key) != x.Value))
+			if(LocalConfigFiles.Any(x => FileReference.GetLastWriteTimeUtc(x.Key) != x.Value))
 			{
 				UpdateProjectConfigFile();
 				if(OnUpdateMetadata != null)
@@ -704,7 +705,7 @@ namespace UnrealGameSync
 			LatestProjectConfigFile = ReadProjectConfigFile(Perforce, BranchClientPath, SelectedClientFileName, CacheFolder, LocalConfigFiles, LogWriter);
 		}
 
-		public static ConfigFile ReadProjectConfigFile(PerforceConnection Perforce, string BranchClientPath, string SelectedClientFileName, string CacheFolder, List<KeyValuePair<string, DateTime>> LocalConfigFiles, TextWriter Log)
+		public static ConfigFile ReadProjectConfigFile(PerforceConnection Perforce, string BranchClientPath, string SelectedClientFileName, DirectoryReference CacheFolder, List<KeyValuePair<FileReference, DateTime>> LocalConfigFiles, TextWriter Log)
 		{
 			List<string> ConfigFilePaths = Utility.GetDepotConfigPaths(BranchClientPath + "/Engine", SelectedClientFileName);
 
@@ -730,7 +731,7 @@ namespace UnrealGameSync
 						try
 						{
 							DateTime LastModifiedTime = File.GetLastWriteTimeUtc(LocalFileName);
-							LocalConfigFiles.Add(new KeyValuePair<string, DateTime>(LocalFileName, LastModifiedTime));
+							LocalConfigFiles.Add(new KeyValuePair<FileReference, DateTime>(new FileReference(LocalFileName), LastModifiedTime));
 							Lines = File.ReadAllLines(LocalFileName).ToList();
 						}
 						catch(Exception Ex)
diff --git a/Engine/Source/Programs/UnrealGameSync/UnrealGameSync/Program.cs b/Engine/Source/Programs/UnrealGameSync/UnrealGameSync/Program.cs
--- a/Engine/Source/Programs/UnrealGameSync/UnrealGameSync/Program.cs
+++ b/Engine/Source/Programs/UnrealGameSync/UnrealGameSync/Program.cs
@@ -1,5 +1,6 @@
 // Copyright Epic Games, Inc. All Rights Reserved.
 
+using EpicGames.Core;
 using System;
 using System.Collections.Generic;
 using System.Data.SqlClient;
@@ -97,7 +98,7 @@ namespace UnrealGameSync
 			string Uri;
 			ParseArgument(RemainingArgs, "-uri=", out Uri);
 
-			string UpdateConfigFile = Path.Combine(Path.GetDirectoryName(Assembly.GetExecutingAssembly().Location), "AutoUpdate.ini");
+			FileReference UpdateConfigFile = FileReference.Combine(new FileReference(Assembly.GetExecutingAssembly().Location).Directory, "AutoUpdate.ini");
 			MergeUpdateSettings(UpdateConfigFile, ref UpdatePath, ref UpdateSpawn);
 
 			// Set the current working directory to the update directory to prevent child-process file handles from disrupting auto-updates
@@ -119,14 +120,14 @@ namespace UnrealGameSync
 				}
 			}
 
-			string DataFolder = Path.Combine(Environment.GetFolderPath(Environment.SpecialFolder.LocalApplicationData), "UnrealGameSync");
-			Directory.CreateDirectory(DataFolder);
+			DirectoryReference DataFolder = DirectoryReference.Combine(DirectoryReference.GetSpecialFolder(Environment.SpecialFolder.LocalApplicationData), "UnrealGameSync");
+			DirectoryReference.CreateDirectory(DataFolder);
 
 			// Enable TLS 1.1 and 1.2. TLS 1.0 is now deprecated and not allowed by default in NET Core servers.
 			ServicePointManager.SecurityProtocol |= SecurityProtocolType.Tls11 | SecurityProtocolType.Tls12;
 
 			// Create the log file
-			using (TimestampLogWriter Log = new TimestampLogWriter(new BoundedLogWriter(Path.Combine(DataFolder, "UnrealGameSync.log"))))
+			using (TimestampLogWriter Log = new TimestampLogWriter(new BoundedLogWriter(FileReference.Combine(DataFolder, "UnrealGameSync.log"))))
 			{
 				Log.WriteLine("Application version: {0}", Assembly.GetExecutingAssembly().GetName().Version);
 				Log.WriteLine("Started at {0}", DateTime.Now.ToString());
@@ -141,7 +142,7 @@ namespace UnrealGameSync
 					Utility.SaveGlobalPerforceSettings(ServerAndPort, UserName, BaseUpdatePath);
 				}
 
-				using (BoundedLogWriter TelemetryLog = new BoundedLogWriter(Path.Combine(DataFolder, "Telemetry.log")))
+				using (BoundedLogWriter TelemetryLog = new BoundedLogWriter(FileReference.Combine(DataFolder, "Telemetry.log")))
 				{
 					TelemetryLog.WriteLine("Creating telemetry sink for session {0}", SessionId);
 
@@ -217,12 +218,12 @@ namespace UnrealGameSync
 			}
 		}
 
-		static void MergeUpdateSettings(string UpdateConfigFile, ref string UpdatePath, ref string UpdateSpawn)
+		static void MergeUpdateSettings(FileReference UpdateConfigFile, ref string UpdatePath, ref string UpdateSpawn)
 		{
 			try
 			{
 				ConfigFile UpdateConfig = new ConfigFile();
-				if(File.Exists(UpdateConfigFile))
+				if(FileReference.Exists(UpdateConfigFile))
 				{
 					UpdateConfig.Load(UpdateConfigFile);
 				}
diff --git a/Engine/Source/Programs/UnrealGameSync/UnrealGameSync/ProgramApplicationContext.cs b/Engine/Source/Programs/UnrealGameSync/UnrealGameSync/ProgramApplicationContext.cs
--- a/Engine/Source/Programs/UnrealGameSync/UnrealGameSync/ProgramApplicationContext.cs
+++ b/Engine/Source/Programs/UnrealGameSync/UnrealGameSync/ProgramApplicationContext.cs
@@ -1,5 +1,6 @@
 ﻿// Copyright Epic Games, Inc. All Rights Reserved.
 
+using EpicGames.Core;
 using System;
 using System.Collections.Generic;
 using System.ComponentModel;
@@ -22,8 +23,8 @@ namespace UnrealGameSync
 		PerforceConnection DefaultConnection;
 		UpdateMonitor UpdateMonitor;
 		string ApiUrl;
-		string DataFolder;
-		string CacheFolder;
+		DirectoryReference DataFolder;
+		DirectoryReference CacheFolder;
 		bool bRestoreState;
 		string UpdateSpawn;
 		bool bUnstable;
@@ -51,13 +52,13 @@ namespace UnrealGameSync
 
 		OIDCTokenManager OIDCTokenManager;
 
-		public ProgramApplicationContext(PerforceConnection DefaultConnection, UpdateMonitor UpdateMonitor, string ApiUrl, string DataFolder, EventWaitHandle ActivateEvent, bool bRestoreState, string UpdateSpawn, string ProjectFileName, bool bUnstable, TimestampLogWriter Log, string Uri)
+		public ProgramApplicationContext(PerforceConnection DefaultConnection, UpdateMonitor UpdateMonitor, string ApiUrl, DirectoryReference DataFolder, EventWaitHandle ActivateEvent, bool bRestoreState, string UpdateSpawn, string ProjectFileName, bool bUnstable, TimestampLogWriter Log, string Uri)
 		{
 			this.DefaultConnection = DefaultConnection;
 			this.UpdateMonitor = UpdateMonitor;
 			this.ApiUrl = ApiUrl;
 			this.DataFolder = DataFolder;
-			this.CacheFolder = Path.Combine(DataFolder, "Cache");
+			this.CacheFolder = DirectoryReference.Combine(DataFolder, "Cache");
 			this.bRestoreState = bRestoreState;
 			this.UpdateSpawn = UpdateSpawn;
 			this.bUnstable = bUnstable;
@@ -65,8 +66,8 @@ namespace UnrealGameSync
 			this.Uri = Uri;
 
 			// Create the directories
-			Directory.CreateDirectory(DataFolder);
-			Directory.CreateDirectory(CacheFolder);
+			DirectoryReference.CreateDirectory(DataFolder);
+			DirectoryReference.CreateDirectory(CacheFolder);
 
 			// Make sure a synchronization context is set. We spawn a bunch of threads (eg. UpdateMonitor) at startup, and need to make sure we can post messages 
 			// back to the main thread at any time.
@@ -79,7 +80,7 @@ namespace UnrealGameSync
 			MainThreadSynchronizationContext = WindowsFormsSynchronizationContext.Current;
 
 			// Read the user's settings
-			Settings = new UserSettings(Path.Combine(DataFolder, "UnrealGameSync.ini"), Log);
+			Settings = new UserSettings(FileReference.Combine(DataFolder, "UnrealGameSync.ini"), Log);
 			if(!String.IsNullOrEmpty(ProjectFileName))
 			{
 				string FullProjectFileName = Path.GetFullPath(ProjectFileName);
diff --git a/Engine/Source/Programs/UnrealGameSync/UnrealGameSync/ToolUpdateMonitor.cs b/Engine/Source/Programs/UnrealGameSync/UnrealGameSync/ToolUpdateMonitor.cs
--- a/Engine/Source/Programs/UnrealGameSync/UnrealGameSync/ToolUpdateMonitor.cs
+++ b/Engine/Source/Programs/UnrealGameSync/UnrealGameSync/ToolUpdateMonitor.cs
@@ -1,5 +1,6 @@
 // Copyright Epic Games, Inc. All Rights Reserved.
 
+using EpicGames.Core;
 using System;
 using System.Collections.Generic;
 using System.Diagnostics;
@@ -40,25 +41,25 @@ namespace UnrealGameSync
 		Thread WorkerThread;
 		AutoResetEvent WakeEvent;
 		bool bQuit;
-		string LogFile;
+		FileReference LogFile;
 		public List<ToolDefinition> Tools { get; private set; } = new List<ToolDefinition>();
 		int LastChange = -1;
 
 		public PerforceConnection Perforce { get; }
-		string ToolsDir { get; }
+		DirectoryReference ToolsDir { get; }
 		UserSettings Settings { get; }
 
 		public event Action OnChange;
 
-		public ToolUpdateMonitor(PerforceConnection InPerforce, string DataDir, UserSettings Settings)
+		public ToolUpdateMonitor(PerforceConnection InPerforce, DirectoryReference DataDir, UserSettings Settings)
 		{
 			this.Perforce = InPerforce;
-			this.ToolsDir = Path.Combine(DataDir, "Tools");
+			this.ToolsDir = DirectoryReference.Combine(DataDir, "Tools");
 			this.Settings = Settings;
 
-			LogFile = Path.Combine(DataDir, "Tools.log");
+			LogFile = FileReference.Combine(DataDir, "Tools.log");
 
-			Directory.CreateDirectory(ToolsDir);
+			DirectoryReference.CreateDirectory(ToolsDir);
 
 			WakeEvent = new AutoResetEvent(false);
 		}
@@ -93,9 +94,9 @@ namespace UnrealGameSync
 			Close();
 		}
 
-		string GetToolPathInternal(string ToolName)
+		DirectoryReference GetToolPathInternal(string ToolName)
 		{
-			return Path.Combine(ToolsDir, ToolName, "Current");
+			return DirectoryReference.Combine(ToolsDir, ToolName, "Current");
 		}
 
 		public string GetToolName(Guid ToolId)
@@ -110,7 +111,7 @@ namespace UnrealGameSync
 			return null;
 		}
 
-		public string GetToolPath(string ToolName)
+		public DirectoryReference GetToolPath(string ToolName)
 		{
 			if (GetToolChange(ToolName) != 0)
 			{
@@ -294,7 +295,7 @@ namespace UnrealGameSync
 
 		void RunCommand(string ToolName, string Command, TextWriter Log)
 		{
-			string ToolPath = GetToolPathInternal(ToolName);
+			DirectoryReference ToolPath = GetToolPathInternal(ToolName);
 
 			string CommandExe = Command;
 			string CommandArgs = string.Empty;
@@ -306,7 +307,7 @@ namespace UnrealGameSync
 				CommandArgs = Command.Substring(SpaceIdx + 1);
 			}
 
-			int ExitCode = Utility.ExecuteProcess(Path.Combine(ToolPath, CommandExe), ToolPath, CommandArgs, null, Log);
+			int ExitCode = Utility.ExecuteProcess(FileReference.Combine(ToolPath, CommandExe).FullName, ToolPath.FullName, CommandArgs, null, Log);
 			Log.WriteLine("(Exit: {0})", ExitCode);
 		}
 
@@ -322,7 +323,7 @@ namespace UnrealGameSync
 
 			SetToolChange(ToolName, null);
 
-			string ToolPath = GetToolPath(ToolName);
+			DirectoryReference ToolPath = GetToolPath(ToolName);
 			if (ToolPath != null)
 			{
 				Log.WriteLine("Removing {0}", ToolPath);
@@ -330,9 +331,9 @@ namespace UnrealGameSync
 			}
 		}
 
-		static void ForceDeleteDirectory(string DirectoryName)
+		static void ForceDeleteDirectory(DirectoryReference DirectoryName)
 		{
-			DirectoryInfo BaseDir = new DirectoryInfo(DirectoryName);
+			DirectoryInfo BaseDir = DirectoryName.ToDirectoryInfo();
 			if (BaseDir.Exists)
 			{
 				foreach (FileInfo File in BaseDir.EnumerateFiles("*", SearchOption.AllDirectories))
@@ -343,7 +344,7 @@ namespace UnrealGameSync
 			}
 		}
 
-		static bool TryDeleteDirectory(string DirectoryName)
+		static bool TryDeleteDirectory(DirectoryReference DirectoryName)
 		{
 			try
 			{
@@ -358,25 +359,25 @@ namespace UnrealGameSync
 
 		bool UpdateTool(string ToolName, int Change, List<PerforceFileRecord> Records, Action<TextWriter> InstallAction, TextWriter Log)
 		{
-			string ToolDir = Path.Combine(ToolsDir, ToolName);
-			Directory.CreateDirectory(ToolDir);
+			DirectoryReference ToolDir = DirectoryReference.Combine(ToolsDir, ToolName);
+			DirectoryReference.CreateDirectory(ToolDir);
 
-			foreach (DirectoryInfo ExistingDir in new DirectoryInfo(ToolDir).EnumerateDirectories("Prev-*"))
+			foreach (DirectoryReference ExistingDir in DirectoryReference.EnumerateDirectories(ToolDir, "Prev-*"))
 			{
-				TryDeleteDirectory(ExistingDir.FullName);
+				TryDeleteDirectory(ExistingDir);
 			}
 
-			string NextToolDir = Path.Combine(ToolDir, "Next");
+			DirectoryReference NextToolDir = DirectoryReference.Combine(ToolDir, "Next");
 			ForceDeleteDirectory(NextToolDir);
-			Directory.CreateDirectory(NextToolDir);
+			DirectoryReference.CreateDirectory(NextToolDir);
 
-			string NextToolZipsDir = Path.Combine(NextToolDir, ".zips");
-			Directory.CreateDirectory(NextToolZipsDir);
+			DirectoryReference NextToolZipsDir = DirectoryReference.Combine(NextToolDir, ".zips");
+			DirectoryReference.CreateDirectory(NextToolZipsDir);
 
 			for (int Idx = 0; Idx < Records.Count; Idx++)
 			{
-				string ZipFile = Path.Combine(NextToolZipsDir, String.Format("{0}.{1}.zip", ToolName, Idx));
-				if (!Perforce.PrintToFile(String.Format("{0}#{1}", Records[Idx].DepotPath, Records[Idx].HeadRevision), ZipFile, Log) || !File.Exists(ZipFile))
+				FileReference ZipFile = FileReference.Combine(NextToolZipsDir, String.Format("{0}.{1}.zip", ToolName, Idx));
+				if (!Perforce.PrintToFile(String.Format("{0}#{1}", Records[Idx].DepotPath, Records[Idx].HeadRevision), ZipFile.FullName, Log) || !FileReference.Exists(ZipFile))
 				{
 					Log.WriteLine("Unable to print {0}", Records[Idx].DepotPath);
 					return false;
@@ -386,15 +387,15 @@ namespace UnrealGameSync
 
 			SetToolChange(ToolName, null);
 
-			string CurrentToolDir = Path.Combine(ToolDir, "Current");
-			if (Directory.Exists(CurrentToolDir))
+			DirectoryReference CurrentToolDir = DirectoryReference.Combine(ToolDir, "Current");
+			if (DirectoryReference.Exists(CurrentToolDir))
 			{
-				string PrevDirectoryName = Path.Combine(ToolDir, String.Format("Prev-{0:X16}", Stopwatch.GetTimestamp()));
-				Directory.Move(CurrentToolDir, PrevDirectoryName);
+				DirectoryReference PrevDirectoryName = DirectoryReference.Combine(ToolDir, String.Format("Prev-{0:X16}", Stopwatch.GetTimestamp()));
+				Directory.Move(CurrentToolDir.FullName, PrevDirectoryName.FullName);
 				TryDeleteDirectory(PrevDirectoryName);
 			}
 
-			Directory.Move(NextToolDir, CurrentToolDir);
+			Directory.Move(NextToolDir.FullName, CurrentToolDir.FullName);
 
 			if (InstallAction != null)
 			{
@@ -406,17 +407,17 @@ namespace UnrealGameSync
 			return true;
 		}
 
-		string GetConfigFilePath(string ToolName)
+		FileReference GetConfigFilePath(string ToolName)
 		{
-			return Path.Combine(ToolsDir, ToolName, ToolName + ".ini");
+			return FileReference.Combine(ToolsDir, ToolName, ToolName + ".ini");
 		}
 
 		int GetToolChange(string ToolName)
 		{
 			try
 			{
-				string ConfigFilePath = GetConfigFilePath(ToolName);
-				if (File.Exists(ConfigFilePath))
+				FileReference ConfigFilePath = GetConfigFilePath(ToolName);
+				if (FileReference.Exists(ConfigFilePath))
 				{
 					ConfigFile ConfigFile = new ConfigFile();
 					ConfigFile.Load(ConfigFilePath);
@@ -431,7 +432,7 @@ namespace UnrealGameSync
 
 		void SetToolChange(string ToolName, int? Change)
 		{
-			string ConfigFilePath = GetConfigFilePath(ToolName);
+			FileReference ConfigFilePath = GetConfigFilePath(ToolName);
 			if (Change.HasValue)
 			{
 				ConfigFile ConfigFile = new ConfigFile();
@@ -440,7 +441,7 @@ namespace UnrealGameSync
 			}
 			else
 			{
-				File.Delete(ConfigFilePath);
+				FileReference.Delete(ConfigFilePath);
 			}
 		}
 	}
diff --git a/Engine/Source/Programs/UnrealGameSync/UnrealGameSync/UserSettings.cs b/Engine/Source/Programs/UnrealGameSync/UnrealGameSync/UserSettings.cs
--- a/Engine/Source/Programs/UnrealGameSync/UnrealGameSync/UserSettings.cs
+++ b/Engine/Source/Programs/UnrealGameSync/UnrealGameSync/UserSettings.cs
@@ -251,7 +251,7 @@ namespace UnrealGameSync
 
 		// Settings for the currently synced project in this workspace. CurrentChangeNumber is only valid for this workspace if CurrentProjectPath is the current project.
 		public string CurrentProjectIdentifier { get; set; }
-		public int CurrentChangeNumber { get; set; }
+		public int CurrentChangeNumber { get; set; } = -1;
 		public string CurrentSyncFilterHash { get; set; }
 		public List<int> AdditionalChangeNumbers { get; set; } = new List<int>();
 
@@ -283,7 +283,13 @@ namespace UnrealGameSync
 			Entry.State = State;
 		}
 
-		public void Save() => Utility.SaveJson(File, this);
+		public void Save()
+		{
+			lock(File)
+			{
+				Utility.SaveJson(File, this);
+			}
+		}
 	}
 
 	class UserProjectSettings
@@ -310,7 +316,7 @@ namespace UnrealGameSync
 			None		// Show no robomerge changes
 		};
 
-		string FileName;
+		FileReference FileName;
 		ConfigFile ConfigFile = new ConfigFile();
 
 		// General settings
@@ -410,7 +416,7 @@ namespace UnrealGameSync
 			return Projects;
 		}
 
-		public UserSettings(string InFileName, TextWriter Log)
+		public UserSettings(FileReference InFileName, TextWriter Log)
 		{
 			FileName = InFileName;
 			ConfigFile.TryLoad(FileName, Log);
@@ -662,7 +668,7 @@ namespace UnrealGameSync
 			return ConfigDir;
 		}
 
-		public UserWorkspaceState FindOrAddWorkspaceState(DirectoryReference RootDir, string SelectedClientFileName)
+		public UserWorkspaceState FindOrAddWorkspaceState(DirectoryReference RootDir, string SelectedClientFileName, string SelectedProjectIdentifier)
 		{
 			FileReference ConfigFile = FileReference.Combine(GetConfigDir(RootDir), "global.json");
 			if (!WorkspaceKeyToState.TryGetValue(ConfigFile, out UserWorkspaceState State))
@@ -678,11 +684,12 @@ namespace UnrealGameSync
 				State.File = ConfigFile;
 				WorkspaceKeyToState.Add(ConfigFile, State);
 			}
-			if (!String.Equals(State.SelectedClientFileName, SelectedClientFileName))
+			if (!String.Equals(State.SelectedClientFileName, SelectedClientFileName, StringComparison.Ordinal) || !String.Equals(State.CurrentProjectIdentifier, SelectedProjectIdentifier, StringComparison.Ordinal))
 			{
 				State = new UserWorkspaceState();
 				State.File = ConfigFile;
 				State.SelectedClientFileName = SelectedClientFileName;
+				State.CurrentProjectIdentifier = SelectedProjectIdentifier;
 				WorkspaceKeyToState[ConfigFile] = State;
 			}
 			return State;
diff --git a/Engine/Source/Programs/UnrealGameSync/UnrealGameSync/Utility.cs b/Engine/Source/Programs/UnrealGameSync/UnrealGameSync/Utility.cs
--- a/Engine/Source/Programs/UnrealGameSync/UnrealGameSync/Utility.cs
+++ b/Engine/Source/Programs/UnrealGameSync/UnrealGameSync/Utility.cs
@@ -510,21 +510,21 @@ namespace UnrealGameSync
 			}
 		}
 
-		public static bool TryPrintFileUsingCache(PerforceConnection Perforce, string DepotPath, string CacheFolder, string Digest, out List<string> Lines, TextWriter Log)
+		public static bool TryPrintFileUsingCache(PerforceConnection Perforce, string DepotPath, DirectoryReference CacheFolder, string Digest, out List<string> Lines, TextWriter Log)
 		{
 			if(Digest == null)
 			{
 				return Perforce.Print(DepotPath, out Lines, Log);
 			}
 
-			string CacheFile = Path.Combine(CacheFolder, Digest);
-			if(File.Exists(CacheFile))
+			FileReference CacheFile = FileReference.Combine(CacheFolder, Digest);
+			if(FileReference.Exists(CacheFile))
 			{
 				Log.WriteLine("Reading cached copy of {0} from {1}", DepotPath, CacheFile);
-				Lines = new List<string>(File.ReadAllLines(CacheFile));
+				Lines = new List<string>(FileReference.ReadAllLines(CacheFile));
 				try
 				{
-					File.SetLastWriteTimeUtc(CacheFile, DateTime.UtcNow);
+					FileReference.SetLastWriteTimeUtc(CacheFile, DateTime.UtcNow);
 				}
 				catch(Exception Ex)
 				{
@@ -534,26 +534,26 @@ namespace UnrealGameSync
 			}
 			else
 			{
-				string TempFile = String.Format("{0}.{1}.temp", CacheFile, Guid.NewGuid());
-				if(!Perforce.PrintToFile(DepotPath, TempFile, Log))
+				FileReference TempFile = new FileReference(String.Format("{0}.{1}.temp", CacheFile.FullName, Guid.NewGuid()));
+				if(!Perforce.PrintToFile(DepotPath, TempFile.FullName, Log))
 				{
 					Lines = null;
 					return false;
 				}
 				else
 				{
-					Lines = new List<string>(File.ReadAllLines(TempFile));
+					Lines = new List<string>(FileReference.ReadAllLines(TempFile));
 					try
 					{
-						File.SetAttributes(TempFile, FileAttributes.Normal);
-						File.SetLastWriteTimeUtc(TempFile, DateTime.UtcNow);
-						File.Move(TempFile, CacheFile);
+						FileReference.SetAttributes(TempFile, FileAttributes.Normal);
+						FileReference.SetLastWriteTimeUtc(TempFile, DateTime.UtcNow);
+						FileReference.Move(TempFile, CacheFile);
 					}
 					catch
 					{
 						try
 						{
-							File.Delete(TempFile);
+							FileReference.Delete(TempFile);
 						}
 						catch
 						{
@@ -564,9 +564,9 @@ namespace UnrealGameSync
 			}
 		}
 
-		public static void ClearPrintCache(string CacheFolder)
+		public static void ClearPrintCache(DirectoryReference CacheFolder)
 		{
-			DirectoryInfo CacheDir = new DirectoryInfo(CacheFolder);
+			DirectoryInfo CacheDir = CacheFolder.ToDirectoryInfo();
 			if(CacheDir.Exists)
 			{
 				DateTime DeleteTime = DateTime.UtcNow - TimeSpan.FromDays(5.0);
diff --git a/Engine/Source/Programs/UnrealGameSync/UnrealGameSync/Workspace.cs b/Engine/Source/Programs/UnrealGameSync/UnrealGameSync/Workspace.cs
--- a/Engine/Source/Programs/UnrealGameSync/UnrealGameSync/Workspace.cs
+++ b/Engine/Source/Programs/UnrealGameSync/UnrealGameSync/Workspace.cs
@@ -66,6 +66,10 @@ namespace UnrealGameSync
 		public PerforceSyncOptions PerforceSyncOptions;
 		public List<PerforceFileRecord> HaveFiles; // Cached when sync filter has changed
 
+		// May be updated during sync
+		public ConfigFile ProjectConfigFile;
+		public IReadOnlyList<string> ProjectStreamFilter;
+
 		public WorkspaceUpdateContext(int InChangeNumber, WorkspaceUpdateOptions InOptions, string[] InSyncFilter, Dictionary<Guid, ConfigObject> InDefaultBuildSteps, List<ConfigObject> InUserBuildSteps, HashSet<Guid> InCustomBuildSteps, Dictionary<string, string> InVariables)
 		{
 			ChangeNumber = InChangeNumber;
@@ -133,6 +137,26 @@ namespace UnrealGameSync
 		}
 	}
 
+	public class ProjectInfo
+	{
+		public DirectoryReference LocalRootPath { get; }
+		public FileReference SelectedLocalFileName { get; }
+		public string ClientRootPath { get; }
+		public string SelectedClientFileName { get; }
+		public string TelemetryProjectPath { get; }
+		public bool bIsEnterpriseProject { get; }
+
+		public ProjectInfo(DirectoryReference InLocalRootPath, FileReference InSelectedLocalFileName, string InClientRootPath, string InSelectedClientFileName, string InTelemetryProjectPath, bool bInIsEnterpriseProject)
+		{
+			LocalRootPath = InLocalRootPath;
+			SelectedLocalFileName = InSelectedLocalFileName;
+			ClientRootPath = InClientRootPath;
+			SelectedClientFileName = InSelectedClientFileName;
+			TelemetryProjectPath = InTelemetryProjectPath;
+			bIsEnterpriseProject = bInIsEnterpriseProject;
+		}
+	}
+
 	class Workspace : IDisposable
 	{
 		const string BuildVersionFileName = "/Engine/Build/Build.version";
@@ -143,19 +167,14 @@ namespace UnrealGameSync
 		static readonly string LocalObjectVersionFileName = ObjectVersionFileName.Replace('/', '\\');
 
 		public readonly PerforceConnection Perforce;
-		public readonly string LocalRootPath;
-		public readonly string LocalRootPrefix;
-		public readonly string SelectedLocalFileName;
-		public readonly string ClientRootPath;
-		public readonly string SelectedClientFileName;
-		public readonly string TelemetryProjectPath;
-		public readonly bool bIsEnterpriseProject;
+		public ProjectInfo Project { get; }
+		public UserWorkspaceState State { get; }
 		Thread WorkerThread;
 		TextWriter Log;
 		bool bSyncing;
 		ProgressValue Progress = new ProgressValue();
 
-		static Workspace ActiveWorkspace;
+		static UserWorkspaceState ActiveWorkspaceState;
 
 		public event Action<WorkspaceUpdateContext, WorkspaceUpdateResult, string> OnUpdateComplete;
 
@@ -311,23 +330,14 @@ namespace UnrealGameSync
 			}
 		}
 
-		public Workspace(PerforceConnection InPerforce, string InLocalRootPath, string InSelectedLocalFileName, string InClientRootPath, string InSelectedClientFileName, int InInitialChangeNumber, string InInitialSyncFilterHash, int InLastBuiltChangeNumber, string InTelemetryProjectPath, bool bInIsEnterpriseProject, TextWriter InLog)
+		public Workspace(PerforceConnection InPerforce, ProjectInfo InProject, UserWorkspaceState InState, TextWriter InLog)
 		{
 			Perforce = InPerforce;
-			LocalRootPath = InLocalRootPath;
-			LocalRootPrefix = InLocalRootPath.TrimEnd(Path.DirectorySeparatorChar) + Path.DirectorySeparatorChar;
-			SelectedLocalFileName = InSelectedLocalFileName;
-			ClientRootPath = InClientRootPath;
-			SelectedClientFileName = InSelectedClientFileName;
-			TelemetryProjectPath = InTelemetryProjectPath;
-			CurrentChangeNumber = InInitialChangeNumber;
-			CurrentSyncFilterHash = InInitialSyncFilterHash;
-			PendingChangeNumber = InInitialChangeNumber;
-			LastBuiltChangeNumber = InLastBuiltChangeNumber;
-			bIsEnterpriseProject = bInIsEnterpriseProject;
+			Project = InProject;
+			State = InState;
 			Log = InLog;
 
-			ProjectConfigFile = ReadProjectConfigFile(InLocalRootPath, InSelectedLocalFileName, Log);
+			ProjectConfigFile = ReadProjectConfigFile(Project.LocalRootPath, Project.SelectedLocalFileName, Log);
 			ProjectStreamFilter = ReadProjectStreamFilter(Perforce, ProjectConfigFile, Log);
 		}
 
@@ -355,10 +365,6 @@ namespace UnrealGameSync
 			if(CurrentChangeNumber != Context.ChangeNumber)
 			{
 				PendingChangeNumber = Context.ChangeNumber;
-				if(!Context.Options.HasFlag(WorkspaceUpdateOptions.SyncSingleChange))
-				{
-					CurrentChangeNumber = -1;
-				}
 			}
 			Progress.Clear();
 			bSyncing = true;
@@ -382,7 +388,7 @@ namespace UnrealGameSync
 				}
 				PendingChangeNumber = CurrentChangeNumber;
 				bSyncing = false;
-				Interlocked.CompareExchange(ref ActiveWorkspace, null, this);
+				Interlocked.CompareExchange(ref ActiveWorkspaceState, null, State);
 			}
 		}
 
@@ -390,10 +396,13 @@ namespace UnrealGameSync
 		{
 			string StatusMessage;
 
+			Context.ProjectConfigFile = ProjectConfigFile;
+			Context.ProjectStreamFilter = ProjectStreamFilter;
+
 			WorkspaceUpdateResult Result = WorkspaceUpdateResult.FailedToSync;
 			try
 			{
-				Result = UpdateWorkspaceInternal(Context, out StatusMessage);
+				Result = UpdateWorkspaceInternal(Perforce, Project, State, Context, Progress, Log, out StatusMessage);
 				if(Result != WorkspaceUpdateResult.Success)
 				{
 					Log.WriteLine("{0}", StatusMessage);
@@ -415,9 +424,19 @@ namespace UnrealGameSync
 				Log.WriteException(Ex, "Failed with exception");
 			}
 
+			ProjectConfigFile = Context.ProjectConfigFile;
+			ProjectStreamFilter = Context.ProjectStreamFilter;
+
+			State.LastSyncChangeNumber = Context.ChangeNumber;
+			State.LastSyncResult = Result;
+			State.LastSyncResultMessage = StatusMessage;
+			State.LastSyncTime = DateTime.UtcNow;
+			State.LastSyncDurationSeconds = (int)(State.LastSyncTime.Value - Context.StartTime).TotalSeconds;
+			State.Save();
+
 			bSyncing = false;
 			PendingChangeNumber = CurrentChangeNumber;
-			Interlocked.CompareExchange(ref ActiveWorkspace, null, this);
+			Interlocked.CompareExchange(ref ActiveWorkspaceState, null, State);
 
 			if(OnUpdateComplete != null)
 			{
@@ -425,8 +444,10 @@ namespace UnrealGameSync
 			}
 		}
 
-		WorkspaceUpdateResult UpdateWorkspaceInternal(WorkspaceUpdateContext Context, out string StatusMessage)
+		static WorkspaceUpdateResult UpdateWorkspaceInternal(PerforceConnection Perforce, ProjectInfo Project, UserWorkspaceState State, WorkspaceUpdateContext Context, ProgressValue Progress, TextWriter Log, out string StatusMessage)
 		{
+//			string LocalRootPrefix = Project.LocalRootPath.FullName.TrimEnd(Path.DirectorySeparatorChar) + Path.DirectorySeparatorChar;
+
 			string CmdExe = Path.Combine(Environment.GetFolderPath(Environment.SpecialFolder.System), "cmd.exe");
 			if(!File.Exists(CmdExe))
 			{
@@ -439,7 +460,7 @@ namespace UnrealGameSync
 			int NumFilesSynced = 0;
 			if(Context.Options.HasFlag(WorkspaceUpdateOptions.Sync) || Context.Options.HasFlag(WorkspaceUpdateOptions.SyncSingleChange))
 			{
-				using(TelemetryStopwatch SyncTelemetryStopwatch = new TelemetryStopwatch("Workspace_Sync", TelemetryProjectPath))
+				using(TelemetryStopwatch SyncTelemetryStopwatch = new TelemetryStopwatch("Workspace_Sync", Project.TelemetryProjectPath))
 				{
 					Log.WriteLine("Syncing to {0}...", Context.ChangeNumber);
 
@@ -457,8 +478,8 @@ namespace UnrealGameSync
 					}
 
 					// Figure out which paths to sync
-					List<string> RelativeSyncPaths = GetRelativeSyncPaths((Context.Options & WorkspaceUpdateOptions.SyncAllProjects) != 0, Context.SyncFilter);
-					List<string> SyncPaths = new List<string>(RelativeSyncPaths.Select(x => ClientRootPath + x));
+					List<string> RelativeSyncPaths = GetRelativeSyncPaths(Project, (Context.Options & WorkspaceUpdateOptions.SyncAllProjects) != 0, Context.SyncFilter);
+					List<string> SyncPaths = new List<string>(RelativeSyncPaths.Select(x => Project.ClientRootPath + x));
 
 					// Get the user's sync filter
 					FileFilter UserFilter = new FileFilter(FileFilterType.Include);
@@ -485,11 +506,11 @@ namespace UnrealGameSync
 					}
 
 					// If the hash differs, enumerate everything in the workspace to find what needs to be removed
-					if (NextSyncFilterHash != CurrentSyncFilterHash)
+					if (NextSyncFilterHash != State.CurrentSyncFilterHash)
 					{
-						using (TelemetryStopwatch FilterStopwatch = new TelemetryStopwatch("Workspace_Sync_FilterChanged", TelemetryProjectPath))
+						using (TelemetryStopwatch FilterStopwatch = new TelemetryStopwatch("Workspace_Sync_FilterChanged", Project.TelemetryProjectPath))
 						{
-							Log.WriteLine("Filter has changed ({0} -> {1}); finding files in workspace that need to be removed.", (String.IsNullOrEmpty(CurrentSyncFilterHash)) ? "None" : CurrentSyncFilterHash, NextSyncFilterHash);
+							Log.WriteLine("Filter has changed ({0} -> {1}); finding files in workspace that need to be removed.", (String.IsNullOrEmpty(State.CurrentSyncFilterHash)) ? "None" : State.CurrentSyncFilterHash, NextSyncFilterHash);
 
 							// Find all the files that are in this workspace
 							List<PerforceFileRecord> HaveFiles = Context.HaveFiles;
@@ -520,8 +541,8 @@ namespace UnrealGameSync
 							{
 								try
 								{
-									string FullPath = Path.GetFullPath(HaveFile.Path);
-									if (MatchFilter(FullPath, SyncPathsFilter) && !MatchFilter(FullPath, UserFilter))
+									FileReference FullPath = new FileReference(HaveFile.Path);
+									if (MatchFilter(Project, FullPath, SyncPathsFilter) && !MatchFilter(Project, FullPath, UserFilter))
 									{
 										Log.WriteLine("  {0}", HaveFile.DepotPath);
 										RemoveDepotPaths.Add(HaveFile.DepotPath);
@@ -564,7 +585,8 @@ namespace UnrealGameSync
 							if (RemoveDepotPaths.Count > 0)
 							{
 								// Clear the current sync filter hash. If the sync is canceled, we'll be in an indeterminate state, and we should always clean next time round.
-								CurrentSyncFilterHash = "INVALID";
+								State.CurrentSyncFilterHash = "INVALID";
+								State.Save();
 
 								// Find all the depot paths that will be synced
 								HashSet<string> RemainingDepotPathsToRemove = new HashSet<string>(StringComparer.OrdinalIgnoreCase);
@@ -574,7 +596,7 @@ namespace UnrealGameSync
 								List<string> RevisionsToRemove = new List<string>();
 								RevisionsToRemove.AddRange(RemoveDepotPaths.Select(x => String.Format("{0}#0", x)));
 
-								WorkspaceUpdateResult RemoveResult = SyncFileRevisions("Removing files...", Context, RevisionsToRemove, RemainingDepotPathsToRemove, out StatusMessage);
+								WorkspaceUpdateResult RemoveResult = SyncFileRevisions(Perforce, "Removing files...", Context, RevisionsToRemove, RemainingDepotPathsToRemove, Progress, Log, out StatusMessage);
 								if (RemoveResult != WorkspaceUpdateResult.Success)
 								{
 									return RemoveResult;
@@ -582,7 +604,8 @@ namespace UnrealGameSync
 							}
 
 							// Update the sync filter hash. We've removed any files we need to at this point.
-							CurrentSyncFilterHash = NextSyncFilterHash;
+							State.CurrentSyncFilterHash = NextSyncFilterHash;
+							State.Save();
 						}
 					}
 
@@ -656,10 +679,10 @@ namespace UnrealGameSync
 								}
 
 								// Get the full local path
-								string FullName;
+								FileReference SyncFile;
 								try
 								{
-									FullName = Path.GetFullPath(SyncRecord.ClientPath);
+									SyncFile = new FileReference(SyncRecord.ClientPath);
 								}
 								catch(PathTooLongException)
 								{
@@ -670,7 +693,7 @@ namespace UnrealGameSync
 
 								// Make sure it's under the current directory. Not sure why this would happen, just being safe.
 								// This occurs for files returned from GetOpenFiles as SyncRecord.ClientPath is the client workspace path not local path.
-								if (!FullName.StartsWith(LocalRootPrefix, StringComparison.OrdinalIgnoreCase))
+								if (!SyncFile.IsUnderDirectory(Project.LocalRootPath))
 								{
 									BatchBuilder.Add(String.Format("{0}@{1}", SyncRecord.DepotPath, Context.ChangeNumber), SyncRecord.FileSize);
 									SyncDepotPaths.Add(SyncRecord.DepotPath);
@@ -679,7 +702,7 @@ namespace UnrealGameSync
 								}
 
 								// Check that it matches the filter
-								string RelativePath = FullName.Substring(LocalRootPrefix.Length).Replace('\\', '/');
+								string RelativePath = SyncFile.MakeRelativeTo(Project.LocalRootPath).Replace('\\', '/');
 								if (Filter.Matches(RelativePath))
 								{
 									long FileSize = SyncRecord.FileSize;
@@ -691,7 +714,7 @@ namespace UnrealGameSync
 									SyncTree.IncludeFile(PerforceUtils.EscapePath(RelativePath), FileSize);
 									SyncDepotPaths.Add(SyncRecord.DepotPath);
 									RequiredFreeSpace += FileSize;
-									FileInfo LocalFileInfo = new FileInfo(FullName);
+									FileInfo LocalFileInfo = SyncFile.ToFileInfo();
 
 									// If the file exists the required free space can be reduced as those bytes will be replaced.
 									if (LocalFileInfo.Exists)
@@ -708,7 +731,7 @@ namespace UnrealGameSync
 
 						try
 						{
-							DirectoryInfo LocalRootInfo = new DirectoryInfo(LocalRootPrefix);
+							DirectoryInfo LocalRootInfo = Project.LocalRootPath.ToDirectoryInfo();
 							DriveInfo Drive = new DriveInfo(LocalRootInfo.Root.FullName);
 
 							if (Drive.AvailableFreeSpace < RequiredFreeSpace)
@@ -720,20 +743,28 @@ namespace UnrealGameSync
 						}
 						catch (SystemException)
 						{
-							Log.WriteLine("Unable to check available free space for {0}.", LocalRootPrefix);
+							Log.WriteLine("Unable to check available free space for {0}.", Project.LocalRootPath);
 						}
 					}
-					SyncTree.GetOptimizedSyncCommands(ClientRootPath, PendingChangeNumber, BatchBuilder);
+					SyncTree.GetOptimizedSyncCommands(Project.ClientRootPath, Context.ChangeNumber, BatchBuilder);
+
+					// Clear the current sync changelist, in case we cancel
+					if (!Context.Options.HasFlag(WorkspaceUpdateOptions.SyncSingleChange))
+					{
+						State.CurrentChangeNumber = -1;
+						State.AdditionalChangeNumbers.Clear();
+						State.Save();
+					}
 
 					// Find all the depot paths that will be synced
 					HashSet<string> RemainingDepotPaths = new HashSet<string>(StringComparer.OrdinalIgnoreCase);
 					RemainingDepotPaths.UnionWith(SyncDepotPaths);
 
-					using (TelemetryStopwatch TransferStopwatch = new TelemetryStopwatch("Workspace_Sync_TransferFiles", TelemetryProjectPath))
+					using (TelemetryStopwatch TransferStopwatch = new TelemetryStopwatch("Workspace_Sync_TransferFiles", Project.TelemetryProjectPath))
 					{
 						TransferStopwatch.AddData(new { MachineName = Environment.MachineName, DomainName = Environment.UserDomainName, ServerAndPort = Perforce.ServerAndPort, UserName = Perforce.UserName, IncludedFiles = SyncTree.TotalIncludedFiles, ExcludedFiles = SyncTree.TotalExcludedFiles, Size = SyncTree.TotalSize, NumThreads = Context.PerforceSyncOptions.NumThreads });
 
-						WorkspaceUpdateResult SyncResult = SyncFileRevisions("Syncing files...", Context, BatchBuilder.Batches, RemainingDepotPaths, out StatusMessage);
+						WorkspaceUpdateResult SyncResult = SyncFileRevisions(Perforce, "Syncing files...", Context, BatchBuilder.Batches, RemainingDepotPaths, Progress, Log, out StatusMessage);
 						if (SyncResult != WorkspaceUpdateResult.Success)
 						{
 							TransferStopwatch.AddData(new { SyncResult = SyncResult.ToString(), CompletedFilesFiles = SyncDepotPaths.Count - RemainingDepotPaths.Count });
@@ -748,8 +779,8 @@ namespace UnrealGameSync
 					if(Context.Options.HasFlag(WorkspaceUpdateOptions.Sync) && !Context.Options.HasFlag(WorkspaceUpdateOptions.UpdateFilter))
 					{
 						// Read the new config file
-						ProjectConfigFile = ReadProjectConfigFile(LocalRootPath, SelectedLocalFileName, Log);
-						ProjectStreamFilter = ReadProjectStreamFilter(Perforce, ProjectConfigFile, Log);
+						Context.ProjectConfigFile = ReadProjectConfigFile(Project.LocalRootPath, Project.SelectedLocalFileName, Log);
+						Context.ProjectStreamFilter = ReadProjectStreamFilter(Perforce, Context.ProjectConfigFile, Log);
 
 						// Get the branch name
 						string BranchOrStreamName;
@@ -775,9 +806,9 @@ namespace UnrealGameSync
 						{
 							// Otherwise use the depot path for GenerateProjectFiles.bat in the root of the workspace
 							string DepotFileName;
-							if(!Perforce.ConvertToDepotPath(ClientRootPath + "/GenerateProjectFiles.bat", out DepotFileName, Log))
+							if(!Perforce.ConvertToDepotPath(Project.ClientRootPath + "/GenerateProjectFiles.bat", out DepotFileName, Log))
 							{
-								StatusMessage = String.Format("Couldn't determine branch name for {0}.", SelectedClientFileName);
+								StatusMessage = String.Format("Couldn't determine branch name for {0}.", Project.SelectedClientFileName);
 								return WorkspaceUpdateResult.FailedToSync;
 							}
 							BranchOrStreamName = PerforceUtils.GetClientOrDepotDirectoryName(DepotFileName);
@@ -785,41 +816,41 @@ namespace UnrealGameSync
 
 						// Find the last code change before this changelist. For consistency in versioning between local builds and precompiled binaries, we need to use the last submitted code changelist as our version number.
 						List<PerforceChangeSummary> CodeChanges;
-						if(!Perforce.FindChanges(new string[]{ ".cs", ".h", ".cpp", ".usf", ".ush", ".uproject", ".uplugin" }.SelectMany(x => SyncPaths.Select(y => String.Format("{0}{1}@<={2}", y, x, PendingChangeNumber))), 1, out CodeChanges, Log))
+						if(!Perforce.FindChanges(new string[]{ ".cs", ".h", ".cpp", ".usf", ".ush", ".uproject", ".uplugin" }.SelectMany(x => SyncPaths.Select(y => String.Format("{0}{1}@<={2}", y, x, Context.ChangeNumber))), 1, out CodeChanges, Log))
 						{
-							StatusMessage = String.Format("Couldn't determine last code changelist before CL {0}.", PendingChangeNumber);
+							StatusMessage = String.Format("Couldn't determine last code changelist before CL {0}.", Context.ChangeNumber);
 							return WorkspaceUpdateResult.FailedToSync;
 						}
 						if(CodeChanges.Count == 0)
 						{
-							StatusMessage = String.Format("Could not find any code changes before CL {0}.", PendingChangeNumber);
+							StatusMessage = String.Format("Could not find any code changes before CL {0}.", Context.ChangeNumber);
 							return WorkspaceUpdateResult.FailedToSync;
 						}
 
 						// Get the last code change
-						if(ProjectConfigFile.GetValue("Options.VersionToLastCodeChange", true))
+						if(Context.ProjectConfigFile.GetValue("Options.VersionToLastCodeChange", true))
 						{
 							VersionChangeNumber = CodeChanges.Max(x => x.Number);
 						}
 						else
 						{
-							VersionChangeNumber = PendingChangeNumber;
+							VersionChangeNumber = Context.ChangeNumber;
 						}
 
 						// Update the version files
-						if(ProjectConfigFile.GetValue("Options.UseFastModularVersioningV2", false))
+						if(Context.ProjectConfigFile.GetValue("Options.UseFastModularVersioningV2", false))
 						{
-							bool bIsLicenseeVersion = IsLicenseeVersion();
-							if (!UpdateVersionFile(ClientRootPath + BuildVersionFileName, PendingChangeNumber, Text => UpdateBuildVersion(Text, PendingChangeNumber, VersionChangeNumber, BranchOrStreamName, bIsLicenseeVersion)))
+							bool bIsLicenseeVersion = IsLicenseeVersion(Perforce, Project, Log);
+							if (!UpdateVersionFile(Perforce, Project.ClientRootPath + BuildVersionFileName, Context.ChangeNumber, Text => UpdateBuildVersion(Text, Context.ChangeNumber, VersionChangeNumber, BranchOrStreamName, bIsLicenseeVersion), Log))
 							{
 								StatusMessage = String.Format("Failed to update {0}.", BuildVersionFileName);
 								return WorkspaceUpdateResult.FailedToSync;
 							}
 						}
-						else if(ProjectConfigFile.GetValue("Options.UseFastModularVersioning", false))
+						else if(Context.ProjectConfigFile.GetValue("Options.UseFastModularVersioning", false))
 						{
-							bool bIsLicenseeVersion = IsLicenseeVersion();
-							if (!UpdateVersionFile(ClientRootPath + BuildVersionFileName, PendingChangeNumber, Text => UpdateBuildVersion(Text, PendingChangeNumber, VersionChangeNumber, BranchOrStreamName, bIsLicenseeVersion)))
+							bool bIsLicenseeVersion = IsLicenseeVersion(Perforce, Project, Log);
+							if (!UpdateVersionFile(Perforce, Project.ClientRootPath + BuildVersionFileName, Context.ChangeNumber, Text => UpdateBuildVersion(Text, Context.ChangeNumber, VersionChangeNumber, BranchOrStreamName, bIsLicenseeVersion), Log))
 							{
 								StatusMessage = String.Format("Failed to update {0}.", BuildVersionFileName);
 								return WorkspaceUpdateResult.FailedToSync;
@@ -829,12 +860,12 @@ namespace UnrealGameSync
 							VersionHeaderStrings["#define ENGINE_IS_PROMOTED_BUILD"] = " (0)";
 							VersionHeaderStrings["#define BUILT_FROM_CHANGELIST"] = " 0";
 							VersionHeaderStrings["#define BRANCH_NAME"] = " \"" + BranchOrStreamName.Replace('/', '+') + "\"";
-							if(!UpdateVersionFile(ClientRootPath + VersionHeaderFileName, VersionHeaderStrings, PendingChangeNumber))
+							if(!UpdateVersionFile(Perforce, Project.ClientRootPath + VersionHeaderFileName, VersionHeaderStrings, Context.ChangeNumber, Log))
 							{
 								StatusMessage = String.Format("Failed to update {0}.", VersionHeaderFileName);
 								return WorkspaceUpdateResult.FailedToSync;
 							}
-							if(!UpdateVersionFile(ClientRootPath + ObjectVersionFileName, new Dictionary<string,string>(), PendingChangeNumber))
+							if(!UpdateVersionFile(Perforce, Project.ClientRootPath + ObjectVersionFileName, new Dictionary<string,string>(), Context.ChangeNumber, Log))
 							{
 								StatusMessage = String.Format("Failed to update {0}.", ObjectVersionFileName);
 								return WorkspaceUpdateResult.FailedToSync;
@@ -842,7 +873,7 @@ namespace UnrealGameSync
 						}
 						else
 						{
-							if(!UpdateVersionFile(ClientRootPath + BuildVersionFileName, new Dictionary<string, string>(), PendingChangeNumber))
+							if(!UpdateVersionFile(Perforce, Project.ClientRootPath + BuildVersionFileName, new Dictionary<string, string>(), Context.ChangeNumber, Log))
 							{
 								StatusMessage = String.Format("Failed to update {0}", BuildVersionFileName);
 								return WorkspaceUpdateResult.FailedToSync;
@@ -853,12 +884,12 @@ namespace UnrealGameSync
 							VersionStrings["#define ENGINE_IS_PROMOTED_BUILD"] = " (0)";
 							VersionStrings["#define BUILT_FROM_CHANGELIST"] = " " + VersionChangeNumber.ToString();
 							VersionStrings["#define BRANCH_NAME"] = " \"" + BranchOrStreamName.Replace('/', '+') + "\"";
-							if(!UpdateVersionFile(ClientRootPath + VersionHeaderFileName, VersionStrings, PendingChangeNumber))
+							if(!UpdateVersionFile(Perforce, Project.ClientRootPath + VersionHeaderFileName, VersionStrings, Context.ChangeNumber, Log))
 							{
 								StatusMessage = String.Format("Failed to update {0}", VersionHeaderFileName);
 								return WorkspaceUpdateResult.FailedToSync;
 							}
-							if(!UpdateVersionFile(ClientRootPath + ObjectVersionFileName, VersionStrings, PendingChangeNumber))
+							if(!UpdateVersionFile(Perforce, Project.ClientRootPath + ObjectVersionFileName, VersionStrings, Context.ChangeNumber, Log))
 							{
 								StatusMessage = String.Format("Failed to update {0}", ObjectVersionFileName);
 								return WorkspaceUpdateResult.FailedToSync;
@@ -866,15 +897,15 @@ namespace UnrealGameSync
 						}
 
 						// Remove all the receipts for build targets in this branch
-						if(SelectedClientFileName.EndsWith(".uproject", StringComparison.OrdinalIgnoreCase))
+						if(Project.SelectedClientFileName.EndsWith(".uproject", StringComparison.OrdinalIgnoreCase))
 						{
-							Perforce.Sync(PerforceUtils.GetClientOrDepotDirectoryName(SelectedClientFileName) + "/Build/Receipts/...#0", Log);
+							Perforce.Sync(PerforceUtils.GetClientOrDepotDirectoryName(Project.SelectedClientFileName) + "/Build/Receipts/...#0", Log);
 						}
 					}
 
 					// Check if there are any files which need resolving
 					List<PerforceFileRecord> UnresolvedFiles;
-					if(!FindUnresolvedFiles(SyncPaths, out UnresolvedFiles))
+					if(!FindUnresolvedFiles(Perforce, SyncPaths, Log, out UnresolvedFiles))
 					{
 						StatusMessage = "Couldn't get list of unresolved files.";
 						return WorkspaceUpdateResult.FailedToSync;
@@ -885,7 +916,7 @@ namespace UnrealGameSync
 						{
 							Perforce.AutoResolveFile(UnresolvedFile.DepotPath, Log);
 						}
-						if(!FindUnresolvedFiles(SyncPaths, out UnresolvedFiles))
+						if(!FindUnresolvedFiles(Perforce, SyncPaths, Log, out UnresolvedFiles))
 						{
 							StatusMessage = "Couldn't get list of unresolved files.";
 							return WorkspaceUpdateResult.FailedToSync;
@@ -906,14 +937,14 @@ namespace UnrealGameSync
 					if (Context.Options.HasFlag(WorkspaceUpdateOptions.Sync) && !Context.Options.HasFlag(WorkspaceUpdateOptions.UpdateFilter))
 					{
 						// Execute any project specific post-sync steps
-						string[] PostSyncSteps = ProjectConfigFile.GetValues("Sync.Step", null);
+						string[] PostSyncSteps = Context.ProjectConfigFile.GetValues("Sync.Step", null);
 						if (PostSyncSteps != null)
 						{
 							Log.WriteLine();
 							Log.WriteLine("Executing post-sync steps...");
 
 							Dictionary<string, string> PostSyncVariables = new Dictionary<string, string>(Context.Variables);
-							PostSyncVariables["Change"] = PendingChangeNumber.ToString();
+							PostSyncVariables["Change"] = Context.ChangeNumber.ToString();
 							PostSyncVariables["CodeChange"] = VersionChangeNumber.ToString();
 
 							foreach (string PostSyncStep in PostSyncSteps.Select(x => x.Trim()))
@@ -936,10 +967,18 @@ namespace UnrealGameSync
 								}
 							}
 						}
+					}
 
-						// Update the current change number. Everything else happens for the new change.
-						CurrentChangeNumber = PendingChangeNumber;
+					// Update the current state
+					if (Context.Options.HasFlag(WorkspaceUpdateOptions.SyncSingleChange))
+					{
+						State.AdditionalChangeNumbers.Add(Context.ChangeNumber);
 					}
+					else
+					{
+						State.CurrentChangeNumber = Context.ChangeNumber;
+					}
+					State.Save();
 
 					// Update the timing info
 					Times.Add(new Tuple<string,TimeSpan>("Sync", SyncTelemetryStopwatch.Stop("Success")));
@@ -953,19 +992,19 @@ namespace UnrealGameSync
 			// Extract an archive from the depot path
 			if(Context.Options.HasFlag(WorkspaceUpdateOptions.SyncArchives))
 			{
-				using(TelemetryStopwatch Stopwatch = new TelemetryStopwatch("Workspace_SyncArchives", TelemetryProjectPath))
+				using(TelemetryStopwatch Stopwatch = new TelemetryStopwatch("Workspace_SyncArchives", Project.TelemetryProjectPath))
 				{
 					// Create the directory for extracted archive manifests
-					string ManifestDirectoryName;
-					if(SelectedLocalFileName.EndsWith(".uproject", StringComparison.OrdinalIgnoreCase))
+					DirectoryReference ManifestDirectoryName;
+					if(Project.SelectedLocalFileName.HasExtension(".uproject"))
 					{
-						ManifestDirectoryName = Path.Combine(Path.GetDirectoryName(SelectedLocalFileName), "Saved", "UnrealGameSync");
+						ManifestDirectoryName = DirectoryReference.Combine(Project.SelectedLocalFileName.Directory, "Saved", "UnrealGameSync");
 					}
 					else
 					{
-						ManifestDirectoryName = Path.Combine(Path.GetDirectoryName(SelectedLocalFileName), "Engine", "Saved", "UnrealGameSync");
+						ManifestDirectoryName = DirectoryReference.Combine(Project.SelectedLocalFileName.Directory, "Engine", "Saved", "UnrealGameSync");
 					}
-					Directory.CreateDirectory(ManifestDirectoryName);
+					DirectoryReference.CreateDirectory(ManifestDirectoryName);
 
 					// Sync and extract (or just remove) the given archives
 					foreach(KeyValuePair<string, Tuple<IArchiveInfo, string>> ArchiveTypeAndArchive in Context.ArchiveTypeToArchive)
@@ -973,13 +1012,13 @@ namespace UnrealGameSync
 						string ArchiveType = ArchiveTypeAndArchive.Key;
 
 						// Remove any existing binaries
-						string ManifestFileName = Path.Combine(ManifestDirectoryName, String.Format("{0}.zipmanifest", ArchiveType));
-						if(File.Exists(ManifestFileName))
+						FileReference ManifestFileName = FileReference.Combine(ManifestDirectoryName, String.Format("{0}.zipmanifest", ArchiveType));
+						if(FileReference.Exists(ManifestFileName))
 						{
 							Log.WriteLine("Removing {0} binaries...", ArchiveType);
 							Progress.Set(String.Format("Removing {0} binaries...", ArchiveType), 0.0f);
-							ArchiveUtils.RemoveExtractedFiles(LocalRootPath, ManifestFileName, Progress, Log);
-							File.Delete(ManifestFileName);
+							ArchiveUtils.RemoveExtractedFiles(Project.LocalRootPath, ManifestFileName, Progress, Log);
+							FileReference.Delete(ManifestFileName);
 							Log.WriteLine();
 						}
 
@@ -991,7 +1030,7 @@ namespace UnrealGameSync
 
 							Log.WriteLine("Syncing {0} binaries...", ArchiveType.ToLowerInvariant());
 							Progress.Set(String.Format("Syncing {0} binaries...", ArchiveType.ToLowerInvariant()), 0.0f);
-							if (!ArchiveInfo.DownloadArchive(ArchiveKey, LocalRootPath, ManifestFileName, Log, Progress))
+							if (!ArchiveInfo.DownloadArchive(ArchiveKey, Project.LocalRootPath, ManifestFileName, Log, Progress))
 							{
 								StatusMessage = String.Format("Couldn't read {0}", ArchiveKey);
 								return WorkspaceUpdateResult.FailedToSync;
@@ -999,6 +1038,10 @@ namespace UnrealGameSync
 						}
 					}
 
+					// Update the state
+					State.ExpandedArchiveTypes = Context.ArchiveTypeToArchive.Where(x => x.Value != null).Select(x => x.Key).ToArray();
+					State.Save();
+
 					// Add the finish time
 					Times.Add(new Tuple<string,TimeSpan>("Archive", Stopwatch.Stop("Success")));
 				}
@@ -1007,10 +1050,10 @@ namespace UnrealGameSync
 			// Take the lock before doing anything else. Building and generating project files can only be done on one workspace at a time.
 			if(Context.Options.HasFlag(WorkspaceUpdateOptions.GenerateProjectFiles) || Context.Options.HasFlag(WorkspaceUpdateOptions.Build))
 			{
-				if(Interlocked.CompareExchange(ref ActiveWorkspace, this, null) != null)
+				if(Interlocked.CompareExchange(ref ActiveWorkspaceState, State, null) != null)
 				{
 					Log.WriteLine("Waiting for other workspaces to finish...");
-					while(Interlocked.CompareExchange(ref ActiveWorkspace, this, null) != null)
+					while(Interlocked.CompareExchange(ref ActiveWorkspaceState, State, null) != null)
 					{
 						Thread.Sleep(100);
 					}
@@ -1020,17 +1063,17 @@ namespace UnrealGameSync
 			// Generate project files in the workspace
 			if(Context.Options.HasFlag(WorkspaceUpdateOptions.GenerateProjectFiles))
 			{
-				using(TelemetryStopwatch Stopwatch = new TelemetryStopwatch("Workspace_GenerateProjectFiles", TelemetryProjectPath))
+				using(TelemetryStopwatch Stopwatch = new TelemetryStopwatch("Workspace_GenerateProjectFiles", Project.TelemetryProjectPath))
 				{
 					Progress.Set("Generating project files...", 0.0f);
 
 					StringBuilder CommandLine = new StringBuilder();
-					CommandLine.AppendFormat("/C \"\"{0}\"", Path.Combine(LocalRootPath, "GenerateProjectFiles.bat"));
+					CommandLine.AppendFormat("/C \"\"{0}\"", FileReference.Combine(Project.LocalRootPath, "GenerateProjectFiles.bat"));
 					if((Context.Options & WorkspaceUpdateOptions.SyncAllProjects) == 0 && (Context.Options & WorkspaceUpdateOptions.IncludeAllProjectsInSolution) == 0)
 					{
-						if(SelectedLocalFileName.EndsWith(".uproject", StringComparison.OrdinalIgnoreCase))
+						if(Project.SelectedLocalFileName.HasExtension(".uproject"))
 						{
-							CommandLine.AppendFormat(" \"{0}\"", SelectedLocalFileName);
+							CommandLine.AppendFormat(" \"{0}\"", Project.SelectedLocalFileName);
 						}
 					}
 					CommandLine.Append(" -progress\"");
@@ -1055,7 +1098,7 @@ namespace UnrealGameSync
 			{
 				// Compile all the build steps together
 				Dictionary<Guid, ConfigObject> BuildStepObjects = Context.DefaultBuildSteps.ToDictionary(x => x.Key, x => new ConfigObject(x.Value));
-				BuildStep.MergeBuildStepObjects(BuildStepObjects, ProjectConfigFile.GetValues("Build.Step", new string[0]).Select(x => new ConfigObject(x)));
+				BuildStep.MergeBuildStepObjects(BuildStepObjects, Context.ProjectConfigFile.GetValues("Build.Step", new string[0]).Select(x => new ConfigObject(x)));
 				BuildStep.MergeBuildStepObjects(BuildStepObjects, Context.UserBuildStepObjects);
 
 				// Construct build steps from them
@@ -1075,14 +1118,14 @@ namespace UnrealGameSync
 
 				// Check if the last successful build was before a change that we need to force a clean for
 				bool bForceClean = false;
-				if(LastBuiltChangeNumber != 0)
+				if(State.LastBuiltChangeNumber != 0)
 				{
-					foreach(string CleanBuildChange in ProjectConfigFile.GetValues("ForceClean.Changelist", new string[0]))
+					foreach(string CleanBuildChange in Context.ProjectConfigFile.GetValues("ForceClean.Changelist", new string[0]))
 					{
 						int ChangeNumber;
 						if(int.TryParse(CleanBuildChange, out ChangeNumber))
 						{
-							if((LastBuiltChangeNumber >= ChangeNumber) != (CurrentChangeNumber >= ChangeNumber))
+							if((State.LastBuiltChangeNumber >= ChangeNumber) != (State.CurrentChangeNumber >= ChangeNumber))
 							{
 								Log.WriteLine("Forcing clean build due to changelist {0}.", ChangeNumber);
 								Log.WriteLine();
@@ -1094,7 +1137,7 @@ namespace UnrealGameSync
 				}
 
 				// Execute them all
-				using(TelemetryStopwatch Stopwatch = new TelemetryStopwatch("Workspace_Build", TelemetryProjectPath))
+				using(TelemetryStopwatch Stopwatch = new TelemetryStopwatch("Workspace_Build", Project.TelemetryProjectPath))
 				{
 					Progress.Set("Starting build...", 0.0f);
 
@@ -1114,11 +1157,11 @@ namespace UnrealGameSync
 							switch(Step.Type)
 							{
 								case BuildStepType.Compile:
-									using(TelemetryStopwatch StepStopwatch = new TelemetryStopwatch("Workspace_Execute_Compile", TelemetryProjectPath))
+									using(TelemetryStopwatch StepStopwatch = new TelemetryStopwatch("Workspace_Execute_Compile", Project.TelemetryProjectPath))
 									{
 										StepStopwatch.AddData(new { Target = Step.Target });
 
-										string BuildBat = Path.Combine(LocalRootPath, "Engine", "Build", "BatchFiles", "Build.bat");
+										FileReference BuildBat = FileReference.Combine(Project.LocalRootPath, "Engine", "Build", "BatchFiles", "Build.bat");
 										string CommandLine = String.Format("{0} {1} {2} {3} -NoHotReloadFromIDE", Step.Target, Step.Platform, Step.Configuration, Utility.ExpandVariables(Step.Arguments ?? "", Context.Variables));
 										if(Context.Options.HasFlag(WorkspaceUpdateOptions.Clean) || bForceClean)
 										{
@@ -1133,19 +1176,19 @@ namespace UnrealGameSync
 										{
 											StepStopwatch.Stop("Failed");
 											StatusMessage = String.Format("Failed to compile {0}.", Step.Target);
-											return (HasModifiedSourceFiles() || Context.UserBuildStepObjects.Count > 0)? WorkspaceUpdateResult.FailedToCompile : WorkspaceUpdateResult.FailedToCompileWithCleanWorkspace;
+											return (HasModifiedSourceFiles(Perforce, Project, Log) || Context.UserBuildStepObjects.Count > 0)? WorkspaceUpdateResult.FailedToCompile : WorkspaceUpdateResult.FailedToCompileWithCleanWorkspace;
 										}
 
 										StepStopwatch.Stop("Success");
 									}
 									break;
 								case BuildStepType.Cook:
-									using(TelemetryStopwatch StepStopwatch = new TelemetryStopwatch("Workspace_Execute_Cook", TelemetryProjectPath))
+									using(TelemetryStopwatch StepStopwatch = new TelemetryStopwatch("Workspace_Execute_Cook", Project.TelemetryProjectPath))
 									{
 										StepStopwatch.AddData(new { Project = Path.GetFileNameWithoutExtension(Step.FileName) });
 
-										string LocalRunUAT = Path.Combine(LocalRootPath, "Engine", "Build", "BatchFiles", "RunUAT.bat");
-										string Arguments = String.Format("/C \"\"{0}\" -profile=\"{1}\"\"", LocalRunUAT, Path.Combine(LocalRootPath, Step.FileName));
+										FileReference LocalRunUAT = FileReference.Combine(Project.LocalRootPath, "Engine", "Build", "BatchFiles", "RunUAT.bat");
+										string Arguments = String.Format("/C \"\"{0}\" -profile=\"{1}\"\"", LocalRunUAT, FileReference.Combine(Project.LocalRootPath, Step.FileName));
 										Log.WriteLine("uat> Running {0} {1}", LocalRunUAT, Arguments);
 
 										int ResultFromUAT = Utility.ExecuteProcess(CmdExe, null, Arguments, null, new ProgressTextWriter(Progress, new PrefixedTextWriter("uat> ", Log)));
@@ -1160,18 +1203,18 @@ namespace UnrealGameSync
 									}
 									break;
 								case BuildStepType.Other:
-									using(TelemetryStopwatch StepStopwatch = new TelemetryStopwatch("Workspace_Execute_Custom", TelemetryProjectPath))
+									using(TelemetryStopwatch StepStopwatch = new TelemetryStopwatch("Workspace_Execute_Custom", Project.TelemetryProjectPath))
 									{
 										StepStopwatch.AddData(new { FileName = Path.GetFileNameWithoutExtension(Step.FileName) });
 
-										string ToolFileName = Path.Combine(LocalRootPath, Utility.ExpandVariables(Step.FileName, Context.Variables));
-										string ToolWorkingDir = String.IsNullOrWhiteSpace(Step.WorkingDir) ? Path.GetDirectoryName(ToolFileName) : Utility.ExpandVariables(Step.WorkingDir, Context.Variables);
+										FileReference ToolFileName = FileReference.Combine(Project.LocalRootPath, Utility.ExpandVariables(Step.FileName, Context.Variables));
+										string ToolWorkingDir = String.IsNullOrWhiteSpace(Step.WorkingDir) ? ToolFileName.Directory.FullName : Utility.ExpandVariables(Step.WorkingDir, Context.Variables);
 										string ToolArguments = Utility.ExpandVariables(Step.Arguments ?? "", Context.Variables);
 										Log.WriteLine("tool> Running {0} {1}", ToolFileName, ToolArguments);
 
 										if(Step.bUseLogWindow)
 										{
-											int ResultFromTool = Utility.ExecuteProcess(ToolFileName, ToolWorkingDir, ToolArguments, null, new ProgressTextWriter(Progress, new PrefixedTextWriter("tool> ", Log)));
+											int ResultFromTool = Utility.ExecuteProcess(ToolFileName.FullName, ToolWorkingDir, ToolArguments, null, new ProgressTextWriter(Progress, new PrefixedTextWriter("tool> ", Log)));
 											if(ResultFromTool != 0)
 											{
 												StepStopwatch.Stop("Failed");
@@ -1181,7 +1224,7 @@ namespace UnrealGameSync
 										}
 										else
 										{
-											ProcessStartInfo StartInfo = new ProcessStartInfo(ToolFileName, ToolArguments);
+											ProcessStartInfo StartInfo = new ProcessStartInfo(ToolFileName.FullName, ToolArguments);
 											StartInfo.WorkingDirectory = ToolWorkingDir;
 											using(Process.Start(StartInfo))
 											{
@@ -1204,7 +1247,8 @@ namespace UnrealGameSync
 				// Update the last successful build change number
 				if(Context.CustomBuildSteps == null || Context.CustomBuildSteps.Count == 0)
 				{
-					LastBuiltChangeNumber = CurrentChangeNumber;
+					State.LastBuiltChangeNumber = State.CurrentChangeNumber;
+					State.Save();
 				}
 			}
 
@@ -1227,39 +1271,39 @@ namespace UnrealGameSync
 			return WorkspaceUpdateResult.Success;
 		}
 
-		bool IsLicenseeVersion()
+		static bool IsLicenseeVersion(PerforceConnection Perforce, ProjectInfo Project, TextWriter Log)
 		{
 			bool bIsEpicInternal;
-			if (Perforce.FileExists(ClientRootPath + "/Engine/Build/NotForLicensees/EpicInternal.txt", out bIsEpicInternal, Log) && bIsEpicInternal)
+			if (Perforce.FileExists(Project.ClientRootPath + "/Engine/Build/NotForLicensees/EpicInternal.txt", out bIsEpicInternal, Log) && bIsEpicInternal)
 			{
 				return false;
 			}
-			if (Perforce.FileExists(ClientRootPath + "/Engine/Restricted/NotForLicensees/Build/EpicInternal.txt", out bIsEpicInternal, Log) && bIsEpicInternal)
+			if (Perforce.FileExists(Project.ClientRootPath + "/Engine/Restricted/NotForLicensees/Build/EpicInternal.txt", out bIsEpicInternal, Log) && bIsEpicInternal)
 			{
 				return false;
 			}
 			return true;
 		}
 
-		public List<string> GetSyncPaths(bool bSyncAllProjects, string[] SyncFilter)
+		public static List<string> GetSyncPaths(ProjectInfo Project, bool bSyncAllProjects, string[] SyncFilter)
 		{
-			List<string> SyncPaths = GetRelativeSyncPaths(bSyncAllProjects, SyncFilter);
-			return SyncPaths.Select(x => ClientRootPath + x).ToList();
+			List<string> SyncPaths = GetRelativeSyncPaths(Project, bSyncAllProjects, SyncFilter);
+			return SyncPaths.Select(x => Project.ClientRootPath + x).ToList();
 		}
 
-		public List<string> GetRelativeSyncPaths(bool bSyncAllProjects, string[] SyncFilter)
+		public static List<string> GetRelativeSyncPaths(ProjectInfo Project, bool bSyncAllProjects, string[] SyncFilter)
 		{
 			List<string> SyncPaths = new List<string>();
 
 			// Check the client path is formatted correctly
-			if (!SelectedClientFileName.StartsWith(ClientRootPath + "/"))
+			if (!Project.SelectedClientFileName.StartsWith(Project.ClientRootPath + "/"))
 			{
-				throw new Exception(String.Format("Expected '{0}' to start with '{1}'", SelectedClientFileName, ClientRootPath));
+				throw new Exception(String.Format("Expected '{0}' to start with '{1}'", Project.SelectedClientFileName, Project.ClientRootPath));
 			}
 
 			// Add the default project paths
-			int LastSlashIdx = SelectedClientFileName.LastIndexOf('/');
-			if (bSyncAllProjects || !SelectedClientFileName.EndsWith(".uproject", StringComparison.OrdinalIgnoreCase) || LastSlashIdx <= ClientRootPath.Length)
+			int LastSlashIdx = Project.SelectedClientFileName.LastIndexOf('/');
+			if (bSyncAllProjects || !Project.SelectedClientFileName.EndsWith(".uproject", StringComparison.OrdinalIgnoreCase) || LastSlashIdx <= Project.ClientRootPath.Length)
 			{
 				SyncPaths.Add("/...");
 			}
@@ -1267,11 +1311,11 @@ namespace UnrealGameSync
 			{
 				SyncPaths.Add("/*");
 				SyncPaths.Add("/Engine/...");
-				if(bIsEnterpriseProject)
+				if(Project.bIsEnterpriseProject)
 				{
 					SyncPaths.Add("/Enterprise/...");
 				}
-				SyncPaths.Add(SelectedClientFileName.Substring(ClientRootPath.Length, LastSlashIdx - ClientRootPath.Length) + "/...");
+				SyncPaths.Add(Project.SelectedClientFileName.Substring(Project.ClientRootPath.Length, LastSlashIdx - Project.ClientRootPath.Length) + "/...");
 			}
 
 			// Apply the sync filter to that list. We only want inclusive rules in the output list, but we can manually apply exclusions to previous entries.
@@ -1312,12 +1356,13 @@ namespace UnrealGameSync
 			return SyncPaths;
 		}
 
-		public bool MatchFilter(string FileName, FileFilter Filter)
+		public static bool MatchFilter(ProjectInfo Project, FileReference FileName, FileFilter Filter)
 		{
 			bool bMatch = true;
-			if(FileName.StartsWith(LocalRootPath, StringComparison.OrdinalIgnoreCase))
+			if(FileName.IsUnderDirectory(Project.LocalRootPath))
 			{
-				if(!Filter.Matches(FileName.Substring(LocalRootPath.Length)))
+				string RelativePath = FileName.MakeRelativeTo(Project.LocalRootPath);
+				if(!Filter.Matches(RelativePath))
 				{
 					bMatch = false;
 				}
@@ -1334,14 +1379,14 @@ namespace UnrealGameSync
 			public WorkspaceUpdateResult Result = WorkspaceUpdateResult.Success;
 		}
 
-		WorkspaceUpdateResult SyncFileRevisions(string Prefix, WorkspaceUpdateContext Context, List<string> SyncCommands, HashSet<string> RemainingDepotPaths, out string StatusMessage)
+		static WorkspaceUpdateResult SyncFileRevisions(PerforceConnection Perforce, string Prefix, WorkspaceUpdateContext Context, List<string> SyncCommands, HashSet<string> RemainingDepotPaths, ProgressValue Progress, TextWriter Log, out string StatusMessage)
 		{
 			Queue<List<string>> SyncCommandLists = new Queue<List<string>>();
 			SyncCommandLists.Enqueue(SyncCommands);
-			return SyncFileRevisions(Prefix, Context, SyncCommandLists, RemainingDepotPaths, out StatusMessage);
+			return SyncFileRevisions(Perforce, Prefix, Context, SyncCommandLists, RemainingDepotPaths, Progress, Log, out StatusMessage);
 		}
 
-		WorkspaceUpdateResult SyncFileRevisions(string Prefix, WorkspaceUpdateContext Context, Queue<List<string>> SyncCommandLists, HashSet<string> RemainingDepotPaths, out string StatusMessage)
+		static WorkspaceUpdateResult SyncFileRevisions(PerforceConnection Perforce, string Prefix, WorkspaceUpdateContext Context, Queue<List<string>> SyncCommandLists, HashSet<string> RemainingDepotPaths, ProgressValue Progress, TextWriter Log, out string StatusMessage)
 		{
 			// Figure out the number of additional background threads we want to run with. We can run worker on the current thread.
 			int NumExtraThreads = Math.Max(Math.Min(SyncCommandLists.Count, Context.PerforceSyncOptions.NumThreads) - 1, 0);
@@ -1368,14 +1413,14 @@ namespace UnrealGameSync
 				for (int ThreadIdx = 0; ThreadIdx < NumExtraThreads; ThreadIdx++)
 				{
 					int ThreadNumber = ThreadIdx + 2;
-					Thread ChildThread = new Thread(() => StaticSyncWorker(ThreadNumber, Perforce, PendingChangeNumber, Context, State, SyncOutput, LogWrapper));
+					Thread ChildThread = new Thread(() => StaticSyncWorker(ThreadNumber, Perforce, Context, State, SyncOutput, LogWrapper));
 					ChildThread.Name = "Sync File Revisions";
 					ChildThreads.Add(ChildThread);
 					ChildThread.Start();
 				}
 
 				// Run one worker on the current thread
-				StaticSyncWorker(1, Perforce, PendingChangeNumber, Context, State, SyncOutput, LogWrapper);
+				StaticSyncWorker(1, Perforce, Context, State, SyncOutput, LogWrapper);
 
 				// Wait for all the background threads to finish
 				foreach (Thread ChildThread in ChildThreads)
@@ -1424,7 +1469,7 @@ namespace UnrealGameSync
 			}
 		}
 
-		static void StaticSyncWorker(int ThreadNumber, PerforceConnection Perforce, int PendingChangeNumber, WorkspaceUpdateContext Context, SyncState State, Action<PerforceFileRecord, LineBasedTextWriter> SyncOutput, LineBasedTextWriter GlobalLog)
+		static void StaticSyncWorker(int ThreadNumber, PerforceConnection Perforce, WorkspaceUpdateContext Context, SyncState State, Action<PerforceFileRecord, LineBasedTextWriter> SyncOutput, LineBasedTextWriter GlobalLog)
 		{
 			PrefixedTextWriter ThreadLog = new PrefixedTextWriter(String.Format("{0}:", ThreadNumber), GlobalLog);
 			for (; ; )
@@ -1448,7 +1493,7 @@ namespace UnrealGameSync
 				WorkspaceUpdateResult Result;
 				try
 				{
-					Result = StaticSyncFileRevisions(Perforce, PendingChangeNumber, Context, SyncCommands, Record => SyncOutput(Record, ThreadLog), ThreadLog, out StatusMessage);
+					Result = StaticSyncFileRevisions(Perforce, Context, SyncCommands, Record => SyncOutput(Record, ThreadLog), ThreadLog, out StatusMessage);
 				}
 				catch (ThreadInterruptedException)
 				{
@@ -1472,7 +1517,7 @@ namespace UnrealGameSync
 			}
 		}
 
-		static WorkspaceUpdateResult StaticSyncFileRevisions(PerforceConnection Perforce, int PendingChangeNumber, WorkspaceUpdateContext Context, List<string> SyncCommands, Action<PerforceFileRecord> SyncOutput, TextWriter Log, out string StatusMessage)
+		static WorkspaceUpdateResult StaticSyncFileRevisions(PerforceConnection Perforce, WorkspaceUpdateContext Context, List<string> SyncCommands, Action<PerforceFileRecord> SyncOutput, TextWriter Log, out string StatusMessage)
 		{
 			// Sync them all
 			List<string> TamperedFiles = new List<string>();
@@ -1506,7 +1551,7 @@ namespace UnrealGameSync
 				}
 				foreach(string TamperedFile in TamperedFiles)
 				{
-					if(Context.ClobberFiles[TamperedFile] && !Perforce.ForceSync(TamperedFile, PendingChangeNumber, Log))
+					if(Context.ClobberFiles[TamperedFile] && !Perforce.ForceSync(TamperedFile, Context.ChangeNumber, Log))
 					{
 						StatusMessage = String.Format("Couldn't sync {0}.", TamperedFile);
 						return WorkspaceUpdateResult.FailedToSync;
@@ -1519,11 +1564,11 @@ namespace UnrealGameSync
 			return WorkspaceUpdateResult.Success;
 		}
 
-		static ConfigFile ReadProjectConfigFile(string LocalRootPath, string SelectedLocalFileName, TextWriter Log)
+		static ConfigFile ReadProjectConfigFile(DirectoryReference LocalRootPath, FileReference SelectedLocalFileName, TextWriter Log)
 		{
 			// Find the valid config file paths
-			DirectoryInfo EngineDir = new DirectoryInfo(Path.Combine(LocalRootPath, "Engine"));
-			List<FileInfo> LocalConfigFiles = Utility.GetLocalConfigPaths(EngineDir, new FileInfo(SelectedLocalFileName));
+			DirectoryInfo EngineDir = DirectoryReference.Combine(LocalRootPath, "Engine").ToDirectoryInfo();
+			List<FileInfo> LocalConfigFiles = Utility.GetLocalConfigPaths(EngineDir, SelectedLocalFileName.ToFileInfo());
 
 			// Read them in
 			ConfigFile ProjectConfig = new ConfigFile();
@@ -1572,10 +1617,10 @@ namespace UnrealGameSync
 			}
 		}
 
-		bool HasModifiedSourceFiles()
+		static bool HasModifiedSourceFiles(PerforceConnection Perforce, ProjectInfo Project, TextWriter Log)
 		{
 			List<PerforceFileRecord> OpenFiles;
-			if(!Perforce.GetOpenFiles(ClientRootPath + "/...", out OpenFiles, Log))
+			if(!Perforce.GetOpenFiles(Project.ClientRootPath + "/...", out OpenFiles, Log))
 			{
 				return true;
 			}
@@ -1586,7 +1631,7 @@ namespace UnrealGameSync
 			return false;
 		}
 
-		bool FindUnresolvedFiles(IEnumerable<string> SyncPaths, out List<PerforceFileRecord> UnresolvedFiles)
+		static bool FindUnresolvedFiles(PerforceConnection Perforce, IEnumerable<string> SyncPaths, TextWriter Log, out List<PerforceFileRecord> UnresolvedFiles)
 		{
 			UnresolvedFiles = new List<PerforceFileRecord>();
 			foreach(string SyncPath in SyncPaths)
@@ -1602,12 +1647,12 @@ namespace UnrealGameSync
 			return true;
 		}
 
-		bool UpdateVersionFile(string ClientPath, Dictionary<string, string> VersionStrings, int ChangeNumber)
+		static bool UpdateVersionFile(PerforceConnection Perforce, string ClientPath, Dictionary<string, string> VersionStrings, int ChangeNumber, TextWriter Log)
 		{
-			return UpdateVersionFile(ClientPath, ChangeNumber, Text => UpdateVersionStrings(Text, VersionStrings));
+			return UpdateVersionFile(Perforce, ClientPath, ChangeNumber, Text => UpdateVersionStrings(Text, VersionStrings), Log);
 		}
 
-		bool UpdateVersionFile(string ClientPath, int ChangeNumber, Func<string, string> Update)
+		static bool UpdateVersionFile(PerforceConnection Perforce, string ClientPath, int ChangeNumber, Func<string, string> Update, TextWriter Log)
 		{
 			List<PerforceFileRecord> Records;
 			if(!Perforce.Stat(ClientPath, out Records, Log))
@@ -1645,10 +1690,10 @@ namespace UnrealGameSync
 
 			string Text = String.Join("\n", Lines);
 			Text = Update(Text);
-			return WriteVersionFile(LocalPath, DepotPath, Text);
+			return WriteVersionFile(Perforce, LocalPath, DepotPath, Text, Log);
 		}
 
-		string UpdateVersionStrings(string Text, Dictionary<string, string> VersionStrings)
+		static string UpdateVersionStrings(string Text, Dictionary<string, string> VersionStrings)
 		{
 			StringWriter Writer = new StringWriter();
 			foreach (string Line in Text.Split('\n'))
@@ -1666,7 +1711,7 @@ namespace UnrealGameSync
 			return Writer.ToString();
 		}
 
-		string UpdateBuildVersion(string Text, int Changelist, int CodeChangelist, string BranchOrStreamName, bool bIsLicenseeVersion)
+		static string UpdateBuildVersion(string Text, int Changelist, int CodeChangelist, string BranchOrStreamName, bool bIsLicenseeVersion)
 		{
 			Dictionary<string, object> Object = JsonSerializer.Deserialize<Dictionary<string, object>>(Text, Program.DefaultJsonSerializerOptions);
 
@@ -1694,7 +1739,7 @@ namespace UnrealGameSync
 			});
 		}
 
-		bool WriteVersionFile(string LocalPath, string DepotPath, string NewText)
+		static bool WriteVersionFile(PerforceConnection Perforce, string LocalPath, string DepotPath, string NewText, TextWriter Log)
 		{
 			try
 			{
@@ -1722,7 +1767,7 @@ namespace UnrealGameSync
 			}
 		}
 
-		bool UpdateVersionLine(ref string Line, string Prefix, string Suffix)
+		static bool UpdateVersionLine(ref string Line, string Prefix, string Suffix)
 		{
 			int LineIdx = 0;
 			int PrefixIdx = 0;
@@ -1744,7 +1789,7 @@ namespace UnrealGameSync
 			return true;
 		}
 
-		string ReadToken(string Line, ref int LineIdx)
+		static string ReadToken(string Line, ref int LineIdx)
 		{
 			for(;; LineIdx++)
 			{
@@ -1780,11 +1825,7 @@ namespace UnrealGameSync
 			get { return Progress.Current; }
 		}
 
-		public int CurrentChangeNumber
-		{
-			get;
-			private set;
-		}
+		public int CurrentChangeNumber => State.CurrentChangeNumber;
 
 		public int PendingChangeNumber
 		{
@@ -1792,18 +1833,6 @@ namespace UnrealGameSync
 			private set;
 		}
 
-		public int LastBuiltChangeNumber
-		{
-			get;
-			private set;
-		}
-
-		public string CurrentSyncFilterHash
-		{
-			get;
-			private set;
-		}
-
 		public string ClientName
 		{
 			get { return Perforce.ClientName; }
