UGS: Significant refactor of UGS internals.

- Now uses the custom C++ Perforce library developed for Horde. This removes dependencies on an installed P4.EXE tool, improves performance, and paves the way to cross-platform support via command line invocations.
- As part of the migration to the new Perforce library, the codebase now uses async/await extensively instead of background threads, and supports a proper cancellation path rather than aborting threads.
- ILogger is now used for writing log files.
- Started converting to use of nullable annotations.

diff --git a/Engine/Source/Programs/Shared/EpicGames.Core/ManagedProcess.cs b/Engine/Source/Programs/Shared/EpicGames.Core/ManagedProcess.cs
--- a/Engine/Source/Programs/Shared/EpicGames.Core/ManagedProcess.cs
+++ b/Engine/Source/Programs/Shared/EpicGames.Core/ManagedProcess.cs
@@ -436,6 +436,11 @@ namespace EpicGames.Core
 		/// </summary>
 		ManagedProcessGroup? AccountingProcessGroup;
 
+		byte[] Buffer = Array.Empty<byte>();
+		int BufferPos = 0;
+		int BufferLen = 0;
+		bool bCarriageReturn = false;
+
 		/// <summary>
 		/// Spawns a new managed process.
 		/// </summary>
@@ -970,7 +975,7 @@ namespace EpicGames.Core
 		/// <param name="Offset">Offset within the buffer to write to</param>
 		/// <param name="Count">Maximum number of bytes to read</param>
 		/// <returns>Number of bytes read</returns>
-		public async Task<int> ReadAsync(byte[] Buffer, int Offset, int Count, CancellationToken CancellationToken)
+		public async ValueTask<int> ReadAsync(byte[] Buffer, int Offset, int Count, CancellationToken CancellationToken)
 		{
 			return await StdOut!.ReadAsync(Buffer, Offset, Count, CancellationToken);
 		}
@@ -1027,6 +1032,66 @@ namespace EpicGames.Core
 			}
 		}
 
+		/// <summary>
+		/// Reads an individual line from the process
+		/// </summary>
+		/// <param name="CancellationToken"></param>
+		/// <returns></returns>
+		public async ValueTask<string?> ReadLineAsync(CancellationToken CancellationToken = default)
+		{
+			for (; ; )
+			{
+				// Skip any '\n' that's part of a '\r\n' sequence
+				if (bCarriageReturn && BufferPos < Buffer.Length)
+				{
+					if (Buffer[BufferPos] == '\n')
+					{
+						BufferPos++;
+					}
+					bCarriageReturn = false;
+				}
+
+				// Pull out the first complete output line
+				for (int Idx = BufferPos; Idx < BufferLen; Idx++)
+				{
+					byte Character = Buffer[Idx];
+					if (Character == '\r' || Character == '\n')
+					{
+						string Line = Console.OutputEncoding.GetString(Buffer, BufferPos, Idx - BufferPos);
+						bCarriageReturn = (Character == '\r');
+						BufferPos = Idx + 1;
+						return Line;
+					}
+				}
+
+				// Create some space in the output buffer
+				if (BufferPos > 0)
+				{
+					Array.Copy(Buffer, BufferPos, Buffer, 0, BufferLen - BufferPos);
+					BufferLen -= BufferPos;
+					BufferPos = 0;
+				}
+				else if (BufferLen == Buffer.Length)
+				{
+					Array.Resize(ref Buffer, Buffer.Length + (32 * 1024));
+				}
+
+				// Read more data into the buffer
+				int NumBytesRead = await ReadAsync(Buffer, BufferPos, Buffer.Length - BufferPos, CancellationToken);
+				if (NumBytesRead == 0)
+				{
+					string? Line = null;
+					if (BufferPos < BufferLen)
+					{
+						Line = Console.OutputEncoding.GetString(Buffer, BufferPos, BufferLen - BufferPos);
+						BufferPos = BufferLen;
+					}
+					return Line;
+				}
+				BufferLen += NumBytesRead;
+			}
+		}
+
 		/// <summary>
 		/// Read all the output from the process. Does not return until the process terminates.
 		/// </summary>
diff --git a/Engine/Source/Programs/Shared/EpicGames.Perforce/PerforceConnection.cs b/Engine/Source/Programs/Shared/EpicGames.Perforce/PerforceConnection.cs
--- a/Engine/Source/Programs/Shared/EpicGames.Perforce/PerforceConnection.cs
+++ b/Engine/Source/Programs/Shared/EpicGames.Perforce/PerforceConnection.cs
@@ -3104,7 +3104,7 @@ namespace EpicGames.Perforce
 		/// <param name="StreamPath">The path for streams to enumerate (eg. "//UE4/...")</param>
 		/// <param name="CancellationToken">Token used to cancel the operation</param>
 		/// <returns>List of streams matching the given criteria</returns>
-		public static async Task<List<StreamsRecord>> GetStreamsAsync(this IPerforceConnection Connection, string StreamPath, CancellationToken CancellationToken = default)
+		public static async Task<List<StreamsRecord>> GetStreamsAsync(this IPerforceConnection Connection, string? StreamPath, CancellationToken CancellationToken = default)
 		{
 			return (await TryGetStreamsAsync(Connection, StreamPath, CancellationToken)).Data;
 		}
@@ -3116,7 +3116,7 @@ namespace EpicGames.Perforce
 		/// <param name="StreamPath">The path for streams to enumerate (eg. "//UE4/...")</param>
 		/// <param name="CancellationToken">Token used to cancel the operation</param>
 		/// <returns>List of streams matching the given criteria</returns>
-		public static Task<PerforceResponseList<StreamsRecord>> TryGetStreamsAsync(this IPerforceConnection Connection, string StreamPath, CancellationToken CancellationToken = default)
+		public static Task<PerforceResponseList<StreamsRecord>> TryGetStreamsAsync(this IPerforceConnection Connection, string? StreamPath, CancellationToken CancellationToken = default)
 		{
 			return TryGetStreamsAsync(Connection, StreamPath, -1, null, false, CancellationToken);
 		}
@@ -3131,7 +3131,7 @@ namespace EpicGames.Perforce
 		/// <param name="bUnloaded">Whether to enumerate unloaded workspaces</param>
 		/// <param name="CancellationToken">Token used to cancel the operation</param>
 		/// <returns>List of streams matching the given criteria</returns>
-		public static async Task<List<StreamsRecord>> GetStreamsAsync(this IPerforceConnection Connection, string StreamPath, int MaxResults, string? Filter, bool bUnloaded, CancellationToken CancellationToken = default)
+		public static async Task<List<StreamsRecord>> GetStreamsAsync(this IPerforceConnection Connection, string? StreamPath, int MaxResults, string? Filter, bool bUnloaded, CancellationToken CancellationToken = default)
 		{
 			return (await TryGetStreamsAsync(Connection, StreamPath, MaxResults, Filter, bUnloaded, CancellationToken)).Data;
 		}
@@ -3146,7 +3146,7 @@ namespace EpicGames.Perforce
 		/// <param name="bUnloaded">Whether to enumerate unloaded workspaces</param>
 		/// <param name="CancellationToken">Token used to cancel the operation</param>
 		/// <returns>List of streams matching the given criteria</returns>
-		public static Task<PerforceResponseList<StreamsRecord>> TryGetStreamsAsync(this IPerforceConnection Connection, string StreamPath, int MaxResults, string? Filter, bool bUnloaded, CancellationToken CancellationToken = default)
+		public static Task<PerforceResponseList<StreamsRecord>> TryGetStreamsAsync(this IPerforceConnection Connection, string? StreamPath, int MaxResults, string? Filter, bool bUnloaded, CancellationToken CancellationToken = default)
 		{
 			// Build the command line
 			List<string> Arguments = new List<string>();
diff --git a/Engine/Source/Programs/UnrealGameSync/UnrealGameSync/ActivationListener.cs b/Engine/Source/Programs/UnrealGameSync/UnrealGameSync/ActivationListener.cs
--- a/Engine/Source/Programs/UnrealGameSync/UnrealGameSync/ActivationListener.cs
+++ b/Engine/Source/Programs/UnrealGameSync/UnrealGameSync/ActivationListener.cs
@@ -1,4 +1,4 @@
-﻿// Copyright Epic Games, Inc. All Rights Reserved.
+// Copyright Epic Games, Inc. All Rights Reserved.
 
 using System;
 using System.Collections.Generic;
@@ -7,15 +7,17 @@ using System.Text;
 using System.Threading;
 using System.Threading.Tasks;
 
+#nullable enable
+
 namespace UnrealGameSync
 {
 	class ActivationListener : IDisposable
 	{
-		Thread BackgroundThread;
+		Thread? BackgroundThread;
 		EventWaitHandle ActivateEventHandle;
 		EventWaitHandle QuitEventHandle;
 
-		public event Action OnActivate;
+		public event Action? OnActivate;
 
 		public ActivationListener(EventWaitHandle InActivateEventHandle)
 		{
@@ -47,16 +49,8 @@ namespace UnrealGameSync
 		{
 			Stop();
 
-			if(QuitEventHandle != null)
-			{
-				QuitEventHandle.Dispose();
-				QuitEventHandle = null;
-			}
-			if(ActivateEventHandle != null)
-			{
-				ActivateEventHandle.Dispose();
-				ActivateEventHandle = null;
-			}
+			QuitEventHandle.Dispose();
+			ActivateEventHandle.Dispose();
 		}
 
 		void ThreadProc()
@@ -66,7 +60,7 @@ namespace UnrealGameSync
 				int Index = EventWaitHandle.WaitAny(new WaitHandle[]{ ActivateEventHandle, QuitEventHandle }, Timeout.Infinite);
 				if(Index == 0)
 				{
-					OnActivate();
+					OnActivate?.Invoke();
 				}
 				else
 				{
diff --git a/Engine/Source/Programs/UnrealGameSync/UnrealGameSync/ArchiveManifest.cs b/Engine/Source/Programs/UnrealGameSync/UnrealGameSync/ArchiveManifest.cs
--- a/Engine/Source/Programs/UnrealGameSync/UnrealGameSync/ArchiveManifest.cs
+++ b/Engine/Source/Programs/UnrealGameSync/UnrealGameSync/ArchiveManifest.cs
@@ -1,6 +1,7 @@
 // Copyright Epic Games, Inc. All Rights Reserved.
 
 using EpicGames.Core;
+using Microsoft.Extensions.Logging;
 using System;
 using System.Collections.Generic;
 using System.Diagnostics;
@@ -10,6 +11,8 @@ using System.Linq;
 using System.Text;
 using System.Threading.Tasks;
 
+#nullable enable
+
 namespace UnrealGameSync
 {
 	[DebuggerDisplay("{FileName}")]
@@ -80,7 +83,7 @@ namespace UnrealGameSync
 
 	static class ArchiveUtils
 	{
-		public static void ExtractFiles(FileReference ArchiveFileName, DirectoryReference BaseDirectoryName, FileReference ManifestFileName, ProgressValue Progress, TextWriter LogWriter)
+		public static void ExtractFiles(FileReference ArchiveFileName, DirectoryReference BaseDirectoryName, FileReference? ManifestFileName, ProgressValue Progress, ILogger Logger)
 		{
 			DateTime TimeStamp = DateTime.UtcNow;
 			using (ZipArchive Zip = new ZipArchive(File.OpenRead(ArchiveFileName.FullName)))
@@ -116,11 +119,7 @@ namespace UnrealGameSync
 					{
 						FileReference FileName = FileReference.Combine(BaseDirectoryName, Entry.FullName);
 						DirectoryReference.CreateDirectory(FileName.Directory);
-
-						lock(LogWriter)
-						{
-							LogWriter.WriteLine("Writing {0}", FileName);
-						}
+						Logger.LogInformation("Writing {0}", FileName);
 
 						Entry.ExtractToFile(FileName.FullName, true);
 						FileReference.SetLastWriteTimeUtc(FileName, TimeStamp);
@@ -130,7 +129,7 @@ namespace UnrealGameSync
 			}
 		}
 
-		public static void RemoveExtractedFiles(DirectoryReference BaseDirectoryName, FileReference ManifestFileName, ProgressValue Progress, TextWriter LogWriter)
+		public static void RemoveExtractedFiles(DirectoryReference BaseDirectoryName, FileReference ManifestFileName, ProgressValue Progress, ILogger LogWriter)
 		{
 			// Read the manifest in
 			ArchiveManifest Manifest;
@@ -147,15 +146,15 @@ namespace UnrealGameSync
 				{
 					if(File.Length != Manifest.Files[Idx].Length)
 					{
-						LogWriter.WriteLine("Skipping {0} due to modified length", File.FullName);
+						LogWriter.LogInformation("Skipping {FileName} due to modified length", File.FullName);
 					}
 					else if(Math.Abs((File.LastWriteTimeUtc - Manifest.Files[Idx].LastWriteTimeUtc).TotalSeconds) > 2.0)
 					{
-						LogWriter.WriteLine("Skipping {0} due to modified timestamp", File.FullName);
+						LogWriter.LogInformation("Skipping {FileName} due to modified timestamp", File.FullName);
 					}
 					else
 					{
-						LogWriter.WriteLine("Removing {0}", File.FullName);
+						LogWriter.LogInformation("Removing {FileName}", File.FullName);
 						File.Delete();
 					}
 				}
diff --git a/Engine/Source/Programs/UnrealGameSync/UnrealGameSync/AsyncDisposer.cs b/Engine/Source/Programs/UnrealGameSync/UnrealGameSync/AsyncDisposer.cs
new file mode 100644
--- /dev/null
+++ b/Engine/Source/Programs/UnrealGameSync/UnrealGameSync/AsyncDisposer.cs
@@ -0,0 +1,65 @@
+// Copyright Epic Games, Inc. All Rights Reserved.
+
+using Microsoft.Extensions.Logging;
+using System;
+using System.Collections.Generic;
+using System.Text;
+using System.Threading;
+using System.Threading.Tasks;
+
+namespace UnrealGameSync
+{
+	interface IAsyncDisposer
+	{
+		void Add(Task Task);
+	}
+
+	internal class AsyncDisposer : IAsyncDisposer, IAsyncDisposable
+	{
+		object LockObject = new object();
+		List<Task> Tasks = new List<Task>();
+		ILogger Logger;
+
+		public AsyncDisposer(ILogger<AsyncDisposer> Logger)
+		{
+			this.Logger = Logger;
+		}
+
+		public void Add(Task Task)
+		{
+			lock (LockObject)
+			{
+				Tasks.Add(Task);
+			}
+			Task.ContinueWith(Remove);
+		}
+
+		private void Remove(Task Task)
+		{
+			if (Task.IsFaulted)
+			{
+				Logger.LogError(Task.Exception, "Exception while disposing task");
+			}
+			lock(LockObject)
+			{
+				Tasks.Remove(Task);
+			}
+		}
+
+		public async ValueTask DisposeAsync()
+		{
+			List<Task> TasksCopy;
+			lock (LockObject)
+			{
+				TasksCopy = Tasks;
+			}
+
+			Task WaitTask = Task.WhenAll(TasksCopy);
+			while (!WaitTask.IsCompleted)
+			{
+				Logger.LogInformation("Waiting for {NumTasks} tasks to complete", TasksCopy.Count);
+				await Task.WhenAny(WaitTask, Task.Delay(TimeSpan.FromSeconds(5.0)));
+			}
+		}
+	}
+}
diff --git a/Engine/Source/Programs/UnrealGameSync/UnrealGameSync/Automation/P4Automation.cs b/Engine/Source/Programs/UnrealGameSync/UnrealGameSync/Automation/P4Automation.cs
--- a/Engine/Source/Programs/UnrealGameSync/UnrealGameSync/Automation/P4Automation.cs
+++ b/Engine/Source/Programs/UnrealGameSync/UnrealGameSync/Automation/P4Automation.cs
@@ -1,6 +1,14 @@
 // Copyright Epic Games, Inc. All Rights Reserved.
 
+using EpicGames.Core;
+using EpicGames.Perforce;
+using Microsoft.Extensions.Logging;
+using System;
 using System.IO;
+using System.Threading;
+using System.Threading.Tasks;
+
+#nullable enable
 
 namespace UnrealGameSync
 {
@@ -8,47 +16,32 @@ namespace UnrealGameSync
 	static class P4Automation
 	{
 
-		public static PerforceConnection GetConnection(out string ErrorMessage)
+		public static IPerforceSettings GetConnectionSettings()
 		{
-			ErrorMessage = null;
-
 			// Read the settings
-			string ServerAndPort = null;
-			string UserName = null;
-			string DepotPathSettings = null;
+			string? ServerAndPort = null;
+			string? UserName = null;
+			string? DepotPathSettings = null;
 
 			Utility.ReadGlobalPerforceSettings(ref ServerAndPort, ref UserName, ref DepotPathSettings);
 
-			if (string.IsNullOrEmpty(UserName))
-			{
-				ErrorMessage = "Unable to get UGS perforce username from registry";
-				return null;
-			}
-
-			if (string.IsNullOrEmpty(ServerAndPort))
-			{
-				ServerAndPort = "perforce:1666";
-			}
-
-			return new PerforceConnection(UserName, null, ServerAndPort);
+			return Utility.OverridePerforceSettings(PerforceSettings.Default, ServerAndPort, UserName);
 
 		}
 
-		public static bool PrintToTempFile(PerforceConnection Connection, string DepotPath, out string TempFileName, out string ErrorMessage)
+		public static Task<string> PrintToTempFile(IPerforceConnection? Connection, string DepotPath, ILogger Logger)
 		{
-			TempFileName = null;
-			ErrorMessage = null;
+			return PrintToTempFileAsync(Connection, DepotPath, CancellationToken.None, Logger);
+		}
 
+		public static async Task<string> PrintToTempFileAsync(IPerforceConnection? Connection, string DepotPath, CancellationToken CancellationToken, ILogger Logger)
+		{
 			if (Connection == null)
 			{
-				Connection = GetConnection(out ErrorMessage);
-				if (Connection == null)
-				{
-					return false;
-				}
+				IPerforceSettings Settings = GetConnectionSettings();
+				Connection = await PerforceConnection.CreateAsync(Logger);
 			}
 
-			BufferedTextWriter Log = new BufferedTextWriter();
 			string DepotFileName = Path.GetFileName(DepotPath);
 
 			// Reorder CL and extension
@@ -61,18 +54,12 @@ namespace UnrealGameSync
 
 			string CL = DepotFileName.Substring(Index + 1);
 			string FileName = DepotFileName.Substring(0, Index);
-			TempFileName = string.Format("{0}@{1}{2}", Path.GetFileNameWithoutExtension(FileName), CL, Path.GetExtension(FileName));
+			string TempFileName = string.Format("{0}@{1}{2}", Path.GetFileNameWithoutExtension(FileName), CL, Path.GetExtension(FileName));
 
 			TempFileName = Path.Combine(Path.GetTempPath(), TempFileName);
+			await Connection.PrintAsync(TempFileName, DepotFileName, CancellationToken);
 
-			if (!Connection.PrintToFile(DepotPath, TempFileName, Log) || !File.Exists(TempFileName))
-			{
-				ErrorMessage = string.Format("Unable to p4 print file {0}\\n{1}\\n", DepotPath, string.Join("\n", Log.Lines));
-				return false;
-			}		
-
-			return true;
-
+			return TempFileName;
 		}
 
 	}
diff --git a/Engine/Source/Programs/UnrealGameSync/UnrealGameSync/Automation/VisualStudioAutomation.cs b/Engine/Source/Programs/UnrealGameSync/UnrealGameSync/Automation/VisualStudioAutomation.cs
--- a/Engine/Source/Programs/UnrealGameSync/UnrealGameSync/Automation/VisualStudioAutomation.cs
+++ b/Engine/Source/Programs/UnrealGameSync/UnrealGameSync/Automation/VisualStudioAutomation.cs
@@ -11,17 +11,19 @@ using Microsoft.VisualStudio.Setup.Configuration;
 using EnvDTE;
 using System.Diagnostics;
 
+#nullable enable
+
 namespace UnrealGameSync
 {
 
 	static class VisualStudioAutomation
 	{
-		public static bool OpenFile(string FileName, out string ErrorMessage, int Line = -1)
+		public static bool OpenFile(string FileName, out string? ErrorMessage, int Line = -1)
 		{
 			ErrorMessage = null;
 
 			// first try to open via DTE
-			DTE DTE = VisualStudioAccessor.GetDTE();
+			DTE? DTE = VisualStudioAccessor.GetDTE();
 			if (DTE != null)
 			{
 				DTE.ItemOperations.OpenFile(FileName);
@@ -29,7 +31,7 @@ namespace UnrealGameSync
 
 				if (Line != -1)
 				{
-					(DTE.ActiveDocument.Selection as TextSelection).GotoLine(Line);
+					(DTE.ActiveDocument.Selection as TextSelection)?.GotoLine(Line);
 				}
 
 				return true;
@@ -60,7 +62,7 @@ namespace UnrealGameSync
 		public static bool LaunchVisualStudio(string Arguments, out string ErrorMessage)
 		{
 			ErrorMessage = "";
-			VisualStudioInstallation Install = VisualStudioInstallations.GetPreferredInstallation();
+			VisualStudioInstallation? Install = VisualStudioInstallations.GetPreferredInstallation();
 
 			if (Install == null)
 			{
@@ -85,7 +87,7 @@ namespace UnrealGameSync
 		}
 
 		[STAThread]
-		public static DTE GetDTE()
+		public static DTE? GetDTE()
 		{
 			IRunningObjectTable Table;
 			if (Succeeded(GetRunningObjectTable(0, out Table)) && Table != null)
@@ -101,7 +103,7 @@ namespace UnrealGameSync
 				MonikersTable.Reset();
 
 				// Look for all visual studio instances in the ROT
-				IMoniker[] Monikers = new IMoniker[] { null };
+				IMoniker[] Monikers = new IMoniker[1];
 				while (MonikersTable.Next(1, Monikers, IntPtr.Zero) == 0)
 				{
 					IBindCtx BindContext;
@@ -147,7 +149,8 @@ namespace UnrealGameSync
 
 			for (int Idx = 0; Idx < Installs.Length; Idx++)
 			{
-				if (InName.StartsWith(Installs[Idx].ROTMoniker))
+				string? Moniker = Installs[Idx].ROTMoniker;
+				if (Moniker != null && InName.StartsWith(Moniker))
 				{
 					return true;
 				}
@@ -175,13 +178,13 @@ namespace UnrealGameSync
 		/// <summary>
 		/// Base directory for the installation
 		/// </summary>
-		public string BaseDir;
+		public string? BaseDir;
 
 
 		/// <summary>
 		/// Path of the devenv executable
 		/// </summary>
-		public string DevEnvPath;
+		public string? DevEnvPath;
 
 		/// <summary>
 		/// Visual Studio major version number
@@ -189,9 +192,9 @@ namespace UnrealGameSync
 		public int MajorVersion;
 
 		/// <summary>
-		///  Running Object Table moniker for this installation
+		/// Running Object Table moniker for this installation
 		/// </summary>
-		public string ROTMoniker;
+		public string? ROTMoniker;
 
 	}
 
@@ -201,7 +204,7 @@ namespace UnrealGameSync
 	static class VisualStudioInstallations
 	{
 
-		public static VisualStudioInstallation GetPreferredInstallation(int MajorVersion = 0)
+		public static VisualStudioInstallation? GetPreferredInstallation(int MajorVersion = 0)
 		{
 
 			if (CachedInstalls.Count == 0)
diff --git a/Engine/Source/Programs/UnrealGameSync/UnrealGameSync/AutomationServer.cs b/Engine/Source/Programs/UnrealGameSync/UnrealGameSync/AutomationServer.cs
--- a/Engine/Source/Programs/UnrealGameSync/UnrealGameSync/AutomationServer.cs
+++ b/Engine/Source/Programs/UnrealGameSync/UnrealGameSync/AutomationServer.cs
@@ -1,5 +1,6 @@
 // Copyright Epic Games, Inc. All Rights Reserved.
 
+using Microsoft.Extensions.Logging;
 using Microsoft.Win32;
 using System;
 using System.IO;
@@ -10,6 +11,8 @@ using System.Text;
 using System.Threading;
 using System.Windows.Forms;
 
+#nullable enable
+
 namespace UnrealGameSync
 {
 	enum AutomationRequestType
@@ -105,7 +108,7 @@ namespace UnrealGameSync
 	class AutomationRequest : IDisposable
 	{
 		public AutomationRequestInput Input;
-		public AutomationRequestOutput Output;
+		public AutomationRequestOutput? Output;
 		public ManualResetEventSlim Complete;
 
 		public AutomationRequest(AutomationRequestInput Input)
@@ -117,47 +120,43 @@ namespace UnrealGameSync
 		public void SetOutput(AutomationRequestOutput Output)
 		{
 			this.Output = Output;
-			Complete.Set();
+			Complete?.Set();
 		}
 
 		public void Dispose()
 		{
-			if(Complete != null)
-			{
-				Complete.Dispose();
-				Complete = null;
-			}
+			Complete.Dispose();
 		}
 	}
 
 	class AutomationServer : IDisposable
 	{
-		TcpListener Listener;
+		TcpListener? Listener;
 		public const int DefaultPortNumber = 30422;
 
-		NamedPipeServerStream IPCStream;
+		NamedPipeServerStream? IPCStream;
 		static UnicodeEncoding StreamEncoding = new UnicodeEncoding();
 		const string UGSChannel = @"\.\pipe\UGSChannel";
 
-		Thread UriThread;
-		Thread TcpThread;
+		Thread? UriThread;
+		Thread? TcpThread;
 
 		EventWaitHandle ShutdownEvent;		
 		Action<AutomationRequest> PostRequest;
 
 		bool bDisposing;
-		TextWriter Log;
+		ILogger Logger;
 		string CommandLineUri;
 
-		public AutomationServer(Action<AutomationRequest> PostRequest, TextWriter Log, string Uri)
+		public AutomationServer(Action<AutomationRequest> PostRequest, string Uri, ILogger<AutomationServer> Logger)
 		{
+			ShutdownEvent = new ManualResetEvent(false);
+			this.PostRequest = PostRequest;
+			this.CommandLineUri = Uri;
+			this.Logger = Logger;
+
 			try
 			{
-				ShutdownEvent = new ManualResetEvent(false);
-				this.PostRequest = PostRequest;
-				this.Log = Log;
-				this.CommandLineUri = Uri;
-				
 				// IPC named pipe
 				IPCStream = new NamedPipeServerStream(UGSChannel, PipeDirection.In, 1, PipeTransmissionMode.Message, PipeOptions.Asynchronous);
 
@@ -176,7 +175,7 @@ namespace UnrealGameSync
 					catch (Exception Ex)
 					{
 						Listener = null;
-						Log.WriteLine("Unable to start automation server tcp listener: {0}", Ex.ToString());
+						Logger.LogError(Ex, "Unable to start automation server tcp listener");
 					}
 				}
 
@@ -185,7 +184,7 @@ namespace UnrealGameSync
 			}
 			catch (Exception Ex)
 			{
-				Log.WriteLine("Unable to start automation server: {0}", Ex.ToString());
+				Logger.LogError(Ex, "Unable to start automation server");
 			}
 		}
 
@@ -227,7 +226,7 @@ namespace UnrealGameSync
 				try
 				{
 
-					IAsyncResult IPCResult = IPCStream.BeginWaitForConnection(null, null);
+					IAsyncResult IPCResult = IPCStream!.BeginWaitForConnection(null, null);
 
 					int WaitResult = WaitHandle.WaitAny(new WaitHandle[] { ShutdownEvent, IPCResult.AsyncWaitHandle });
 
@@ -241,12 +240,12 @@ namespace UnrealGameSync
 					{
 						IPCStream.EndWaitForConnection(IPCResult);
 
-						Log.WriteLine("Accepted Uri connection");
+						Logger.LogInformation("Accepted Uri connection");
 
 						// Read URI
 						string Uri = ReadString(IPCStream);
 
-						Log.WriteLine("Received Uri: {0}", Uri);
+						Logger.LogInformation("Received Uri: {Uri}", Uri);
 
 						IPCStream.Disconnect();
 
@@ -255,14 +254,14 @@ namespace UnrealGameSync
 					}
 					catch (Exception Ex)
 					{
-						Log.WriteLine("Exception: {0}", Ex.ToString());
+						Logger.LogError(Ex, "Error during automation connection");
 					}
 				}
 				catch (Exception Ex)
 				{
 					if (!bDisposing)
 					{
-						Log.WriteLine("Exception: {0}", Ex.ToString());
+						Logger.LogError(Ex, "Error during automation connection");
 					}
 				}
 			}
@@ -277,7 +276,7 @@ namespace UnrealGameSync
 				try
 				{
 
-					IAsyncResult TCPResult = Listener.BeginAcceptTcpClient(null, null);
+					IAsyncResult TCPResult = Listener!.BeginAcceptTcpClient(null, null);
 
 					int WaitResult = WaitHandle.WaitAny(new WaitHandle[] { ShutdownEvent, TCPResult.AsyncWaitHandle });
 
@@ -291,39 +290,38 @@ namespace UnrealGameSync
 					{
 						TcpClient Client = Listener.EndAcceptTcpClient(TCPResult);
 
-						Log.WriteLine("Accepted connection from {0}", Client.Client.RemoteEndPoint);
+						Logger.LogInformation("Accepted connection from {Remote}", Client.Client.RemoteEndPoint);
 
 						NetworkStream Stream = Client.GetStream();
 
 						AutomationRequestInput Input = AutomationRequestInput.Read(Stream);
-						Log.WriteLine("Received input: {0} (+{1} bytes)", Input.Type, Input.Data.Length);
+						Logger.LogInformation("Received input: {Type} (+{NumBytes} bytes)", Input.Type, Input.Data.Length);
 
 						AutomationRequestOutput Output;
 						using (AutomationRequest Request = new AutomationRequest(Input))
 						{
 							PostRequest(Request);
 							Request.Complete.Wait();
-							Output = Request.Output;
+							Output = Request.Output!;
 						}
 
 						Output.Write(Stream);
-						Log.WriteLine("Sent output: {0} (+{1} bytes)", Output.Result, Output.Data.Length);
+						Logger.LogInformation("Sent output: {Result} (+{NumBytes} bytes)", Output.Result, Output.Data.Length);
 					}
 					catch (Exception Ex)
 					{
-						Log.WriteLine("Exception: {0}", Ex.ToString());
+						Logger.LogError(Ex, "Exception during automation");
 					}
 					finally
 					{
-						TCPResult = null;
-						Log.WriteLine("Closed connection.");
+						Logger.LogInformation("Closed connection.");
 					}
 				}
 				catch (Exception Ex)
 				{
 					if (!bDisposing)
 					{
-						Log.WriteLine("Exception: {0}", Ex.ToString());
+						Logger.LogError(Ex, "Exception during automation operation");
 					}
 				}
 			}
diff --git a/Engine/Source/Programs/UnrealGameSync/UnrealGameSync/ConfigFile.cs b/Engine/Source/Programs/UnrealGameSync/UnrealGameSync/ConfigFile.cs
--- a/Engine/Source/Programs/UnrealGameSync/UnrealGameSync/ConfigFile.cs
+++ b/Engine/Source/Programs/UnrealGameSync/UnrealGameSync/ConfigFile.cs
@@ -1,17 +1,21 @@
 // Copyright Epic Games, Inc. All Rights Reserved.
 
 using EpicGames.Core;
+using Microsoft.Extensions.Logging;
 using System;
 using System.Collections.Generic;
 using System.Diagnostics;
+using System.Diagnostics.CodeAnalysis;
 using System.IO;
 using System.Linq;
 using System.Text;
 using System.Threading.Tasks;
 
+#nullable enable
+
 namespace UnrealGameSync
 {
-	class ConfigObject
+	public class ConfigObject
 	{
 		const string ConfigSeparatorCharacters = "(),= \t\"";
 
@@ -47,17 +51,25 @@ namespace UnrealGameSync
 				while(Idx < Text.Length)
 				{
 					// Read the next key/value pair
-					string Key = ParseConfigToken(Text, ref Idx);
+					string? Key = ParseConfigToken(Text, ref Idx);
+					if (Key == null)
+					{
+						return;
+					}
 					if(ParseConfigToken(Text, ref Idx) == "=")
 					{
-						string Value = ParseConfigValueToken(Text, ref Idx);
+						string? Value = ParseConfigValueToken(Text, ref Idx);
+						if (Value == null)
+						{
+							return;
+						}
 						SetValue(Key, Value);
 					}
 
 					// Check for the end of the list, or a comma before the next pair
 					for(;;)
 					{
-						string Token = ParseConfigValueToken(Text, ref Idx);
+						string? Token = ParseConfigValueToken(Text, ref Idx);
 						if(Token == ",")
 						{
 							break;
@@ -71,7 +83,7 @@ namespace UnrealGameSync
 			}
 		}
 
-		static string ParseConfigToken(string Text, ref int Idx)
+		static string? ParseConfigToken(string Text, ref int Idx)
 		{
 			// Skip whitespace
 			while(Idx < Text.Length && Char.IsWhiteSpace(Text[Idx]))
@@ -117,15 +129,15 @@ namespace UnrealGameSync
 			}
 		}
 
-		static string ParseConfigValueToken(string Text, ref int Idx)
+		static string? ParseConfigValueToken(string Text, ref int Idx)
 		{
-			string Token = ParseConfigToken(Text, ref Idx);
+			string? Token = ParseConfigToken(Text, ref Idx);
 			if (Token == "(")
 			{
 				int StartIdx = Idx - 1;
 				for (; ; )
 				{
-					string NextToken = ParseConfigValueToken(Text, ref Idx);
+					string? NextToken = ParseConfigValueToken(Text, ref Idx);
 					if (NextToken == null || NextToken == ")")
 					{
 						break;
@@ -136,7 +148,8 @@ namespace UnrealGameSync
 			return Token;
 		}
 
-		public string GetValue(string Key, string DefaultValue = null)
+		[return: NotNullIfNotNull("DefaultValue")]
+		public string? GetValue(string Key, string? DefaultValue = null)
 		{
 			for(int Idx = 0; Idx < Pairs.Count; Idx++)
 			{
@@ -150,7 +163,7 @@ namespace UnrealGameSync
 
 		public Guid GetValue(string Key, Guid DefaultValue)
 		{
-			string StringValue = GetValue(Key);
+			string? StringValue = GetValue(Key);
 			if(StringValue != null)
 			{
 				Guid Value;
@@ -164,7 +177,7 @@ namespace UnrealGameSync
 
 		public int GetValue(string Key, int DefaultValue)
 		{
-			string StringValue = GetValue(Key);
+			string? StringValue = GetValue(Key);
 			if(StringValue != null)
 			{
 				int Value;
@@ -178,7 +191,7 @@ namespace UnrealGameSync
 
 		public bool GetValue(string Key, bool DefaultValue)
 		{
-			string StringValue = GetValue(Key);
+			string? StringValue = GetValue(Key);
 			if(StringValue != null)
 			{
 				bool Value;
@@ -190,17 +203,24 @@ namespace UnrealGameSync
 			return DefaultValue;
 		}
 
-		public void SetValue(string Key, string Value)
+		public void SetValue(string Key, string? Value)
 		{
-			for(int Idx = 0; Idx < Pairs.Count; Idx++)
+			if (Value == null)
 			{
-				if(Pairs[Idx].Key.Equals(Key, StringComparison.InvariantCultureIgnoreCase))
+				Pairs.RemoveAll(x => x.Key.Equals(Key, StringComparison.OrdinalIgnoreCase));
+			}
+			else
+			{
+				for (int Idx = 0; Idx < Pairs.Count; Idx++)
 				{
-					Pairs[Idx] = new KeyValuePair<string,string>(Key, Value);
-					return;
+					if (Pairs[Idx].Key.Equals(Key, StringComparison.OrdinalIgnoreCase))
+					{
+						Pairs[Idx] = new KeyValuePair<string, string>(Key, Value);
+						return;
+					}
 				}
+				Pairs.Add(new KeyValuePair<string, string>(Key, Value));
 			}
-			Pairs.Add(new KeyValuePair<string,string>(Key, Value));
 		}
 
 		public void SetValue(string Key, Guid Value)
@@ -218,7 +238,7 @@ namespace UnrealGameSync
 			SetValue(Key, Value.ToString());
 		}
 
-		public string this[string Key]
+		public string? this[string Key]
 		{
 			get { return GetValue(Key); }
 			set { SetValue(Key, value); }
@@ -246,7 +266,7 @@ namespace UnrealGameSync
 			}
 		}
 
-		public string ToString(ConfigObject BaseObject)
+		public string ToString(ConfigObject? BaseObject)
 		{
 			StringBuilder Result = new StringBuilder();
 			Result.Append("(");
@@ -281,7 +301,7 @@ namespace UnrealGameSync
 	}
 
 	[DebuggerDisplay("{Name}")]
-	class ConfigSection
+	public class ConfigSection
 	{
 		public string Name;
 		public Dictionary<string, string> Pairs = new Dictionary<string,string>();
@@ -344,7 +364,7 @@ namespace UnrealGameSync
 
 		public void AppendValue(string Key, string Value)
 		{
-			string CurrentValue;
+			string? CurrentValue;
 			if(Pairs.TryGetValue(Key, out CurrentValue))
 			{
 				Pairs[Key] = CurrentValue + "\n" + Value;
@@ -362,7 +382,7 @@ namespace UnrealGameSync
 
 		public int GetValue(string Key, int DefaultValue)
 		{
-			string ValueString = GetValue(Key);
+			string? ValueString = GetValue(Key);
 			if(ValueString != null)
 			{
 				int Value;
@@ -379,9 +399,10 @@ namespace UnrealGameSync
 			return GetValue(Key, DefaultValue? 1 : 0) != 0;
 		}
 
-		public string GetValue(string Key, string DefaultValue = null)
+		[return: NotNullIfNotNull("DefaultValue")]
+		public string? GetValue(string Key, string? DefaultValue = null)
 		{
-			string Value;
+			string? Value;
 			if(!Pairs.TryGetValue(Key, out Value))
 			{
 				Value = DefaultValue;
@@ -389,9 +410,10 @@ namespace UnrealGameSync
 			return Value;
 		}
 
-		public string[] GetValues(string Key, string[] DefaultValue = null)
+		[return: NotNullIfNotNull("DefaultValue")]
+		public string[]? GetValues(string Key, string[]? DefaultValue = null)
 		{
-			string Value = GetValue(Key, null);
+			string? Value = GetValue(Key, null);
 			if(Value == null)
 			{
 				return DefaultValue;
@@ -402,9 +424,10 @@ namespace UnrealGameSync
 			}
 		}
 
-		public Guid[] GetValues(string Key, Guid[] DefaultValue = null)
+		[return: NotNullIfNotNull("DefaultValue")]
+		public Guid[]? GetValues(string Key, Guid[]? DefaultValue = null)
 		{
-			string[] StringValues = GetValues(Key, (string[])null);
+			string[]? StringValues = GetValues(Key, (string[]?)null);
 			if(StringValues == null)
 			{
 				return DefaultValue;
@@ -424,7 +447,7 @@ namespace UnrealGameSync
 		}
 	}
 
-	class ConfigFile
+	public class ConfigFile
 	{
 		List<ConfigSection> Sections = new List<ConfigSection>();
 
@@ -442,12 +465,12 @@ namespace UnrealGameSync
 			Parse(FileReference.ReadAllLines(FileName));
 		}
 
-		public bool TryLoad(FileReference FileName, TextWriter Log)
+		public bool TryLoad(FileReference FileName, ILogger Logger)
 		{
 			FileInfo FileInfo = FileName.ToFileInfo();
 			if (FileInfo.Exists)
 			{
-				Log.WriteLine("Loading config file from {0} ({1} bytes)", FileInfo.FullName, FileInfo.Length);
+				Logger.LogInformation("Loading config file from {File} ({Size} bytes)", FileInfo.FullName, FileInfo.Length);
 				Load(FileName);
 				return true;
 			}
@@ -457,18 +480,18 @@ namespace UnrealGameSync
 			FileInfo TempFileInfo = TempFileName.ToFileInfo();
 			if (TempFileInfo.Exists)
 			{
-				Log.WriteLine("Loading temporary config file from {0} ({1} bytes)", TempFileInfo.FullName, TempFileInfo.Length);
+				Logger.LogInformation("Loading temporary config file from {File} ({Size} bytes)", TempFileInfo.FullName, TempFileInfo.Length);
 				Load(TempFileName);
 				return true;
 			}
 
-			Log.WriteLine("No existing config file at {0}", FileName);
+			Logger.LogInformation("No existing config file at {File}", FileName);
 			return false;
 		}
 
 		public void Parse(string[] Lines)
 		{
-			ConfigSection CurrentSection = null;
+			ConfigSection? CurrentSection = null;
 			foreach(string Line in Lines)
 			{
 				string TrimLine = Line.Trim();
@@ -592,21 +615,24 @@ namespace UnrealGameSync
 			return (Section == null)? DefaultValue : Section.GetValue(Key.Substring(DotIdx + 1), DefaultValue);
 		}
 
-		public string GetValue(string Key, string DefaultValue)
+		[return: NotNullIfNotNull("DefaultValue")]
+		public string? GetValue(string Key, string? DefaultValue)
 		{
 			int DotIdx = Key.IndexOf('.');
 			ConfigSection Section = FindSection(Key.Substring(0, DotIdx));
 			return (Section == null)? DefaultValue : Section.GetValue(Key.Substring(DotIdx + 1), DefaultValue);
 		}
 
-		public string[] GetValues(string Key, string[] DefaultValue)
+		[return: NotNullIfNotNull("DefaultValue")]
+		public string[]? GetValues(string Key, string[]? DefaultValue)
 		{
 			int DotIdx = Key.IndexOf('.');
 			ConfigSection Section = FindSection(Key.Substring(0, DotIdx));
 			return (Section == null)? DefaultValue : Section.GetValues(Key.Substring(DotIdx + 1), DefaultValue);
 		}
 
-		public Guid[] GetGuidValues(string Key, Guid[] DefaultValue)
+		[return: NotNullIfNotNull("DefaultValue")]
+		public Guid[]? GetGuidValues(string Key, Guid[]? DefaultValue)
 		{
 			int DotIdx = Key.IndexOf('.');
 			ConfigSection Section = FindSection(Key.Substring(0, DotIdx));
diff --git a/Engine/Source/Programs/UnrealGameSync/UnrealGameSync/Controls/CustomListViewControl.cs b/Engine/Source/Programs/UnrealGameSync/UnrealGameSync/Controls/CustomListViewControl.cs
--- a/Engine/Source/Programs/UnrealGameSync/UnrealGameSync/Controls/CustomListViewControl.cs
+++ b/Engine/Source/Programs/UnrealGameSync/UnrealGameSync/Controls/CustomListViewControl.cs
@@ -11,6 +11,8 @@ using System.Threading.Tasks;
 using System.Windows.Forms;
 using System.Windows.Forms.VisualStyles;
 
+#nullable enable
+
 namespace UnrealGameSync
 {
 	abstract class CustomListViewWidget
@@ -47,7 +49,7 @@ namespace UnrealGameSync
 
 		public void ConditionalLayout(Control Owner, Rectangle Bounds)
 		{
-			if(RequiresLayout() || PreviousBounds.Value != Bounds)
+			if(RequiresLayout() || PreviousBounds == null || PreviousBounds.Value != Bounds)
 			{
 				using(Graphics Graphics = Owner.CreateGraphics())
 				{
@@ -88,8 +90,8 @@ namespace UnrealGameSync
 		}
 
 		Rectangle Bounds;
-		StatusElement MouseDownElement;
-		StatusElement MouseOverElement;
+		StatusElement? MouseDownElement;
+		StatusElement? MouseOverElement;
 
 		public StatusLineListViewWidget(ListViewItem Item, StatusElementResources Resources)
 			: base(Item)
@@ -198,12 +200,12 @@ namespace UnrealGameSync
 
 	partial class CustomListViewControl : ListView
 	{
-		VisualStyleRenderer SelectedItemRenderer;
-		VisualStyleRenderer TrackedItemRenderer;
+		VisualStyleRenderer? SelectedItemRenderer;
+		VisualStyleRenderer? TrackedItemRenderer;
 		public int HoverItem = -1;
 
-		CustomListViewWidget MouseOverWidget;
-		CustomListViewWidget MouseDownWidget;
+		CustomListViewWidget? MouseOverWidget;
+		CustomListViewWidget? MouseDownWidget;
 
 		public CustomListViewControl()
 		{
@@ -249,16 +251,16 @@ namespace UnrealGameSync
 			}
 		}
 
-		protected CustomListViewWidget FindWidget(Point Location)
+		protected CustomListViewWidget? FindWidget(Point Location)
 		{
 			return FindWidget(HitTest(Location));
 		}
 
-		protected CustomListViewWidget FindWidget(ListViewHitTestInfo HitTest)
+		protected CustomListViewWidget? FindWidget(ListViewHitTestInfo HitTest)
 		{
 			if(HitTest.Item != null)
 			{
-				CustomListViewWidget Widget = HitTest.Item.Tag as CustomListViewWidget;
+				CustomListViewWidget? Widget = HitTest.Item.Tag as CustomListViewWidget;
 				if(Widget != null)
 				{
 					return Widget;
@@ -266,7 +268,7 @@ namespace UnrealGameSync
 			}
 			if(HitTest.SubItem != null)
 			{
-				CustomListViewWidget Widget = HitTest.SubItem.Tag as CustomListViewWidget;
+				CustomListViewWidget? Widget = HitTest.SubItem.Tag as CustomListViewWidget;
 				if(Widget != null)
 				{
 					return Widget;
@@ -338,7 +340,7 @@ namespace UnrealGameSync
 		{
 			ListViewHitTestInfo HitTest = this.HitTest(e.Location);
 
-			CustomListViewWidget NewMouseOverWidget = FindWidget(HitTest);
+			CustomListViewWidget? NewMouseOverWidget = FindWidget(HitTest);
 			if(MouseOverWidget != null && MouseOverWidget != NewMouseOverWidget)
 			{
 				Cursor = Cursors.Arrow;
@@ -405,12 +407,12 @@ namespace UnrealGameSync
 
 		public void DrawCustomSubItem(Graphics Graphics, ListViewItem.ListViewSubItem SubItem)
 		{
-			CustomListViewWidget Widget = SubItem.Tag as CustomListViewWidget;
+			CustomListViewWidget? Widget = SubItem.Tag as CustomListViewWidget;
 			if(Widget != null)
 			{
-				foreach(ListViewItem.ListViewSubItem OtherSubItem in Widget.Item.SubItems)
+				foreach(ListViewItem.ListViewSubItem? OtherSubItem in Widget.Item.SubItems)
 				{
-					if(OtherSubItem.Tag == Widget)
+					if(OtherSubItem != null && OtherSubItem.Tag == Widget)
 					{
 						if(OtherSubItem == SubItem)
 						{
@@ -446,7 +448,7 @@ namespace UnrealGameSync
 
 		public void DrawSelectedBackground(Graphics Graphics, Rectangle Bounds)
 		{
-			if(Application.RenderWithVisualStyles)
+			if (SelectedItemRenderer != null)
 			{
 				SelectedItemRenderer.DrawBackground(Graphics, Bounds);
 			}
@@ -458,7 +460,7 @@ namespace UnrealGameSync
 
 		public void DrawTrackedBackground(Graphics Graphics, Rectangle Bounds)
 		{
-			if(Application.RenderWithVisualStyles)
+			if (TrackedItemRenderer != null)
 			{
 				TrackedItemRenderer.DrawBackground(Graphics, Bounds);
 			}
diff --git a/Engine/Source/Programs/UnrealGameSync/UnrealGameSync/Controls/LogControl.cs b/Engine/Source/Programs/UnrealGameSync/UnrealGameSync/Controls/LogControl.cs
--- a/Engine/Source/Programs/UnrealGameSync/UnrealGameSync/Controls/LogControl.cs
+++ b/Engine/Source/Programs/UnrealGameSync/UnrealGameSync/Controls/LogControl.cs
@@ -14,6 +14,7 @@ using System.Collections.Concurrent;
 using System.IO;
 using System.Text.RegularExpressions;
 using EpicGames.Core;
+using Microsoft.Extensions.Logging;
 
 namespace UnrealGameSync
 {
@@ -906,8 +907,13 @@ namespace UnrealGameSync
 		}
 	}
 
-	public class LogControlTextWriter : LineBasedTextWriter
+	public class LogControlTextWriter : ILogger
 	{
+		class NullDisposable : IDisposable
+		{
+			public void Dispose() { }
+		}
+
 		LogControl LogControl;
 
 		public LogControlTextWriter(LogControl InLogControl)
@@ -915,9 +921,22 @@ namespace UnrealGameSync
 			LogControl = InLogControl;
 		}
 
-		protected override void FlushLine(string Line)
+		public void Log<TState>(LogLevel LogLevel, EventId EventId, TState State, Exception Exception, Func<TState, Exception, string> Formatter)
 		{
-			LogControl.AppendLine(Line);
+			string Message = Formatter(State, Exception);
+			LogControl.AppendLine(Message);
+
+			if (Exception != null)
+			{
+				foreach (string Line in Exception.ToString().Trim().Split('\n'))
+				{
+					LogControl.AppendLine(Message);
+				}
+			}
 		}
+
+		public bool IsEnabled(LogLevel logLevel) => true;
+
+		public IDisposable BeginScope<TState>(TState state) => new NullDisposable();
 	}
 }
diff --git a/Engine/Source/Programs/UnrealGameSync/UnrealGameSync/Controls/WorkspaceControl.cs b/Engine/Source/Programs/UnrealGameSync/UnrealGameSync/Controls/WorkspaceControl.cs
--- a/Engine/Source/Programs/UnrealGameSync/UnrealGameSync/Controls/WorkspaceControl.cs
+++ b/Engine/Source/Programs/UnrealGameSync/UnrealGameSync/Controls/WorkspaceControl.cs
@@ -16,6 +16,10 @@ using System.Windows.Forms;
 using System.Windows.Forms.VisualStyles;
 using System.Threading;
 using EpicGames.Core;
+using EpicGames.Perforce;
+using System.Threading.Tasks;
+using Microsoft.Extensions.Logging;
+using Microsoft.Extensions.DependencyInjection;
 
 namespace UnrealGameSync
 {
@@ -32,7 +36,7 @@ namespace UnrealGameSync
 		public List<UncontrolledChangelist> Changelists { get; set; }
 	}
 
-public enum LatestChangeType
+	public enum LatestChangeType
 	{
 		Any,
 		Good,
@@ -160,8 +164,8 @@ public enum LatestChangeType
 		{
 			public int Compare(object A, object B)
 			{
-				PerforceChangeSummary ChangeA = ((ListViewItem)A).Tag as PerforceChangeSummary;
-				PerforceChangeSummary ChangeB = ((ListViewItem)B).Tag as PerforceChangeSummary;
+				ChangesRecord ChangeA = ((ListViewItem)A).Tag as ChangesRecord;
+				ChangesRecord ChangeB = ((ListViewItem)B).Tag as ChangesRecord;
 				if (ChangeA == null)
 				{
 					return +1;
@@ -207,7 +211,10 @@ public enum LatestChangeType
 		IWorkspaceControlOwner Owner;
 		string ApiUrl;
 		DirectoryReference DataFolder;
-		LineBasedTextWriter Log;
+		IServiceProvider ServiceProvider;
+		ILogger Logger;
+		IPerforceSettings PerforceSettings;
+		ProjectInfo ProjectInfo { get; }
 
 		UserSettings Settings;
 		UserWorkspaceState WorkspaceState;
@@ -277,7 +284,7 @@ public enum LatestChangeType
 
 		string HoverBadgeUniqueId = null;
 		bool bHoverSync;
-		PerforceChangeSummary ContextMenuChange;
+		ChangesRecord ContextMenuChange;
 		Font BuildFont;
 		Font SelectedBuildFont;
 		Font BadgeFont;
@@ -320,7 +327,7 @@ public enum LatestChangeType
 		// Placeholder text that is in the control and cleared when the user starts editing.
 		static string AuthorFilterPlaceholderText = "<username>";
 
-		public WorkspaceControl(IWorkspaceControlOwner InOwner, string InApiUrl, string InOriginalExecutableFileName, bool bInUnstable, DetectProjectSettingsTask DetectSettings, LineBasedTextWriter InLog, UserSettings InSettings, OIDCTokenManager InOidcTokenManager)
+		public WorkspaceControl(IWorkspaceControlOwner InOwner, string InApiUrl, string InOriginalExecutableFileName, bool bInUnstable, WorkspaceSettings WorkspaceSettings, IServiceProvider InServiceProvider, UserSettings InSettings, OIDCTokenManager InOidcTokenManager)
 		{
 			InitializeComponent();
 
@@ -328,16 +335,19 @@ public enum LatestChangeType
 
 			Owner = InOwner;
 			ApiUrl = InApiUrl;
-			DataFolder = DetectSettings.DataFolder;
+			DataFolder = WorkspaceSettings.DataFolder;
 			OriginalExecutableFileName = InOriginalExecutableFileName;
 			bUnstable = bInUnstable;
-			Log = InLog;
+			ServiceProvider = InServiceProvider;
+			Logger = InServiceProvider.GetRequiredService<ILogger<WorkspaceControl>>();
+			PerforceSettings = WorkspaceSettings.PerforceSettings;
+			ProjectInfo = WorkspaceSettings.ProjectInfo;
 			Settings = InSettings;
-			WorkspaceState = InSettings.FindOrAddWorkspaceState(DetectSettings.BranchDirectoryName, DetectSettings.NewSelectedClientFileName, DetectSettings.NewSelectedProjectIdentifier);
-			WorkspaceSettings = InSettings.FindOrAddWorkspaceSettings(DetectSettings.BranchDirectoryName, DetectSettings.BranchClientPath);
+			WorkspaceState = InSettings.FindOrAddWorkspaceState(WorkspaceSettings.ProjectInfo.LocalRootPath, WorkspaceSettings.ProjectInfo.ClientFileName, WorkspaceSettings.ProjectInfo.ProjectIdentifier);
+			this.WorkspaceSettings = InSettings.FindOrAddWorkspaceSettings(WorkspaceSettings.ProjectInfo.LocalRootPath, WorkspaceSettings.ProjectInfo.ClientFileName);
 
-			string ProjectPath = Regex.Replace(DetectSettings.NewSelectedClientFileName, "^//[^/]+/", "");
-			ProjectSettings = InSettings.FindOrAddProjectSettings(DetectSettings.BranchDirectoryName, ProjectPath);
+			string ProjectPath = Regex.Replace(WorkspaceSettings.ProjectInfo.ClientFileName, "^//[^/]+/", "");
+			ProjectSettings = InSettings.FindOrAddProjectSettings(WorkspaceSettings.ProjectInfo.LocalRootPath, ProjectPath);
 
 			DesiredTaskbarState = Tuple.Create(TaskbarState.NoProgress, 0.0f);
 
@@ -362,21 +372,37 @@ public enum LatestChangeType
 
 			// Set the project logo on the status panel and notification window
 			NotificationWindow = new NotificationWindow(Properties.Resources.DefaultNotificationLogo);
-			StatusPanel.SetProjectLogo(DetectSettings.ProjectLogo, true);
-			DetectSettings.ProjectLogo = null;
+			if (ProjectInfo.LocalFileName.HasExtension(".uproject"))
+			{
+				FileReference LogoFileName = FileReference.Combine(ProjectInfo.LocalFileName.Directory, "Build", "UnrealGameSync.png");
+				if (FileReference.Exists(LogoFileName))
+				{
+					try
+					{
+						// Duplicate the image, otherwise we'll leave the file locked
+						using (Image Image = Image.FromFile(LogoFileName.FullName))
+						{
+							StatusPanel.SetProjectLogo(new Bitmap(Image), true);
+						}
+					}
+					catch
+					{
+					}
+				}
+			}
 
 			// Commit all the new project info
-			PerforceConnection PerforceClient = DetectSettings.PerforceClient;
-			ClientName = PerforceClient.ClientName;
-			SelectedProject = DetectSettings.SelectedProject;
-			SelectedProjectIdentifier = DetectSettings.NewSelectedProjectIdentifier;
-			SelectedFileName = DetectSettings.NewSelectedFileName;
-			EditorTargetName = DetectSettings.NewProjectEditorTarget;
-			bIsEnterpriseProject = DetectSettings.bIsEnterpriseProject;
-			StreamName = DetectSettings.StreamName;
+			IPerforceSettings PerforceClientSettings = WorkspaceSettings.PerforceSettings;
+			ClientName = PerforceUtils.GetClientOrDepotDirectoryName(WorkspaceSettings.ProjectInfo.ClientFileName);
+			SelectedProject = WorkspaceSettings.SelectedProject;
+			SelectedProjectIdentifier = WorkspaceSettings.NewSelectedProjectIdentifier;
+			SelectedFileName = WorkspaceSettings.ProjectInfo.LocalFileName;
+			EditorTargetName = WorkspaceSettings.NewProjectEditorTarget;
+			bIsEnterpriseProject = WorkspaceSettings.ProjectInfo.bIsEnterpriseProject;
+			StreamName = WorkspaceSettings.StreamName;
 
 			// Update the branch directory
-			BranchDirectoryName = DetectSettings.BranchDirectoryName;
+			BranchDirectoryName = WorkspaceSettings.ProjectInfo.LocalRootPath;
 
 			// Check if we've the project we've got open in this workspace is the one we're actually synced to
 			int CurrentChangeNumber = -1;
@@ -389,30 +415,31 @@ public enum LatestChangeType
 
 			// Figure out which API server to use
 			string NewApiUrl;
-			if (TryGetProjectSetting(DetectSettings.LatestProjectConfigFile, "ApiUrl", out NewApiUrl))
+			if (TryGetProjectSetting(WorkspaceSettings.LatestProjectConfigFile, "ApiUrl", out NewApiUrl))
 			{
 				ApiUrl = NewApiUrl;
 			}
 
-			string TelemetryProjectIdentifier = PerforceUtils.GetClientOrDepotDirectoryName(DetectSettings.NewSelectedProjectIdentifier);
-
-			ProjectInfo Project = new ProjectInfo(BranchDirectoryName, SelectedFileName, DetectSettings.BranchClientPath, DetectSettings.NewSelectedClientFileName, TelemetryProjectIdentifier, DetectSettings.bIsEnterpriseProject);
+			ProjectInfo Project = WorkspaceSettings.ProjectInfo;
 
-			Workspace = new Workspace(PerforceClient, Project, WorkspaceState, new LogControlTextWriter(SyncLog));
+			Workspace = new Workspace(PerforceClientSettings, Project, WorkspaceState, WorkspaceSettings.WorkspaceProjectConfigFile, WorkspaceSettings.WorkspaceProjectStreamFilter, new LogControlTextWriter(SyncLog), ServiceProvider);
 			Workspace.OnUpdateComplete += UpdateCompleteCallback;
 
-			FileReference ProjectLogBaseName = FileReference.Combine(DataFolder, String.Format("{0}@{1}", PerforceClient.ClientName, DetectSettings.BranchClientPath.Replace("//" + PerforceClient.ClientName + "/", "").Trim('/').Replace("/", "$")));
+			FileReference ProjectLogBaseName = FileReference.Combine(DataFolder, "sync.log");
 
-			PerforceMonitor = new PerforceMonitor(PerforceClient, DetectSettings.BranchClientPath, DetectSettings.NewSelectedClientFileName, DetectSettings.NewSelectedProjectIdentifier, ProjectLogBaseName + ".p4.log", DetectSettings.bIsEnterpriseProject, DetectSettings.LatestProjectConfigFile, DetectSettings.CacheFolder, DetectSettings.LocalConfigFiles);
+			ILogger PerforceLogger = ServiceProvider.GetRequiredService<ILogger<PerforceMonitor>>();
+			PerforceMonitor = new PerforceMonitor(PerforceClientSettings, WorkspaceSettings.ProjectInfo, WorkspaceSettings.LatestProjectConfigFile, WorkspaceSettings.CacheFolder, WorkspaceSettings.LocalConfigFiles, ServiceProvider);
 			PerforceMonitor.OnUpdate += UpdateBuildListCallback;
 			PerforceMonitor.OnUpdateMetadata += UpdateBuildMetadataCallback;
-			PerforceMonitor.OnStreamChange += StreamChangedCallback;
-			PerforceMonitor.OnLoginExpired += LoginExpiredCallback;
+			PerforceMonitor.OnStreamChange += StreamChanged;
+			PerforceMonitor.OnLoginExpired += LoginExpired;
 
-			EventMonitor = new EventMonitor(ApiUrl, PerforceUtils.GetClientOrDepotDirectoryName(SelectedProjectIdentifier), DetectSettings.PerforceClient.UserName, ProjectLogBaseName + ".review.log");
+			ILogger EventLogger = ServiceProvider.GetRequiredService < ILogger<EventMonitor>>();
+			EventMonitor = new EventMonitor(ApiUrl, PerforceUtils.GetClientOrDepotDirectoryName(SelectedProjectIdentifier), WorkspaceSettings.PerforceSettings.UserName, ServiceProvider);
 			EventMonitor.OnUpdatesReady += UpdateReviewsCallback;
 
-			JupiterMonitor = JupiterMonitor.CreateFromConfigFile(InOidcTokenManager, ProjectLogBaseName + ".jupiter.log", DetectSettings.LatestProjectConfigFile, SelectedProjectIdentifier);
+			ILogger<JupiterMonitor> JupiterLogger = ServiceProvider.GetRequiredService<ILogger<JupiterMonitor>>();
+			JupiterMonitor = JupiterMonitor.CreateFromConfigFile(InOidcTokenManager, JupiterLogger, WorkspaceSettings.LatestProjectConfigFile, SelectedProjectIdentifier);
 
 			UpdateColumnSettings();
 
@@ -434,7 +461,7 @@ public enum LatestChangeType
 			StartupCallbacks = new List<WorkspaceStartupCallback>();
 
 			string IssuesApiUrl = GetIssuesApiUrl();
-			IssueMonitor = InOwner.CreateIssueMonitor(IssuesApiUrl, DetectSettings.PerforceClient.UserName);
+			IssueMonitor = InOwner.CreateIssueMonitor(IssuesApiUrl, WorkspaceSettings.PerforceSettings.UserName);
 			IssueMonitor.OnIssuesChanged += IssueMonitor_OnIssuesChangedAsync;
 
 			if (SelectedFileName.HasExtension(".uproject"))
@@ -460,7 +487,7 @@ public enum LatestChangeType
 
 		public List<IssueData> GetOpenIssuesForUser()
 		{
-			return IssueMonitor.GetIssues().Where(x => x.FixChange <= 0 && String.Compare(x.Owner, Workspace.Perforce.UserName, StringComparison.OrdinalIgnoreCase) == 0).ToList();
+			return IssueMonitor.GetIssues().Where(x => x.FixChange <= 0 && String.Compare(x.Owner, PerforceSettings.UserName, StringComparison.OrdinalIgnoreCase) == 0).ToList();
 		}
 
 		public void IssueMonitor_OnIssuesChanged()
@@ -804,7 +831,7 @@ public enum LatestChangeType
 		private void ShowErrorDialog(string Format, params object[] Args)
 		{
 			string Message = String.Format(Format, Args);
-			Log.WriteLine(Message);
+			Logger.LogError("{Message}", Message);
 			MessageBox.Show(Message);
 		}
 
@@ -829,25 +856,15 @@ public enum LatestChangeType
 			StatusPanel.ResumeDisplay();*/
 		}
 
-		private void StreamChangedCallback()
-		{
-			MainThreadSynchronizationContext.Post((o) => StreamChanged(), null);
-		}
-
 		private void LoginExpired()
 		{
 			if (!bIsDisposing)
 			{
-				Log.WriteLine("Login has expired. Requesting project to be closed.");
+				Logger.LogInformation("Login has expired. Requesting project to be closed.");
 				Owner.RequestProjectChange(this, SelectedProject, false);
 			}
 		}
 
-		private void LoginExpiredCallback()
-		{
-			MainThreadSynchronizationContext.Post((o) => LoginExpired(), null);
-		}
-
 		private void BuildList_OnScroll()
 		{
 			PendingSelectedChangeNumber = -1;
@@ -1010,10 +1027,10 @@ public enum LatestChangeType
 			Context.StartTime = DateTime.UtcNow;
 			Context.PerforceSyncOptions = (PerforceSyncOptions)Settings.SyncOptions.Clone();
 
-			Log.WriteLine("Updating workspace at {0}...", Context.StartTime.ToLocalTime().ToString());
-			Log.WriteLine("  ChangeNumber={0}", Context.ChangeNumber);
-			Log.WriteLine("  Options={0}", Context.Options.ToString());
-			Log.WriteLine("  Clobbering {0} files", Context.ClobberFiles.Count);
+			Logger.LogInformation("Updating workspace at {Time}...", Context.StartTime.ToLocalTime().ToString());
+			Logger.LogInformation("  ChangeNumber={Change}", Context.ChangeNumber);
+			Logger.LogInformation("  Options={Options}", Context.Options.ToString());
+			Logger.LogInformation("  Clobbering {NumFiles} files", Context.ClobberFiles.Count);
 
 			if (Context.Options.HasFlag(WorkspaceUpdateOptions.Sync))
 			{
@@ -1185,7 +1202,7 @@ public enum LatestChangeType
 				ArchiveToChangeNumberToArchiveKey.Clear();
 				ChangeNumberToLayoutInfo.Clear();
 
-				List<PerforceChangeSummary> Changes = PerforceMonitor.GetChanges();
+				List<ChangesRecord> Changes = PerforceMonitor.GetChanges();
 				EventMonitor.FilterChanges(Changes.Select(x => x.Number));
 
 				PromotedChangeNumbers = PerforceMonitor.GetPromotedChangeNumbers();
@@ -1207,10 +1224,10 @@ public enum LatestChangeType
 				ListIndexToChangeIndex = new List<int>();
 				SortedChangeNumbers = new List<int>();
 
-				List<PerforceChangeSummary> FilteredChanges = new List<PerforceChangeSummary>();
+				List<ChangesRecord> FilteredChanges = new List<ChangesRecord>();
 				for (int ChangeIdx = 0; ChangeIdx < Changes.Count; ChangeIdx++)
 				{
-					PerforceChangeSummary Change = Changes[ChangeIdx];
+					ChangesRecord Change = Changes[ChangeIdx];
 					if (ShouldShowChange(Change, ExcludeChanges) || PromotedChangeNumbers.Contains(Change.Number))
 					{
 						SortedChangeNumbers.Add(Change.Number);
@@ -1249,18 +1266,18 @@ public enum LatestChangeType
 			UpdateSyncActionCheckboxes();
 		}
 
-		void UpdateBuildListInternal(List<PerforceChangeSummary> Changes, bool bShowExpandItem)
+		void UpdateBuildListInternal(List<ChangesRecord> Changes, bool bShowExpandItem)
 		{
 			BuildList.BeginUpdate();
 
 			// Remove any changes that no longer exist, and update the rest
-			Dictionary<int, PerforceChangeSummary> ChangeNumberToSummary = Changes.ToDictionary(x => x.Number, x => x);
+			Dictionary<int, ChangesRecord> ChangeNumberToSummary = Changes.ToDictionary(x => x.Number, x => x);
 			for (int Idx = BuildList.Items.Count - 1; Idx >= 0; Idx--)
 			{
-				PerforceChangeSummary Change = BuildList.Items[Idx].Tag as PerforceChangeSummary;
+				ChangesRecord Change = BuildList.Items[Idx].Tag as ChangesRecord;
 				if (Change != null)
 				{
-					PerforceChangeSummary Summary;
+					ChangesRecord Summary;
 					if (ChangeNumberToSummary.TryGetValue(Change.Number, out Summary))
 					{
 						// Update
@@ -1276,7 +1293,7 @@ public enum LatestChangeType
 			}
 
 			// Add everything left over
-			foreach (PerforceChangeSummary Change in ChangeNumberToSummary.Values)
+			foreach (ChangesRecord Change in ChangeNumberToSummary.Values)
 			{
 				BuildList_AddItem(Change);
 			}
@@ -1329,7 +1346,7 @@ public enum LatestChangeType
 			BuildList.EndUpdate();
 		}
 
-		ListViewGroup BuildList_FindOrAddGroup(PerforceChangeSummary Change, string GroupName)
+		ListViewGroup BuildList_FindOrAddGroup(ChangesRecord Change, string GroupName)
 		{
 			// Find or add the new group
 			int GroupIndex = 0;
@@ -1341,10 +1358,10 @@ public enum LatestChangeType
 					return NextGroup;
 				}
 
-				PerforceChangeSummary LastChange = null;
+				ChangesRecord LastChange = null;
 				for (int Idx = NextGroup.Items.Count - 1; Idx >= 0 && LastChange == null; Idx--)
 				{
-					LastChange = NextGroup.Items[Idx].Tag as PerforceChangeSummary;
+					LastChange = NextGroup.Items[Idx].Tag as ChangesRecord;
 				}
 				if (LastChange == null || Change.Number > LastChange.Number)
 				{
@@ -1358,10 +1375,10 @@ public enum LatestChangeType
 			return Group;
 		}
 
-		void BuildList_AddItem(PerforceChangeSummary Change)
+		void BuildList_AddItem(ChangesRecord Change)
 		{
 			// Get the display time for this item
-			DateTime DisplayTime = Change.Date;
+			DateTime DisplayTime = Change.Time;
 			if (Settings.bShowLocalTimes)
 			{
 				DisplayTime = (DisplayTime - PerforceMonitor.ServerTimeZone).ToLocalTime();
@@ -1399,7 +1416,7 @@ public enum LatestChangeType
 			int GroupInsertIdx = 0;
 			for (; GroupInsertIdx < Group.Items.Count; GroupInsertIdx++)
 			{
-				PerforceChangeSummary OtherChange = Group.Items[GroupInsertIdx].Tag as PerforceChangeSummary;
+				ChangesRecord OtherChange = Group.Items[GroupInsertIdx].Tag as ChangesRecord;
 				if (OtherChange == null || Change.Number >= OtherChange.Number)
 				{
 					break;
@@ -1416,14 +1433,14 @@ public enum LatestChangeType
 		/// </summary>
 		/// <param name="Change"></param>
 		/// <returns></returns>
-		bool IsRobomergeChange(PerforceChangeSummary Change)
+		bool IsRobomergeChange(ChangesRecord Change)
 		{
 			// If hiding robomerge, filter out based on workspace name. Note - don't look at the description because we
 			// *do* want to see conflicts that were manually merged by someone
 			return Change.Client.IndexOf("ROBOMERGE", StringComparison.OrdinalIgnoreCase) >= 0;
 		}
 
-		bool ShouldShowChange(PerforceChangeSummary Change, string[] ExcludeChanges)
+		bool ShouldShowChange(ChangesRecord Change, string[] ExcludeChanges)
 		{
 			if (ProjectSettings.FilterType != FilterType.None)
 			{
@@ -1523,8 +1540,8 @@ public enum LatestChangeType
 			string IssuesApiUrl = GetIssuesApiUrl();
 			if (IssuesApiUrl != IssueMonitor.ApiUrl)
 			{
-				Log.WriteLine("Changing issues API url from {0} to {1}", IssueMonitor.ApiUrl, IssuesApiUrl);
-				IssueMonitor NewIssueMonitor = Owner.CreateIssueMonitor(IssuesApiUrl, Workspace.Perforce.UserName);
+				Logger.LogInformation("Changing issues API url from {OldApiUrl} to {ApiUrl}", IssueMonitor.ApiUrl, IssuesApiUrl);
+				IssueMonitor NewIssueMonitor = Owner.CreateIssueMonitor(IssuesApiUrl, PerforceSettings.UserName);
 				Owner.ReleaseIssueMonitor(IssueMonitor);
 				IssueMonitor = NewIssueMonitor;
 			}
@@ -1558,7 +1575,7 @@ public enum LatestChangeType
 			{
 				if (Item.Tag != null)
 				{
-					EventSummary Summary = EventMonitor.GetSummaryForChange(((PerforceChangeSummary)Item.Tag).Number);
+					EventSummary Summary = EventMonitor.GetSummaryForChange(((ChangesRecord)Item.Tag).Number);
 					if (Summary != null)
 					{
 						foreach (BadgeData Badge in Summary.Badges)
@@ -1804,7 +1821,7 @@ public enum LatestChangeType
 				}
 				catch (Exception Ex)
 				{
-					Log.WriteLine("Exception while parsing Uncontrolled Changelist persistency file: {0}", Ex.ToString());
+					Logger.LogError(Ex, "Exception while parsing Uncontrolled Changelist persistency file");
 				}
 			}
 		}
@@ -2044,7 +2061,7 @@ public enum LatestChangeType
 				{
 					BuildList.DrawTrackedBackground(e.Graphics, e.Bounds);
 				}
-				else if (((PerforceChangeSummary)e.Item.Tag).Number == Workspace.PendingChangeNumber)
+				else if (((ChangesRecord)e.Item.Tag).Number == Workspace.PendingChangeNumber)
 				{
 					BuildList.DrawTrackedBackground(e.Graphics, e.Bounds);
 				}
@@ -2104,7 +2121,7 @@ public enum LatestChangeType
 			return SelectedArchives.Count == 0 || SelectedArchives.All(x => GetArchiveKeyForChangeNumber(x, ChangeNumber) != null);
 		}
 
-		private ChangeLayoutInfo GetChangeLayoutInfo(PerforceChangeSummary Change)
+		private ChangeLayoutInfo GetChangeLayoutInfo(ChangesRecord Change)
 		{
 			ChangeLayoutInfo LayoutInfo;
 			if (!ChangeNumberToLayoutInfo.TryGetValue(Change.Number, out LayoutInfo))
@@ -2149,7 +2166,7 @@ public enum LatestChangeType
 
 		private void BuildList_DrawSubItem(object sender, DrawListViewSubItemEventArgs e)
 		{
-			PerforceChangeSummary Change = (PerforceChangeSummary)e.Item.Tag;
+			ChangesRecord Change = (ChangesRecord)e.Item.Tag;
 			if (Change == null)
 			{
 				return;
@@ -2481,7 +2498,7 @@ public enum LatestChangeType
 			}
 		}
 
-		private List<BadgeInfo> CreateDescriptionBadges(PerforceChangeSummary Change)
+		private List<BadgeInfo> CreateDescriptionBadges(ChangesRecord Change)
 		{
 			string Description = Change.Description;
 
@@ -2995,7 +3012,7 @@ public enum LatestChangeType
 				ListViewHitTestInfo HitTest = BuildList.HitTest(Args.Location);
 				if (HitTest.Item != null)
 				{
-					PerforceChangeSummary Change = (PerforceChangeSummary)HitTest.Item.Tag;
+					ChangesRecord Change = (ChangesRecord)HitTest.Item.Tag;
 					if (Change != null)
 					{
 						if (Change.Number == Workspace.CurrentChangeNumber)
@@ -3096,7 +3113,7 @@ public enum LatestChangeType
 			{
 				if (FileReference.Exists(ReceiptFileName))
 				{
-					Log.WriteLine("Reading {0}", ReceiptFileName);
+					Logger.LogInformation("Reading {FileName}", ReceiptFileName);
 					try
 					{
 						string Text = FileReference.ReadAllText(ReceiptFileName);
@@ -3112,7 +3129,7 @@ public enum LatestChangeType
 					}
 					catch (Exception Ex)
 					{
-						Log.WriteLine("  Exception while parsing receipt: {0}", Ex.ToString());
+						Logger.LogError(Ex, "Exception while parsing receipt");
 					}
 					break;
 				}
@@ -3306,7 +3323,7 @@ public enum LatestChangeType
 					}
 					catch(Exception Ex)
 					{
-						Log.WriteLine("Unable to parse config filter '{0}': {1}", Filter, Ex);
+						Logger.LogError(Ex, "Unable to parse config filter '{Filter}'", Filter);
 					}
 				}
 			}
@@ -3317,9 +3334,9 @@ public enum LatestChangeType
 		private void OpenPerforce()
 		{
 			StringBuilder CommandLine = new StringBuilder();
-			if (Workspace != null && Workspace.Perforce != null)
+			if (Workspace != null)
 			{
-				CommandLine.Append(Workspace.Perforce.GetConnectionOptions());
+				CommandLine.Append(PerforceSettings.GetArgumentsForExternalProgram(true));
 			}
 			SafeProcessStart("p4v.exe", CommandLine.ToString());
 		}
@@ -3635,7 +3652,7 @@ public enum LatestChangeType
 				{
 					Issues.Add(Issue);
 				}
-				else if (Issue.FixChange == 0 && Issue.Owner != null && String.Compare(Issue.Owner, Workspace.Perforce.UserName, StringComparison.OrdinalIgnoreCase) == 0)
+				else if (Issue.FixChange == 0 && Issue.Owner != null && String.Compare(Issue.Owner, PerforceSettings.UserName, StringComparison.OrdinalIgnoreCase) == 0)
 				{
 					bHasOtherAssignedIssue = true;
 				}
@@ -3686,7 +3703,7 @@ public enum LatestChangeType
 					{
 						Item.ForeColor = SystemColors.GrayText;
 					}
-					else if (Issue.Owner != null && String.Compare(Issue.Owner, Workspace.Perforce.UserName, StringComparison.OrdinalIgnoreCase) == 0)
+					else if (Issue.Owner != null && String.Compare(Issue.Owner, PerforceSettings.UserName, StringComparison.OrdinalIgnoreCase) == 0)
 					{
 						Item.Font = new Font(Item.Font, FontStyle.Bold);
 					}
@@ -3715,14 +3732,14 @@ public enum LatestChangeType
 
 		public void ShowIssueDetails(IssueData Issue)
 		{
-			IssueDetailsWindow.Show(ParentForm, IssueMonitor, Workspace.Perforce.ServerAndPort, Workspace.Perforce.UserName, GetServerTimeOffset(), Issue, Log, StreamName);
+			IssueDetailsWindow.Show(ParentForm, IssueMonitor, PerforceSettings, GetServerTimeOffset(), Issue, ServiceProvider, StreamName);
 		}
 
 		private void BuildHealthContextMenu_Browse_Click(object sender, EventArgs e)
 		{
 			string DefaultFilter = GetDefaultIssueFilter();
 			Dictionary<string, Func<IssueData, bool>> CustomFilters = GetCustomIssueFilters();
-			IssueBrowserWindow.Show(ParentForm, IssueMonitor, Workspace.Perforce.ServerAndPort, Workspace.Perforce.UserName, GetServerTimeOffset(), Log, StreamName, CustomFilters, DefaultFilter);
+			IssueBrowserWindow.Show(ParentForm, IssueMonitor, PerforceSettings, GetServerTimeOffset(), ServiceProvider, StreamName, CustomFilters, DefaultFilter);
 		}
 
 		private void BuildHealthContextMenu_Settings_Click(object sender, EventArgs e)
@@ -3893,7 +3910,7 @@ public enum LatestChangeType
 		private void SelectOtherStreamDialog()
 		{
 			string NewStreamName;
-			if (SelectStreamWindow.ShowModal(this, Workspace.Perforce, StreamName, Log, out NewStreamName))
+			if (SelectStreamWindow.ShowModal(this, PerforceSettings, StreamName, ServiceProvider, out NewStreamName))
 			{
 				SelectStream(NewStreamName);
 			}
@@ -3907,15 +3924,31 @@ public enum LatestChangeType
 				{
 					MessageBox.Show("Please retry after the current sync has finished.", "Sync in Progress");
 				}
-				else if (Workspace.Perforce != null)
+				else
 				{
-					if (!Workspace.Perforce.HasOpenFiles(Log) || MessageBox.Show("You have files open for edit in this workspace. If you continue, you will not be able to submit them until you switch back.\n\nContinue switching streams?", "Files checked out", MessageBoxButtons.YesNo) == DialogResult.Yes)
+					for (int Idx = 0; Idx < 2; Idx++)
 					{
-						if (Workspace.Perforce.SwitchStream(NewStreamName, Log))
+						Func<IPerforceConnection, CancellationToken, Task<bool>> SwitchFunc = async (Perforce, CancellationToken) =>
+						{
+							if (Idx == 1 || !await Perforce.OpenedAsync(OpenedOptions.None, FileSpecList.Any, CancellationToken).AnyAsync())
+							{
+								await Perforce.SwitchClientToStreamAsync(Perforce.Settings.ClientName, NewStreamName, SwitchClientOptions.IgnoreOpenFiles, CancellationToken);
+								return true;
+							}
+							return false;
+						};
+
+						ModalTask<bool> SwitchTask = PerforceModalTask.Execute<bool>(this, "Switching streams", "Please wait...", PerforceSettings, SwitchFunc, Logger);
+						if (!SwitchTask.Succeeded)
+						{
+							break;
+						}
+						if (SwitchTask.Result)
 						{
 							StatusPanel.SuspendLayout();
 							StreamChanged();
 							StatusPanel.ResumeLayout();
+							break;
 						}
 					}
 				}
@@ -3963,7 +3996,7 @@ public enum LatestChangeType
 
 			foreach (ListViewItem Item in BuildList.Items)
 			{
-				PerforceChangeSummary Summary = (PerforceChangeSummary)Item.Tag;
+				ChangesRecord Summary = (ChangesRecord)Item.Tag;
 				if (Summary != null && Summary.Number <= ChangeNumber)
 				{
 					Item.Selected = true;
@@ -4010,7 +4043,7 @@ public enum LatestChangeType
 					}
 					else
 					{
-						PerforceChangeSummary Change = (PerforceChangeSummary)HitTest.Item.Tag;
+						ChangesRecord Change = (ChangesRecord)HitTest.Item.Tag;
 						if (Workspace.PendingChangeNumber == Change.Number)
 						{
 							Rectangle SubItemRect = HitTest.Item.SubItems[StatusColumn.Index].Bounds;
@@ -4093,7 +4126,7 @@ public enum LatestChangeType
 						{
 							if (CustomColumn.Index < HitTest.Item.SubItems.Count && HitTest.Item.SubItems[CustomColumn.Index] == HitTest.SubItem)
 							{
-								ChangeLayoutInfo LayoutInfo = GetChangeLayoutInfo((PerforceChangeSummary)HitTest.Item.Tag);
+								ChangeLayoutInfo LayoutInfo = GetChangeLayoutInfo((ChangesRecord)HitTest.Item.Tag);
 
 								List<BadgeInfo> Badges;
 								if (LayoutInfo.CustomBadges.TryGetValue(CustomColumn.Text, out Badges) && Badges.Count > 0)
@@ -4129,7 +4162,7 @@ public enum LatestChangeType
 					}
 					else
 					{
-						ContextMenuChange = (PerforceChangeSummary)HitTest.Item.Tag;
+						ContextMenuChange = (ChangesRecord)HitTest.Item.Tag;
 
 						BuildListContextMenu_WithdrawReview.Visible = (EventMonitor.GetReviewByCurrentUser(ContextMenuChange.Number) != null);
 						BuildListContextMenu_StartInvestigating.Visible = !EventMonitor.IsUnderInvestigationByCurrentUser(ContextMenuChange.Number);
@@ -4238,12 +4271,12 @@ public enum LatestChangeType
 			}
 
 			string NewHoverBadgeUniqueId = null;
-			if (HitTest.Item != null && HitTest.Item.Tag is PerforceChangeSummary)
+			if (HitTest.Item != null && HitTest.Item.Tag is ChangesRecord)
 			{
 				int ColumnIndex = HitTest.Item.SubItems.IndexOf(HitTest.SubItem);
 				if (ColumnIndex == DescriptionColumn.Index)
 				{
-					ChangeLayoutInfo LayoutInfo = GetChangeLayoutInfo((PerforceChangeSummary)HitTest.Item.Tag);
+					ChangeLayoutInfo LayoutInfo = GetChangeLayoutInfo((ChangesRecord)HitTest.Item.Tag);
 					if (LayoutInfo.DescriptionBadges.Count > 0)
 					{
 						Point ListLocation = GetBadgeListLocation(LayoutInfo.DescriptionBadges, HitTest.SubItem.Bounds, HorizontalAlign.Right, VerticalAlignment.Middle);
@@ -4252,7 +4285,7 @@ public enum LatestChangeType
 				}
 				else if (ColumnIndex == CISColumn.Index)
 				{
-					ChangeLayoutInfo LayoutInfo = GetChangeLayoutInfo((PerforceChangeSummary)HitTest.Item.Tag);
+					ChangeLayoutInfo LayoutInfo = GetChangeLayoutInfo((ChangesRecord)HitTest.Item.Tag);
 					if (LayoutInfo.BuildBadges.Count > 0)
 					{
 						Point BuildListLocation = GetBadgeListLocation(LayoutInfo.BuildBadges, HitTest.SubItem.Bounds, HorizontalAlign.Center, VerticalAlignment.Middle);
@@ -4278,7 +4311,7 @@ public enum LatestChangeType
 				{
 					ColumnHeader Column = BuildList.Columns[ColumnIndex];
 
-					ChangeLayoutInfo LayoutInfo = GetChangeLayoutInfo((PerforceChangeSummary)HitTest.Item.Tag);
+					ChangeLayoutInfo LayoutInfo = GetChangeLayoutInfo((ChangesRecord)HitTest.Item.Tag);
 
 					List<BadgeInfo> Badges;
 					if (LayoutInfo.CustomBadges.TryGetValue(Column.Text, out Badges) && Badges.Count > 0)
@@ -4632,7 +4665,7 @@ public enum LatestChangeType
 
 		private void BuildListContextMenu_MoreInfo_Click(object sender, EventArgs e)
 		{
-			Utility.SpawnP4VC(String.Format("{0} change {1}", Workspace.Perforce.GetConnectionOptions(), ContextMenuChange.Number));
+			Utility.SpawnP4VC(String.Format("{0} change {1}", PerforceSettings.GetArgumentsForExternalProgram(true), ContextMenuChange.Number));
 		}
 
 		private void BuildListContextMenu_ViewInSwarm_Click(object sender, EventArgs e)
@@ -4698,7 +4731,7 @@ public enum LatestChangeType
 					string Message = String.Format("Mark all changes between {0} and {1} as bad?", StartChangeNumber, ContextMenuChange.Number);
 					if (MessageBox.Show(Message, "Finish investigating", MessageBoxButtons.YesNo) == DialogResult.Yes)
 					{
-						foreach (PerforceChangeSummary Change in PerforceMonitor.GetChanges())
+						foreach (ChangesRecord Change in PerforceMonitor.GetChanges())
 						{
 							if (Change.Number >= StartChangeNumber && Change.Number < ContextMenuChange.Number)
 							{
@@ -4745,7 +4778,7 @@ public enum LatestChangeType
 			{
 				if (Args.Item.Bounds.Contains(ClientPoint))
 				{
-					PerforceChangeSummary Change = (PerforceChangeSummary)Args.Item.Tag;
+					ChangesRecord Change = (ChangesRecord)Args.Item.Tag;
 					if (Change == null)
 					{
 						return;
@@ -4890,7 +4923,7 @@ public enum LatestChangeType
 				}
 				catch (Exception Ex)
 				{
-					Log.WriteException(Ex, "Unable to read '{0}'", MasterProjectNameFileName);
+					Logger.LogError(Ex, "Unable to read '{File}'", MasterProjectNameFileName);
 				}
 			}
 
@@ -4953,16 +4986,16 @@ public enum LatestChangeType
 		{
 			if (Settings.bScheduleEnabled)
 			{
-				Log.WriteLine("Scheduled sync at {0} for {1}.", DateTime.Now, Settings.ScheduleChange);
+				Logger.LogInformation("Scheduled sync at {Time} for {Change}.", DateTime.Now, Settings.ScheduleChange);
 
 				int ChangeNumber;
 				if (!FindChangeToSync(Settings.ScheduleChange, out ChangeNumber))
 				{
-					Log.WriteLine("Couldn't find any matching change");
+					Logger.LogInformation("Couldn't find any matching change");
 				}
 				else if (Workspace.CurrentChangeNumber >= ChangeNumber)
 				{
-					Log.WriteLine("Sync ignored; already at or ahead of CL ({0} >= {1})", Workspace.CurrentChangeNumber, ChangeNumber);
+					Logger.LogInformation("Sync ignored; already at or ahead of CL ({CurrentChange} >= {Change})", Workspace.CurrentChangeNumber, ChangeNumber);
 				}
 				else
 				{
@@ -4980,7 +5013,7 @@ public enum LatestChangeType
 		{
 			for (int Idx = 0; Idx < BuildList.Items.Count; Idx++)
 			{
-				PerforceChangeSummary Change = (PerforceChangeSummary)BuildList.Items[Idx].Tag;
+				ChangesRecord Change = (ChangesRecord)BuildList.Items[Idx].Tag;
 				if (Change != null)
 				{
 					if (ChangeType == LatestChangeType.Any)
@@ -5103,7 +5136,7 @@ public enum LatestChangeType
 			string[] CombinedSyncFilter = UserSettings.GetCombinedSyncFilter(GetSyncCategories(), Settings.SyncView, Settings.SyncCategories, WorkspaceSettings.SyncView, WorkspaceSettings.SyncCategoriesDict);
 			List<string> SyncPaths = Workspace.GetSyncPaths(Workspace.Project, WorkspaceSettings.bSyncAllProjects ?? Settings.bSyncAllProjects, CombinedSyncFilter);
 
-			CleanWorkspaceWindow.DoClean(ParentForm, Workspace.Perforce, BranchDirectoryName, Workspace.Project.ClientRootPath, SyncPaths, ExtraSafeToDeleteFolders.Split('\n'), ExtraSafeToDeleteExtensions.Split('\n'), Log);
+			CleanWorkspaceWindow.DoClean(ParentForm, PerforceSettings, BranchDirectoryName, Workspace.Project.ClientRootPath, SyncPaths, ExtraSafeToDeleteFolders.Split('\n'), ExtraSafeToDeleteExtensions.Split('\n'), ServiceProvider.GetRequiredService<ILogger<CleanWorkspaceWindow>>());
 		}
 
 		private void UpdateBuildSteps()
@@ -5430,9 +5463,9 @@ public enum LatestChangeType
 			DiagnosticsText.AppendFormat("Selected file: {0}\n", (SelectedFileName == null) ? "(none)" : SelectedFileName.FullName);
 			if (Workspace != null)
 			{
-				DiagnosticsText.AppendFormat("P4 server: {0}\n", Workspace.Perforce.ServerAndPort ?? "(default)");
-				DiagnosticsText.AppendFormat("P4 user: {0}\n", Workspace.Perforce.UserName);
-				DiagnosticsText.AppendFormat("P4 workspace: {0}\n", Workspace.Perforce.ClientName);
+				DiagnosticsText.AppendFormat("P4 server: {0}\n", PerforceSettings.ServerAndPort ?? "(default)");
+				DiagnosticsText.AppendFormat("P4 user: {0}\n", PerforceSettings.UserName);
+				DiagnosticsText.AppendFormat("P4 workspace: {0}\n", PerforceSettings.ClientName);
 			}
 			DiagnosticsText.AppendFormat("Perforce monitor: {0}\n", (PerforceMonitor == null) ? "(inactive)" : PerforceMonitor.LastStatusMessage);
 			DiagnosticsText.AppendFormat("Event monitor: {0}\n", (EventMonitor == null) ? "(inactive)" : EventMonitor.LastStatusMessage);
@@ -5505,7 +5538,7 @@ public enum LatestChangeType
 		{
 			if (Args.Control && Args.KeyCode == Keys.C && BuildList.SelectedItems.Count > 0)
 			{
-				int SelectedChange = ((PerforceChangeSummary)BuildList.SelectedItems[0].Tag).Number;
+				int SelectedChange = ((ChangesRecord)BuildList.SelectedItems[0].Tag).Number;
 				Clipboard.SetText(String.Format("{0}", SelectedChange));
 			}
 		}
@@ -5772,7 +5805,7 @@ public enum LatestChangeType
 				Dictionary<int, BisectState> ChangeNumberToBisectState = new Dictionary<int, BisectState>();
 				foreach (ListViewItem SelectedItem in BuildList.SelectedItems)
 				{
-					PerforceChangeSummary Change = (PerforceChangeSummary)SelectedItem.Tag;
+					ChangesRecord Change = (ChangesRecord)SelectedItem.Tag;
 					ChangeNumberToBisectState[Change.Number] = BisectState.Include;
 				}
 
@@ -5800,7 +5833,7 @@ public enum LatestChangeType
 		{
 			foreach (ListViewItem SelectedItem in BuildList.SelectedItems)
 			{
-				PerforceChangeSummary Change = (PerforceChangeSummary)SelectedItem.Tag;
+				ChangesRecord Change = (ChangesRecord)SelectedItem.Tag;
 				if (Change != null)
 				{
 					WorkspaceState.SetBisectState(Change.Number, State);
diff --git a/Engine/Source/Programs/UnrealGameSync/UnrealGameSync/DeleteFilesTask.cs b/Engine/Source/Programs/UnrealGameSync/UnrealGameSync/DeleteFilesTask.cs
--- a/Engine/Source/Programs/UnrealGameSync/UnrealGameSync/DeleteFilesTask.cs
+++ b/Engine/Source/Programs/UnrealGameSync/UnrealGameSync/DeleteFilesTask.cs
@@ -1,45 +1,39 @@
 // Copyright Epic Games, Inc. All Rights Reserved.
 
+using EpicGames.Perforce;
+using Microsoft.Extensions.Logging;
 using System;
 using System.Collections.Generic;
 using System.IO;
 using System.Linq;
 using System.Text;
+using System.Threading;
 using System.Threading.Tasks;
 
+#nullable enable
+
 namespace UnrealGameSync
 {
-	class DeleteFilesTask : IModalTask
+	static class DeleteFilesTask
 	{
-		PerforceConnection Perforce;
-		List<FileInfo> FilesToSync;
-		List<FileInfo> FilesToDelete;
-		List<DirectoryInfo> DirectoriesToDelete;
-
-		public DeleteFilesTask(PerforceConnection Perforce, List<FileInfo> FilesToSync, List<FileInfo> FilesToDelete, List<DirectoryInfo> DirectoriesToDelete)
-		{
-			this.Perforce = Perforce;
-			this.FilesToSync = FilesToSync;
-			this.FilesToDelete = FilesToDelete;
-			this.DirectoriesToDelete = DirectoriesToDelete;
-		}
-
-		public bool Run(out string ErrorMessage)
+		public static async Task RunAsync(IPerforceSettings PerforceSettings, List<FileInfo> FilesToSync, List<FileInfo> FilesToDelete, List<DirectoryInfo> DirectoriesToDelete, ILogger Logger, CancellationToken CancellationToken)
 		{
 			StringBuilder FailMessage = new StringBuilder();
 
 			if(FilesToSync.Count > 0)
 			{
+				using IPerforceConnection Perforce = await PerforceConnection.CreateAsync(PerforceSettings, Logger);
+
 				List<string> RevisionsToSync = new List<string>();
 				foreach(FileInfo FileToSync in FilesToSync)
 				{
 					RevisionsToSync.Add(String.Format("{0}#have", PerforceUtils.EscapePath(FileToSync.FullName)));
 				}
 
-				StringWriter Log = new StringWriter();
-				if(!Perforce.Sync(RevisionsToSync, x => { }, new List<string>(), true, null, Log))
+				List<PerforceResponse<SyncRecord>> FailedRecords = await Perforce.TrySyncAsync(RevisionsToSync, CancellationToken).Where(x => x.Failed).ToListAsync(CancellationToken);
+				foreach (PerforceResponse<SyncRecord> FailedRecord in FailedRecords)
 				{
-					FailMessage.Append(Log.ToString());
+					FailMessage.Append(FailedRecord.ToString());
 				}
 			}
 
@@ -52,6 +46,7 @@ namespace UnrealGameSync
 				}
 				catch(Exception Ex)
 				{
+					Logger.LogWarning(Ex, "Unable to delete {File}", FileToDelete.FullName);
 					FailMessage.AppendFormat("{0} ({1})\r\n", FileToDelete.FullName, Ex.Message.Trim());
 				}
 			}
@@ -63,19 +58,14 @@ namespace UnrealGameSync
 				}
 				catch(Exception Ex)
 				{
+					Logger.LogWarning(Ex, "Unable to delete {Directory}", DirectoryToDelete.FullName);
 					FailMessage.AppendFormat("{0} ({1})\r\n", DirectoryToDelete.FullName, Ex.Message.Trim());
 				}
 			}
 
-			if(FailMessage.Length == 0)
-			{
-				ErrorMessage = null;
-				return true;
-			}
-			else 
+			if(FailMessage.Length > 0)
 			{
-				ErrorMessage = FailMessage.ToString();
-				return false;
+				throw new UserErrorException(FailMessage.ToString());
 			}
 		}
 	}
diff --git a/Engine/Source/Programs/UnrealGameSync/UnrealGameSync/DeploymentSettings.cs b/Engine/Source/Programs/UnrealGameSync/UnrealGameSync/DeploymentSettings.cs
--- a/Engine/Source/Programs/UnrealGameSync/UnrealGameSync/DeploymentSettings.cs
+++ b/Engine/Source/Programs/UnrealGameSync/UnrealGameSync/DeploymentSettings.cs
@@ -1,5 +1,6 @@
 // Copyright Epic Games, Inc. All Rights Reserved.
 
+using Microsoft.Extensions.Logging;
 using System;
 using System.Collections.Generic;
 using System.IO;
@@ -20,9 +21,9 @@ namespace UnrealGameSync
 		/// </summary>
 		/// <param name="UserName">The default Perforce user name</param>
 		/// <param name="SessionId">Unique identifier for this session</param>
-		/// <param name="Log">Log writer</param>
+		/// <param name="Logger">Log writer</param>
 		/// <returns>New telemetry sink instance</returns>
-		public delegate ITelemetrySink CreateTelemetrySinkDelegate(string UserName, string SessionId, TextWriter Log);
+		public delegate ITelemetrySink CreateTelemetrySinkDelegate(string UserName, string SessionId, ILogger Logger);
 #endif
 
 		/// <summary>
@@ -68,11 +69,11 @@ namespace UnrealGameSync
 		/// <summary>
 		/// Delegate to allow validating a project being opened
 		/// </summary>
-		/// <param name="Task">The detected settings for the project</param>
-		/// <param name="Log">The logger</param>
+		/// <param name="Workspace">The detected settings for the project</param>
+		/// <param name="Logger">The logger</param>
 		/// <param name="Error">Receives an error on failure</param>
 		/// <returns></returns>
-		public delegate bool DetectProjectSettingsEvent(DetectProjectSettingsTask Task, TextWriter Log, out string Error);
+		public delegate bool DetectProjectSettingsEvent(WorkspaceSettings Workspace, ILogger Logger, out string Error);
 
 		/// <summary>
 		/// Called to validate the project settings
