UGS: Significant refactor of UGS internals.

diff --git a/Engine/Source/Programs/UnrealGameSync/UnrealGameSync/DetectProjectSettingsTask.cs b/Engine/Source/Programs/UnrealGameSync/UnrealGameSync/DetectProjectSettingsTask.cs
--- a/Engine/Source/Programs/UnrealGameSync/UnrealGameSync/DetectProjectSettingsTask.cs
+++ b/Engine/Source/Programs/UnrealGameSync/UnrealGameSync/DetectProjectSettingsTask.cs
@@ -1,492 +1,401 @@
 // Copyright Epic Games, Inc. All Rights Reserved.
 
 using EpicGames.Core;
+using EpicGames.Perforce;
+using Microsoft.Extensions.Logging;
 using System;
 using System.Collections.Generic;
 using System.Drawing;
 using System.IO;
 using System.Linq;
 using System.Text.RegularExpressions;
+using System.Threading;
 using System.Threading.Tasks;
 
+#nullable enable
+
 namespace UnrealGameSync
 {
-	class DetectProjectSettingsTask : IPerforceModalTask, IDisposable
+	public class WorkspaceSettings
 	{
-		public UserSelectedProjectSettings SelectedProject;
-		public PerforceConnection PerforceClient;
-		public FileReference NewSelectedFileName;
+		public IPerforceSettings PerforceSettings;
+		public UserSelectedProjectSettings SelectedProject { get; }
+		public ProjectInfo ProjectInfo { get; private set; }
 		public string NewSelectedProjectIdentifier;
 		public string NewProjectEditorTarget;
-		public string BranchClientPath;
-		public DirectoryReference BranchDirectoryName;
-		public string NewSelectedClientFileName;
 		public string StreamName;
-		public Image ProjectLogo;
-		public DirectoryReference DataFolder;
-		public DirectoryReference CacheFolder;
-		public bool bIsEnterpriseProject;
+		public DirectoryReference DataFolder => GetDataFolder(ProjectInfo.LocalRootPath);
+		public DirectoryReference CacheFolder => GetCacheFolder(ProjectInfo.LocalRootPath);
 		public ConfigFile LatestProjectConfigFile;
+		public ConfigFile WorkspaceProjectConfigFile;
+		public IReadOnlyList<string>? WorkspaceProjectStreamFilter;
 		public List<KeyValuePair<FileReference, DateTime>> LocalConfigFiles;
-		TextWriter Log;
 
-		public DetectProjectSettingsTask(UserSelectedProjectSettings SelectedProject, DirectoryReference InDataFolder, DirectoryReference InCacheFolder, TextWriter InLog)
+		public WorkspaceSettings(IPerforceSettings PerforceSettings, UserSelectedProjectSettings ProjectSettings, ProjectInfo ProjectInfo, string NewSelectedProjectIdentifier, string NewProjectEditorTarget, string StreamName, ConfigFile LatestProjectConfigFile, ConfigFile WorkspaceProjectConfigFile, IReadOnlyList<string>? WorkspaceProjectStreamFilter, List<KeyValuePair<FileReference, DateTime>> LocalConfigFiles)
 		{
-			this.SelectedProject = SelectedProject;
-			DataFolder = InDataFolder;
-			CacheFolder = InCacheFolder;
-			Log = InLog;
+			this.PerforceSettings = PerforceSettings;
+			this.SelectedProject = ProjectSettings;
+			this.ProjectInfo = ProjectInfo;
+			this.NewSelectedProjectIdentifier = NewSelectedProjectIdentifier;
+			this.NewProjectEditorTarget = NewProjectEditorTarget;
+			this.StreamName = StreamName;
+			this.LatestProjectConfigFile = LatestProjectConfigFile;
+			this.WorkspaceProjectConfigFile = WorkspaceProjectConfigFile;
+			this.WorkspaceProjectStreamFilter = WorkspaceProjectStreamFilter;
+			this.LocalConfigFiles = LocalConfigFiles;
 		}
 
-		public void Dispose()
-		{
-			if(ProjectLogo != null)
-			{
-				ProjectLogo.Dispose();
-				ProjectLogo = null;
-			}
-		}
+		public static DirectoryReference GetDataFolder(DirectoryReference WorkspaceDir) => DirectoryReference.Combine(WorkspaceDir, ".ugs");
+		public static DirectoryReference GetCacheFolder(DirectoryReference WorkspaceDir) => DirectoryReference.Combine(WorkspaceDir, ".ugs", "cache");
 
-		public bool Run(PerforceConnection Perforce, out string ErrorMessage)
+		public static async Task<WorkspaceSettings> CreateAsync(IPerforceSettings DefaultPerforceSettings, UserSelectedProjectSettings SelectedProject, ILogger<WorkspaceSettings> Logger, CancellationToken CancellationToken)
 		{
+			PerforceSettings PerforceSettings = Utility.OverridePerforceSettings(DefaultPerforceSettings, SelectedProject.ServerAndPort, SelectedProject.UserName);
+			using IPerforceConnection Perforce = await PerforceConnection.CreateAsync(PerforceSettings, Logger);
+
 			// Make sure we're logged in
-			bool bLoggedIn;
-			if(!Perforce.GetLoggedInState(out bLoggedIn, Log))
-			{
-				ErrorMessage = "Error while checking login status.";
-				return false;
-			}
-			if(!bLoggedIn)
+			PerforceResponse<LoginRecord> LoginState = await Perforce.TryGetLoginStateAsync(CancellationToken);
+			if (!LoginState.Succeeded)
 			{
-				ErrorMessage = "User is not logged in to Perforce.";
-				return false;
+				throw new UserErrorException("User is not logged in to Perforce.");
 			}
 
 			// Execute like a regular task
-			return Run(Perforce, Log, out ErrorMessage);
+			return await CreateAsync(Perforce, SelectedProject, Logger, CancellationToken);
 		}
 
-		public bool Run(PerforceConnection Perforce, TextWriter Log, out string ErrorMessage)
+		public static async Task<WorkspaceSettings> CreateAsync(IPerforceConnection DefaultConnection, UserSelectedProjectSettings SelectedProject, ILogger<WorkspaceSettings> Logger, CancellationToken CancellationToken)
 		{
+			using IDisposable LoggerScope = Logger.BeginScope("Project {SelectedProject}", SelectedProject.ToString());
+			Logger.LogInformation("Detecting settings for {Project}", SelectedProject);
+
 			// Use the cached client path to the file if it's available; it's much quicker than trying to find the correct workspace.
-			if(!String.IsNullOrEmpty(SelectedProject.ClientPath))
+			IPerforceConnection? PerforceClient = null;
+			try
 			{
-				// Get the client path
-				NewSelectedClientFileName = SelectedProject.ClientPath;
+				IPerforceSettings PerforceSettings;
 
-				// Get the client name
-				string ClientName;
-				if(!PerforceUtils.TryGetClientName(NewSelectedClientFileName, out ClientName))
+				FileReference NewSelectedFileName;
+				string NewSelectedClientFileName;
+				if (!String.IsNullOrEmpty(SelectedProject.ClientPath))
 				{
-					ErrorMessage = String.Format("Couldn't get client name from {0}", NewSelectedClientFileName);
-					return false;
-				}
+					// Get the client path
+					NewSelectedClientFileName = SelectedProject.ClientPath;
 
-				// Create the client
-				PerforceClient = new PerforceConnection(Perforce.UserName, ClientName, Perforce.ServerAndPort);
+					// Get the client name
+					string? ClientName;
+					if (!PerforceUtils.TryGetClientName(NewSelectedClientFileName, out ClientName))
+					{
+						throw new UserErrorException($"Couldn't get client name from {NewSelectedClientFileName}");
+					}
 
-				// Figure out the path on the client. Use the cached location if it's valid.
-				string LocalPath = SelectedProject.LocalPath;
-				if (LocalPath == null || !File.Exists(LocalPath))
-				{
-					if (!PerforceClient.ConvertToLocalPath(NewSelectedClientFileName, out LocalPath, Log))
+					// Create the client
+					PerforceSettings = new PerforceSettings(DefaultConnection.Settings) { ClientName = ClientName };
+					PerforceClient = await PerforceConnection.CreateAsync(PerforceSettings, Logger);
+
+					// Figure out the path on the client. Use the cached location if it's valid.
+					string? LocalPath = SelectedProject.LocalPath;
+					if (LocalPath == null || !File.Exists(LocalPath))
 					{
-						ErrorMessage = String.Format("Couldn't get client path for {0}", NewSelectedClientFileName);
-						return false;
+						List<WhereRecord> Records = await PerforceClient.WhereAsync(NewSelectedClientFileName, CancellationToken).Where(x => !x.Unmap).ToListAsync(CancellationToken);
+						if (Records.Count != 1)
+						{
+							throw new UserErrorException($"Couldn't get client path for {NewSelectedClientFileName}");
+						}
+						LocalPath = Path.GetFullPath(Records[0].Path);
 					}
+					NewSelectedFileName = new FileReference(LocalPath);
 				}
-				NewSelectedFileName = new FileReference(LocalPath);
-			}
-			else
-			{
-				// Get the perforce server settings
-				PerforceInfoRecord PerforceInfo;
-				if(!Perforce.Info(out PerforceInfo, Log))
+				else
 				{
-					ErrorMessage = String.Format("Couldn't get Perforce server info");
-					return false;
-				}
+					PerforceSettings = DefaultConnection.Settings;
 
-				// Use the path as the selected filename
-				NewSelectedFileName = new FileReference(SelectedProject.LocalPath);
+					// Get the perforce server settings
+					InfoRecord PerforceInfo = await DefaultConnection.GetInfoAsync(InfoOptions.ShortOutput, CancellationToken);
 
-				// Make sure the project exists
-				if (!FileReference.Exists(NewSelectedFileName))
-				{
-					ErrorMessage = String.Format("{0} does not exist.", SelectedProject.LocalPath);
-					return false;
-				}
+					// Use the path as the selected filename
+					NewSelectedFileName = new FileReference(SelectedProject.LocalPath!);
 
-				// Find all the clients for this user
-				Log.WriteLine("Enumerating clients for {0}...", PerforceInfo.UserName);
+					// Make sure the project exists
+					if (!FileReference.Exists(NewSelectedFileName))
+					{
+						throw new UserErrorException($"{SelectedProject.LocalPath} does not exist.");
+					}
 
-				List<PerforceClientRecord> Clients;
-				if(!Perforce.FindClients(PerforceInfo.UserName, out Clients, Log))
-				{
-					ErrorMessage = String.Format("Couldn't find any clients for this host.");
-					return false;
-				}
+					// Find all the clients for this user
+					Logger.LogInformation("Enumerating clients for {UserName}...", PerforceInfo.UserName);
 
-				List<PerforceConnection> CandidateClients = FilterClients(Clients, Perforce.ServerAndPort, PerforceInfo.HostName, PerforceInfo.UserName);
-				if(CandidateClients.Count == 0)
-				{
-					// Search through all workspaces. We may find a suitable workspace which is for any user.
-					Log.WriteLine("Enumerating shared clients...");
-					if(!Perforce.FindClients("", out Clients, Log))
+					List<ClientsRecord> Clients = await DefaultConnection.GetClientsAsync(ClientsOptions.None, DefaultConnection.Settings.UserName, CancellationToken);
+
+					List<IPerforceSettings> CandidateClients = await FilterClients(Clients, NewSelectedFileName, DefaultConnection.Settings, PerforceInfo.ClientHost, Logger, CancellationToken);
+					if (CandidateClients.Count == 0)
 					{
-						ErrorMessage = "Failed to enumerate clients.";
-						return false;
-					}
+						// Search through all workspaces. We may find a suitable workspace which is for any user.
+						Logger.LogInformation("Enumerating shared clients...");
+						Clients = await DefaultConnection.GetClientsAsync(ClientsOptions.None, "", CancellationToken);
 
-					// Filter this list of clients
-					CandidateClients = FilterClients(Clients, Perforce.ServerAndPort, PerforceInfo.HostName, PerforceInfo.UserName);
+						// Filter this list of clients
+						CandidateClients = await FilterClients(Clients, NewSelectedFileName, DefaultConnection.Settings, PerforceInfo.ClientHost, Logger, CancellationToken);
 
-					// If we still couldn't find any, fail.
-					if(CandidateClients.Count == 0)
+						// If we still couldn't find any, fail.
+						if (CandidateClients.Count == 0)
+						{
+							throw new UserErrorException($"Couldn't find any Perforce workspace containing {NewSelectedFileName}. Check your connection settings.");
+						}
+					}
+
+					// Check there's only one client
+					if (CandidateClients.Count > 1)
 					{
-						ErrorMessage = String.Format("Couldn't find any Perforce workspace containing {0}. Check your connection settings.", NewSelectedFileName);
-						return false;
+						throw new UserErrorException(String.Format("Found multiple workspaces containing {0}:\n\n{1}\n\nCannot determine which to use.", Path.GetFileName(NewSelectedFileName.GetFileName()), String.Join("\n", CandidateClients.Select(x => x.ClientName))));
 					}
-				}
 
-				// Check there's only one client
-				if(CandidateClients.Count > 1)
-				{
-					ErrorMessage = String.Format("Found multiple workspaces containing {0}:\n\n{1}\n\nCannot determine which to use.", Path.GetFileName(NewSelectedFileName.GetFileName()), String.Join("\n", CandidateClients.Select(x => x.ClientName)));
-					return false;
-				}
+					// Take the client we've chosen
+					PerforceClient = await PerforceConnection.CreateAsync(CandidateClients[0], Logger);
 
-				// Take the client we've chosen
-				PerforceClient = CandidateClients[0];
+					// Get the client path for the project file
+					List<WhereRecord> Records = await PerforceClient.WhereAsync(NewSelectedFileName.FullName, CancellationToken).Where(x => !x.Unmap).ToListAsync(CancellationToken);
+					if (Records.Count == 0)
+					{
+						throw new UserErrorException("File is not mapped to any client");
+					}
+					else if (Records.Count > 1)
+					{
+						throw new UserErrorException($"File is mapped to {Records.Count} locations: {String.Join(", ", Records.Select(x => x.Path))}");
+					}
 
-				// Get the client path for the project file
-				if(!PerforceClient.ConvertToClientPath(NewSelectedFileName.FullName, out NewSelectedClientFileName, Log))
-				{
-					ErrorMessage = String.Format("Couldn't get client path for {0}", NewSelectedFileName);
-					return false;
+					NewSelectedClientFileName = Records[0].ClientFile;
 				}
-			}
 
-			// Make sure the path case is correct. This can cause UBT intermediates to be out of date if the case mismatches.
-			NewSelectedFileName = FileReference.FindCorrectCase(NewSelectedFileName);
+				// Make sure the path case is correct. This can cause UBT intermediates to be out of date if the case mismatches.
+				NewSelectedFileName = FileReference.FindCorrectCase(NewSelectedFileName);
 
-			// Update the selected project with all the data we've found
-			SelectedProject = new UserSelectedProjectSettings(SelectedProject.ServerAndPort, SelectedProject.UserName, SelectedProject.Type, NewSelectedClientFileName, NewSelectedFileName.FullName);
+				// Update the selected project with all the data we've found
+				SelectedProject = new UserSelectedProjectSettings(SelectedProject.ServerAndPort, SelectedProject.UserName, SelectedProject.Type, NewSelectedClientFileName, NewSelectedFileName.FullName);
 
-			// Figure out where the engine is in relation to it
-			int EndIdx = NewSelectedClientFileName.Length - 1;
-			if(EndIdx != -1 && NewSelectedClientFileName.EndsWith(".uproject", StringComparison.InvariantCultureIgnoreCase))
-			{
-				EndIdx = NewSelectedClientFileName.LastIndexOf('/') - 1;
-			}
-			for(;;EndIdx--)
-			{
-				if(EndIdx < 2)
+				// Get the local branch root
+				string? BranchClientPath = null;
+				DirectoryReference? BranchDirectoryName = null;
+
+				// Figure out where the engine is in relation to it
+				int EndIdx = NewSelectedClientFileName.Length - 1;
+				if (EndIdx != -1 && NewSelectedClientFileName.EndsWith(".uproject", StringComparison.InvariantCultureIgnoreCase))
 				{
-					ErrorMessage = String.Format("Could not find engine in Perforce relative to project path ({0})", NewSelectedClientFileName);
-					return false;
+					EndIdx = NewSelectedClientFileName.LastIndexOf('/') - 1;
 				}
-				if(NewSelectedClientFileName[EndIdx] == '/')
+				for (; EndIdx >= 2; EndIdx--)
 				{
-					List<PerforceFileRecord> FileRecords;
-					if(PerforceClient.Stat(NewSelectedClientFileName.Substring(0, EndIdx) + "/Engine/Build/Build.version", out FileRecords, Log) && FileRecords.Count > 0)
+					if (NewSelectedClientFileName[EndIdx] == '/')
 					{
-						if(FileRecords[0].ClientPath == null)
+						List<PerforceResponse<FStatRecord>> FileRecords = await PerforceClient.TryFStatAsync(FStatOptions.None, NewSelectedClientFileName.Substring(0, EndIdx) + "/Engine/Build/Build.version", CancellationToken).ToListAsync();
+						if (FileRecords.Succeeded())
 						{
-							ErrorMessage = String.Format("Missing client path for {0}", FileRecords[0].DepotPath);
-							return false;
+							FStatRecord FileRecord = FileRecords[0].Data;
+							if (FileRecord.ClientFile == null)
+							{
+								throw new UserErrorException($"Missing client path for {FileRecord.DepotFile}");
+							}
+
+							BranchClientPath = NewSelectedClientFileName.Substring(0, EndIdx);
+							BranchDirectoryName = new FileReference(FileRecord.ClientFile).Directory.ParentDirectory?.ParentDirectory;
+							break;
 						}
-
-						BranchClientPath = NewSelectedClientFileName.Substring(0, EndIdx);
-						BranchDirectoryName = new DirectoryReference(Path.Combine(Path.GetDirectoryName(FileRecords[0].ClientPath), "..", ".."));
-						break;
 					}
 				}
-			}
-			Log.WriteLine("Found branch root at {0}", BranchClientPath);
-
-			// Find the editor target for this project
-			if(NewSelectedFileName.HasExtension(".uproject"))
-			{
-				List<PerforceFileRecord> Files;
-				if(PerforceClient.FindFiles(PerforceUtils.GetClientOrDepotDirectoryName(NewSelectedClientFileName) + "/Source/*Editor.Target.cs", out Files, Log) && Files.Count >= 1)
+				if(BranchClientPath == null || BranchDirectoryName == null)
 				{
-					string DepotPath = Files[0].DepotPath;
-					NewProjectEditorTarget = Path.GetFileNameWithoutExtension(Path.GetFileNameWithoutExtension(DepotPath.Substring(DepotPath.LastIndexOf('/') + 1)));
-					Log.WriteLine("Using {0} as editor target name (from {1})", NewProjectEditorTarget, Files[0]);
+					throw new UserErrorException($"Could not find engine in Perforce relative to project path ({NewSelectedClientFileName})");
 				}
-				if (NewProjectEditorTarget == null)
-				{
-					Log.WriteLine("Couldn't find any editor targets for this project.");
-				}
-			}
 
-			// Get a unique name for the project that's selected. For regular branches, this can be the depot path. For streams, we want to include the stream name to encode imports.
-			if(PerforceClient.GetActiveStream(out StreamName, Log))
-			{
-				string ExpectedPrefix = String.Format("//{0}/", PerforceClient.ClientName);
-				if(!NewSelectedClientFileName.StartsWith(ExpectedPrefix, StringComparison.InvariantCultureIgnoreCase))
+				Logger.LogInformation("Found branch root at {RootPath}", BranchClientPath);
+
+				// Find the editor target for this project
+				string? NewProjectEditorTarget = null;
+				if (NewSelectedFileName.HasExtension(".uproject"))
 				{
-					ErrorMessage = String.Format("Unexpected client path; expected '{0}' to begin with '{1}'", NewSelectedClientFileName, ExpectedPrefix);
-					return false;
+					List<FStatRecord> Records = await PerforceClient.FStatAsync(PerforceUtils.GetClientOrDepotDirectoryName(NewSelectedClientFileName) + "/Source/*Editor.Target.cs", CancellationToken).ToListAsync(CancellationToken);
+
+					FStatRecord Record = Records.FirstOrDefault(x => x.Action != FileAction.Delete || x.Action != FileAction.MoveDelete);
+					if (Record != null)
+					{
+						string DepotPath = Record.DepotFile!;
+						NewProjectEditorTarget = Path.GetFileNameWithoutExtension(Path.GetFileNameWithoutExtension(DepotPath.Substring(DepotPath.LastIndexOf('/') + 1)));
+						Logger.LogInformation("Using {TargetName} as editor target name (from {DepotPath})", NewProjectEditorTarget, DepotPath);
+					}
+
+					if (NewProjectEditorTarget == null)
+					{
+						Logger.LogInformation("Couldn't find any editor targets for this project.");
+					}
 				}
-				string StreamPrefix;
-				if(!TryGetStreamPrefix(PerforceClient, StreamName, Log, out StreamPrefix))
+
+				string NewSelectedProjectIdentifier;
+				// Get a unique name for the project that's selected. For regular branches, this can be the depot path. For streams, we want to include the stream name to encode imports.
+				string? StreamName = await PerforceClient.GetCurrentStreamAsync(CancellationToken);
+				if (StreamName != null)
 				{
-					ErrorMessage = String.Format("Failed to get stream info for {0}", StreamName);
-					return false;
+					string ExpectedPrefix = String.Format("//{0}/", PerforceClient.Settings.ClientName);
+					if (!NewSelectedClientFileName.StartsWith(ExpectedPrefix, StringComparison.InvariantCultureIgnoreCase))
+					{
+						throw new UserErrorException($"Unexpected client path; expected '{NewSelectedClientFileName}' to begin with '{ExpectedPrefix}'");
+					}
+					string? StreamPrefix = await TryGetStreamPrefixAsync(PerforceClient, StreamName, CancellationToken);
+					if (StreamPrefix == null)
+					{
+						throw new UserErrorException("Unable to get stream prefix");
+					}
+					NewSelectedProjectIdentifier = String.Format("{0}/{1}", StreamPrefix, NewSelectedClientFileName.Substring(ExpectedPrefix.Length));
 				}
-				NewSelectedProjectIdentifier = String.Format("{0}/{1}", StreamPrefix, NewSelectedClientFileName.Substring(ExpectedPrefix.Length));
-			}
-			else
-			{
-				if(!PerforceClient.ConvertToDepotPath(NewSelectedClientFileName, out NewSelectedProjectIdentifier, Log))
+				else
 				{
-					ErrorMessage = String.Format("Couldn't get depot path for {0}", NewSelectedFileName);
-					return false;
-				}
+					List<PerforceResponse<WhereRecord>> Records = await PerforceClient.TryWhereAsync(NewSelectedClientFileName, CancellationToken).Where(x => !x.Succeeded || !x.Data.Unmap).ToListAsync(CancellationToken);
+					if (!Records.Succeeded() || Records.Count != 1)
+					{
+						throw new UserErrorException($"Couldn't get depot path for {NewSelectedFileName}");
+					}
 
-				Match Match = Regex.Match(NewSelectedProjectIdentifier, "//([^/]+)/");
-				if (Match.Success)
-				{
-					PerforceSpec Spec;
-					if (PerforceClient.TryGetDepotSpec(Match.Groups[1].Value, out Spec, Log) && Spec.GetField("Type") == "stream")
+					NewSelectedProjectIdentifier = Records[0].Data.DepotFile;
+
+					Match Match = Regex.Match(NewSelectedProjectIdentifier, "//([^/]+)/");
+					if (Match.Success)
 					{
-						ErrorMessage = String.Format("Cannot use a legacy client ({0}) with a stream depot ({1}).", PerforceClient.ClientName, Match.Groups[1].Value);
-						return false;
+						DepotRecord Depot = await PerforceClient.GetDepotAsync(Match.Groups[1].Value, CancellationToken);
+						if (Depot.Type == "stream")
+						{
+							throw new UserErrorException($"Cannot use a legacy client ({PerforceClient.Settings.ClientName}) with a stream depot ({Depot.Depot}).");
+						}
 					}
 				}
-			}
 
-			// Read the project logo
-			if (NewSelectedFileName.HasExtension(".uproject"))
-			{
-				FileReference LogoFileName = FileReference.Combine(NewSelectedFileName.Directory, "Build", "UnrealGameSync.png");
-				if(FileReference.Exists(LogoFileName))
+				// Figure out if it's an enterprise project. Use the synced version if we have it.
+				bool bIsEnterpriseProject = false;
+				if (NewSelectedClientFileName.EndsWith(".uproject", StringComparison.InvariantCultureIgnoreCase))
 				{
-					try
+					string Text;
+					if (FileReference.Exists(NewSelectedFileName))
 					{
-						// Duplicate the image, otherwise we'll leave the file locked
-						using(Image Image = Image.FromFile(LogoFileName.FullName))
-						{
-							ProjectLogo = new Bitmap(Image);
-						}
+						Text = FileReference.ReadAllText(NewSelectedFileName);
 					}
-					catch
+					else
 					{
-						ProjectLogo = null;
+						PerforceResponse<PrintRecord<string[]>> ProjectLines = await PerforceClient.TryPrintLinesAsync(NewSelectedClientFileName, CancellationToken);
+						if (!ProjectLines.Succeeded)
+						{
+							throw new UserErrorException($"Unable to get contents of {NewSelectedClientFileName}");
+						}
+						Text = String.Join("\n", ProjectLines.Data.Contents!);
 					}
+					bIsEnterpriseProject = Utility.IsEnterpriseProjectFromText(Text);
 				}
-			}
 
-			// Figure out if it's an enterprise project. Use the synced version if we have it.
-			if(NewSelectedClientFileName.EndsWith(".uproject", StringComparison.InvariantCultureIgnoreCase))
-			{
-				string Text;
-				if(FileReference.Exists(NewSelectedFileName))
+				// Make sure the drive containing the project exists, to prevent other errors down the line
+				string PathRoot = Path.GetPathRoot(NewSelectedFileName.FullName)!;
+				if (!Directory.Exists(PathRoot))
 				{
-					Text = FileReference.ReadAllText(NewSelectedFileName);
+					throw new UserErrorException($"Path '{NewSelectedFileName}' is invalid");
 				}
-				else
+
+				// Read the initial config file
+				List<KeyValuePair<FileReference, DateTime>> LocalConfigFiles = new List<KeyValuePair<FileReference, DateTime>>();
+				ConfigFile LatestProjectConfigFile = await PerforceMonitor.ReadProjectConfigFileAsync(PerforceClient, BranchClientPath, NewSelectedClientFileName, GetCacheFolder(BranchDirectoryName), LocalConfigFiles, Logger, CancellationToken);
+
+				// Get the local config file and stream filter
+				ConfigFile WorkspaceProjectConfigFile = await Workspace.ReadProjectConfigFile(BranchDirectoryName, NewSelectedFileName, Logger);
+				IReadOnlyList<string>? WorkspaceProjectStreamFilter = await Workspace.ReadProjectStreamFilter(PerforceClient, WorkspaceProjectConfigFile, Logger, CancellationToken.None);
+
+				ProjectInfo ProjectInfo = new ProjectInfo(BranchDirectoryName, NewSelectedFileName, BranchClientPath, NewSelectedClientFileName, NewSelectedProjectIdentifier, bIsEnterpriseProject);
+
+				WorkspaceSettings WorkspaceSettings = new WorkspaceSettings(PerforceSettings, SelectedProject, ProjectInfo, NewSelectedProjectIdentifier, NewProjectEditorTarget!, StreamName!, LatestProjectConfigFile, WorkspaceProjectConfigFile, WorkspaceProjectStreamFilter, LocalConfigFiles);
+				DirectoryReference.CreateDirectory(WorkspaceSettings.DataFolder);
+				DirectoryReference.CreateDirectory(WorkspaceSettings.CacheFolder);
+
+				// Run any event hooks
+				if (DeploymentSettings.OnDetectProjectSettings != null)
 				{
-					List<string> ProjectLines;
-					if(!PerforceClient.Print(NewSelectedClientFileName, out ProjectLines, Log))
+					string Message;
+					if (!DeploymentSettings.OnDetectProjectSettings(WorkspaceSettings, Logger, out Message))
 					{
-						ErrorMessage = String.Format("Unable to get contents of {0}", NewSelectedClientFileName);
-						return false;
+						throw new UserErrorException(Message);
 					}
-					Text = String.Join("\n", ProjectLines);
 				}
-				bIsEnterpriseProject = Utility.IsEnterpriseProjectFromText(Text);
-			}
 
-			// Make sure the drive containing the project exists, to prevent other errors down the line
-			string PathRoot = Path.GetPathRoot(NewSelectedFileName.FullName);
-			if (!Directory.Exists(PathRoot))
-			{
-				ErrorMessage = String.Format("Path '{0}' is invalid", NewSelectedFileName);
-				return false;
+				return WorkspaceSettings;
 			}
-
-			// Read the initial config file
-			LocalConfigFiles = new List<KeyValuePair<FileReference, DateTime>>();
-			LatestProjectConfigFile = PerforceMonitor.ReadProjectConfigFile(PerforceClient, BranchClientPath, NewSelectedClientFileName, CacheFolder, LocalConfigFiles, Log);
-
-			// Run any event hooks
-			if (DeploymentSettings.OnDetectProjectSettings != null)
+			finally
 			{
-				string Message;
-				if (!DeploymentSettings.OnDetectProjectSettings(this, Log, out Message))
-				{
-					ErrorMessage = Message;
-					return false;
-				}
+				PerforceClient?.Dispose();
 			}
-
-			// Succeed!
-			ErrorMessage = null;
-			return true;
 		}
 
-		List<PerforceConnection> FilterClients(List<PerforceClientRecord> Clients, string ServerAndPort, string HostName, string UserName)
+		static async Task<List<IPerforceSettings>> FilterClients(List<ClientsRecord> Clients, FileReference NewSelectedFileName, IPerforceSettings DefaultPerforceSettings, string? HostName, ILogger Logger, CancellationToken CancellationToken)
 		{
-			List<PerforceConnection> CandidateClients = new List<PerforceConnection>();
-			foreach(PerforceClientRecord Client in Clients)
+			List<IPerforceSettings> CandidateClients = new List<IPerforceSettings>();
+			foreach(ClientsRecord Client in Clients)
 			{
 				// Make sure the client is well formed
 				if(!String.IsNullOrEmpty(Client.Name) && (!String.IsNullOrEmpty(Client.Host) || !String.IsNullOrEmpty(Client.Owner)) && !String.IsNullOrEmpty(Client.Root))
 				{
 					// Require either a username or host name match
-					if((String.IsNullOrEmpty(Client.Host) || String.Compare(Client.Host, HostName, StringComparison.InvariantCultureIgnoreCase) == 0) && (String.IsNullOrEmpty(Client.Owner) || String.Compare(Client.Owner, UserName, StringComparison.InvariantCultureIgnoreCase) == 0))
+					if((String.IsNullOrEmpty(Client.Host) || String.Compare(Client.Host, HostName, StringComparison.OrdinalIgnoreCase) == 0) && (String.IsNullOrEmpty(Client.Owner) || String.Compare(Client.Owner, DefaultPerforceSettings.UserName, StringComparison.OrdinalIgnoreCase) == 0))
 					{
 						if(!Utility.SafeIsFileUnderDirectory(NewSelectedFileName.FullName, Client.Root))
 						{
-							Log.WriteLine("Rejecting {0} due to root mismatch ({1})", Client.Name, Client.Root);
+							Logger.LogInformation("Rejecting {ClientName} due to root mismatch ({RootPath})", Client.Name, Client.Root);
 							continue;
 						}
 
-						PerforceConnection CandidateClient = new PerforceConnection(UserName, Client.Name, ServerAndPort);
+						PerforceSettings CandidateSettings = new PerforceSettings(DefaultPerforceSettings) { ClientName = Client.Name };
+						using IPerforceConnection CandidateClient = await PerforceConnection.CreateAsync(CandidateSettings, Logger);
 
-						bool bFileExists;
-						if(!CandidateClient.FileExists(NewSelectedFileName.FullName, out bFileExists, Log) || !bFileExists)
+						List<PerforceResponse<WhereRecord>> WhereRecords = await CandidateClient.TryWhereAsync(NewSelectedFileName.FullName, CancellationToken).Where(x => x.Failed || !x.Data.Unmap).ToListAsync(CancellationToken);
+						if(!WhereRecords.Succeeded() || WhereRecords.Count != 1)
 						{
-							Log.WriteLine("Rejecting {0} due to file not existing in workspace", Client.Name);
+							Logger.LogInformation("Rejecting {ClientName} due to file not existing in workspace", Client.Name);
 							continue;
 						}
 
-						List<PerforceFileRecord> Records;
-						if(!CandidateClient.Stat(NewSelectedFileName.FullName, out Records, Log))
+						List<PerforceResponse<FStatRecord>> Records = await CandidateClient.TryFStatAsync(FStatOptions.None, NewSelectedFileName.FullName, CancellationToken).ToListAsync(CancellationToken);
+						if (!Records.Succeeded())
 						{
-							Log.WriteLine("Rejecting {0} due to {1} not in depot", Client.Name, NewSelectedFileName);
+							Logger.LogInformation("Rejecting {ClientName} due to {FileName} not in depot", Client.Name, NewSelectedFileName);
 							continue;
 						}
 
-						Records.RemoveAll(x => !x.IsMapped);
+						Records.RemoveAll(x => !x.Data.IsMapped);
 						if(Records.Count == 0)
 						{
-							Log.WriteLine("Rejecting {0} due to {1} matching records", Client.Name, Records.Count);
+							Logger.LogInformation("Rejecting {ClientName} due to {NumRecords} matching records", Client.Name, Records.Count);
 							continue;
 						}
 
-						Log.WriteLine("Found valid client {0}", Client.Name);
-						CandidateClients.Add(CandidateClient);
+						Logger.LogInformation("Found valid client {ClientName}", Client.Name);
+						CandidateClients.Add(CandidateSettings);
 					}
 				}
 			}
 			return CandidateClients;
 		}
 
-		bool TryGetStreamPrefix(PerforceConnection Perforce, string StreamName, TextWriter Log, out string StreamPrefix)
+		static async Task<string?> TryGetStreamPrefixAsync(IPerforceConnection Perforce, string StreamName, CancellationToken CancellationToken)
 		{ 
-			string CurrentStreamName = StreamName;
-			for(;;)
+			string? CurrentStreamName = StreamName;
+			while(!String.IsNullOrEmpty(CurrentStreamName))
 			{
-				PerforceSpec StreamSpec;
-				if(!Perforce.TryGetStreamSpec(CurrentStreamName, out StreamSpec, Log))
+				PerforceResponse<StreamRecord> Response = await Perforce.TryGetStreamAsync(CurrentStreamName, false, CancellationToken);
+				if (!Response.Succeeded)
 				{
-					StreamPrefix = null;
-					return false;
+					return null;
 				}
-				if(StreamSpec.GetField("Type") != "virtual")
-				{
-					StreamPrefix = CurrentStreamName;
-					return true;
-				}
-				CurrentStreamName = StreamSpec.GetField("Parent");
-			}
-		}
-	}
-
-	class DetectProjectSettingsResult : IDisposable
-	{
-		public DetectProjectSettingsTask Task;
-		public bool bSucceeded;
-		public string ErrorMessage;
-
-		public DetectProjectSettingsResult(DetectProjectSettingsTask Task, bool bSucceeded, string ErrorMessage)
-		{
-			this.Task = Task;
-			this.bSucceeded = bSucceeded;
-			this.ErrorMessage = ErrorMessage;
-		}
-
-		public void Dispose()
-		{
-			if(Task != null)
-			{
-				Task.Dispose();
-				Task = null;
-			}
-		}
-	}
-
-	class DetectMultipleProjectSettingsTask : IModalTask, IDisposable
-	{
-		PerforceConnection DefaultConnection;
-		public DetectProjectSettingsTask[] Tasks;
-		public DetectProjectSettingsResult[] Results;
 
-		public DetectMultipleProjectSettingsTask(PerforceConnection DefaultConnection, IEnumerable<DetectProjectSettingsTask> Tasks)
-		{
-			this.DefaultConnection = DefaultConnection;
-			this.Tasks = Tasks.ToArray();
-		}
-
-		public void Dispose()
-		{
-			if(Results != null)
-			{
-				foreach(DetectProjectSettingsTask Task in Tasks.Where(x => x != null))
+				StreamRecord StreamSpec = Response.Data;
+				if(StreamSpec.Type != "virtual")
 				{
-					Task.Dispose();
+					return CurrentStreamName;
 				}
-				foreach(DetectProjectSettingsResult Result in Results.Where(x => x != null))
-				{
-					Result.Dispose();
-				}
-				Results = null;
-			}
-		}
 
-		public bool Run(out string ErrorMessage)
-		{
-			Results = new DetectProjectSettingsResult[Tasks.Length];
-			Parallel.For(0, Tasks.Length, new ParallelOptions(){ MaxDegreeOfParallelism = 4 }, Idx => RunTask(Idx));
-
-			ErrorMessage = null;
-			return true;
-		}
-
-		void RunTask(int Idx)
-		{
-			DetectProjectSettingsTask Task = Tasks[Idx];
-
-			UserSelectedProjectSettings Project = Task.SelectedProject;
-			PerforceConnection Perforce = Utility.OverridePerforceSettings(DefaultConnection, Project.ServerAndPort, Project.UserName);
-
-			string TaskErrorMessage;
-			bool bTaskSucceeded = Tasks[Idx].Run(Perforce, out TaskErrorMessage);
-			Results[Idx] = new DetectProjectSettingsResult(Tasks[Idx], bTaskSucceeded, TaskErrorMessage);
-			Tasks[Idx] = null;
-		}
-	}
-
-	class QuietDetectProjectSettingsTask : IModalTask
-	{
-		PerforceConnection Perforce;
-		DetectProjectSettingsTask Inner;
-		TextWriter Log;
-
-		public QuietDetectProjectSettingsTask(PerforceConnection Perforce, DetectProjectSettingsTask Inner, TextWriter Log)
-		{
-			this.Perforce = Perforce;
-			this.Inner = Inner;
-			this.Log = Log;
-		}
-
-		public bool Run(out string ErrorMessage)
-		{
-			return Inner.Run(Perforce, out ErrorMessage);
+				CurrentStreamName = StreamSpec.Parent;
+			}
+			return null;
 		}
 	}
 }
diff --git a/Engine/Source/Programs/UnrealGameSync/UnrealGameSync/EventMonitor.cs b/Engine/Source/Programs/UnrealGameSync/UnrealGameSync/EventMonitor.cs
--- a/Engine/Source/Programs/UnrealGameSync/UnrealGameSync/EventMonitor.cs
+++ b/Engine/Source/Programs/UnrealGameSync/UnrealGameSync/EventMonitor.cs
@@ -10,6 +10,9 @@ using System.Linq;
 using System.Text.Json;
 using System.Threading;
 using EpicGames.Core;
+using System.Threading.Tasks;
+using Microsoft.Extensions.Logging;
+using Microsoft.Extensions.DependencyInjection;
 
 namespace UnrealGameSync
 {
@@ -212,8 +215,10 @@ namespace UnrealGameSync
 		int ApiVersion;
 		string Project;
 		string CurrentUserName;
-		Thread WorkerThread;
-		AutoResetEvent RefreshEvent = new AutoResetEvent(false);
+		SynchronizationContext SynchronizationContext;
+		CancellationTokenSource CancellationSource;
+		Task WorkerTask;
+		AsyncEvent RefreshEvent = new AsyncEvent();
 		ConcurrentQueue<EventData> OutgoingEvents = new ConcurrentQueue<EventData>();
 		ConcurrentQueue<EventData> IncomingEvents = new ConcurrentQueue<EventData>();
 		ConcurrentQueue<CommentData> OutgoingComments = new ConcurrentQueue<CommentData>();
@@ -222,8 +227,8 @@ namespace UnrealGameSync
 		SortedDictionary<int, EventSummary> ChangeNumberToSummary = new SortedDictionary<int, EventSummary>();
 		Dictionary<string, EventData> UserNameToLastSyncEvent = new Dictionary<string, EventData>(StringComparer.InvariantCultureIgnoreCase);
 		Dictionary<string, BadgeData> BadgeNameToLatestData = new Dictionary<string, BadgeData>();
-		BoundedLogWriter LogWriter;
-		bool bDisposing;
+		ILogger Logger;
+		IAsyncDisposer AsyncDisposer;
 		LatestData LatestIds;
 		HashSet<int> FilterChangeNumbers = new HashSet<int>();
 		List<EventData> InvestigationEvents = new List<EventData>();
@@ -239,11 +244,16 @@ namespace UnrealGameSync
 
 		public event Action OnUpdatesReady;
 
-		public EventMonitor(string InApiUrl, string InProject, string InCurrentUserName, FileReference InLogFileName)
+		public EventMonitor(string InApiUrl, string InProject, string InCurrentUserName, IServiceProvider ServiceProvider)
 		{
 			ApiUrl = InApiUrl;
 			Project = InProject;
 			CurrentUserName = InCurrentUserName;
+			Logger = ServiceProvider.GetRequiredService<ILogger<EventMonitor>>();
+			AsyncDisposer = ServiceProvider.GetRequiredService<IAsyncDisposer>();
+			SynchronizationContext = SynchronizationContext.Current;
+			CancellationSource = new CancellationTokenSource();
+
 			LatestIds = new LatestData { LastBuildId = 0, LastCommentId = 0, LastEventId = 0 };
 
 			MetadataStream = Project.ToLowerInvariant().TrimEnd('/');
@@ -261,40 +271,33 @@ namespace UnrealGameSync
 				}
 			}
 
-			LogWriter = new BoundedLogWriter(InLogFileName);
 			if(ApiUrl == null)
 			{
 				LastStatusMessage = "Database functionality disabled due to empty ApiUrl.";
 			}
 			else
 			{
-				LogWriter.WriteLine("Using connection string: {0}", ApiUrl);
+				Logger.LogInformation("Using connection string: {ApiUrl}", ApiUrl);
 			}
 		}
 
 		public void Start()
 		{
-			WorkerThread = new Thread(() => PollForUpdates());
-			WorkerThread.Start();
+			if (WorkerTask == null)
+			{
+				WorkerTask = Task.Run(() => PollForUpdatesAsync(CancellationSource.Token));
+			}
 		}
 
 		public void Dispose()
 		{
-			bDisposing = true;
-			if(WorkerThread != null)
-			{
-				RefreshEvent.Set();
-				if(!WorkerThread.Join(100))
-				{
-					WorkerThread.Interrupt();
-					WorkerThread.Join();
-				}
-				WorkerThread = null;
-			}
-			if(LogWriter != null)
+			OnUpdatesReady = null;
+
+			if(WorkerTask != null)
 			{
-				LogWriter.Dispose();
-				LogWriter = null;
+				CancellationSource.Cancel();
+				AsyncDisposer.Add(WorkerTask.ContinueWith(_ => CancellationSource.Dispose()));
+				WorkerTask = null;
 			}
 		}
 
@@ -699,30 +702,31 @@ namespace UnrealGameSync
 			}
 		}
 
-		void PollForUpdates()
+		async Task PollForUpdatesAsync(CancellationToken CancellationToken)
 		{
-			Stopwatch NetCoreTimer = Stopwatch.StartNew();
 			EventData Event = null;
 			CommentData Comment = null;
 			bool bUpdateThrottledRequests = true;
 			double RequestThrottle = 90; // seconds to wait for throttled request;
 			Stopwatch Timer = Stopwatch.StartNew();
-			while (!bDisposing)
+			while (!CancellationToken.IsCancellationRequested)
 			{
+				Task RefreshTask = RefreshEvent.Task;
+
 				// If there's no connection string, just empty out the queue
 				if (ApiUrl != null)
 				{
 					// Post all the reviews to the database. We don't send them out of order, so keep the review outside the queue until the next update if it fails
 					while(Event != null || OutgoingEvents.TryDequeue(out Event))
 					{
-						SendEventToBackend(Event);
+						await SendEventToBackendAsync(Event, CancellationToken);
 						Event = null;
 					}
 
 					// Post all the comments to the database.
 					while(Comment != null || OutgoingComments.TryDequeue(out Comment))
 					{
-						SendCommentToBackend(Comment);
+						await SendCommentToBackendAsync(Comment, CancellationToken);
 						Comment = null;
 					}
 
@@ -733,30 +737,21 @@ namespace UnrealGameSync
 					}
 
 					// Read all the new reviews, pass whether or not to fire the throttled requests
-					ReadEventsFromBackend(bUpdateThrottledRequests);
+					await ReadEventsFromBackendAsync(bUpdateThrottledRequests, CancellationToken);
 
 					// Send a notification that we're ready to update
-					if((IncomingMetadata.Count > 0 || IncomingEvents.Count > 0 || IncomingBadges.Count > 0 || IncomingComments.Count > 0) && OnUpdatesReady != null)
+					if(IncomingMetadata.Count > 0 || IncomingEvents.Count > 0 || IncomingBadges.Count > 0 || IncomingComments.Count > 0)
 					{
-						OnUpdatesReady();
+						SynchronizationContext.Post(_ => OnUpdatesReady?.Invoke(), null);
 					}
 				}
 
-				// Post info about whether net core is installed
-				if (DeploymentSettings.NetCoreTelemetryUrl != null && NetCoreTimer.Elapsed > TimeSpan.FromMinutes(60.0))
+				// Wait for something else to do
+				Task DelayTask = Task.Delay(TimeSpan.FromSeconds(30.0), CancellationToken);
+				if (await Task.WhenAny(DelayTask, RefreshTask) == DelayTask)
 				{
-					try
-					{
-						RESTApi.POST(DeploymentSettings.NetCoreTelemetryUrl, "netcore", "{ }", "User=" + CurrentUserName, "Machine=" + Environment.MachineName, "NetCore=" + HasNetCore3().ToString());
-					}
-					catch
-					{
-					}
-					NetCoreTimer.Restart();
+					bUpdateThrottledRequests = true;
 				}
-
-				// Wait for something else to do
-				bUpdateThrottledRequests = RefreshEvent.WaitOne(30 * 1000);
 			}
 		}
 
@@ -776,54 +771,54 @@ namespace UnrealGameSync
 			return false;
 		}
 
-		bool SendEventToBackend(EventData Event)
+		async Task<bool> SendEventToBackendAsync(EventData Event, CancellationToken CancellationToken)
 		{
 			try
 			{
 				Stopwatch Timer = Stopwatch.StartNew();
-				LogWriter.WriteLine("Posting event... ({0}, {1}, {2})", Event.Change, Event.UserName, Event.Type);
+				Logger.LogInformation("Posting event... ({Change}, {UserName}, {Type})", Event.Change, Event.UserName, Event.Type);
 				if (ApiVersion == 2)
 				{
-					SendMetadataUpdateUpdateV2(Event.Change, Event.Project, Event.UserName, Event.Type, null);
+					await SendMetadataUpdateUpdateV2Async(Event.Change, Event.Project, Event.UserName, Event.Type, null, CancellationToken);
 				}
 				else
 				{
-					RESTApi.POST(ApiUrl, "event", JsonSerializer.Serialize(Event));
+					await RESTApi.PostAsync($"{ApiUrl}/api/event", JsonSerializer.Serialize(Event), CancellationToken);
 				}
 				return true;
 			}
 			catch(Exception Ex)
 			{
-				LogWriter.WriteException(Ex, "Failed with exception.");
+				Logger.LogError(Ex, "Failed with exception.");
 				return false;
 			}
 		}
 
-		bool SendCommentToBackend(CommentData Comment)
+		async Task<bool> SendCommentToBackendAsync(CommentData Comment, CancellationToken CancellationToken)
 		{
 			try
 			{
 				Stopwatch Timer = Stopwatch.StartNew();
-				LogWriter.WriteLine("Posting comment... ({0}, {1}, {2}, {3})", Comment.ChangeNumber, Comment.UserName, Comment.Text, Comment.Project);
+				Logger.LogInformation("Posting comment... ({Change}, {User}, {Text}, {Project})", Comment.ChangeNumber, Comment.UserName, Comment.Text, Comment.Project);
 				if (ApiVersion == 2)
 				{
-					SendMetadataUpdateUpdateV2(Comment.ChangeNumber, Comment.Project, Comment.UserName, null, Comment.Text);
+					await SendMetadataUpdateUpdateV2Async(Comment.ChangeNumber, Comment.Project, Comment.UserName, null, Comment.Text, CancellationToken);
 				}
 				else
 				{
-					RESTApi.POST(ApiUrl, "comment", JsonSerializer.Serialize(Comment));
+					await RESTApi.PostAsync($"{ApiUrl}/api/comment", JsonSerializer.Serialize(Comment), CancellationToken);
 				}
-				LogWriter.WriteLine("Done in {0}ms.", Timer.ElapsedMilliseconds);
+				Logger.LogInformation("Done in {Time}ms.", Timer.ElapsedMilliseconds);
 				return true;
 			}
 			catch(Exception Ex)
 			{
-				LogWriter.WriteException(Ex, "Failed with exception.");
+				Logger.LogError(Ex, "Failed with exception.");
 				return false;
 			}
 		}
 
-		void SendMetadataUpdateUpdateV2(int Change, string Project, string UserName, EventType? Event, string Comment)
+		async Task SendMetadataUpdateUpdateV2Async(int Change, string Project, string UserName, EventType? Event, string Comment, CancellationToken CancellationToken)
 		{
 			UpdateMetadataRequestV2 Update = new UpdateMetadataRequestV2();
 			Update.Stream = MetadataStream;
@@ -869,22 +864,22 @@ namespace UnrealGameSync
 			}
 			Update.Comment = Comment;
 
-			RESTApi.POST(ApiUrl, "metadata", JsonSerializer.Serialize(Update));
+			await RESTApi.PostAsync($"{ApiUrl}/api/metadata", JsonSerializer.Serialize(Update), CancellationToken);
 		}
 
-		bool ReadEventsFromBackend(bool bFireThrottledRequests)
+		async Task<bool> ReadEventsFromBackendAsync(bool bFireThrottledRequests, CancellationToken CancellationToken)
 		{
 			try
 			{
 				Stopwatch Timer = Stopwatch.StartNew();
-				LogWriter.WriteLine();
-				LogWriter.WriteLine("Polling for reviews at {0}...", DateTime.Now.ToString());
+				Logger.LogInformation("Polling for events...");
+
 				//////////////
 				/// Initial Ids 
 				//////////////
 				if (ApiVersion == 0)
 				{
-					LatestData InitialIds = RESTApi.GET<LatestData>(ApiUrl, "latest", string.Format("project={0}", Project));
+					LatestData InitialIds = await RESTApi.GetAsync<LatestData>($"{ApiUrl}/api/latest?project={Project}", CancellationToken);
 					ApiVersion = (InitialIds.Version == 0)? 1 : InitialIds.Version;
 					LatestIds.LastBuildId = InitialIds.LastBuildId;
 					LatestIds.LastCommentId = InitialIds.LastCommentId;
@@ -909,7 +904,7 @@ namespace UnrealGameSync
 						// Fetch any updates in the current range of changes
 						if (MinChange != 0)
 						{
-							GetMetadataListResponseV2 NewEventList = RESTApi.GET<GetMetadataListResponseV2>(ApiUrl, "metadata", String.Format("{0}&minchange={1}&sequence={2}", CommonArgs, MinChange, MetadataSequenceNumber));
+							GetMetadataListResponseV2 NewEventList = await RESTApi.GetAsync<GetMetadataListResponseV2>($"{ApiUrl}/api/metadata?{CommonArgs}&minchange={MinChange}&sequence={MetadataSequenceNumber}", CancellationToken);
 							foreach (GetMetadataResponseV2 NewEvent in NewEventList.Items)
 							{
 								IncomingMetadata.Enqueue(NewEvent);
@@ -920,7 +915,7 @@ namespace UnrealGameSync
 						// Fetch any new changes
 						if (NewMinChangeCopy < MinChange)
 						{
-							GetMetadataListResponseV2 NewEvents = RESTApi.GET<GetMetadataListResponseV2>(ApiUrl, "metadata", String.Format("{0}&minchange={1}&maxchange={2}", CommonArgs, NewMinChangeCopy, MinChange));
+							GetMetadataListResponseV2 NewEvents = await RESTApi.GetAsync<GetMetadataListResponseV2>($"{ApiUrl}/api/metadata?{CommonArgs}&minchange={NewMinChangeCopy}&maxchange={MinChange}", CancellationToken);
 							foreach (GetMetadataResponseV2 NewEvent in NewEvents.Items)
 							{
 								IncomingMetadata.Enqueue(NewEvent);
@@ -930,51 +925,51 @@ namespace UnrealGameSync
 					}
 				}
 				else
-				{
-				//////////////
-				/// Builds
-				//////////////
-				List<BadgeData> Builds = RESTApi.GET<List<BadgeData>>(ApiUrl, "build", string.Format("project={0}", Project), string.Format("lastbuildid={0}", LatestIds.LastBuildId));
-				foreach (BadgeData Build in Builds)
-				{
-					IncomingBadges.Enqueue(Build);
-					LatestIds.LastBuildId = Math.Max(LatestIds.LastBuildId, Build.Id);
-				}
-
-				//////////////////////////
-				/// Throttled Requests
-				//////////////////////////
-				if (bFireThrottledRequests)
 				{
 					//////////////
-					/// Reviews 
+					/// Builds
 					//////////////
-					List<EventData> Events = RESTApi.GET<List<EventData>>(ApiUrl, "event", string.Format("project={0}", Project), string.Format("lasteventid={0}", LatestIds.LastEventId));
-					foreach (EventData Review in Events)
+					List<BadgeData> Builds = await RESTApi.GetAsync<List<BadgeData>>($"{ApiUrl}/api/build?project={Project}&lastbuildid={LatestIds.LastBuildId}", CancellationToken);
+					foreach (BadgeData Build in Builds)
 					{
-						IncomingEvents.Enqueue(Review);
-						LatestIds.LastEventId = Math.Max(LatestIds.LastEventId, Review.Id);
+						IncomingBadges.Enqueue(Build);
+						LatestIds.LastBuildId = Math.Max(LatestIds.LastBuildId, Build.Id);
 					}
 
-					//////////////
-					/// Comments 
-					//////////////
-					List<CommentData> Comments = RESTApi.GET<List<CommentData>>(ApiUrl, "comment", string.Format("project={0}", Project), string.Format("lastcommentid={0}", LatestIds.LastCommentId));
-					foreach (CommentData Comment in Comments)
+					//////////////////////////
+					/// Throttled Requests
+					//////////////////////////
+					if (bFireThrottledRequests)
 					{
-						IncomingComments.Enqueue(Comment);
-						LatestIds.LastCommentId = Math.Max(LatestIds.LastCommentId, Comment.Id);
+						//////////////
+						/// Reviews 
+						//////////////
+						List<EventData> Events = await RESTApi.GetAsync<List<EventData>>($"{ApiUrl}/api/event?project={Project}&lasteventid={LatestIds.LastEventId}", CancellationToken);
+						foreach (EventData Review in Events)
+						{
+							IncomingEvents.Enqueue(Review);
+							LatestIds.LastEventId = Math.Max(LatestIds.LastEventId, Review.Id);
+						}
+
+						//////////////
+						/// Comments 
+						//////////////
+						List<CommentData> Comments = await RESTApi.GetAsync<List<CommentData>>($"{ApiUrl}/api/comment?project={Project}&lastcommentid={LatestIds.LastCommentId}", CancellationToken);
+						foreach (CommentData Comment in Comments)
+						{
+							IncomingComments.Enqueue(Comment);
+							LatestIds.LastCommentId = Math.Max(LatestIds.LastCommentId, Comment.Id);
+						}
 					}
 				}
-				}
 
 				LastStatusMessage = String.Format("Last update took {0}ms", Timer.ElapsedMilliseconds);
-				LogWriter.WriteLine("Done in {0}ms.", Timer.ElapsedMilliseconds);
+				Logger.LogInformation("Done in {Time}ms.", Timer.ElapsedMilliseconds);
 				return true;
 			}
 			catch(Exception Ex)
 			{
-				LogWriter.WriteException(Ex, "Failed with exception.");
+				Logger.LogError(Ex, "Failed with exception.");
 				LastStatusMessage = String.Format("Last update failed: ({0})", Ex.ToString());
 				return false;
 			}
diff --git a/Engine/Source/Programs/UnrealGameSync/UnrealGameSync/FindFoldersToCleanTask.cs b/Engine/Source/Programs/UnrealGameSync/UnrealGameSync/FindFoldersToCleanTask.cs
--- a/Engine/Source/Programs/UnrealGameSync/UnrealGameSync/FindFoldersToCleanTask.cs
+++ b/Engine/Source/Programs/UnrealGameSync/UnrealGameSync/FindFoldersToCleanTask.cs
@@ -1,5 +1,7 @@
 // Copyright Epic Games, Inc. All Rights Reserved.
 
+using EpicGames.Perforce;
+using Microsoft.Extensions.Logging;
 using System;
 using System.Collections.Generic;
 using System.Diagnostics;
@@ -9,6 +11,8 @@ using System.Text;
 using System.Threading;
 using System.Threading.Tasks;
 
+#nullable enable
+
 namespace UnrealGameSync
 {
 	class FolderToClean
@@ -37,33 +41,33 @@ namespace UnrealGameSync
 		}
 	}
 
-	class FindFoldersToCleanTask : IModalTask, IDisposable
+	class FindFoldersToCleanTask : IDisposable
 	{
 		class PerforceHaveFolder
 		{
-			public Dictionary<string, PerforceHaveFolder> NameToSubFolder = new Dictionary<string,PerforceHaveFolder>(StringComparer.InvariantCultureIgnoreCase);
-			public Dictionary<string, PerforceFileRecord> NameToFile = new Dictionary<string,PerforceFileRecord>(StringComparer.InvariantCultureIgnoreCase);
+			public Dictionary<string, PerforceHaveFolder> NameToSubFolder = new Dictionary<string,PerforceHaveFolder>(StringComparer.OrdinalIgnoreCase);
+			public Dictionary<string, FStatRecord> NameToFile = new Dictionary<string, FStatRecord>(StringComparer.OrdinalIgnoreCase);
 		}
 
-		PerforceConnection PerforceClient;
+		IPerforceSettings PerforceSettings;
 		string ClientRootPath;
 		IReadOnlyList<string> SyncPaths;
-		TextWriter Log;
+		ILogger Logger;
 		FolderToClean RootFolderToClean;
 
 		int RemainingFoldersToScan;
 		ManualResetEvent FinishedScan = new ManualResetEvent(false);
 		bool bAbortScan;
-		string ScanError;
+		string? ScanError;
 
 		public List<string> FileNames = new List<string>();
 
-		public FindFoldersToCleanTask(PerforceConnection InPerforceClient, FolderToClean InRootFolderToClean, string InClientRootPath, IReadOnlyList<string> InSyncPaths, TextWriter InLog)
+		public FindFoldersToCleanTask(IPerforceSettings InPerforceSettings, FolderToClean InRootFolderToClean, string InClientRootPath, IReadOnlyList<string> InSyncPaths, ILogger Logger)
 		{
-			PerforceClient = InPerforceClient;
+			PerforceSettings = InPerforceSettings;
 			ClientRootPath = InClientRootPath.TrimEnd('/') + "/";
 			SyncPaths = new List<string>(InSyncPaths);
-			Log = InLog;
+			this.Logger = Logger;
 			RootFolderToClean = InRootFolderToClean;
 			FinishedScan = new ManualResetEvent(true);
 		}
@@ -112,7 +116,7 @@ namespace UnrealGameSync
 			}
 		}
 
-		void MergeTrees(FolderToClean LocalFolder, PerforceHaveFolder PerforceFolder, HashSet<string> OpenClientPaths, string PerforceConfigFile)
+		void MergeTrees(FolderToClean LocalFolder, PerforceHaveFolder? PerforceFolder, HashSet<string> OpenClientPaths, string? PerforceConfigFile)
 		{
 			if(PerforceFolder == null)
 			{
@@ -130,7 +134,7 @@ namespace UnrealGameSync
 				// Loop through all the local sub-folders
 				foreach(FolderToClean LocalSubFolder in LocalFolder.NameToSubFolder.Values)
 				{
-					PerforceHaveFolder PerforceSubFolder;
+					PerforceHaveFolder? PerforceSubFolder;
 					PerforceFolder.NameToSubFolder.TryGetValue(LocalSubFolder.Name, out PerforceSubFolder);
 					MergeTrees(LocalSubFolder, PerforceSubFolder, OpenClientPaths, PerforceConfigFile);
 				}
@@ -138,7 +142,7 @@ namespace UnrealGameSync
 				// Also merge all the Perforce folders that no longer exist
 				foreach(KeyValuePair<string, PerforceHaveFolder> PerforceSubFolderPair in PerforceFolder.NameToSubFolder)
 				{
-					FolderToClean LocalSubFolder;
+					FolderToClean? LocalSubFolder;
 					if(!LocalFolder.NameToSubFolder.TryGetValue(PerforceSubFolderPair.Key, out LocalSubFolder))
 					{
 						LocalSubFolder = new FolderToClean(new DirectoryInfo(Path.Combine(LocalFolder.Directory.FullName, PerforceSubFolderPair.Key)));
@@ -148,14 +152,14 @@ namespace UnrealGameSync
 				}
 
 				// Find all the files that need to be re-synced
-				foreach(KeyValuePair<string, PerforceFileRecord> FilePair in PerforceFolder.NameToFile)
+				foreach(KeyValuePair<string, FStatRecord> FilePair in PerforceFolder.NameToFile)
 				{
-					FileInfo LocalFile;
+					FileInfo? LocalFile;
 					if(!LocalFolder.NameToFile.TryGetValue(FilePair.Key, out LocalFile))
 					{
 						LocalFolder.FilesToSync.Add(new FileInfo(Path.Combine(LocalFolder.Directory.FullName, FilePair.Key)));
 					}
-					else if((FilePair.Value.Flags & PerforceFileFlags.AlwaysWritable) == 0 && (LocalFile.Attributes & FileAttributes.ReadOnly) == 0 && !OpenClientPaths.Contains(FilePair.Value.ClientPath))
+					else if(IsFileTypeWritable(FilePair.Value?.HeadType ?? FilePair.Value?.Type) && (LocalFile.Attributes & FileAttributes.ReadOnly) == 0 && !OpenClientPaths.Contains(FilePair.Value?.ClientFile ?? ""))
 					{
 						LocalFolder.FilesToSync.Add(LocalFile);
 					}
@@ -184,6 +188,19 @@ namespace UnrealGameSync
 			LocalFolder.bEmptyAfterClean = LocalFolder.NameToSubFolder.Values.All(x => x.bEmptyAfterClean) && LocalFolder.FilesToDelete.Count == LocalFolder.NameToFile.Count && LocalFolder.FilesToSync.Count == 0;
 		}
 
+		static bool IsFileTypeWritable(string? Type)
+		{
+			if (Type != null)
+			{
+				int Idx = Type.IndexOf('+');
+				if (Idx != -1)
+				{
+					return Type.IndexOf('w', Idx) != -1;
+				}
+			}
+			return true;
+		}
+
 		void RemoveEmptyFolders(FolderToClean Folder)
 		{
 			foreach(FolderToClean SubFolder in Folder.NameToSubFolder.Values)
@@ -202,14 +219,16 @@ namespace UnrealGameSync
 			{
 				FinishedScan.WaitOne();
 				FinishedScan.Dispose();
-				FinishedScan = null;
+				FinishedScan = null!;
 			}
 		}
 
-		public bool Run(out string ErrorMessage)
+		public async Task RunAsync(CancellationToken CancellationToken)
 		{
-			Log.WriteLine("Finding files in workspace...");
-			Log.WriteLine();
+			using IPerforceConnection PerforceClient = await PerforceConnection.CreateAsync(PerforceSettings, Logger);
+
+			Logger.LogInformation("Finding files in workspace...");
+			Logger.LogInformation("");
 
 			// Clear the current error
 			ScanError = null;
@@ -225,7 +244,7 @@ namespace UnrealGameSync
 					FolderToClean SyncFolder = RootFolderToClean;
 					for(int Idx = 0; Idx < Fragments.Length - 1; Idx++)
 					{
-						FolderToClean NextSyncFolder;
+						FolderToClean? NextSyncFolder;
 						if(!SyncFolder.NameToSubFolder.TryGetValue(Fragments[Idx], out NextSyncFolder))
 						{
 							NextSyncFolder = new FolderToClean(new DirectoryInfo(Path.Combine(SyncFolder.Directory.FullName, Fragments[Idx])));
@@ -259,26 +278,20 @@ namespace UnrealGameSync
 			PerforceHaveFolder RootHaveFolder = new PerforceHaveFolder();
 			foreach(string SyncPath in SyncPaths)
 			{
-				List<PerforceFileRecord> FileRecords;
-				if(!PerforceClient.Stat(String.Format("{0}#have", SyncPath), out FileRecords, Log))
-				{
-					ErrorMessage = "Couldn't query have table from Perforce.";
-					return false;
-				}
-				foreach(PerforceFileRecord FileRecord in FileRecords)
+				List<FStatRecord> FileRecords = await PerforceClient.FStatAsync($"{SyncPath}#have", CancellationToken).ToListAsync(CancellationToken);
+				foreach(FStatRecord FileRecord in FileRecords)
 				{
-					if(!FileRecord.ClientPath.StartsWith(LocalRootPrefix, StringComparison.InvariantCultureIgnoreCase))
+					if(FileRecord.ClientFile == null || !FileRecord.ClientFile.StartsWith(LocalRootPrefix, StringComparison.InvariantCultureIgnoreCase))
 					{
-						ErrorMessage = String.Format("Failed to get have table; file '{0}' doesn't start with root path ('{1}')", FileRecord.ClientPath, RootFolderToClean.Directory.FullName);
-						return false;
+						throw new UserErrorException($"Failed to get have table; file '{FileRecord.ClientFile}' doesn't start with root path ('{RootFolderToClean.Directory.FullName}')");
 					}
 
-					string[] Tokens = FileRecord.ClientPath.Substring(LocalRootPrefix.Length).Split('/', '\\');
+					string[] Tokens = FileRecord.ClientFile.Substring(LocalRootPrefix.Length).Split('/', '\\');
 
 					PerforceHaveFolder FileFolder = RootHaveFolder;
 					for(int Idx = 0; Idx < Tokens.Length - 1; Idx++)
 					{
-						PerforceHaveFolder NextFileFolder;
+						PerforceHaveFolder? NextFileFolder;
 						if(!FileFolder.NameToSubFolder.TryGetValue(Tokens[Idx], out NextFileFolder))
 						{
 							NextFileFolder = new PerforceHaveFolder();
@@ -291,23 +304,17 @@ namespace UnrealGameSync
 			}
 
 			// Find all the files which are currently open for edit. We don't want to force sync these.
-			List<PerforceFileRecord> OpenFileRecords;
-			if(!PerforceClient.GetOpenFiles("//...", out OpenFileRecords, Log))
-			{
-				ErrorMessage = "Couldn't query open files from Perforce.";
-				return false;
-			}
+			List<OpenedRecord> OpenFileRecords = await PerforceClient.OpenedAsync(OpenedOptions.None, -1, null, null, -1, "//...", CancellationToken).ToListAsync(CancellationToken);
 
 			// Build a set of all the open local files
 			HashSet<string> OpenLocalFiles = new HashSet<string>(StringComparer.InvariantCultureIgnoreCase);
-			foreach(PerforceFileRecord OpenFileRecord in OpenFileRecords)
+			foreach (OpenedRecord OpenFileRecord in OpenFileRecords)
 			{
-				if(!OpenFileRecord.ClientPath.StartsWith(ClientRootPath, StringComparison.InvariantCultureIgnoreCase))
+				if(!OpenFileRecord.ClientFile.StartsWith(ClientRootPath, StringComparison.InvariantCultureIgnoreCase))
 				{
-					ErrorMessage = String.Format("Failed to get open file list; file '{0}' doesn't start with client root path ('{1}')", OpenFileRecord.ClientPath, ClientRootPath);
-					return false;
+					throw new UserErrorException($"Failed to get open file list; file '{OpenFileRecord.ClientFile}' doesn't start with client root path ('{ClientRootPath}')");
 				}
-				OpenLocalFiles.Add(LocalRootPrefix + PerforceUtils.UnescapePath(OpenFileRecord.ClientPath).Substring(ClientRootPath.Length).Replace('/', Path.DirectorySeparatorChar));
+				OpenLocalFiles.Add(LocalRootPrefix + PerforceUtils.UnescapePath(OpenFileRecord.ClientFile).Substring(ClientRootPath.Length).Replace('/', Path.DirectorySeparatorChar));
 			}
 
 			// Wait to finish scanning the directory
@@ -316,21 +323,17 @@ namespace UnrealGameSync
 			// Check if there was an error
 			if (ScanError != null)
 			{
-				ErrorMessage = ScanError;
-				return false;
+				throw new UserErrorException(ScanError);
 			}
 
 			// Find the value of the P4CONFIG variable
-			string PerforceConfigFile;
-			PerforceClient.GetSetting("P4CONFIG", out PerforceConfigFile, Log);
+			string? PerforceConfigFile = PerforceEnvironment.Default.GetValue("P4CONFIG");
 
 			// Merge the trees
 			MergeTrees(RootFolderToClean, RootHaveFolder, OpenLocalFiles, PerforceConfigFile);
 
 			// Remove all the empty folders
 			RemoveEmptyFolders(RootFolderToClean);
-			ErrorMessage = null;
-			return true;
 		}
 	}
 }
diff --git a/Engine/Source/Programs/UnrealGameSync/UnrealGameSync/Forms/ApplicationSettingsWindow.cs b/Engine/Source/Programs/UnrealGameSync/UnrealGameSync/Forms/ApplicationSettingsWindow.cs
--- a/Engine/Source/Programs/UnrealGameSync/UnrealGameSync/Forms/ApplicationSettingsWindow.cs
+++ b/Engine/Source/Programs/UnrealGameSync/UnrealGameSync/Forms/ApplicationSettingsWindow.cs
@@ -1,5 +1,7 @@
 // Copyright Epic Games, Inc. All Rights Reserved.
 
+using EpicGames.Perforce;
+using Microsoft.Extensions.Logging;
 using Microsoft.Win32;
 using System;
 using System.Collections.Generic;
@@ -15,42 +17,34 @@ using System.Threading.Tasks;
 using System.Windows.Forms;
 using UnrealGameSync;
 
+#nullable enable
+
 namespace UnrealGameSync
 {
 	partial class ApplicationSettingsWindow : Form
 	{
-		class PerforceTestConnectionTask : IPerforceModalTask
+		static class PerforceTestConnectionTask
 		{
-			string DepotPath;
-
-			public PerforceTestConnectionTask(string DepotPath)
-			{
-				this.DepotPath = DepotPath ?? DeploymentSettings.DefaultDepotPath;
-			}
-
-			public bool Run(PerforceConnection Perforce, TextWriter Log, out string ErrorMessage)
+			public static async Task RunAsync(IPerforceConnection Perforce, string DepotPath, CancellationToken CancellationToken)
 			{
 				string CheckFilePath = String.Format("{0}/Release/UnrealGameSync.exe", DepotPath);
 
-				List<PerforceFileRecord> FileRecords;
-				if(!Perforce.FindFiles(CheckFilePath, out FileRecords, Log) || FileRecords.Count == 0)
+				List<FStatRecord> FileRecords = await Perforce.FStatAsync(CheckFilePath, CancellationToken).ToListAsync(CancellationToken);
+				if(FileRecords.Count == 0)
 				{
-					ErrorMessage = String.Format("Unable to find {0}", CheckFilePath);
-					return false;
+					throw new UserErrorException($"Unable to find {CheckFilePath}");
 				}
-
-				ErrorMessage = null;
-				return true;
 			}
 		}
 
 		string OriginalExecutableFileName;
+		IPerforceSettings DefaultPerforceSettings;
 		UserSettings Settings;
-		TextWriter Log;
+		ILogger Logger;
 
-		string InitialServerAndPort;
-		string InitialUserName;
-		string InitialDepotPath;
+		string? InitialServerAndPort;
+		string? InitialUserName;
+		string? InitialDepotPath;
 		bool bInitialUnstable;
 		int InitialAutomationPortNumber;
 		ProtocolHandlerState InitialProtocolHandlerState;
@@ -74,14 +68,15 @@ namespace UnrealGameSync
 			}
 		}
 
-		private ApplicationSettingsWindow(string DefaultServerAndPort, string DefaultUserName, bool bUnstable, string OriginalExecutableFileName, UserSettings Settings, ToolUpdateMonitor ToolUpdateMonitor, TextWriter Log)
+		private ApplicationSettingsWindow(IPerforceSettings DefaultPerforceSettings, bool bUnstable, string OriginalExecutableFileName, UserSettings Settings, ToolUpdateMonitor ToolUpdateMonitor, ILogger<ApplicationSettingsWindow> Logger)
 		{
 			InitializeComponent();
 
 			this.OriginalExecutableFileName = OriginalExecutableFileName;
+			this.DefaultPerforceSettings = DefaultPerforceSettings;
 			this.Settings = Settings;
 			this.ToolUpdateMonitor = ToolUpdateMonitor;
-			this.Log = Log;
+			this.Logger = Logger;
 
 			Utility.ReadGlobalPerforceSettings(ref InitialServerAndPort, ref InitialUserName, ref InitialDepotPath);
 			bInitialUnstable = bUnstable;
@@ -94,11 +89,11 @@ namespace UnrealGameSync
 						
 			this.ServerTextBox.Text = InitialServerAndPort;
 			this.ServerTextBox.Select(ServerTextBox.TextLength, 0);
-			this.ServerTextBox.CueBanner = (DefaultServerAndPort == null)? "Default" : String.Format("Default ({0})", DefaultServerAndPort);
+			this.ServerTextBox.CueBanner = $"Default ({DefaultPerforceSettings.ServerAndPort})";
 
 			this.UserNameTextBox.Text = InitialUserName;
 			this.UserNameTextBox.Select(UserNameTextBox.TextLength, 0);
-			this.UserNameTextBox.CueBanner = (DefaultUserName == null)? "Default" : String.Format("Default ({0})", DefaultUserName);
+			this.UserNameTextBox.CueBanner = $"Default ({DefaultPerforceSettings.UserName})";
 
 			this.ParallelSyncThreadsSpinner.Value = Math.Max(Math.Min(Settings.SyncOptions.NumThreads, ParallelSyncThreadsSpinner.Maximum), ParallelSyncThreadsSpinner.Minimum);
 
@@ -141,9 +136,9 @@ namespace UnrealGameSync
 			}
 		}
 
-		public static bool? ShowModal(IWin32Window Owner, PerforceConnection DefaultConnection, bool bUnstable, string OriginalExecutableFileName, UserSettings Settings, ToolUpdateMonitor ToolUpdateMonitor, TextWriter Log)
+		public static bool? ShowModal(IWin32Window Owner, IPerforceSettings DefaultPerforceSettings, bool bUnstable, string OriginalExecutableFileName, UserSettings Settings, ToolUpdateMonitor ToolUpdateMonitor, ILogger<ApplicationSettingsWindow> Logger)
 		{
-			ApplicationSettingsWindow ApplicationSettings = new ApplicationSettingsWindow(DefaultConnection.ServerAndPort, DefaultConnection.UserName, bUnstable, OriginalExecutableFileName, Settings, ToolUpdateMonitor, Log);
+			ApplicationSettingsWindow ApplicationSettings = new ApplicationSettingsWindow(DefaultPerforceSettings, bUnstable, OriginalExecutableFileName, Settings, ToolUpdateMonitor, Logger);
 			if(ApplicationSettings.ShowDialog() == DialogResult.OK)
 			{
 				return ApplicationSettings.bRestartUnstable;
@@ -163,19 +158,19 @@ namespace UnrealGameSync
 		private void OkBtn_Click(object sender, EventArgs e)
 		{
 			// Update the settings
-			string ServerAndPort = ServerTextBox.Text.Trim();
+			string? ServerAndPort = ServerTextBox.Text.Trim();
 			if(ServerAndPort.Length == 0)
 			{
 				ServerAndPort = null;
 			}
 
-			string UserName = UserNameTextBox.Text.Trim();
+			string? UserName = UserNameTextBox.Text.Trim();
 			if(UserName.Length == 0)
 			{
 				UserName = null;
 			}
 
-			string DepotPath = DepotPathTextBox.Text.Trim();
+			string? DepotPath = DepotPathTextBox.Text.Trim();
 			if(DepotPath.Length == 0 || DepotPath == DeploymentSettings.DefaultDepotPath)
 			{
 				DepotPath = null;
@@ -195,14 +190,13 @@ namespace UnrealGameSync
 				// Try to log in to the new server, and check the application is there
 				if(ServerAndPort != InitialServerAndPort || UserName != InitialUserName || DepotPath != InitialDepotPath)
 				{
-					string ErrorMessage;
-					ModalTaskResult Result = PerforceModalTask.Execute(this, new PerforceConnection(UserName, null, ServerAndPort), new PerforceTestConnectionTask(DepotPath), "Connecting", "Checking connection, please wait...", Log, out ErrorMessage);
-					if(Result != ModalTaskResult.Succeeded)
+					PerforceSettings Settings = Utility.OverridePerforceSettings(DefaultPerforceSettings, ServerAndPort, UserName);
+
+					string TestDepotPath = DepotPath ?? DeploymentSettings.DefaultDepotPath;
+
+					ModalTask? Task = PerforceModalTask.Execute(this, "Checking connection", "Checking connection, please wait...", Settings, (p, c) => PerforceTestConnectionTask.RunAsync(p, TestDepotPath, c), Logger);
+					if (Task == null || !Task.Succeeded)
 					{
-						if(Result == ModalTaskResult.Failed)
-						{
-							MessageBox.Show(ErrorMessage, "Unable to connect");
-						}
 						return;
 					}
 				}
@@ -235,9 +229,12 @@ namespace UnrealGameSync
 			}
 
 			List<Guid> NewEnabledTools = new List<Guid>();
-			foreach (ToolItem Item in CustomToolsListBox.CheckedItems)
+			foreach (ToolItem? Item in CustomToolsListBox.CheckedItems)
 			{
-				NewEnabledTools.Add(Item.Definition.Id);
+				if (Item != null)
+				{
+					NewEnabledTools.Add(Item.Definition.Id);
+				}
 			}
 			if (!NewEnabledTools.SequenceEqual(Settings.EnabledTools))
 			{
diff --git a/Engine/Source/Programs/UnrealGameSync/UnrealGameSync/Forms/AutomatedBuildWindow.cs b/Engine/Source/Programs/UnrealGameSync/UnrealGameSync/Forms/AutomatedBuildWindow.cs
--- a/Engine/Source/Programs/UnrealGameSync/UnrealGameSync/Forms/AutomatedBuildWindow.cs
+++ b/Engine/Source/Programs/UnrealGameSync/UnrealGameSync/Forms/AutomatedBuildWindow.cs
@@ -1,5 +1,7 @@
 // Copyright Epic Games, Inc. All Rights Reserved.
 
+using EpicGames.Perforce;
+using Microsoft.Extensions.Logging;
 using System;
 using System.Collections.Generic;
 using System.ComponentModel;
@@ -25,19 +27,19 @@ namespace UnrealGameSync.Forms
 		}
 
 		string StreamName;
-		TextWriter Log;
+		IServiceProvider ServiceProvider;
 
 		string ServerAndPortOverride;
 		string UserNameOverride;
-		PerforceConnection DefaultConnection;
+		IPerforceSettings DefaultPerforceSettings;
 
 		BuildInfo Result;
 
-		private AutomatedBuildWindow(string StreamName, int Changelist, string Command, PerforceConnection DefaultConnection, string DefaultWorkspaceName, string DefaultProjectPath, TextWriter Log)
+		private AutomatedBuildWindow(string StreamName, int Changelist, string Command, IPerforceSettings DefaultPerforceSettings, string DefaultWorkspaceName, string DefaultProjectPath, IServiceProvider ServiceProvider)
 		{
 			this.StreamName = StreamName;
-			this.DefaultConnection = DefaultConnection;
-			this.Log = Log;
+			this.DefaultPerforceSettings = DefaultPerforceSettings;
+			this.ServiceProvider = ServiceProvider;
 
 			InitializeComponent();
 
@@ -66,14 +68,14 @@ namespace UnrealGameSync.Forms
 			UpdateWorkspacePathBrowseButton();
 		}
 
-		private PerforceConnection Perforce
+		private IPerforceSettings Perforce
 		{
-			get { return Utility.OverridePerforceSettings(DefaultConnection, ServerAndPortOverride, UserNameOverride); }
+			get => Utility.OverridePerforceSettings(DefaultPerforceSettings, ServerAndPortOverride, UserNameOverride);
 		}
 
-		public static bool ShowModal(IWin32Window Owner, PerforceConnection DefaultConnection, string StreamName, string ProjectPath, int Changelist, string Command, UserSettings Settings, TextWriter Log, out BuildInfo BuildInfo)
+		public static bool ShowModal(IWin32Window Owner, IPerforceSettings DefaultPerforceSettings, string StreamName, string ProjectPath, int Changelist, string Command, UserSettings Settings, IServiceProvider LoggerFactory, out BuildInfo BuildInfo)
 		{
-			string DefaultWorkspaceName = AutomatedSyncWindow.FindDefaultWorkspace(Owner, DefaultConnection, StreamName, Log);
+			string DefaultWorkspaceName = AutomatedSyncWindow.FindDefaultWorkspace(Owner, DefaultPerforceSettings, StreamName, LoggerFactory);
 
 			string DefaultProjectPath = null;
 			if(!String.IsNullOrEmpty(ProjectPath))
@@ -85,7 +87,7 @@ namespace UnrealGameSync.Forms
 				string ClientPrefix = String.Format("//{0}/", DefaultWorkspaceName);
 				foreach (UserSelectedProjectSettings ProjectSettings in Settings.RecentProjects)
 				{
-					if (ProjectSettings.ClientPath.StartsWith(ClientPrefix, StringComparison.OrdinalIgnoreCase))
+					if (ProjectSettings.ClientPath != null && ProjectSettings.ClientPath.StartsWith(ClientPrefix, StringComparison.OrdinalIgnoreCase))
 					{
 						DefaultProjectPath = ProjectSettings.ClientPath.Substring(ClientPrefix.Length - 1);
 						break;
@@ -93,7 +95,7 @@ namespace UnrealGameSync.Forms
 				}
 			}
 
-			AutomatedBuildWindow Window = new AutomatedBuildWindow(StreamName, Changelist, Command, DefaultConnection, DefaultWorkspaceName, DefaultProjectPath, Log);
+			AutomatedBuildWindow Window = new AutomatedBuildWindow(StreamName, Changelist, Command, DefaultPerforceSettings, DefaultWorkspaceName, DefaultProjectPath, LoggerFactory);
 			if (Window.ShowDialog() == DialogResult.OK)
 			{
 				BuildInfo = Window.Result;
@@ -108,7 +110,7 @@ namespace UnrealGameSync.Forms
 
 		private void ChangeLink_LinkClicked(object sender, LinkLabelLinkClickedEventArgs e)
 		{
-			if (ConnectWindow.ShowModal(this, DefaultConnection, ref ServerAndPortOverride, ref UserNameOverride, Log))
+			if (ConnectWindow.ShowModal(this, DefaultPerforceSettings, ref ServerAndPortOverride, ref UserNameOverride, ServiceProvider))
 			{
 				UpdateServerLabel();
 			}
@@ -116,13 +118,13 @@ namespace UnrealGameSync.Forms
 
 		private void UpdateServerLabel()
 		{
-			ServerLabel.Text = OpenProjectWindow.GetServerLabelText(DefaultConnection, ServerAndPortOverride, UserNameOverride);
+			ServerLabel.Text = OpenProjectWindow.GetServerLabelText(DefaultPerforceSettings, ServerAndPortOverride, UserNameOverride);
 		}
 
 		private void WorkspaceNameNewBtn_Click(object sender, EventArgs e)
 		{
 			string WorkspaceName;
-			if (NewWorkspaceWindow.ShowModal(this, Perforce, StreamName, WorkspaceNameTextBox.Text, Log, out WorkspaceName))
+			if (NewWorkspaceWindow.ShowModal(this, Perforce, StreamName, WorkspaceNameTextBox.Text, ServiceProvider, out WorkspaceName))
 			{
 				WorkspaceNameTextBox.Text = WorkspaceName;
 				UpdateOkButton();
@@ -132,7 +134,7 @@ namespace UnrealGameSync.Forms
 		private void WorkspaceNameBrowseBtn_Click(object sender, EventArgs e)
 		{
 			string WorkspaceName = WorkspaceNameTextBox.Text;
-			if (SelectWorkspaceWindow.ShowModal(this, Perforce, WorkspaceName, Log, out WorkspaceName))
+			if (SelectWorkspaceWindow.ShowModal(this, Perforce, WorkspaceName, ServiceProvider, out WorkspaceName))
 			{
 				WorkspaceNameTextBox.Text = WorkspaceName;
 				UpdateOkButton();
@@ -147,7 +149,7 @@ namespace UnrealGameSync.Forms
 		private void OkBtn_Click(object sender, EventArgs e)
 		{
 			AutomatedSyncWindow.WorkspaceInfo SelectedWorkspaceInfo;
-			if (AutomatedSyncWindow.ValidateWorkspace(this, Perforce, WorkspaceNameTextBox.Text, StreamName, Log, out SelectedWorkspaceInfo))
+			if (AutomatedSyncWindow.ValidateWorkspace(this, Perforce, WorkspaceNameTextBox.Text, StreamName, ServiceProvider, out SelectedWorkspaceInfo))
 			{
 				Result = new BuildInfo();
 				Result.SelectedWorkspaceInfo = SelectedWorkspaceInfo;
@@ -177,7 +179,7 @@ namespace UnrealGameSync.Forms
 			if (TryGetWorkspaceName(out WorkspaceName))
 			{
 				string WorkspacePath = WorkspacePathTextBox.Text.Trim();
-				if (SelectProjectFromWorkspaceWindow.ShowModal(this, Perforce, WorkspaceName, WorkspacePath, Log, out WorkspacePath))
+				if (SelectProjectFromWorkspaceWindow.ShowModal(this, Perforce, WorkspaceName, WorkspacePath, ServiceProvider, out WorkspacePath))
 				{
 					WorkspacePathTextBox.Text = WorkspacePath;
 					UpdateOkButton();
diff --git a/Engine/Source/Programs/UnrealGameSync/UnrealGameSync/Forms/AutomatedSyncWindow.cs b/Engine/Source/Programs/UnrealGameSync/UnrealGameSync/Forms/AutomatedSyncWindow.cs
--- a/Engine/Source/Programs/UnrealGameSync/UnrealGameSync/Forms/AutomatedSyncWindow.cs
+++ b/Engine/Source/Programs/UnrealGameSync/UnrealGameSync/Forms/AutomatedSyncWindow.cs
@@ -1,5 +1,8 @@
 // Copyright Epic Games, Inc. All Rights Reserved.
 
+using EpicGames.Perforce;
+using Microsoft.Extensions.DependencyInjection;
+using Microsoft.Extensions.Logging;
 using System;
 using System.Collections.Generic;
 using System.ComponentModel;
@@ -8,43 +11,28 @@ using System.Drawing;
 using System.IO;
 using System.Linq;
 using System.Text;
+using System.Threading;
 using System.Threading.Tasks;
 using System.Windows.Forms;
 
+#nullable enable
+
 namespace UnrealGameSync
 {
 	partial class AutomatedSyncWindow : Form
 	{
-		class FindDefaultWorkspaceTask : IPerforceModalTask
+		static class FindDefaultWorkspaceTask
 		{
-			string StreamName;
-			public string WorkspaceName;
-
-			public FindDefaultWorkspaceTask(string StreamName)
+			public static async Task<string?> RunAsync(IPerforceConnection Perforce, string StreamName, CancellationToken CancellationToken)
 			{
-				this.StreamName = StreamName;
-			}
+				InfoRecord Info = await Perforce.GetInfoAsync(InfoOptions.ShortOutput, CancellationToken);
 
-			public bool Run(PerforceConnection Perforce, TextWriter Log, out string ErrorMessage)
-			{
-				PerforceInfoRecord Info;
-				if(!Perforce.Info(out Info, Log))
-				{
-					ErrorMessage = "Unable to query Perforce info.";
-					return false;
-				}
-				
-				List<PerforceClientRecord> Clients;
-				if(!Perforce.FindClients(Info.UserName, out Clients, Log))
-				{
-					ErrorMessage = "Unable to enumerate clients from Perforce";
-					return false;
-				}
+				List<ClientsRecord> Clients = await Perforce.GetClientsAsync(ClientsOptions.None, Perforce.Settings.UserName, CancellationToken);
 
-				List<PerforceClientRecord> CandidateClients = new List<PerforceClientRecord>();
-				foreach(PerforceClientRecord Client in Clients)
+				List<ClientsRecord> CandidateClients = new List<ClientsRecord>();
+				foreach(ClientsRecord Client in Clients)
 				{
-					if(Client.Host == null || Client.Host.Equals(Info.HostName, StringComparison.OrdinalIgnoreCase))
+					if(Client.Host == null || Client.Host.Equals(Info.ClientHost, StringComparison.OrdinalIgnoreCase))
 					{
 						if(Client.Stream != null && Client.Stream.Equals(StreamName, StringComparison.OrdinalIgnoreCase))
 						{
@@ -53,22 +41,19 @@ namespace UnrealGameSync
 					}
 				}
 
+				string? WorkspaceName = null;
 				if(CandidateClients.Count >= 1)
 				{
 					WorkspaceName = CandidateClients.OrderByDescending(x => x.Access).First().Name;
 				}
-
-				ErrorMessage = null;
-				return true;
+				return WorkspaceName;
 			}
 		}
 
-		class ValidateWorkspaceTask : IPerforceModalTask
+		class ValidateWorkspaceTask
 		{
 			public string WorkspaceName;
 			public string StreamName;
-			public string ServerAndPort;
-			public string UserName;
 			public bool bRequiresStreamSwitch;
 			public bool bHasOpenFiles;
 
@@ -78,51 +63,50 @@ namespace UnrealGameSync
 				this.StreamName = StreamName;
 			}
 
-			public bool Run(PerforceConnection Perforce, TextWriter Log, out string ErrorMessage)
+			public async Task RunAsync(IPerforceConnection Perforce, CancellationToken CancellationToken)
 			{
-				this.ServerAndPort = Perforce.ServerAndPort;
-				this.UserName = Perforce.UserName;
-
-				PerforceSpec Spec;
-				if(!Perforce.TryGetClientSpec(WorkspaceName, out Spec, Log))
-				{
-					ErrorMessage = String.Format("Unable to get info for client '{0}'", WorkspaceName);
-					return false;
-				}
+				ClientRecord Spec = await Perforce.GetClientAsync(WorkspaceName, CancellationToken);
 
-				string CurrentStreamName = Spec.GetField("Stream");
+				string? CurrentStreamName = Spec.Stream;
 				if(CurrentStreamName == null || CurrentStreamName != StreamName)
 				{
 					bRequiresStreamSwitch = true;
-					bHasOpenFiles = Perforce.HasOpenFiles(Log);
+
+					List<PerforceResponse<OpenedRecord>> Records = await Perforce.TryOpenedAsync(OpenedOptions.None, FileSpecList.Any, CancellationToken).ToListAsync(CancellationToken);
+					bHasOpenFiles = Records.Succeeded() && Records.Count > 0;
 				}
-			
-				ErrorMessage = null;
-				return true;
 			}
 		}
 
 		public class WorkspaceInfo
 		{
-			public string ServerAndPort;
-			public string UserName;
-			public string WorkspaceName;
-			public bool bRequiresStreamSwitch;
+			public string ServerAndPort { get; }
+			public string UserName { get; }
+			public string WorkspaceName { get; }
+			public bool bRequiresStreamSwitch { get; }
+
+			public WorkspaceInfo(string ServerAndPort, string UserName, string WorkspaceName, bool bRequiresStreamSwitch)
+			{
+				this.ServerAndPort = ServerAndPort;
+				this.UserName = UserName;
+				this.WorkspaceName = WorkspaceName;
+				this.bRequiresStreamSwitch = bRequiresStreamSwitch;
+			}
 		}
 
 		string StreamName;
-		TextWriter Log;
+		IServiceProvider ServiceProvider;
 
-		string ServerAndPortOverride;
-		string UserNameOverride;
-        PerforceConnection DefaultConnection;
-		WorkspaceInfo SelectedWorkspaceInfo;
+		string? ServerAndPortOverride;
+		string? UserNameOverride;
+		IPerforceSettings DefaultPerforceSettings;
+		WorkspaceInfo? SelectedWorkspaceInfo;
 
-		private AutomatedSyncWindow(string StreamName, string ProjectPath, string WorkspaceName, PerforceConnection DefaultConnection, TextWriter Log)
+		private AutomatedSyncWindow(string StreamName, string ProjectPath, string? WorkspaceName, IPerforceSettings DefaultPerforceSettings, IServiceProvider ServiceProvider)
 		{
 			this.StreamName = StreamName;
-			this.DefaultConnection = DefaultConnection;
-			this.Log = Log;
+			this.DefaultPerforceSettings = DefaultPerforceSettings;
+			this.ServiceProvider = ServiceProvider;
 
 			InitializeComponent();
 
@@ -143,26 +127,20 @@ namespace UnrealGameSync
 			UpdateOkButton();
 		}
 
-		private PerforceConnection Perforce
-		{
-			get { return Utility.OverridePerforceSettings(DefaultConnection, ServerAndPortOverride, UserNameOverride); }
-		}
+		private IPerforceSettings Perforce => Utility.OverridePerforceSettings(DefaultPerforceSettings, ServerAndPortOverride, UserNameOverride);
 
-		public static string FindDefaultWorkspace(IWin32Window Owner, PerforceConnection DefaultConnection, string StreamName, TextWriter Log)
+		public static string? FindDefaultWorkspace(IWin32Window Owner, IPerforceSettings DefaultPerforceSettings, string StreamName, IServiceProvider ServiceProvider)
 		{
-			FindDefaultWorkspaceTask FindWorkspace = new FindDefaultWorkspaceTask(StreamName);
-
-			string ErrorMessage;
-			PerforceModalTask.Execute(Owner, DefaultConnection, FindWorkspace, "Finding workspace", "Finding default workspace, please wait...", Log, out ErrorMessage);
-
-			return FindWorkspace.WorkspaceName;
+			ILogger Logger = ServiceProvider.GetRequiredService<ILogger<AutomatedSyncWindow>>();
+			ModalTask<string?>? WorkspaceTask = PerforceModalTask.Execute(Owner, "Finding workspace", "Finding default workspace, please wait...", DefaultPerforceSettings, (p, c) => FindDefaultWorkspaceTask.RunAsync(p, StreamName, c), Logger);
+			return (WorkspaceTask != null && WorkspaceTask.Succeeded) ? WorkspaceTask.Result : null;
 		}
 
-		public static bool ShowModal(IWin32Window Owner, PerforceConnection DefaultConnection, string StreamName, string ProjectPath, out WorkspaceInfo WorkspaceInfo, TextWriter Log)
+		public static bool ShowModal(IWin32Window Owner, IPerforceSettings DefaultPerforceSettings, string StreamName, string ProjectPath, out WorkspaceInfo? WorkspaceInfo, IServiceProvider ServiceProvider)
 		{
-			string WorkspaceName = FindDefaultWorkspace(Owner, DefaultConnection, StreamName, Log);
+			string? WorkspaceName = FindDefaultWorkspace(Owner, DefaultPerforceSettings, StreamName, ServiceProvider);
 
-			AutomatedSyncWindow Window = new AutomatedSyncWindow(StreamName, ProjectPath, WorkspaceName, DefaultConnection, Log);
+			AutomatedSyncWindow Window = new AutomatedSyncWindow(StreamName, ProjectPath, WorkspaceName, DefaultPerforceSettings, ServiceProvider);
 			if(Window.ShowDialog() == DialogResult.OK)
 			{
 				WorkspaceInfo = Window.SelectedWorkspaceInfo;
@@ -177,7 +155,7 @@ namespace UnrealGameSync
 
 		private void ChangeLink_LinkClicked(object sender, LinkLabelLinkClickedEventArgs e)
 		{
-			if(ConnectWindow.ShowModal(this, DefaultConnection, ref ServerAndPortOverride, ref UserNameOverride, Log))
+			if(ConnectWindow.ShowModal(this, DefaultPerforceSettings, ref ServerAndPortOverride, ref UserNameOverride, ServiceProvider))
 			{
 				UpdateServerLabel();
 			}
@@ -185,14 +163,14 @@ namespace UnrealGameSync
 
 		private void UpdateServerLabel()
 		{
-			ServerLabel.Text = OpenProjectWindow.GetServerLabelText(DefaultConnection, ServerAndPortOverride, UserNameOverride);
+			ServerLabel.Text = OpenProjectWindow.GetServerLabelText(DefaultPerforceSettings, ServerAndPortOverride, UserNameOverride);
 			ChangeLink.Location = new Point(ServerLabel.Right + 5, ChangeLink.Location.Y);
 		}
 
 		private void WorkspaceNameNewBtn_Click(object sender, EventArgs e)
 		{
-			string WorkspaceName;
-			if(NewWorkspaceWindow.ShowModal(this, Perforce, StreamName, WorkspaceNameTextBox.Text, Log, out WorkspaceName))
+			string? WorkspaceName;
+			if(NewWorkspaceWindow.ShowModal(this, DefaultPerforceSettings, StreamName, WorkspaceNameTextBox.Text, ServiceProvider, out WorkspaceName))
 			{
 				WorkspaceNameTextBox.Text = WorkspaceName;
 				UpdateOkButton();
@@ -201,8 +179,8 @@ namespace UnrealGameSync
 
 		private void WorkspaceNameBrowseBtn_Click(object sender, EventArgs e)
 		{
-			string WorkspaceName = WorkspaceNameTextBox.Text;
-			if(SelectWorkspaceWindow.ShowModal(this, Perforce, WorkspaceName, Log, out WorkspaceName))
+			string? WorkspaceName = WorkspaceNameTextBox.Text;
+			if(SelectWorkspaceWindow.ShowModal(this, DefaultPerforceSettings, WorkspaceName, ServiceProvider, out WorkspaceName))
 			{
 				WorkspaceNameTextBox.Text = WorkspaceName;
 				UpdateOkButton();
@@ -214,17 +192,12 @@ namespace UnrealGameSync
 			OkBtn.Enabled = (WorkspaceNameTextBox.Text.Length > 0);
 		}
 
-		public static bool ValidateWorkspace(IWin32Window Owner, PerforceConnection Perforce, string WorkspaceName, string StreamName, TextWriter Log, out WorkspaceInfo SelectedWorkspaceInfo)
+		public static bool ValidateWorkspace(IWin32Window Owner, IPerforceSettings Perforce, string WorkspaceName, string StreamName, IServiceProvider ServiceProvider, out WorkspaceInfo? SelectedWorkspaceInfo)
 		{
 			ValidateWorkspaceTask ValidateWorkspace = new ValidateWorkspaceTask(WorkspaceName, StreamName);
 
-			string ErrorMessage;
-			ModalTaskResult Result = PerforceModalTask.Execute(Owner, Perforce, ValidateWorkspace, "Checking workspace", "Checking workspace, please wait...", Log, out ErrorMessage);
-			if (Result == ModalTaskResult.Failed)
-			{
-				MessageBox.Show(ErrorMessage);
-			}
-			else if (Result == ModalTaskResult.Succeeded)
+			ModalTask? Task = PerforceModalTask.Execute(Owner, "Checking workspace", "Checking workspace, please wait...", Perforce, ValidateWorkspace.RunAsync, ServiceProvider.GetRequiredService<ILogger<ValidateWorkspaceTask>>());
+			if (Task != null && Task.Succeeded)
 			{
 				if (ValidateWorkspace.bRequiresStreamSwitch)
 				{
@@ -244,7 +217,7 @@ namespace UnrealGameSync
 					}
 				}
 
-				SelectedWorkspaceInfo = new WorkspaceInfo() { ServerAndPort = ValidateWorkspace.ServerAndPort, UserName = ValidateWorkspace.UserName, WorkspaceName = ValidateWorkspace.WorkspaceName, bRequiresStreamSwitch = ValidateWorkspace.bRequiresStreamSwitch };
+				SelectedWorkspaceInfo = new WorkspaceInfo(Perforce.ServerAndPort, Perforce.UserName, ValidateWorkspace.WorkspaceName, ValidateWorkspace.bRequiresStreamSwitch);
 				return true;
 			}
 
@@ -254,7 +227,7 @@ namespace UnrealGameSync
 
 		private void OkBtn_Click(object sender, EventArgs e)
 		{
-			if(ValidateWorkspace(this, Perforce, WorkspaceNameTextBox.Text, StreamName, Log, out SelectedWorkspaceInfo))
+			if(ValidateWorkspace(this, Perforce, WorkspaceNameTextBox.Text, StreamName, ServiceProvider, out SelectedWorkspaceInfo))
 			{
 				DialogResult = DialogResult.OK;
 				Close();
diff --git a/Engine/Source/Programs/UnrealGameSync/UnrealGameSync/Forms/CleanWorkspaceWindow.cs b/Engine/Source/Programs/UnrealGameSync/UnrealGameSync/Forms/CleanWorkspaceWindow.cs
--- a/Engine/Source/Programs/UnrealGameSync/UnrealGameSync/Forms/CleanWorkspaceWindow.cs
+++ b/Engine/Source/Programs/UnrealGameSync/UnrealGameSync/Forms/CleanWorkspaceWindow.cs
@@ -1,6 +1,8 @@
 // Copyright Epic Games, Inc. All Rights Reserved.
 
 using EpicGames.Core;
+using EpicGames.Perforce;
+using Microsoft.Extensions.Logging;
 using System;
 using System.Collections.Generic;
 using System.ComponentModel;
@@ -16,6 +18,8 @@ using System.Threading.Tasks;
 using System.Windows.Forms;
 using System.Windows.Forms.VisualStyles;
 
+#nullable enable
+
 namespace UnrealGameSync
 {
 	partial class CleanWorkspaceWindow : Form
@@ -29,8 +33,8 @@ namespace UnrealGameSync
 		class TreeNodeData
 		{
 			public TreeNodeAction Action;
-			public FileInfo File;
-			public FolderToClean Folder;
+			public FileInfo? File;
+			public FolderToClean? Folder;
 			public int NumFiles;
 			public int NumSelectedFiles;
 			public int NumEmptySelectedFiles;
@@ -54,7 +58,7 @@ namespace UnrealGameSync
 			CheckBoxState.CheckedNormal 
 		};
 
-		PerforceConnection PerforceClient;
+		IPerforceSettings PerforceSettings;
 		FolderToClean RootFolderToClean;
 
 		static readonly string[] SafeToDeleteFolders =
@@ -101,33 +105,31 @@ namespace UnrealGameSync
 
 		IReadOnlyList<string> ExtraSafeToDeleteFolders;
 		IReadOnlyList<string> ExtraSafeToDeleteExtensions;
+		ILogger Logger;
 
 		[DllImport("Shell32.dll", EntryPoint = "ExtractIconExW", CharSet = CharSet.Unicode, ExactSpelling = true, CallingConvention = CallingConvention.StdCall)]
 		private static extern int ExtractIconEx(string sFile, int iIndex, IntPtr piLargeVersion, out IntPtr piSmallVersion, int amountIcons);
 
-		private CleanWorkspaceWindow(PerforceConnection PerforceClient, FolderToClean RootFolderToClean, string[] ExtraSafeToDeleteFolders, string[] ExtraSafeToDeleteExtensions)
+		private CleanWorkspaceWindow(IPerforceSettings PerforceSettings, FolderToClean RootFolderToClean, string[] ExtraSafeToDeleteFolders, string[] ExtraSafeToDeleteExtensions, ILogger<CleanWorkspaceWindow> Logger)
 		{
-			this.PerforceClient = PerforceClient;
+			this.PerforceSettings = PerforceSettings;
 			this.RootFolderToClean = RootFolderToClean;
 			this.ExtraSafeToDeleteFolders = ExtraSafeToDeleteFolders.Select(x => x.Trim().Replace('\\', '/').Trim('/')).Where(x => x.Length > 0).Select(x => String.Format("/{0}/", x.ToLowerInvariant())).ToArray();
 			this.ExtraSafeToDeleteExtensions = ExtraSafeToDeleteExtensions.Select(x => x.Trim().ToLowerInvariant()).Where(x => x.Length > 0).ToArray();
+			this.Logger = Logger;
 
 			InitializeComponent();
 		}
 
-		public static void DoClean(IWin32Window Owner, PerforceConnection PerforceClient, DirectoryReference LocalRootPath, string ClientRootPath, IReadOnlyList<string> SyncPaths, string[] ExtraSafeToDeleteFolders, string[] ExtraSafeToDeleteExtensions, TextWriter Log)
+		public static void DoClean(IWin32Window Owner, IPerforceSettings PerforceSettings, DirectoryReference LocalRootPath, string ClientRootPath, IReadOnlyList<string> SyncPaths, string[] ExtraSafeToDeleteFolders, string[] ExtraSafeToDeleteExtensions, ILogger<CleanWorkspaceWindow> Logger)
 		{
 			// Figure out which folders to clean
 			FolderToClean RootFolderToClean = new FolderToClean(LocalRootPath.ToDirectoryInfo());
-			using(FindFoldersToCleanTask QueryWorkspace = new FindFoldersToCleanTask(PerforceClient, RootFolderToClean, ClientRootPath, SyncPaths, Log))
+			using(FindFoldersToCleanTask QueryWorkspace = new FindFoldersToCleanTask(PerforceSettings, RootFolderToClean, ClientRootPath, SyncPaths, Logger))
 			{
-				string ErrorMessage;
-				if(ModalTask.Execute(Owner, QueryWorkspace, "Clean Workspace", "Querying files in Perforce, please wait...", out ErrorMessage) != ModalTaskResult.Succeeded)
+				ModalTask? Result = ModalTask.Execute(Owner, "Clean Workspace", "Querying files in Perforce, please wait...", x => QueryWorkspace.RunAsync(x), ModalTaskFlags.None);
+				if (Result == null || !Result.Succeeded)
 				{
-					if(!String.IsNullOrEmpty(ErrorMessage))
-					{
-						MessageBox.Show(ErrorMessage);
-					}
 					return;
 				}
 			}
@@ -140,7 +142,7 @@ namespace UnrealGameSync
 			}
 
 			// Populate the tree
-			CleanWorkspaceWindow CleanWorkspace = new CleanWorkspaceWindow(PerforceClient, RootFolderToClean, ExtraSafeToDeleteFolders, ExtraSafeToDeleteExtensions);
+			CleanWorkspaceWindow CleanWorkspace = new CleanWorkspaceWindow(PerforceSettings, RootFolderToClean, ExtraSafeToDeleteFolders, ExtraSafeToDeleteExtensions, Logger);
 			CleanWorkspace.ShowDialog();
 		}
 
@@ -377,9 +379,12 @@ namespace UnrealGameSync
 
 			if(NewNumSelectedFiles != ParentNodeData.NumSelectedFiles)
 			{
-				foreach(TreeNode ChildNode in ParentNode.Nodes)
+				foreach(TreeNode? ChildNode in ParentNode.Nodes)
 				{
-					SetSelectedOnChildren(ChildNode, Type);
+					if (ChildNode != null)
+					{
+						SetSelectedOnChildren(ChildNode, Type);
+					}
 				}
 				ParentNodeData.NumSelectedFiles = NewNumSelectedFiles;
 				UpdateImage(ParentNode);
@@ -400,16 +405,19 @@ namespace UnrealGameSync
 			List<FileInfo> FilesToSync = new List<FileInfo>();
 			List<FileInfo> FilesToDelete = new List<FileInfo>();
 			List<DirectoryInfo> DirectoriesToDelete = new List<DirectoryInfo>();
-			foreach(TreeNode RootNode in TreeView.Nodes)
+			foreach(TreeNode? RootNode in TreeView.Nodes)
 			{
-				FindSelection(RootNode, FilesToSync, FilesToDelete, DirectoriesToDelete);
+				if (RootNode != null)
+				{
+					FindSelection(RootNode, FilesToSync, FilesToDelete, DirectoriesToDelete);
+				}
 			}
 
-			string ErrorMessage;
-			if(ModalTask.Execute(this, new DeleteFilesTask(PerforceClient, FilesToSync, FilesToDelete, DirectoriesToDelete), "Clean Workspace", "Cleaning files, please wait...", out ErrorMessage) != ModalTaskResult.Succeeded && !String.IsNullOrEmpty(ErrorMessage))
+			ModalTask? Result = ModalTask.Execute(this, "Clean Workspace", "Cleaning files, please wait...", x => DeleteFilesTask.RunAsync(PerforceSettings, FilesToSync, FilesToDelete, DirectoriesToDelete, Logger, x), ModalTaskFlags.Quiet);
+			if(Result != null && Result.Failed)
 			{
 				FailedToDeleteWindow FailedToDelete = new FailedToDeleteWindow();
-				FailedToDelete.FileList.Text = ErrorMessage;
+				FailedToDelete.FileList.Text = Result.Error;
 				FailedToDelete.FileList.SelectionStart = 0;
 				FailedToDelete.FileList.SelectionLength = 0;
 				FailedToDelete.ShowDialog();
@@ -435,9 +443,12 @@ namespace UnrealGameSync
 			}
 			else
 			{
-				foreach(TreeNode ChildNode in Node.Nodes)
+				foreach(TreeNode? ChildNode in Node.Nodes)
 				{
-					FindSelection(ChildNode, FilesToSync, FilesToDelete, DirectoriesToDelete);
+					if (ChildNode != null)
+					{
+						FindSelection(ChildNode, FilesToSync, FilesToDelete, DirectoriesToDelete);
+					}
 				}
 				if(NodeData.Folder != null && NodeData.Folder.bEmptyAfterClean && NodeData.NumSelectedFiles == NodeData.NumFiles)
 				{
@@ -448,17 +459,23 @@ namespace UnrealGameSync
 
 		private void SelectAllBtn_Click(object sender, EventArgs e)
 		{
-			foreach(TreeNode Node in TreeView.Nodes)
+			foreach(TreeNode? Node in TreeView.Nodes)
 			{
-				SetSelected(Node, SelectionType.All);
+				if (Node != null)
+				{
+					SetSelected(Node, SelectionType.All);
+				}
 			}
 		}
 
 		private void SelectMissingBtn_Click(object sender, EventArgs e)
 		{
-			foreach(TreeNode Node in TreeView.Nodes)
+			foreach(TreeNode? Node in TreeView.Nodes)
 			{
-				SetSelected(Node, SelectionType.Missing);
+				if (Node != null)
+				{
+					SetSelected(Node, SelectionType.Missing);
+				}
 			}
 		}
 
diff --git a/Engine/Source/Programs/UnrealGameSync/UnrealGameSync/Forms/ClobberWindow.cs b/Engine/Source/Programs/UnrealGameSync/UnrealGameSync/Forms/ClobberWindow.cs
--- a/Engine/Source/Programs/UnrealGameSync/UnrealGameSync/Forms/ClobberWindow.cs
+++ b/Engine/Source/Programs/UnrealGameSync/UnrealGameSync/Forms/ClobberWindow.cs
@@ -11,6 +11,8 @@ using System.Text;
 using System.Threading.Tasks;
 using System.Windows.Forms;
 
+#nullable enable
+
 namespace UnrealGameSync
 {
 	public partial class ClobberWindow : Form
@@ -51,17 +53,23 @@ namespace UnrealGameSync
 
 		private void UncheckAll_Click(object sender, EventArgs e)
 		{
-			foreach(ListViewItem Item in FileList.Items)
+			foreach (ListViewItem? Item in FileList.Items)
 			{
-				Item.Checked = false;
+				if (Item != null)
+				{
+					Item.Checked = false;
+				}
 			}
 		}
 
 		private void ContinueButton_Click(object sender, EventArgs e)
 		{
-			foreach(ListViewItem Item in FileList.Items)
+			foreach(ListViewItem? Item in FileList.Items)
 			{
-				FilesToClobber[(string)Item.Tag] = Item.Checked;
+				if (Item != null)
+				{
+					FilesToClobber[(string)Item.Tag] = Item.Checked;
+				}
 			}
 		}
 	}
