UGS: Significant refactor of UGS internals.

diff --git a/Engine/Source/Programs/UnrealGameSync/UnrealGameSync/Forms/ConnectWindow.cs b/Engine/Source/Programs/UnrealGameSync/UnrealGameSync/Forms/ConnectWindow.cs
--- a/Engine/Source/Programs/UnrealGameSync/UnrealGameSync/Forms/ConnectWindow.cs
+++ b/Engine/Source/Programs/UnrealGameSync/UnrealGameSync/Forms/ConnectWindow.cs
@@ -1,5 +1,7 @@
 ﻿// Copyright Epic Games, Inc. All Rights Reserved.
 
+using EpicGames.Perforce;
+using Microsoft.Extensions.Logging;
 using System;
 using System.Collections.Generic;
 using System.ComponentModel;
@@ -15,17 +17,17 @@ namespace UnrealGameSync
 {
 	partial class ConnectWindow : Form
 	{
-		PerforceConnection DefaultConnection;
+		IPerforceSettings DefaultPerforceSettings;
 		string ServerAndPortOverride;
 		string UserNameOverride;
-		TextWriter Log;
+		IServiceProvider ServiceProvider;
 
-		private ConnectWindow(PerforceConnection DefaultConnection, string ServerAndPortOverride, string UserNameOverride, TextWriter Log)
+		private ConnectWindow(IPerforceSettings DefaultPerforceSettings, string ServerAndPortOverride, string UserNameOverride, IServiceProvider ServiceProvider)
 		{
 			InitializeComponent();
 
-			this.DefaultConnection = DefaultConnection;
-			this.Log = Log;
+			this.DefaultPerforceSettings = DefaultPerforceSettings;
+			this.ServiceProvider = ServiceProvider;
 
 			if(!String.IsNullOrWhiteSpace(ServerAndPortOverride))
 			{
@@ -36,19 +38,19 @@ namespace UnrealGameSync
 				this.UserNameOverride = UserNameOverride.Trim();
 			}
 
-			ServerAndPortTextBox.CueBanner = DefaultConnection.ServerAndPort;
-			ServerAndPortTextBox.Text = this.ServerAndPortOverride ?? DefaultConnection.ServerAndPort;
-			UserNameTextBox.CueBanner = DefaultConnection.UserName;
-			UserNameTextBox.Text = this.UserNameOverride ?? DefaultConnection.UserName;
+			ServerAndPortTextBox.CueBanner = DefaultPerforceSettings.ServerAndPort;
+			ServerAndPortTextBox.Text = this.ServerAndPortOverride ?? DefaultPerforceSettings.ServerAndPort;
+			UserNameTextBox.CueBanner = DefaultPerforceSettings.UserName;
+			UserNameTextBox.Text = this.UserNameOverride ?? DefaultPerforceSettings.UserName;
 			UseDefaultConnectionSettings.Checked = this.ServerAndPortOverride == null && this.UserNameOverride == null;
 
 			UpdateEnabledControls();
 		}
 
-		public static bool ShowModal(IWin32Window Owner, PerforceConnection DefaultConnection, ref string ServerAndPortOverride, ref string UserNameOverride, TextWriter Log)
+		public static bool ShowModal(IWin32Window Owner, IPerforceSettings DefaultSettings, ref string ServerAndPortOverride, ref string UserNameOverride, IServiceProvider ServiceProvider)
 		{
-			ConnectWindow Connect = new ConnectWindow(DefaultConnection, ServerAndPortOverride, UserNameOverride, Log);
-			if(Connect.ShowDialog() == DialogResult.OK)
+			ConnectWindow Connect = new ConnectWindow(DefaultSettings, ServerAndPortOverride, UserNameOverride, ServiceProvider);
+			if(Connect.ShowDialog(Owner) == DialogResult.OK)
 			{
 				ServerAndPortOverride = Connect.ServerAndPortOverride;
 				UserNameOverride = Connect.UserNameOverride;
@@ -75,13 +77,13 @@ namespace UnrealGameSync
 				ServerAndPortOverride = ServerAndPortTextBox.Text.Trim();
 				if(ServerAndPortOverride.Length == 0)
 				{
-					ServerAndPortOverride = DefaultConnection.ServerAndPort;
+					ServerAndPortOverride = DefaultPerforceSettings.ServerAndPort;
 				}
 
 				UserNameOverride = UserNameTextBox.Text.Trim();
 				if(UserNameOverride.Length == 0)
 				{
-					UserNameOverride = DefaultConnection.UserName;
+					UserNameOverride = DefaultPerforceSettings.UserName;
 				}
 			}
 
@@ -108,7 +110,7 @@ namespace UnrealGameSync
 		private void BrowseUserBtn_Click(object sender, EventArgs e)
 		{
 			string NewUserName;
-			if(SelectUserWindow.ShowModal(this, new PerforceConnection(UserNameTextBox.Text, null, ServerAndPortTextBox.Text), Log, out NewUserName))
+			if(SelectUserWindow.ShowModal(this, new PerforceSettings(DefaultPerforceSettings) { UserName = UserNameTextBox.Text, ServerAndPort = ServerAndPortTextBox.Text }, ServiceProvider, out NewUserName))
 			{
 				UserNameTextBox.Text = NewUserName;
 			}
diff --git a/Engine/Source/Programs/UnrealGameSync/UnrealGameSync/Forms/IssueBrowserWindow.cs b/Engine/Source/Programs/UnrealGameSync/UnrealGameSync/Forms/IssueBrowserWindow.cs
--- a/Engine/Source/Programs/UnrealGameSync/UnrealGameSync/Forms/IssueBrowserWindow.cs
+++ b/Engine/Source/Programs/UnrealGameSync/UnrealGameSync/Forms/IssueBrowserWindow.cs
@@ -1,5 +1,8 @@
 // Copyright Epic Games, Inc. All Rights Reserved.
 
+using EpicGames.Perforce;
+using Microsoft.Extensions.DependencyInjection;
+using Microsoft.Extensions.Logging;
 using System;
 using System.Collections.Generic;
 using System.ComponentModel;
@@ -14,36 +17,38 @@ using System.Threading;
 using System.Threading.Tasks;
 using System.Windows.Forms;
 
+#nullable enable
+
 namespace UnrealGameSync
 {
 	partial class IssueBrowserWindow : Form
 	{
 		IssueMonitor IssueMonitor;
-		string ServerAndPort;
-		string UserName;
+		IPerforceSettings PerforceSettings;
 		TimeSpan? ServerTimeOffset;
-		TextWriter Log;
+		IServiceProvider ServiceProvider;
+		ILogger Logger;
 		string CurrentStream;
 		int MaxResults = 0;
 		int PendingMaxResults = 0;
-		string FilterName;
+		string? FilterName;
 		Dictionary<string, Func<IssueData, bool>> CustomFilters = new Dictionary<string, Func<IssueData, bool>>();
 		List<IssueData> Issues = new List<IssueData>();
 		SynchronizationContext MainThreadSynchronizationContext;
-		Thread BackgroundThread;
+		Task? BackgroundTask;
 		bool bDisposed;
 
-		public IssueBrowserWindow(IssueMonitor IssueMonitor, string ServerAndPort, string UserName, TimeSpan? ServerTimeOffset, TextWriter Log, string CurrentStream, Dictionary<string, Func<IssueData, bool>> CustomFilters, string FilterName)
+		public IssueBrowserWindow(IssueMonitor IssueMonitor, IPerforceSettings PerforceSettings, TimeSpan? ServerTimeOffset, IServiceProvider ServiceProvider, string CurrentStream, Dictionary<string, Func<IssueData, bool>> CustomFilters, string FilterName)
 		{
 			this.IssueMonitor = IssueMonitor;
-			this.ServerAndPort = ServerAndPort;
-			this.UserName = UserName;
+			this.PerforceSettings = PerforceSettings;
 			this.ServerTimeOffset = ServerTimeOffset;
-			this.Log = Log;
+			this.ServiceProvider = ServiceProvider;
+			this.Logger = ServiceProvider.GetRequiredService<ILogger<IssueBrowserWindow>>();
 			this.CurrentStream = CurrentStream;
 			this.FilterName = FilterName;
 			this.CustomFilters = CustomFilters;
-			this.MainThreadSynchronizationContext = SynchronizationContext.Current;
+			this.MainThreadSynchronizationContext = SynchronizationContext.Current!;
 
 			IssueMonitor.AddRef();
 
@@ -52,13 +57,16 @@ namespace UnrealGameSync
 			using (Graphics Graphics = Graphics.FromHwnd(IntPtr.Zero))
 			{
 				float DpiScaleX = Graphics.DpiX / 96.0f;
-				foreach (ColumnHeader Column in IssueListView.Columns)
+				foreach (ColumnHeader? Column in IssueListView.Columns)
 				{
-					Column.Width = (int)(Column.Width * DpiScaleX);
+					if (Column != null)
+					{
+						Column.Width = (int)(Column.Width * DpiScaleX);
+					}
 				}
 			}
 
-			System.Reflection.PropertyInfo DoubleBufferedProperty = typeof(Control).GetProperty("DoubleBuffered", System.Reflection.BindingFlags.NonPublic | System.Reflection.BindingFlags.Instance);
+			System.Reflection.PropertyInfo DoubleBufferedProperty = typeof(Control).GetProperty("DoubleBuffered", System.Reflection.BindingFlags.NonPublic | System.Reflection.BindingFlags.Instance)!;
 			DoubleBufferedProperty.SetValue(IssueListView, true, null);
 		}
 
@@ -87,19 +95,19 @@ namespace UnrealGameSync
 			CheckToStartBackgroundThread();
 		}
 
-		void FetchIssues(int NewMaxResults)
+		async Task FetchIssues(int NewMaxResults, CancellationToken CancellationToken)
 		{
 			try
 			{
 				SortedDictionary<int, IssueData> NewSortedIssues = new SortedDictionary<int, IssueData>();
 
-				List<IssueData> OpenIssues = RESTApi.GET<List<IssueData>>(IssueMonitor.ApiUrl, "issues");
+				List<IssueData> OpenIssues = await RESTApi.GetAsync<List<IssueData>>($"{IssueMonitor.ApiUrl}/api/issues", CancellationToken);
 				foreach (IssueData OpenIssue in OpenIssues)
 				{
 					NewSortedIssues[(int)OpenIssue.Id] = OpenIssue;
 				}
 
-				List<IssueData> ResolvedIssues = RESTApi.GET<List<IssueData>>(IssueMonitor.ApiUrl, String.Format("issues?includeresolved=true&maxresults={0}", NewMaxResults));
+				List<IssueData> ResolvedIssues = await RESTApi.GetAsync<List<IssueData>>($"{IssueMonitor.ApiUrl}/api/issues?includeresolved=true&maxresults={NewMaxResults}", CancellationToken);
 				foreach (IssueData ResolvedIssue in ResolvedIssues)
 				{
 					NewSortedIssues[(int)ResolvedIssue.Id] = ResolvedIssue;
@@ -122,14 +130,14 @@ namespace UnrealGameSync
 			Issues = NewIssues;
 			UpdateIssueList();
 
-			BackgroundThread = null;
+			BackgroundTask = null;
 			CheckToStartBackgroundThread();
 		}
 
 		void FetchIssuesFailure(Exception Ex)
 		{
 			StatusLabel.Text = String.Format("Unable to fetch issues ({0})", Ex.Message);
-			BackgroundThread = null;
+			BackgroundTask = null;
 			CheckToStartBackgroundThread();
 		}
 
@@ -143,11 +151,10 @@ namespace UnrealGameSync
 
 		void StartBackgroundThread()
 		{
-			if(BackgroundThread == null)
+			if(BackgroundTask == null)
 			{
 				int NewMaxResultsCopy = PendingMaxResults;
-				BackgroundThread = new Thread(() => FetchIssues(NewMaxResultsCopy));
-				BackgroundThread.Start();
+				BackgroundTask = Task.Run(() => FetchIssues(NewMaxResultsCopy, CancellationToken.None));
 			}
 		}
 
@@ -158,7 +165,7 @@ namespace UnrealGameSync
 			DateTime Midnight = (Now - Now.TimeOfDay).ToUniversalTime();
 
 			// Get the regex for the selected filter
-			Func<IssueData, bool> Filter;
+			Func<IssueData, bool>? Filter;
 			if (String.IsNullOrEmpty(FilterName))
 			{
 				Filter = x => true;
@@ -255,12 +262,12 @@ namespace UnrealGameSync
 
 		static List<IssueBrowserWindow> ExistingWindows = new List<IssueBrowserWindow>();
 
-		public static void Show(Form Owner, IssueMonitor IssueMonitor, string ServerAndPort, string UserName, TimeSpan? ServerTimeOffset, TextWriter Log, string CurrentStream, Dictionary<string, Func<IssueData, bool>> CustomFilters, string DefaultFilter)
+		public static void Show(Form Owner, IssueMonitor IssueMonitor, IPerforceSettings PerforceSettings, TimeSpan? ServerTimeOffset, IServiceProvider ServiceProvider, string CurrentStream, Dictionary<string, Func<IssueData, bool>> CustomFilters, string DefaultFilter)
 		{
 			IssueBrowserWindow Window = ExistingWindows.FirstOrDefault(x => x.IssueMonitor == IssueMonitor);
 			if(Window == null)
 			{
-				Window = new IssueBrowserWindow(IssueMonitor, ServerAndPort, UserName, ServerTimeOffset, Log, CurrentStream, CustomFilters, DefaultFilter);
+				Window = new IssueBrowserWindow(IssueMonitor, PerforceSettings, ServerTimeOffset, ServiceProvider, CurrentStream, CustomFilters, DefaultFilter);
 				Window.Owner = Owner;
 				Window.StartPosition = FormStartPosition.Manual;
 				Window.Location = new Point(Owner.Location.X + (Owner.Width - Window.Width) / 2, Owner.Location.Y + (Owner.Height - Window.Height) / 2);
@@ -349,7 +356,7 @@ namespace UnrealGameSync
 
 		private void ShowIssue(IssueData Issue)
 		{
-			IssueDetailsWindow.Show(Owner, IssueMonitor, ServerAndPort, UserName, ServerTimeOffset, Issue, Log, CurrentStream);
+			IssueDetailsWindow.Show(Owner, IssueMonitor, PerforceSettings, ServerTimeOffset, Issue, ServiceProvider, CurrentStream);
 		}
 
 		private void IssueListView_DrawColumnHeader(object sender, DrawListViewColumnHeaderEventArgs e)
diff --git a/Engine/Source/Programs/UnrealGameSync/UnrealGameSync/Forms/IssueDetailsWindow.cs b/Engine/Source/Programs/UnrealGameSync/UnrealGameSync/Forms/IssueDetailsWindow.cs
--- a/Engine/Source/Programs/UnrealGameSync/UnrealGameSync/Forms/IssueDetailsWindow.cs
+++ b/Engine/Source/Programs/UnrealGameSync/UnrealGameSync/Forms/IssueDetailsWindow.cs
@@ -1,9 +1,14 @@
 // Copyright Epic Games, Inc. All Rights Reserved.
 
+using EpicGames.Core;
+using EpicGames.Perforce;
+using Microsoft.Extensions.DependencyInjection;
+using Microsoft.Extensions.Logging;
 using System;
 using System.Collections.Generic;
 using System.Data;
 using System.Diagnostics;
+using System.Diagnostics.CodeAnalysis;
 using System.Drawing;
 using System.Drawing.Drawing2D;
 using System.IO;
@@ -11,67 +16,88 @@ using System.Linq;
 using System.Text;
 using System.Text.RegularExpressions;
 using System.Threading;
+using System.Threading.Tasks;
 using System.Windows.Forms;
 
+#nullable enable
+
 namespace UnrealGameSync
 {
 	partial class IssueDetailsWindow : Form
 	{
 		class BuildGroup
 		{
-			public string JobName;
-			public string JobUrl;
-			public int Change;
-			public IssueBuildOutcome Outcome;
-			public List<IssueBuildData> Builds;
+			public string JobName { get; }
+			public string JobUrl { get; }
+			public int Change { get; }
+			public IssueBuildOutcome Outcome { get; }
+			public IReadOnlyList<IssueBuildData> Builds { get; }
+
+			public BuildGroup(string JobName, string JobUrl, int Change, IssueBuildOutcome Outcome, IReadOnlyList<IssueBuildData> Builds)
+			{
+				this.JobName = JobName;
+				this.JobUrl = JobUrl;
+				this.Change = Change;
+				this.Outcome = Outcome;
+				this.Builds = Builds;
+			}
 		}
 
 		class PerforceChangeRange
 		{
 			public bool bExpanded;
-			public int MinChange;
-			public int MaxChange;
-			public List<PerforceChangeSummary> Changes;
-			public string ErrorMessage;
-			public BuildGroup BuildGroup;
+			public int MinChange { get; }
+			public int MaxChange { get; }
+			public List<ChangesRecord>? Changes;
+			public string? ErrorMessage;
+			public BuildGroup BuildGroup { get; }
+
+			public PerforceChangeRange(int MinChange, int MaxChange, BuildGroup BuildGroup)
+			{
+				this.MinChange = MinChange;
+				this.MaxChange = MaxChange;
+				this.BuildGroup = BuildGroup;
+			}
 		}
 
 		class PerforceChangeDetailsWithDescribeRecord : PerforceChangeDetails
 		{
-			public PerforceDescribeRecord DescribeRecord;
+			public DescribeRecord DescribeRecord;
 
-			public PerforceChangeDetailsWithDescribeRecord(PerforceDescribeRecord DescribeRecord)
+			public PerforceChangeDetailsWithDescribeRecord(DescribeRecord DescribeRecord)
 				: base(DescribeRecord)
 			{
 				this.DescribeRecord = DescribeRecord;
 			}
 		}
 
-		class PerforceWorkerThread : IDisposable
+		class PerforceWorkerTask : IDisposable
 		{
-			readonly PerforceConnection Perforce;
+			readonly IPerforceSettings PerforceSettings;
 			readonly string Filter;
-			readonly Action<PerforceChangeRange> OnUpdateChanges;
-			readonly Action<PerforceChangeSummary> OnUpdateChangeMetadata;
-			readonly TextWriter Log;
+			Action<PerforceChangeRange>? OnUpdateChanges;
+			Action<ChangesRecord>? OnUpdateChangeMetadata;
+			readonly ILogger Logger;
 			readonly object LockObject = new object();
-			Thread WorkerThread;
-			bool bTerminate;
-			AutoResetEvent RefreshEvent;
+			readonly SynchronizationContext SynchronizationContext;
+			Task? WorkerTask;
+			CancellationTokenSource CancellationSource;
+			AsyncEvent RefreshEvent;
 			List<PerforceChangeRange> Requests = new List<PerforceChangeRange>();
 			Dictionary<int, PerforceChangeDetailsWithDescribeRecord> ChangeNumberToDetails = new Dictionary<int, PerforceChangeDetailsWithDescribeRecord>();
 
-			public PerforceWorkerThread(PerforceConnection Perforce, string Filter, Action<PerforceChangeRange> OnUpdateChanges, Action<PerforceChangeSummary> OnUpdateChangeMetadata, TextWriter Log)
+			public PerforceWorkerTask(IPerforceSettings PerforceSettings, string Filter, Action<PerforceChangeRange>? OnUpdateChanges, Action<ChangesRecord>? OnUpdateChangeMetadata, ILogger Logger)
 			{
-				this.Perforce = Perforce;
+				this.PerforceSettings = PerforceSettings;
 				this.Filter = Filter;
 				this.OnUpdateChanges = OnUpdateChanges;
 				this.OnUpdateChangeMetadata = OnUpdateChangeMetadata;
-				this.Log = Log;
+				this.Logger = Logger;
+				this.SynchronizationContext = SynchronizationContext.Current!;
 
-				RefreshEvent = new AutoResetEvent(true);
-				WorkerThread = new Thread(() => Run());
-				WorkerThread.Start();
+				RefreshEvent = new AsyncEvent();
+				CancellationSource = new CancellationTokenSource();
+				WorkerTask = Task.Run(() => RunAsync(CancellationSource.Token));
 			}
 
 			public void AddRequest(PerforceChangeRange Range)
@@ -85,42 +111,29 @@ namespace UnrealGameSync
 
 			public void Dispose()
 			{
-				if(WorkerThread != null)
+				if (WorkerTask != null)
 				{
-					bTerminate = true;
-					RefreshEvent.Set();
-					if(!WorkerThread.Join(100))
-					{
-						WorkerThread.Interrupt();
-						WorkerThread.Join();
-					}
-					WorkerThread = null;
-				}
-				if(RefreshEvent != null)
-				{
-					RefreshEvent.Dispose();
-					RefreshEvent = null;
-				}
-			}
+					OnUpdateChanges = null;
+					OnUpdateChangeMetadata = null;
 
-			void Run()
-			{
-				try
-				{
-					RunInternal();
-				}
-				catch (ThreadInterruptedException)
-				{
+					CancellationSource.Cancel();
+
+					WorkerTask.ContinueWith(x => CancellationSource.Dispose());
+					WorkerTask = null;
 				}
 			}
 
-			void RunInternal()
+			async Task RunAsync(CancellationToken CancellationToken)
 			{
+				using IPerforceConnection Perforce = await PerforceConnection.CreateAsync(PerforceSettings, Logger);
+
 				List<PerforceChangeRange> CompletedRequests = new List<PerforceChangeRange>(); 
-				while(!bTerminate)
+				while(!CancellationToken.IsCancellationRequested)
 				{
+					Task RefreshTask = RefreshEvent.Task;
+
 					// Check if there's a request in the queue
-					PerforceChangeRange NextRequest = null;
+					PerforceChangeRange? NextRequest = null;
 					lock(LockObject)
 					{
 						if(Requests.Count > 0)
@@ -135,10 +148,10 @@ namespace UnrealGameSync
 					{
 						string RangeFilter = String.Format("{0}@{1},{2}", Filter, NextRequest.MinChange, (NextRequest.MaxChange == -1)? "now" : NextRequest.MaxChange.ToString());
 
-						List<PerforceChangeSummary> NewChanges;
-						if(Perforce.FindChanges(RangeFilter, -1, out NewChanges, Log))
+						PerforceResponseList<ChangesRecord> NewChanges = await Perforce.TryGetChangesAsync(ChangesOptions.None, -1, ChangeStatus.Submitted, RangeFilter, CancellationToken);
+						if(NewChanges.Succeeded)
 						{
-							NextRequest.Changes = NewChanges;
+							NextRequest.Changes = NewChanges.Data;
 							CompletedRequests.Add(NextRequest);
 						}
 						else
@@ -146,51 +159,50 @@ namespace UnrealGameSync
 							NextRequest.ErrorMessage = "Unable to fetch changes. Check your connection settings and try again.";
 						}
 
-						OnUpdateChanges(NextRequest);
+						SynchronizationContext.Post(x => OnUpdateChanges?.Invoke(NextRequest), null);
 						continue;
 					}
 
 					// Figure out which changes to fetch
-					List<PerforceChangeSummary> DescribeChanges;
+					List<ChangesRecord> DescribeChanges;
 					lock(LockObject)
 					{
 						DescribeChanges = CompletedRequests.SelectMany(x => x.Changes).Where(x => !ChangeNumberToDetails.ContainsKey(x.Number)).ToList();
 					}
 
 					// Fetch info on each individual change
-					foreach(PerforceChangeSummary DescribeChange in DescribeChanges)
+					foreach(ChangesRecord DescribeChange in DescribeChanges)
 					{
 						lock(LockObject)
 						{
-							if(bTerminate || Requests.Count > 0)
+							if(CancellationToken.IsCancellationRequested || Requests.Count > 0)
 							{
 								break;
 							}
 						}
 
-						PerforceDescribeRecord Record;
-						if(Perforce.Describe(DescribeChange.Number, out Record, Log))
+						PerforceResponse<DescribeRecord> Response = await Perforce.TryDescribeAsync(DescribeChange.Number);
+						if (Response.Succeeded)
 						{
+							DescribeRecord Record = Response.Data;
 							lock(LockObject)
 							{
-								ChangeNumberToDetails[Record.ChangeNumber] = new PerforceChangeDetailsWithDescribeRecord(Record);
-								if((Record.ChangeNumber & 3) == 0)
-								{
-									ChangeNumberToDetails[Record.ChangeNumber].bContainsCode = true;
-									ChangeNumberToDetails[Record.ChangeNumber].bContainsContent = true;
-								}
+								ChangeNumberToDetails[Record.Number] = new PerforceChangeDetailsWithDescribeRecord(Record);
 							}
 						}
 
-						OnUpdateChangeMetadata(DescribeChange);
+						SynchronizationContext.Post(x => OnUpdateChangeMetadata?.Invoke(DescribeChange), null);
 					}
 
 					// Wait for something to change
-					RefreshEvent.WaitOne();
+					using (IDisposable _ = CancellationToken.Register(() => RefreshEvent.Set()))
+					{
+						await RefreshTask;
+					}
 				}
 			}
 
-			public bool TryGetChangeDetails(int ChangeNumber, out PerforceChangeDetailsWithDescribeRecord Details)
+			public bool TryGetChangeDetails(int ChangeNumber, [NotNullWhen(true)] out PerforceChangeDetailsWithDescribeRecord? Details)
 			{
 				lock(LockObject)
 				{
@@ -255,34 +267,34 @@ namespace UnrealGameSync
 		IssueData Issue;
 		List<IssueBuildData> IssueBuilds;
 		List<IssueDiagnosticData> Diagnostics;
-		PerforceConnection Perforce;
+		IPerforceSettings PerforceSettings;
 		TimeSpan? ServerTimeOffset;
-		PerforceWorkerThread PerforceWorker;
-		TextWriter Log;
+		PerforceWorkerTask? PerforceWorker;
+		IServiceProvider ServiceProvider;
 		SynchronizationContext MainThreadSynchronizationContext;
-		string SelectedStream;
+		string? SelectedStream;
 		List<PerforceChangeRange> SelectedStreamRanges = new List<PerforceChangeRange>();
 		bool bIsDisposing;
-		Font BoldFont;
-		PerforceChangeSummary ContextMenuChange;
-		string LastOwner;
-		string LastDetailsText;
-		System.Windows.Forms.Timer UpdateTimer;
+		Font? BoldFont;
+		ChangesRecord? ContextMenuChange;
+		string? LastOwner;
+		string? LastDetailsText;
+		System.Windows.Forms.Timer? UpdateTimer;
 		StatusElementResources StatusElementResources;
 
-		IssueDetailsWindow(IssueMonitor IssueMonitor, IssueData Issue, List<IssueBuildData> IssueBuilds, List<IssueDiagnosticData> Diagnostics, string ServerAndPort, string UserName, TimeSpan? ServerTimeOffset, TextWriter Log, string CurrentStream)
+		IssueDetailsWindow(IssueMonitor IssueMonitor, IssueData Issue, List<IssueBuildData> IssueBuilds, List<IssueDiagnosticData> Diagnostics, IPerforceSettings PerforceSettings, TimeSpan? ServerTimeOffset, IServiceProvider ServiceProvider, string CurrentStream)
 		{
 			this.IssueMonitor = IssueMonitor;
 			this.Issue = Issue;
 			this.IssueBuilds = IssueBuilds;
 			this.Diagnostics = Diagnostics;
-			this.Perforce = new PerforceConnection(UserName, null, ServerAndPort);
+			this.PerforceSettings = PerforceSettings;
 			this.ServerTimeOffset = ServerTimeOffset;
-			this.Log = Log;
+			this.ServiceProvider = ServiceProvider;
 
 			IssueMonitor.AddRef();
 
-			MainThreadSynchronizationContext = SynchronizationContext.Current;
+			MainThreadSynchronizationContext = SynchronizationContext.Current!;
 
 			InitializeComponent();
 
@@ -295,15 +307,18 @@ namespace UnrealGameSync
 
 			BuildListView.SmallImageList = new ImageList(){ ImageSize = new Size(1, 20) };
 
-			System.Reflection.PropertyInfo DoubleBufferedProperty = typeof(Control).GetProperty("DoubleBuffered", System.Reflection.BindingFlags.NonPublic | System.Reflection.BindingFlags.Instance);
+			System.Reflection.PropertyInfo DoubleBufferedProperty = typeof(Control).GetProperty("DoubleBuffered", System.Reflection.BindingFlags.NonPublic | System.Reflection.BindingFlags.Instance)!;
 			DoubleBufferedProperty.SetValue(BuildListView, true, null); 
 
 			using(Graphics Graphics = Graphics.FromHwnd(IntPtr.Zero))
 			{
 				float DpiScaleX = Graphics.DpiX / 96.0f;
-				foreach(ColumnHeader Column in BuildListView.Columns)
+				foreach(ColumnHeader? Column in BuildListView.Columns)
 				{
-					Column.Width = (int)(Column.Width * DpiScaleX);
+					if (Column != null)
+					{
+						Column.Width = (int)(Column.Width * DpiScaleX);
+					}
 				}
 			}
 
@@ -340,7 +355,7 @@ namespace UnrealGameSync
 				IssueMonitor.StopTracking(Issue.Id);
 				IssueMonitor.OnIssuesChanged -= OnUpdateIssuesAsync;
 				IssueMonitor.Release();
-				IssueMonitor = null;
+				IssueMonitor = null!;
 			}
 
 			base.Dispose(disposing);
@@ -370,7 +385,7 @@ namespace UnrealGameSync
 			}
 		}
 
-		private void UpdateTimer_Tick(object sender, EventArgs e)
+		private void UpdateTimer_Tick(object? sender, EventArgs e)
 		{
 			UpdateBuildList();
 			StopUpdateTimer();
@@ -612,7 +627,7 @@ namespace UnrealGameSync
 			}
 		}
 
-		private void IssueDetailsWindow_Disposed(object sender, EventArgs e)
+		private void IssueDetailsWindow_Disposed(object? sender, EventArgs e)
 		{
 			bIsDisposing = true;
 			DestroyWorker();
@@ -631,12 +646,12 @@ namespace UnrealGameSync
 			if(!Range.bExpanded)
 			{
 				Range.bExpanded = true;
-				PerforceWorker.AddRequest(Range);
+				PerforceWorker!.AddRequest(Range);
 				UpdateBuildList();
 			}
 		}
 
-		bool FilterMatch(Regex FilterRegex, PerforceChangeSummary Summary)
+		bool FilterMatch(Regex FilterRegex, ChangesRecord Summary)
 		{
 			if(FilterRegex.IsMatch(Summary.User))
 			{
@@ -649,7 +664,7 @@ namespace UnrealGameSync
 			return false;
 		}
 
-		bool FilterMatch(Regex FilterRegex, PerforceDescribeRecord DescribeRecord)
+		bool FilterMatch(Regex FilterRegex, DescribeRecord DescribeRecord)
 		{
 			if(FilterRegex.IsMatch(DescribeRecord.User))
 			{
@@ -672,7 +687,7 @@ namespace UnrealGameSync
 			BuildListView.BeginUpdate();
 
 			// Capture the initial selection
-			object PrevSelection = null;
+			object? PrevSelection = null;
 			if(BuildListView.SelectedItems.Count > 0)
 			{
 				PrevSelection = BuildListView.SelectedItems[0].Tag;
@@ -726,10 +741,10 @@ namespace UnrealGameSync
 					}
 					else
 					{
-						foreach(PerforceChangeSummary Change in Range.Changes)
+						foreach(ChangesRecord Change in Range.Changes)
 						{
-							PerforceChangeDetailsWithDescribeRecord Details;
-							PerforceWorker.TryGetChangeDetails(Change.Number, out Details);
+							PerforceChangeDetailsWithDescribeRecord? Details;
+							PerforceWorker!.TryGetChangeDetails(Change.Number, out Details);
 
 							if(FilterRegexes.Count > 0 || bOnlyShowCodeChanges || bOnlyShowContentChanges)
 							{
@@ -770,7 +785,7 @@ namespace UnrealGameSync
 
 							Item.SubItems.Add(Change.Number.ToString());
 
-							DateTime DisplayTime = Change.Date;
+							DateTime DisplayTime = Change.Time;
 							if (ServerTimeOffset.HasValue)
 							{
 								DisplayTime = (DisplayTime - ServerTimeOffset.Value).ToLocalTime();
@@ -844,7 +859,7 @@ namespace UnrealGameSync
 			BuildListView.Invalidate();
 		}
 
-		BuildGroup SelectedBuildGroup;
+		BuildGroup? SelectedBuildGroup;
 
 		void ShowJobContextMenu(Point Point, BuildGroup BuildGroup)
 		{
@@ -868,7 +883,7 @@ namespace UnrealGameSync
 			JobContextMenu.Show(BuildListView, Point, ToolStripDropDownDirection.BelowRight);
 		}
 
-		void UpdateChangeTypeWidget(StatusLineListViewWidget TypeWidget, PerforceChangeDetails Details)
+		void UpdateChangeTypeWidget(StatusLineListViewWidget TypeWidget, PerforceChangeDetails? Details)
 		{
 			TypeWidget.Line.Clear();
 			if(Details == null)
@@ -890,7 +905,7 @@ namespace UnrealGameSync
 
 		void CreateWorker()
 		{
-			string NewSelectedStream = StreamComboBox.SelectedItem as string;
+			string? NewSelectedStream = StreamComboBox.SelectedItem as string;
 			if(SelectedStream != NewSelectedStream)
 			{
 				DestroyWorker();
@@ -907,23 +922,16 @@ namespace UnrealGameSync
 					int MaxChange = -1;
 					foreach(IGrouping<string, IssueBuildData> Group in IssueBuilds.Where(x => x.Stream == SelectedStream).OrderByDescending(x => x.Change).ThenByDescending(x => x.JobUrl).GroupBy(x => x.JobUrl))
 					{
-						BuildGroup BuildGroup = new BuildGroup();
-						BuildGroup.JobName = Group.First().JobName;
-						BuildGroup.JobUrl = Group.Key;
-						BuildGroup.Change = Group.First().Change;
-						BuildGroup.Builds = Group.ToList();
-						BuildGroup.Outcome = Group.Any(x => x.Outcome == IssueBuildOutcome.Error)? IssueBuildOutcome.Error : Group.Any(x => x.Outcome == IssueBuildOutcome.Warning)? IssueBuildOutcome.Warning : IssueBuildOutcome.Success;
-
-						PerforceChangeRange Range = new PerforceChangeRange();
-						Range.MinChange = BuildGroup.Change + 1;
-						Range.MaxChange = MaxChange;
-						Range.BuildGroup = BuildGroup;
+						IssueBuildOutcome Outcome = Group.Any(x => x.Outcome == IssueBuildOutcome.Error) ? IssueBuildOutcome.Error : Group.Any(x => x.Outcome == IssueBuildOutcome.Warning) ? IssueBuildOutcome.Warning : IssueBuildOutcome.Success;
+						BuildGroup BuildGroup = new BuildGroup(Group.First().JobName, Group.Key, Group.First().Change, Outcome, Group.ToList());
+
+						PerforceChangeRange Range = new PerforceChangeRange(BuildGroup.Change + 1, MaxChange, BuildGroup);
 						SelectedStreamRanges.Add(Range);
 
 						MaxChange = BuildGroup.Change;
 					}
 
-					PerforceWorker = new PerforceWorkerThread(Perforce, String.Format("{0}/...", SelectedStream), OnRequestCompleteAsync, UpdateChangeMetadataAsync, Log);
+					PerforceWorker = new PerforceWorkerTask(PerforceSettings, String.Format("{0}/...", SelectedStream), OnRequestComplete, UpdateChangeMetadata, ServiceProvider.GetRequiredService<ILogger<PerforceWorkerTask>>());
 
 					UpdateBuildList();
 
@@ -961,22 +969,12 @@ namespace UnrealGameSync
 			MainThreadSynchronizationContext.Post((o) => { if(!bIsDisposing){ OnUpdateIssues(); } }, null);
 		}
 
-		void OnUpdateError(Exception Ex)
-		{
-			MessageBox.Show(String.Format("Unable to update database.\n\n{0}", Ex));
-		}
-
-		void OnUpdateErrorAsync(Exception Ex)
-		{
-			MainThreadSynchronizationContext.Post((o) => { if(!bIsDisposing){ OnUpdateError(Ex); } }, null);
-		}
-
 		void DestroyWorker()
 		{
 			if(PerforceWorker != null)
 			{
 				PerforceWorker.Dispose();
-				PerforceWorker = null;
+				PerforceWorker = null!;
 			}
 		}
 
@@ -985,12 +983,7 @@ namespace UnrealGameSync
 			UpdateBuildList();
 		}
 
-		void OnRequestCompleteAsync(PerforceChangeRange Range)
-		{
-			MainThreadSynchronizationContext.Post((o) => { if(!bIsDisposing) { OnRequestComplete(Range); } }, null);
-		}
-
-		void UpdateChangeMetadata(PerforceChangeSummary Change)
+		void UpdateChangeMetadata(ChangesRecord Change)
 		{
 			if(FilterTextBox.Text.Length > 0 || FilterTypeComboBox.SelectedIndex != 0)
 			{
@@ -998,12 +991,12 @@ namespace UnrealGameSync
 			}
 			else
 			{
-				foreach(ListViewItem Item in BuildListView.Items)
+				foreach(ListViewItem? Item in BuildListView.Items)
 				{
-					if(Item.Tag == Change)
+					if(Item != null && Item.Tag == Change)
 					{
-						PerforceChangeDetailsWithDescribeRecord Details;
-						PerforceWorker.TryGetChangeDetails(Change.Number, out Details);
+						PerforceChangeDetailsWithDescribeRecord? Details;
+						PerforceWorker!.TryGetChangeDetails(Change.Number, out Details);
 
 						StatusLineListViewWidget TypeWidget = (StatusLineListViewWidget)Item.SubItems[TypeHeader.Index].Tag;
 						UpdateChangeTypeWidget(TypeWidget, Details);
@@ -1015,14 +1008,9 @@ namespace UnrealGameSync
 			}
 		}
 
-		void UpdateChangeMetadataAsync(PerforceChangeSummary Change)
-		{
-			MainThreadSynchronizationContext.Post((o) => { if(!bIsDisposing) { UpdateChangeMetadata(Change); } }, null);
-		}
-
 		static List<IssueDetailsWindow> ExistingWindows = new List<IssueDetailsWindow>();
 
-		class UpdateIssueDetailsTask : IModalTask
+		class UpdateIssueDetailsTask
 		{
 			string ApiUrl;
 			long IssueId;
@@ -1035,30 +1023,23 @@ namespace UnrealGameSync
 				this.Diagnostics = Diagnostics;
 			}
 
-			public bool Run(out string ErrorMessage)
+			public async Task RunAsync(CancellationToken CancellationToken)
 			{
-				Diagnostics.AddRange(RESTApi.GET<List<IssueDiagnosticData>>(ApiUrl, String.Format("issues/{0}/diagnostics", IssueId)));
-
-				ErrorMessage = null;
-				return true;
+				Diagnostics.AddRange(await RESTApi.GetAsync<List<IssueDiagnosticData>>($"{ApiUrl}/api/issues/{IssueId}/diagnostics", CancellationToken));
 			}
 		}
 
-		public static void Show(Form Owner, IssueMonitor IssueMonitor, string ServerAndPort, string UserName, TimeSpan? ServerTimeOffset, IssueData Issue, TextWriter Log, string CurrentStream)
+		public static void Show(Form Owner, IssueMonitor IssueMonitor, IPerforceSettings PerforceSettings, TimeSpan? ServerTimeOffset, IssueData Issue, IServiceProvider ServiceProvider, string CurrentStream)
 		{
-			List<IssueBuildData> IssueBuilds = new List<IssueBuildData>();
-			try
-			{
-				IssueBuilds = RESTApi.GET<List<IssueBuildData>>(IssueMonitor.ApiUrl, String.Format("issues/{0}/builds", Issue.Id));
-			}
-			catch (Exception Ex)
+			Func<CancellationToken, Task<List<IssueBuildData>>> Func = x => RESTApi.GetAsync<List<IssueBuildData>>($"{IssueMonitor.ApiUrl}/api/issues/{Issue.Id}/builds", x);
+			ModalTask<List<IssueBuildData>>? IssueBuilds = ModalTask.Execute(Owner, "Querying issue", "Querying issue, please wait...", Func);
+			if (IssueBuilds != null && IssueBuilds.Succeeded)
 			{
-				MessageBox.Show(Owner, Ex.ToString(), String.Format("Error querying builds for issue {0}", Issue.Id), MessageBoxButtons.OK);
+				Show(Owner, IssueMonitor, PerforceSettings, ServerTimeOffset, Issue, IssueBuilds.Result, ServiceProvider, CurrentStream);
 			}
-			Show(Owner, IssueMonitor, ServerAndPort, UserName, ServerTimeOffset, Issue, IssueBuilds, Log, CurrentStream);
 		}
 
-		public static void Show(Form Owner, IssueMonitor IssueMonitor, string ServerAndPort, string UserName, TimeSpan? ServerTimeOffset, IssueData Issue, List<IssueBuildData> IssueBuilds, TextWriter Log, string CurrentStream)
+		public static void Show(Form Owner, IssueMonitor IssueMonitor, IPerforceSettings PerforceSettings, TimeSpan? ServerTimeOffset, IssueData Issue, List<IssueBuildData> IssueBuilds, IServiceProvider ServiceProvider, string CurrentStream)
 		{
 			IssueDetailsWindow Window = ExistingWindows.FirstOrDefault(x => x.IssueMonitor == IssueMonitor && x.Issue.Id == Issue.Id);
 			if(Window == null)
@@ -1066,12 +1047,13 @@ namespace UnrealGameSync
 				List<IssueDiagnosticData> Diagnostics = new List<IssueDiagnosticData>();
 
 				UpdateIssueDetailsTask Task = new UpdateIssueDetailsTask(IssueMonitor.ApiUrl, Issue.Id, Diagnostics);
-				if(!ModalTask.ExecuteAndShowError(Owner, Task, "Fetching data", "Fetching data, please wait..."))
+				ModalTask? TaskResult = ModalTask.Execute(Owner, "Fetching data", "Fetching data, please wait...", Task.RunAsync);
+				if (TaskResult == null || !TaskResult.Succeeded)
 				{
 					return;
 				}
 
-				Window = new IssueDetailsWindow(IssueMonitor, Issue, IssueBuilds, Diagnostics, ServerAndPort, UserName, ServerTimeOffset, Log, CurrentStream);
+				Window = new IssueDetailsWindow(IssueMonitor, Issue, IssueBuilds, Diagnostics, PerforceSettings, ServerTimeOffset, ServiceProvider, CurrentStream);
 				Window.Owner = Owner;
 				if (Owner.Visible && Owner.WindowState != FormWindowState.Minimized)
 				{
@@ -1164,14 +1146,14 @@ namespace UnrealGameSync
 			{
 				BuildListView.DrawCustomSubItem(e.Graphics, e.SubItem);
 			}
-			else if(e.Item.Tag is PerforceChangeSummary)
+			else if(e.Item.Tag is ChangesRecord)
 			{
-				PerforceChangeSummary Change = (PerforceChangeSummary)e.Item.Tag;
+				ChangesRecord Change = (ChangesRecord)e.Item.Tag;
 
 				Font ChangeFont = BuildListView.Font;
 				if(Issue.Owner != null && String.Compare(Change.User, Issue.Owner, StringComparison.OrdinalIgnoreCase) == 0)
 				{
-					ChangeFont = BoldFont;
+					ChangeFont = BoldFont!;
 				}
 
 				if(e.ColumnIndex == ChangeHeader.Index)
@@ -1288,7 +1270,7 @@ namespace UnrealGameSync
 				ListViewHitTestInfo HitTest = BuildListView.HitTest(Args.Location);
 				if(HitTest.Item != null && HitTest.Item.Tag != null)
 				{
-					ContextMenuChange = HitTest.Item.Tag as PerforceChangeSummary; 
+					ContextMenuChange = HitTest.Item.Tag as ChangesRecord; 
 					if(ContextMenuChange != null)
 					{
 						BuildListContextMenu_Assign.Text = String.Format("Assign to {0}", ContextMenuChange.User);
@@ -1300,12 +1282,15 @@ namespace UnrealGameSync
 
 		private void BuildListContextMenu_MoreInfo_Click(object sender, EventArgs e)
 		{
-			Utility.SpawnP4VC(String.Format("{0} change {1}", Perforce.GetConnectionOptions(), ContextMenuChange.Number));
+			if (ContextMenuChange != null)
+			{
+				Utility.SpawnP4VC(String.Format("{0} change {1}", PerforceSettings.GetArgumentsForExternalProgram(true), ContextMenuChange.Number));
+			}
 		}
 
 		private void BuildListContextMenu_Blame_Click(object sender, EventArgs e)
 		{
-			if(ContextMenuChange != null)
+			if(ContextMenuChange != null && ContextMenuChange.User != null)
 			{
 				AssignToUser(ContextMenuChange.User);
 			}
@@ -1317,14 +1302,14 @@ namespace UnrealGameSync
 			Update.Id = Issue.Id;
 			Update.Owner = User;
 			Update.FixChange = 0;
-			if(String.Compare(User, Perforce.UserName, StringComparison.OrdinalIgnoreCase) == 0)
+			if(String.Compare(User, PerforceSettings.UserName, StringComparison.OrdinalIgnoreCase) == 0)
 			{
 				Update.NominatedBy = "";
 				Update.Acknowledged = true;
 			}
 			else
 			{
-				Update.NominatedBy = Perforce.UserName;
+				Update.NominatedBy = PerforceSettings.UserName;
 				Update.Acknowledged = false;
 			}
 			IssueMonitor.PostUpdate(Update);
@@ -1339,8 +1324,8 @@ namespace UnrealGameSync
 
 		private void AssignToOtherBtn_Click(object sender, EventArgs e)
 		{
-			string SelectedUserName;
-			if(SelectUserWindow.ShowModal(this, Perforce, Log, out SelectedUserName))
+			string? SelectedUserName;
+			if(SelectUserWindow.ShowModal(this, PerforceSettings, ServiceProvider, out SelectedUserName))
 			{
 				AssignToUser(SelectedUserName);
 			}
@@ -1362,7 +1347,7 @@ namespace UnrealGameSync
 			int FixChangeNumber = Issue.FixChange;
 			if(FixChangeNumber == 0)
 			{
-				if(IssueFixedWindow.ShowModal(this, Perforce, ref FixChangeNumber))
+				if(IssueFixedWindow.ShowModal(this, PerforceSettings, ServiceProvider, ref FixChangeNumber))
 				{
 					IssueUpdateData Update = new IssueUpdateData();
 					Update.Id = Issue.Id;
@@ -1390,19 +1375,25 @@ namespace UnrealGameSync
 
 		private void JobContextMenu_ShowError_Click(object sender, EventArgs e)
 		{
-			foreach(IssueBuildData Build in SelectedBuildGroup.Builds)
+			if (SelectedBuildGroup != null)
 			{
-				if(Build.ErrorUrl != null)
+				foreach (IssueBuildData Build in SelectedBuildGroup.Builds)
 				{
-					Utility.OpenUrl(Build.ErrorUrl);
-					break;
+					if (Build.ErrorUrl != null)
+					{
+						Utility.OpenUrl(Build.ErrorUrl);
+						break;
+					}
 				}
 			}
 		}
 
 		private void JobContextMenu_ViewJob_Click(object sender, EventArgs e)
 		{
-			Utility.OpenUrl(SelectedBuildGroup.JobUrl);
+			if (SelectedBuildGroup != null)
+			{
+				Utility.OpenUrl(SelectedBuildGroup.JobUrl);
+			}
 		}
 
 		private void BuildListView_MouseUp(object sender, MouseEventArgs e)
@@ -1412,7 +1403,7 @@ namespace UnrealGameSync
 				ListViewHitTestInfo HitTest = BuildListView.HitTest(e.Location);
 				if(HitTest.Item != null)
 				{
-					BuildGroup Group = HitTest.Item.Tag as BuildGroup;
+					BuildGroup? Group = HitTest.Item.Tag as BuildGroup;
 					if(Group != null)
 					{
 						ShowJobContextMenu(e.Location, Group);
diff --git a/Engine/Source/Programs/UnrealGameSync/UnrealGameSync/Forms/IssueFixedWindow.cs b/Engine/Source/Programs/UnrealGameSync/UnrealGameSync/Forms/IssueFixedWindow.cs
--- a/Engine/Source/Programs/UnrealGameSync/UnrealGameSync/Forms/IssueFixedWindow.cs
+++ b/Engine/Source/Programs/UnrealGameSync/UnrealGameSync/Forms/IssueFixedWindow.cs
@@ -1,5 +1,9 @@
 // Copyright Epic Games, Inc. All Rights Reserved.
 
+using EpicGames.Core;
+using EpicGames.Perforce;
+using Microsoft.Extensions.DependencyInjection;
+using Microsoft.Extensions.Logging;
 using System;
 using System.Collections.Generic;
 using System.ComponentModel;
@@ -11,49 +15,55 @@ using System.Threading;
 using System.Threading.Tasks;
 using System.Windows.Forms;
 
+#nullable enable
+
 namespace UnrealGameSync
 {
 	partial class IssueFixedWindow : Form
 	{
 		class FindChangesWorker : Component
 		{
-			public delegate void OnCompleteDelegate(string UserName, List<PerforceDescribeRecord> Changes);
+			public delegate void OnCompleteDelegate(string UserName, List<DescribeRecord> Changes);
 
-			BufferedTextWriter Log = new BufferedTextWriter();
-			PerforceConnection Perforce;
+			IPerforceSettings PerforceSettings;
 			SynchronizationContext MainThreadSyncContext;
-			AutoResetEvent WakeEvent = new AutoResetEvent(false);
-			string RequestedUserName;
-			bool bStopRequested;
-			Thread Thread;
-			OnCompleteDelegate OnComplete;
-
-			public FindChangesWorker(PerforceConnection Perforce, OnCompleteDelegate OnComplete)
+			AsyncEvent WakeEvent;
+			string? RequestedUserName;
+			CancellationTokenSource CancellationSource;
+			Task BackgroundTask;
+			OnCompleteDelegate? OnComplete;
+			ILogger Logger;
+
+			public FindChangesWorker(IPerforceSettings PerforceSettings, OnCompleteDelegate? OnComplete, ILogger Logger)
 			{
-				this.Perforce = Perforce;
-				this.MainThreadSyncContext = SynchronizationContext.Current;
+				this.PerforceSettings = PerforceSettings;
+				this.MainThreadSyncContext = SynchronizationContext.Current!;
+				WakeEvent = new AsyncEvent();
 				this.OnComplete = OnComplete;
+				this.Logger = Logger;
+
+				CancellationSource = new CancellationTokenSource();
+				BackgroundTask = Task.Run(() => DoWork(CancellationSource.Token));
 			}
 
-			public void Start()
+			public void Stop()
 			{
-				if(Thread == null)
-				{
-					bStopRequested = false;
-					Thread = new Thread(DoWork);
-					Thread.Start();
-				}
+				_ = StopAsync();
 			}
 
-			public void Stop()
+			Task StopAsync()
 			{
-				if(Thread != null)
+				Task StopTask = Task.CompletedTask;
+				if(BackgroundTask != null)
 				{
-					WakeEvent.Set();
-					bStopRequested = true;
-					Thread.Join();
-					Thread = null;
+					OnComplete = null;
+
+					CancellationSource.Cancel();
+					StopTask = BackgroundTask.ContinueWith(_ => CancellationSource.Dispose());
+
+					BackgroundTask = null!;
 				}
+				return StopTask;
 			}
 
 			public void FetchChanges(string UserName)
@@ -69,51 +79,50 @@ namespace UnrealGameSync
 				Stop();
 			}
 
-			public void DoWork()
+			public async Task DoWork(CancellationToken CancellationToken)
 			{
+				using IPerforceConnection Perforce = await PerforceConnection.CreateAsync(PerforceSettings, Logger);
+
+				Task WakeTask = WakeEvent.Task;
+				Task CancelTask = Task.Delay(-1, CancellationToken);
 				for(;;)
 				{
-					try
-					{
-						WakeEvent.WaitOne();
-					}
-					catch (ThreadInterruptedException)
-					{
-					}
+					await Task.WhenAny(WakeTask, CancelTask);
 
-					if(bStopRequested)
+					if (CancellationToken.IsCancellationRequested)
 					{
 						break;
 					}
 
-					string UserName = RequestedUserName;
+					WakeTask = WakeEvent.Task;
 
-					List<PerforceDescribeRecord> Descriptions = null;
-
-					List<PerforceChangeSummary> Changes;
-					if(Perforce.FindChanges(new string[]{ "//..." }, UserName, 100, out Changes, Log))
+					string? UserName = RequestedUserName;
+					if (UserName != null)
 					{
-						Perforce.DescribeMultiple(Changes.Select(x => x.Number), out Descriptions, Log);
+						List<ChangesRecord> Changes = await Perforce.GetChangesAsync(ChangesOptions.IncludeTimes, 100, ChangeStatus.Submitted, FileSpecList.Any, CancellationToken);
+						List<DescribeRecord> Descriptions = await Perforce.DescribeAsync(Changes.Select(x => x.Number).ToArray(), CancellationToken);
+
+						MainThreadSyncContext.Post(_ => OnComplete?.Invoke(UserName, Descriptions), null);
 					}
-					
-					MainThreadSyncContext.Post((o) => { OnComplete(UserName, Descriptions); }, null);
 				}
 			}
 		}
 
-		PerforceConnection Perforce;
+		IPerforceSettings PerforceSettings;
 		int ChangeNumber;
 		FindChangesWorker Worker;
+		IServiceProvider ServiceProvider;
 	
-		public IssueFixedWindow(PerforceConnection Perforce, int InitialChangeNumber)
+		public IssueFixedWindow(IPerforceSettings PerforceSettings, int InitialChangeNumber, IServiceProvider ServiceProvider)
 		{
 			InitializeComponent();
 
-			this.Perforce = Perforce;
-			this.Worker = new FindChangesWorker(Perforce, PopulateChanges);
-			components.Add(Worker);
+			this.PerforceSettings = PerforceSettings;
+			this.Worker = new FindChangesWorker(PerforceSettings, PopulateChanges, ServiceProvider.GetRequiredService<ILogger<FindChangesWorker>>());
+			this.ServiceProvider = ServiceProvider;
+			components!.Add(Worker);
 
-			UserNameTextBox.Text = Perforce.UserName;
+			UserNameTextBox.Text = PerforceSettings.UserName;
 			UserNameTextBox.SelectionStart = UserNameTextBox.Text.Length;
 
 			if(InitialChangeNumber != 0)
@@ -150,13 +159,12 @@ namespace UnrealGameSync
 				ChangesListView.Select();
 			}
 
-			Worker.Start();
 			Worker.FetchChanges(UserNameTextBox.Text);
 		}
 
-		public static bool ShowModal(IWin32Window Owner, PerforceConnection Perforce, ref int FixChangeNumber)
+		public static bool ShowModal(IWin32Window Owner, IPerforceSettings Perforce, IServiceProvider ServiceProvider, ref int FixChangeNumber)
 		{
-			using(IssueFixedWindow FixedWindow = new IssueFixedWindow(Perforce, FixChangeNumber))
+			using(IssueFixedWindow FixedWindow = new IssueFixedWindow(Perforce, FixChangeNumber, ServiceProvider))
 			{
 				if(FixedWindow.ShowDialog(Owner) == DialogResult.OK)
 				{
@@ -185,7 +193,7 @@ namespace UnrealGameSync
 			UpdateSelectedChangeAndClose();
 		}
 
-		private void PopulateChanges(string UserName, List<PerforceDescribeRecord> Changes)
+		private void PopulateChanges(string UserName, List<DescribeRecord> Changes)
 		{
 			if(!IsDisposed)
 			{
@@ -193,7 +201,7 @@ namespace UnrealGameSync
 				ChangesListView.Items.Clear();
 				if(Changes != null)
 				{
-					foreach(PerforceDescribeRecord Change in Changes)
+					foreach(DescribeRecord Change in Changes)
 					{
 						if(Change.Description != null && Change.Description.IndexOf("#ROBOMERGE-SOURCE", 0) == -1)
 						{
@@ -216,7 +224,7 @@ namespace UnrealGameSync
 
 							ListViewItem Item = new ListViewItem("");
 							Item.Tag = Change;
-							Item.SubItems.Add(Change.ChangeNumber.ToString());
+							Item.SubItems.Add(Change.Number.ToString());
 							Item.SubItems.Add(Stream);
 							Item.SubItems.Add(Change.Description.Replace('\n', ' '));
 							ChangesListView.Items.Add(Item);
@@ -234,7 +242,7 @@ namespace UnrealGameSync
 				ListViewHitTestInfo HitTest = ChangesListView.HitTest(Args.Location);
 				if(HitTest.Item != null && HitTest.Item.Tag != null)
 				{
-					PerforceDescribeRecord Record = HitTest.Item.Tag as PerforceDescribeRecord; 
+					DescribeRecord? Record = HitTest.Item.Tag as DescribeRecord; 
 					if(Record != null)
 					{
 						ChangesListContextMenu.Tag = Record;
@@ -251,8 +259,8 @@ namespace UnrealGameSync
 
 		private void ChangesContextMenu_MoreInfo_Click(object sender, EventArgs e)
 		{
-			PerforceDescribeRecord Record = (PerforceDescribeRecord)ChangesListContextMenu.Tag;
-			Utility.SpawnP4VC(String.Format("{0} change {1}", Perforce.GetConnectionOptions(), Record.ChangeNumber));
+			DescribeRecord Record = (DescribeRecord)ChangesListContextMenu.Tag;
+			Utility.SpawnP4VC(String.Format("{0} change {1}", PerforceSettings.GetArgumentsForExternalProgram(true), Record.Number));
 		}
 
 		private void ChangeNumberTextBox_Enter(object sender, EventArgs e)
@@ -302,7 +310,7 @@ namespace UnrealGameSync
 			}
 			else
 			{
-				PerforceDescribeRecord Change = (ChangesListView.SelectedItems.Count > 0)? ChangesListView.SelectedItems[0].Tag as PerforceDescribeRecord : null;
+				DescribeRecord? Change = (ChangesListView.SelectedItems.Count > 0)? ChangesListView.SelectedItems[0].Tag as DescribeRecord : null;
 				if(Change == null)
 				{
 					ChangeNumber = 0;
@@ -310,7 +318,7 @@ namespace UnrealGameSync
 				}
 				else
 				{
-					ChangeNumber = Change.ChangeNumber;
+					ChangeNumber = Change.Number;
 					return true;
 				}
 			}
@@ -354,8 +362,8 @@ namespace UnrealGameSync
 
 		private void UserBrowseBtn_Click(object sender, EventArgs e)
 		{
-			string SelectedUserName;
-			if(SelectUserWindow.ShowModal(this, Perforce, new BufferedTextWriter(), out SelectedUserName))
+			string? SelectedUserName;
+			if(SelectUserWindow.ShowModal(this, PerforceSettings, ServiceProvider, out SelectedUserName))
 			{
 				UserNameTextBox.Text = SelectedUserName;
 			}
diff --git a/Engine/Source/Programs/UnrealGameSync/UnrealGameSync/Forms/MainWindow.cs b/Engine/Source/Programs/UnrealGameSync/UnrealGameSync/Forms/MainWindow.cs
--- a/Engine/Source/Programs/UnrealGameSync/UnrealGameSync/Forms/MainWindow.cs
+++ b/Engine/Source/Programs/UnrealGameSync/UnrealGameSync/Forms/MainWindow.cs
@@ -15,6 +15,10 @@ using System.Diagnostics;
 using UnrealGameSync.Forms;
 using System.Text.RegularExpressions;
 using EpicGames.Core;
+using EpicGames.Perforce;
+using System.Threading.Tasks;
+using Microsoft.Extensions.Logging;
+using Microsoft.Extensions.DependencyInjection;
 
 namespace UnrealGameSync
 {
@@ -79,8 +83,9 @@ namespace UnrealGameSync
 		string ApiUrl;
 		DirectoryReference DataFolder;
 		DirectoryReference CacheFolder;
-		PerforceConnection DefaultConnection;
-		LineBasedTextWriter Log;
+		IPerforceSettings DefaultPerforceSettings;
+		IServiceProvider ServiceProvider;
+		ILogger Logger;
 		UserSettings Settings;
 		int TabMenu_TabIdx = -1;
 		int ChangingWorkspacesRefCount;
@@ -100,7 +105,6 @@ namespace UnrealGameSync
 		IMainWindowTabPanel CurrentTabPanel;
 
 		AutomationServer AutomationServer;
-		TextWriter AutomationLog;
 
 		bool bAllowCreatingHandle;
 
@@ -111,10 +115,12 @@ namespace UnrealGameSync
 
 		NetCoreWindow NetCoreWindow;
 
-		public MainWindow(UpdateMonitor InUpdateMonitor, string InApiUrl, DirectoryReference InDataFolder, DirectoryReference InCacheFolder, bool bInRestoreStateOnLoad, string InOriginalExecutableFileName, bool bInUnstable, DetectProjectSettingsResult[] StartupProjects, PerforceConnection InDefaultConnection, LineBasedTextWriter InLog, UserSettings InSettings, string InUri, OIDCTokenManager InOidcTokenManager)
+		public MainWindow(UpdateMonitor InUpdateMonitor, string InApiUrl, DirectoryReference InDataFolder, DirectoryReference InCacheFolder, bool bInRestoreStateOnLoad, string InOriginalExecutableFileName, bool bInUnstable, List<(UserSelectedProjectSettings, ModalTask<WorkspaceSettings>)> StartupTasks, IPerforceSettings InDefaultPerforceSettings, IServiceProvider InServiceProvider, UserSettings InSettings, string InUri, OIDCTokenManager InOidcTokenManager)
 		{
-			Log = InLog;
-			Log.WriteLine("Opening Main Window for {0} projects. Last Project {1}", StartupProjects.Length, InSettings.LastProject);
+			ServiceProvider = InServiceProvider;
+			Logger = ServiceProvider.GetRequiredService<ILogger<MainWindow>>();
+
+			Logger.LogInformation("Opening Main Window for {NumProject} projects. Last Project {LastProject}", StartupTasks.Count, InSettings.LastProject);
 
 			InitializeComponent();
 
@@ -126,8 +132,8 @@ namespace UnrealGameSync
 			bRestoreStateOnLoad = bInRestoreStateOnLoad;
 			OriginalExecutableFileName = InOriginalExecutableFileName;
 			bUnstable = bInUnstable;
-			DefaultConnection = InDefaultConnection;
-			ToolUpdateMonitor = new ToolUpdateMonitor(DefaultConnection, DataFolder, InSettings);
+			DefaultPerforceSettings = InDefaultPerforceSettings;
+			ToolUpdateMonitor = new ToolUpdateMonitor(DefaultPerforceSettings, DataFolder, InSettings, ServiceProvider);
 
 			Settings = InSettings;
 			OIDCTokenManager = InOidcTokenManager;
@@ -147,20 +153,20 @@ namespace UnrealGameSync
 			SetupDefaultControl();
 
 			int SelectTabIdx = -1;
-			foreach(DetectProjectSettingsResult StartupProject in StartupProjects)
+			foreach((UserSelectedProjectSettings Project, ModalTask<WorkspaceSettings> StartupTask) in StartupTasks)
 			{
 				int TabIdx = -1;
-				if(StartupProject.bSucceeded)
+				if (StartupTask.Succeeded)
 				{
-					TabIdx = TryOpenProject(StartupProject.Task, -1, OpenProjectOptions.Quiet);
+					TabIdx = TryOpenProject(StartupTask.Result, -1, OpenProjectOptions.Quiet);
 				}
-				else if(StartupProject.ErrorMessage != null)
+				else if(StartupTask.Failed)
 				{
-					Log.WriteLine("StartupProject Error: {0}", StartupProject.ErrorMessage);
-					CreateErrorPanel(-1, StartupProject.Task.SelectedProject, StartupProject.ErrorMessage);
+					Logger.LogError("StartupProject Error: {Message}", StartupTask.Error);
+					CreateErrorPanel(-1, Project, StartupTask.Error);
 				}
 
-				if (TabIdx != -1 && Settings.LastProject != null && StartupProject.Task.SelectedProject.Equals(Settings.LastProject))
+				if (TabIdx != -1 && Settings.LastProject != null && Project.Equals(Settings.LastProject))
 				{
 					SelectTabIdx = TabIdx;
 				}
@@ -182,8 +188,8 @@ namespace UnrealGameSync
 				Text += $" {Program.GetVersionString()} (UNSTABLE)";
 			}
 
-			AutomationLog = new TimestampLogWriter(new BoundedLogWriter(FileReference.Combine(DataFolder, "Automation.log")));
-			AutomationServer = new AutomationServer(Request => { MainThreadSynchronizationContext.Post(Obj => PostAutomationRequest(Request), null); }, AutomationLog, InUri);
+			ILogger<AutomationServer> AutomationLogger = ServiceProvider.GetRequiredService<ILogger<AutomationServer>>();
+			AutomationServer = new AutomationServer(Request => { MainThreadSynchronizationContext.Post(Obj => PostAutomationRequest(Request), null); }, InUri, AutomationLogger);
 
 			// Allow creating controls from now on
 			TabPanel.ResumeLayout(false);
@@ -191,7 +197,7 @@ namespace UnrealGameSync
 
 			foreach (string DefaultIssueApiUrl in DeploymentSettings.DefaultIssueApiUrls)
 			{
-				DefaultIssueMonitors.Add(CreateIssueMonitor(DefaultIssueApiUrl, InDefaultConnection.UserName));
+				DefaultIssueMonitors.Add(CreateIssueMonitor(DefaultIssueApiUrl, InDefaultPerforceSettings.UserName));
 			}
 
 			bAllowCreatingHandle = true;
@@ -250,7 +256,7 @@ namespace UnrealGameSync
 			}
 			catch(Exception Ex)
 			{
-				Log.WriteLine("Exception running automation request: {0}", Ex);
+				Logger.LogError(Ex, "Exception running automation request");
 				Request.SetOutput(new AutomationRequestOutput(AutomationRequestResult.Invalid));
 			}
 		}
@@ -264,7 +270,7 @@ namespace UnrealGameSync
 			string ProjectPath = Reader.ReadString();
 
 			AutomatedBuildWindow.BuildInfo BuildInfo;
-			if (!AutomatedBuildWindow.ShowModal(this, DefaultConnection, StreamName, ProjectPath, Changelist, Command.ToString(), Settings, Log, out BuildInfo))
+			if (!AutomatedBuildWindow.ShowModal(this, DefaultPerforceSettings, StreamName, ProjectPath, Changelist, Command.ToString(), Settings, ServiceProvider, out BuildInfo))
 			{
 				return new AutomationRequestOutput(AutomationRequestResult.Canceled);
 			}
@@ -312,7 +318,7 @@ namespace UnrealGameSync
 			string ProjectPath = Reader.ReadString();
 
 			AutomatedSyncWindow.WorkspaceInfo WorkspaceInfo;
-			if(!AutomatedSyncWindow.ShowModal(this, DefaultConnection, StreamName, ProjectPath, out WorkspaceInfo, Log))
+			if(!AutomatedSyncWindow.ShowModal(this, DefaultPerforceSettings, StreamName, ProjectPath, out WorkspaceInfo, ServiceProvider))
 			{
 				return new AutomationRequestOutput(AutomationRequestResult.Canceled);
 			}
@@ -348,10 +354,18 @@ namespace UnrealGameSync
 				}
 
 				// Switch the stream
-				PerforceConnection Perforce = new PerforceConnection(WorkspaceInfo.UserName, WorkspaceInfo.WorkspaceName, WorkspaceInfo.ServerAndPort);
-				if (!Perforce.SwitchStream(StreamName, Log))
+				Func<IPerforceConnection, CancellationToken, Task> SwitchTask = async (Connection, CancellationToken) =>
+				{
+					await Connection.SwitchClientToStreamAsync(Connection.Settings.ClientName, StreamName, SwitchClientOptions.None, CancellationToken);
+				};
+
+				PerforceSettings Settings = new PerforceSettings(WorkspaceInfo.ServerAndPort, WorkspaceInfo.UserName);
+				Settings.ClientName = WorkspaceInfo.WorkspaceName;
+
+				ModalTask Result = PerforceModalTask.Execute(Owner, "Please wait", "Switching streams, please wait...", Settings, SwitchTask, Logger, ModalTaskFlags.Quiet);
+				if (Result == null || !Result.Succeeded)
 				{
-					Log.WriteLine("Unable to switch stream");
+					Logger.LogError("Unable to switch stream ({Message})", Result.Error);
 					OutWorkspace = null;
 					return false;
 				}
@@ -362,7 +376,7 @@ namespace UnrealGameSync
 			int TabIdx = TryOpenProject(SelectedProject, -1, OpenProjectOptions.None);
 			if (TabIdx == -1)
 			{
-				Log.WriteLine("Unable to open project");
+				Logger.LogError("Unable to open project");
 				OutWorkspace = null;
 				return false;
 			}
@@ -370,7 +384,7 @@ namespace UnrealGameSync
 			WorkspaceControl Workspace = TabControl.GetTabData(TabIdx) as WorkspaceControl;
 			if (Workspace == null)
 			{
-				Log.WriteLine("Workspace was unable to open");
+				Logger.LogError("Workspace was unable to open");
 				OutWorkspace = null;
 				return false;
 			}
@@ -454,15 +468,23 @@ namespace UnrealGameSync
 						IssueMonitor.AddRef();
 						try
 						{
-							IssueData Issue = RESTApi.GET<IssueData>(IssueMonitor.ApiUrl, String.Format("issues/{0}", IssueId));
-							ExistingWorkspace.ShowIssueDetails(Issue);
-							return new AutomationRequestOutput(AutomationRequestResult.Ok);
-						}
-						catch (Exception Ex)
-						{
-							MessageBox.Show(String.Format("Unable to query issue {0} from {1}\n\n{2}", IssueId, IssueMonitor.ApiUrl, Ex));
-							Log.WriteException(Ex, "Unable to query issue {0} from {1}", IssueId, IssueMonitor.ApiUrl);
-							return new AutomationRequestOutput(AutomationRequestResult.Error);
+							Func<CancellationToken, Task<IssueData>> Func = x => RESTApi.GetAsync<IssueData>($"{IssueMonitor.ApiUrl}/api/issues/{IssueId}", x);
+							ModalTask<IssueData> IssueTask = ModalTask.Execute(this, "Finding Issue", "Querying issue data, please wait...", Func);
+							if (IssueTask == null)
+							{
+								Logger.LogInformation("Operation cancelled");
+								return new AutomationRequestOutput(AutomationRequestResult.Canceled);
+							}
+							else if (IssueTask.Succeeded)
+							{
+								ExistingWorkspace.ShowIssueDetails(IssueTask.Result);
+								return new AutomationRequestOutput(AutomationRequestResult.Ok);
+							}
+							else
+							{
+								Logger.LogError(IssueTask?.Exception, "Unable to query issue {IssueId} from {ApiUrl}: {Error}", IssueId, IssueMonitor.ApiUrl, IssueTask?.Error);
+								return new AutomationRequestOutput(AutomationRequestResult.Error);
+							}
 						}
 						finally
 						{
@@ -586,15 +608,9 @@ namespace UnrealGameSync
 				AutomationServer = null;
 			}
 
-			if(AutomationLog != null)
-			{
-				AutomationLog.Close();
-				AutomationLog = null;
-			}
-
 			if (ToolUpdateMonitor != null)
 			{
-				ToolUpdateMonitor.Close();
+				ToolUpdateMonitor.Dispose();
 				ToolUpdateMonitor = null;
 			}
 
@@ -672,7 +688,7 @@ namespace UnrealGameSync
 
 		private void CreateErrorPanel(int ReplaceTabIdx, UserSelectedProjectSettings Project, string Message)
 		{
-			Log.WriteLine(Message ?? "Unknown error");
+			Logger.LogError("{Error}", Message ?? "Unknown error");
 
 			ErrorPanel ErrorPanel = new ErrorPanel(Project);
 			ErrorPanel.Parent = TabPanel;
@@ -877,7 +893,7 @@ namespace UnrealGameSync
 				TimeSpan IntervalToFirstTick = NextScheduleTime - CurrentTime;
 				ScheduleTimer = new System.Threading.Timer(x => MainThreadSynchronizationContext.Post((o) => { if(!IsDisposed){ ScheduleTimerElapsed(); } }, null), null, IntervalToFirstTick, TimeSpan.FromDays(1));
 
-				Log.WriteLine("Schedule: Started ScheduleTimer for {0} ({1} remaining)", NextScheduleTime, IntervalToFirstTick);
+				Logger.LogInformation("Schedule: Started ScheduleTimer for {Time} ({Time} remaining)", NextScheduleTime, IntervalToFirstTick);
 			}
 		}
 
@@ -887,20 +903,20 @@ namespace UnrealGameSync
 			{
 				ScheduleTimer.Dispose();
 				ScheduleTimer = null;
-				Log.WriteLine("Schedule: Stopped ScheduleTimer");
+				Logger.LogInformation("Schedule: Stopped ScheduleTimer");
 			}
 			StopScheduleSettledTimer();
 		}
 
 		private void ScheduleTimerElapsed()
 		{
-			Log.WriteLine("Schedule: Timer Elapsed");
+			Logger.LogInformation("Schedule: Timer Elapsed");
 
 			// Try to open any missing tabs. 
 			int NumInitialTabs = TabControl.GetTabCount();
 			foreach (UserSelectedProjectSettings ScheduledProject in Settings.ScheduleProjects)
 			{
-				Log.WriteLine("Schedule: Attempting to open {0}", ScheduledProject);
+				Logger.LogInformation("Schedule: Attempting to open {Project}", ScheduledProject);
 				TryOpenProject(ScheduledProject, -1, OpenProjectOptions.Quiet);
 			}
 
@@ -919,7 +935,7 @@ namespace UnrealGameSync
 		{
 			StopScheduleSettledTimer();
 			ScheduleSettledTimer = new System.Threading.Timer(x => MainThreadSynchronizationContext.Post((o) => { if(!IsDisposed){ ScheduleSettledTimerElapsed(); } }, null), null, TimeSpan.FromSeconds(20.0), TimeSpan.FromMilliseconds(-1.0));
-			Log.WriteLine("Schedule: Started ScheduleSettledTimer");
+			Logger.LogInformation("Schedule: Started ScheduleSettledTimer");
 		}
 
 		private void StopScheduleSettledTimer()
@@ -929,22 +945,22 @@ namespace UnrealGameSync
 				ScheduleSettledTimer.Dispose();
 				ScheduleSettledTimer = null;
 
-				Log.WriteLine("Schedule: Stopped ScheduleSettledTimer");
+				Logger.LogInformation("Schedule: Stopped ScheduleSettledTimer");
 			}
 		}
 
 		private void ScheduleSettledTimerElapsed()
 		{
-			Log.WriteLine("Schedule: Starting Sync");
+			Logger.LogInformation("Schedule: Starting Sync");
 			for(int Idx = 0; Idx < TabControl.GetTabCount(); Idx++)
 			{
 				WorkspaceControl Workspace = TabControl.GetTabData(Idx) as WorkspaceControl;
 				if(Workspace != null)
 				{
-					Log.WriteLine("Schedule: Considering {0}", Workspace.SelectedFileName);
+					Logger.LogInformation("Schedule: Considering {File}", Workspace.SelectedFileName);
 					if(Settings.ScheduleAnyOpenProject || Settings.ScheduleProjects.Any(x => x.LocalPath.Equals(Workspace.SelectedProject.LocalPath, StringComparison.OrdinalIgnoreCase)))
 					{
-						Log.WriteLine("Schedule: Starting Sync");
+						Logger.LogInformation("Schedule: Starting Sync");
 						Workspace.ScheduleTimerElapsed();
 					}
 				}
@@ -1015,24 +1031,23 @@ namespace UnrealGameSync
 
 		public void OpenNewProject()
 		{
-			DetectProjectSettingsTask DetectedProjectSettings;
-			if(OpenProjectWindow.ShowModal(this, null, out DetectedProjectSettings, Settings, DataFolder, CacheFolder, DefaultConnection, Log))
+			WorkspaceSettings WorkspaceSettings = OpenProjectWindow.ShowModal(this, null, Settings, DataFolder, CacheFolder, DefaultPerforceSettings, ServiceProvider);
+			if(WorkspaceSettings != null)
 			{
-				int NewTabIdx = TryOpenProject(DetectedProjectSettings, -1, OpenProjectOptions.None);
+				int NewTabIdx = TryOpenProject(WorkspaceSettings, -1, OpenProjectOptions.None);
 				if(NewTabIdx != -1)
 				{
 					TabControl.SelectTab(NewTabIdx);
 					SaveTabSettings();
-					UpdateRecentProjectsList(DetectedProjectSettings);
+					UpdateRecentProjectsList(WorkspaceSettings.SelectedProject);
 				}
-				DetectedProjectSettings.Dispose();
 			}
 		}
 
-		void UpdateRecentProjectsList(DetectProjectSettingsTask DetectedProjectSettings)
+		void UpdateRecentProjectsList(UserSelectedProjectSettings DetectedProjectSettings)
 		{
-			Settings.RecentProjects.RemoveAll(x => x.LocalPath == DetectedProjectSettings.NewSelectedFileName.FullName);
-			Settings.RecentProjects.Insert(0, DetectedProjectSettings.SelectedProject);
+			Settings.RecentProjects.RemoveAll(x => x.LocalPath.Equals(DetectedProjectSettings.LocalPath, StringComparison.OrdinalIgnoreCase));
+			Settings.RecentProjects.Insert(0, DetectedProjectSettings);
 
 			const int MaxRecentProjects = 10;
 			if (Settings.RecentProjects.Count > MaxRecentProjects)
@@ -1078,50 +1093,42 @@ namespace UnrealGameSync
 
 		public void EditSelectedProject(int TabIdx, UserSelectedProjectSettings SelectedProject)
 		{
-			DetectProjectSettingsTask DetectedProjectSettings;
-			if(OpenProjectWindow.ShowModal(this, SelectedProject, out DetectedProjectSettings, Settings, DataFolder, CacheFolder, DefaultConnection, Log))
+			WorkspaceSettings WorkspaceSettings = OpenProjectWindow.ShowModal(this, SelectedProject, Settings, DataFolder, CacheFolder, DefaultPerforceSettings, ServiceProvider);
+			if(WorkspaceSettings != null)
 			{
-				int NewTabIdx = TryOpenProject(DetectedProjectSettings, TabIdx, OpenProjectOptions.None);
+				int NewTabIdx = TryOpenProject(WorkspaceSettings, TabIdx, OpenProjectOptions.None);
 				if(NewTabIdx != -1)
 				{
 					TabControl.SelectTab(NewTabIdx);
 					SaveTabSettings();
-					UpdateRecentProjectsList(DetectedProjectSettings);
+					UpdateRecentProjectsList(WorkspaceSettings.SelectedProject);
 				}
 			}
 		}
 
 		int TryOpenProject(UserSelectedProjectSettings Project, int ReplaceTabIdx, OpenProjectOptions Options = OpenProjectOptions.None)
 		{
-			Log.WriteLine("Detecting settings for {0}", Project);
-			using(DetectProjectSettingsTask DetectProjectSettings = new DetectProjectSettingsTask(Project, DataFolder, CacheFolder, new PrefixedTextWriter("  ", Log)))
-			{
-				PerforceConnection Perforce = Utility.OverridePerforceSettings(DefaultConnection, Project.ServerAndPort, Project.UserName);
+			ILogger<WorkspaceSettings> ProjectLogger = ServiceProvider.GetRequiredService<ILogger<WorkspaceSettings>>();
 
-				string ErrorMessage;
-
-				ModalTaskResult Result;
-				if((Options & OpenProjectOptions.Quiet) != 0)
-				{
-					Result = ModalTask.Execute(this, new QuietDetectProjectSettingsTask(Perforce, DetectProjectSettings, Log), "Opening Project", "Opening project, please wait...", out ErrorMessage);
-				}
-				else
-				{
-					Result = PerforceModalTask.Execute(this, Perforce, DetectProjectSettings, "Opening Project", "Opening project, please wait...", Log, out ErrorMessage);
-				}
+			ModalTaskFlags TaskFlags = ModalTaskFlags.None;
+			if((Options & OpenProjectOptions.Quiet) != 0)
+			{
+				TaskFlags |= ModalTaskFlags.Quiet;
+			}
 
-				if(Result != ModalTaskResult.Succeeded)
-				{
-					CreateErrorPanel(ReplaceTabIdx, Project, ErrorMessage);
-					return -1;
-				}
-				return TryOpenProject(DetectProjectSettings, ReplaceTabIdx, Options);
+			ModalTask<WorkspaceSettings> SettingsTask = PerforceModalTask.Execute(this, "Opening Project", "Opening project, please wait...", DefaultPerforceSettings, (p, c) => WorkspaceSettings.CreateAsync(p, Project, ProjectLogger, c), ProjectLogger, TaskFlags);
+			if(SettingsTask == null || !SettingsTask.Succeeded)
+			{
+				CreateErrorPanel(ReplaceTabIdx, Project, SettingsTask.Error);
+				return -1;
 			}
+
+			return TryOpenProject(SettingsTask.Result, ReplaceTabIdx, Options);
 		}
 
-		int TryOpenProject(DetectProjectSettingsTask ProjectSettings, int ReplaceTabIdx, OpenProjectOptions Options)
+		int TryOpenProject(WorkspaceSettings WorkspaceSettings, int ReplaceTabIdx, OpenProjectOptions Options)
 		{
-			Log.WriteLine("Trying to open project {0}", ProjectSettings.SelectedProject.ToString());
+			Logger.LogInformation("Trying to open project {Project}", WorkspaceSettings.ProjectInfo.ClientFileName);
 
 			// Check that none of the other tabs already have it open
 			for(int TabIdx = 0; TabIdx < TabControl.GetTabCount(); TabIdx++)
@@ -1131,25 +1138,25 @@ namespace UnrealGameSync
 					WorkspaceControl Workspace = TabControl.GetTabData(TabIdx) as WorkspaceControl;
 					if(Workspace != null)
 					{
-						if(Workspace.SelectedFileName == ProjectSettings.NewSelectedFileName)
+						if(Workspace.SelectedFileName == WorkspaceSettings.ProjectInfo.LocalFileName)
 						{
-							Log.WriteLine("  Already open in tab {0}", TabIdx);
+							Logger.LogInformation("  Already open in tab {TabIdx}", TabIdx);
 							if((Options & OpenProjectOptions.Quiet) == 0)
 							{
 								TabControl.SelectTab(TabIdx);
 							}
 							return TabIdx;
 						}
-						else if(ProjectSettings.NewSelectedFileName.IsUnderDirectory(Workspace.BranchDirectoryName))
+						else if(WorkspaceSettings.ProjectInfo.LocalFileName.IsUnderDirectory(Workspace.BranchDirectoryName))
 						{
-							if((Options & OpenProjectOptions.Quiet) == 0 && MessageBox.Show(String.Format("{0} is already open under {1}.\n\nWould you like to close it?", Workspace.SelectedFileName.GetFileNameWithoutExtension(), Workspace.BranchDirectoryName, ProjectSettings.NewSelectedFileName.GetFileNameWithoutExtension()), "Branch already open", MessageBoxButtons.YesNo) == System.Windows.Forms.DialogResult.Yes)
+							if((Options & OpenProjectOptions.Quiet) == 0 && MessageBox.Show(String.Format("{0} is already open under {1}.\n\nWould you like to close it?", Workspace.SelectedFileName.GetFileNameWithoutExtension(), Workspace.BranchDirectoryName, WorkspaceSettings.ProjectInfo.LocalFileName.GetFileNameWithoutExtension()), "Branch already open", MessageBoxButtons.YesNo) == System.Windows.Forms.DialogResult.Yes)
 							{
-								Log.WriteLine("  Another project already open in this workspace, tab {0}. Replacing.", TabIdx);
+								Logger.LogInformation("  Another project already open in this workspace, tab {TabIdx}. Replacing.", TabIdx);
 								TabControl.RemoveTab(TabIdx);
 							}
 							else
 							{
-								Log.WriteLine("  Another project already open in this workspace, tab {0}. Aborting.", TabIdx);
+								Logger.LogInformation("  Another project already open in this workspace, tab {TabIdx}. Aborting.", TabIdx);
 								return -1;
 							}
 						}
@@ -1167,13 +1174,13 @@ namespace UnrealGameSync
 				if(OldWorkspace != null)
 				{
 					OldWorkspace.Hide();
-					TabControl.SetTabData(ReplaceTabIdx, new ErrorPanel(ProjectSettings.SelectedProject));
+					TabControl.SetTabData(ReplaceTabIdx, new ErrorPanel(WorkspaceSettings.SelectedProject));
 					OldWorkspace.Dispose();
 				}
 			}
 
 			// Now that we have the project settings, we can construct the tab
-			WorkspaceControl NewWorkspace = new WorkspaceControl(this, ApiUrl, OriginalExecutableFileName, bUnstable, ProjectSettings, Log, Settings, OIDCTokenManager);
+			WorkspaceControl NewWorkspace = new WorkspaceControl(this, ApiUrl, OriginalExecutableFileName, bUnstable, WorkspaceSettings, ServiceProvider, Settings, OIDCTokenManager);
 			
 			NewWorkspace.Parent = TabPanel;
 			NewWorkspace.Dock = DockStyle.Fill;
@@ -1184,12 +1191,12 @@ namespace UnrealGameSync
 			if(ReplaceTabIdx == -1)
 			{
 				int NewTabIdx = TabControl.InsertTab(-1, NewTabName, NewWorkspace, NewWorkspace.TintColor);
-				Log.WriteLine("  Inserted tab {0}", NewTabIdx);
+				Logger.LogInformation("  Inserted tab {TabIdx}", NewTabIdx);
 				return NewTabIdx;
 			}
 			else
 			{
-				Log.WriteLine("  Replacing tab {0}", ReplaceTabIdx);
+				Logger.LogInformation("  Replacing tab {TabIdx}", ReplaceTabIdx);
 				TabControl.InsertTab(ReplaceTabIdx + 1, NewTabName, NewWorkspace, NewWorkspace.TintColor);
 				TabControl.RemoveTab(ReplaceTabIdx);
 				return ReplaceTabIdx;
@@ -1267,7 +1274,7 @@ namespace UnrealGameSync
 			// if this failes, return something sensible to avoid blank tabs
 			if (string.IsNullOrEmpty(TabName))
 			{
-				Log.WriteLine("No TabName for {0} for setting {1}. Defaulting to client name", Workspace.ClientName, Settings.TabLabels.ToString());
+				Logger.LogInformation("No TabName for {ClientName} for setting {TabSetting}. Defaulting to client name", Workspace.ClientName, Settings.TabLabels.ToString());
 				TabName = Workspace.ClientName;
 			}
 
@@ -1397,7 +1404,7 @@ namespace UnrealGameSync
 
 		public void ModifyApplicationSettings()
 		{
-			bool? bRelaunchUnstable = ApplicationSettingsWindow.ShowModal(this, DefaultConnection, bUnstable, OriginalExecutableFileName, Settings, ToolUpdateMonitor, Log);
+			bool? bRelaunchUnstable = ApplicationSettingsWindow.ShowModal(this, DefaultPerforceSettings, bUnstable, OriginalExecutableFileName, Settings, ToolUpdateMonitor, ServiceProvider.GetRequiredService<ILogger<ApplicationSettingsWindow>>());
 			if(bRelaunchUnstable.HasValue)
 			{
 				UpdateMonitor.TriggerUpdate(UpdateType.UserInitiated, bRelaunchUnstable);
@@ -1644,7 +1651,7 @@ namespace UnrealGameSync
 				FileReference LogFileName = FileReference.Combine(DataFolder, String.Format("IssueMonitor-{0}-{1}.log", ServerId, UserName));
 
 				WorkspaceIssueMonitor = new WorkspaceIssueMonitor();
-				WorkspaceIssueMonitor.IssueMonitor = new IssueMonitor(ApiUrl, UserName, TimeSpan.FromSeconds(60.0), LogFileName);
+				WorkspaceIssueMonitor.IssueMonitor = new IssueMonitor(ApiUrl, UserName, TimeSpan.FromSeconds(60.0), ServiceProvider);
 				WorkspaceIssueMonitor.IssueMonitor.OnIssuesChanged += IssueMonitor_OnUpdateAsync;
 
 				WorkspaceIssueMonitors.Add(WorkspaceIssueMonitor);
diff --git a/Engine/Source/Programs/UnrealGameSync/UnrealGameSync/Forms/ModalTaskWindow.cs b/Engine/Source/Programs/UnrealGameSync/UnrealGameSync/Forms/ModalTaskWindow.cs
--- a/Engine/Source/Programs/UnrealGameSync/UnrealGameSync/Forms/ModalTaskWindow.cs
+++ b/Engine/Source/Programs/UnrealGameSync/UnrealGameSync/Forms/ModalTaskWindow.cs
@@ -16,37 +16,18 @@ namespace UnrealGameSync
 {
 	partial class ModalTaskWindow : Form
 	{
-		IModalTask Task;
-		SynchronizationContext SyncContext;
-		Thread BackgroundThread;
-		ModalTaskResult? PendingResult;
+		Task Task;
+		CancellationTokenSource CancellationSource;
 
-		public ModalTaskResult Result
-		{
-			get { return PendingResult.Value; }
-		}
-
-		public string ErrorMessage
-		{
-			get;
-			private set;
-		}
-
-		public event Action Complete;
-
-		public ModalTaskWindow(IModalTask InTask, string InTitle, string InMessage, FormStartPosition InStartPosition)
+		public ModalTaskWindow(string InTitle, string InMessage, FormStartPosition InStartPosition, Task InTask, CancellationTokenSource InCancellationSource)//, Func<CancellationToken, Task> InTaskFunc)
 		{
 			InitializeComponent();
 
-			Task = InTask;
 			Text = InTitle;
 			MessageLabel.Text = InMessage;
 			StartPosition = InStartPosition;
-
-			SyncContext = SynchronizationContext.Current;
-
-			BackgroundThread = new Thread(x => ThreadProc());
-			BackgroundThread.Start();
+			Task = InTask;
+			CancellationSource = InCancellationSource;
 		}
 
 		public void ShowAndActivate()
@@ -55,52 +36,19 @@ namespace UnrealGameSync
 			Activate();
 		}
 
-		private void ThreadProc()
+		private void ModalTaskWindow_Load(object sender, EventArgs e)
 		{
-			try
-			{
-				string ErrorMessageValue;
-				if (Task.Run(out ErrorMessageValue))
-				{
-					PendingResult = ModalTaskResult.Succeeded;
-				}
-				else
-				{
-					ErrorMessage = ErrorMessageValue;
-					PendingResult = ModalTaskResult.Failed;
-				}
-			}
-			catch(Exception Ex)
-			{
-				ErrorMessage = Ex.ToString();
-				PendingResult = ModalTaskResult.Failed;
-			}
-			SyncContext.Post((o) => ThreadCompleteCallback(), null);
+			SynchronizationContext SyncContext = SynchronizationContext.Current;
+			Task.ContinueWith(x => SyncContext.Post(y => Close(), null));
 		}
 
-		private void ThreadCompleteCallback()
+		private void ModalTaskWindow_FormClosing(object sender, FormClosingEventArgs e)
 		{
-			if(BackgroundThread != null)
+			CancellationSource.Cancel();
+			if (!Task.IsCompleted)
 			{
-				BackgroundThread.Interrupt();
-				BackgroundThread.Join();
-				BackgroundThread = null;
-
-				if(!PendingResult.HasValue)
-				{
-					PendingResult = ModalTaskResult.Aborted;
-				}
-
-				if(Complete != null)
-				{
-					Complete();
-				}
+				e.Cancel = true;
 			}
 		}
-
-		private void ModalTaskWindow_FormClosing(object sender, FormClosingEventArgs e)
-		{
-			ThreadCompleteCallback();
-		}
 	}
 }
diff --git a/Engine/Source/Programs/UnrealGameSync/UnrealGameSync/Forms/NewWorkspaceWindow.cs b/Engine/Source/Programs/UnrealGameSync/UnrealGameSync/Forms/NewWorkspaceWindow.cs
--- a/Engine/Source/Programs/UnrealGameSync/UnrealGameSync/Forms/NewWorkspaceWindow.cs
+++ b/Engine/Source/Programs/UnrealGameSync/UnrealGameSync/Forms/NewWorkspaceWindow.cs
@@ -1,136 +1,119 @@
 // Copyright Epic Games, Inc. All Rights Reserved.
 
+using EpicGames.Core;
+using EpicGames.Perforce;
+using Microsoft.Extensions.DependencyInjection;
+using Microsoft.Extensions.Logging;
 using System;
 using System.Collections.Generic;
 using System.ComponentModel;
 using System.Data;
+using System.Diagnostics.CodeAnalysis;
 using System.Drawing;
 using System.IO;
 using System.Linq;
 using System.Text;
+using System.Threading;
 using System.Threading.Tasks;
 using System.Windows.Forms;
 
+#nullable enable
+
 namespace UnrealGameSync
 {
 	partial class NewWorkspaceWindow : Form
 	{
 		class NewWorkspaceSettings
 		{
-			public string Name;
-			public string Stream;
-			public string RootDir;
+			public string Name { get; }
+			public string Stream { get; }
+			public DirectoryReference RootDir { get; }
+
+			public NewWorkspaceSettings(string Name, string Stream, DirectoryReference RootDir)
+			{
+				this.Name = Name;
+				this.Stream = Stream;
+				this.RootDir = RootDir;
+			}
 		}
 
-		class FindWorkspaceSettingsTask : IPerforceModalTask
+		class FindWorkspaceSettingsTask
 		{
-			string CurrentWorkspaceName;
-			public string CurrentStream;
-			public PerforceInfoRecord Info;
-			public List<PerforceClientRecord> Clients;
+			public InfoRecord Info;
+			public List<ClientsRecord> Clients;
+			public string? CurrentStream;
 
-			public FindWorkspaceSettingsTask(string WorkspaceName)
+			public FindWorkspaceSettingsTask(InfoRecord Info, List<ClientsRecord> Clients, string? CurrentStream)
 			{
-				this.CurrentWorkspaceName = WorkspaceName;
+				this.Info = Info;
+				this.Clients = Clients;
+				this.CurrentStream = CurrentStream;
 			}
 
-			public bool Run(PerforceConnection Perforce, TextWriter Log, out string ErrorMessage)
+			public static async Task<FindWorkspaceSettingsTask> RunAsync(IPerforceConnection Perforce, string CurrentWorkspaceName, CancellationToken CancellationToken)
 			{
-				if(!Perforce.Info(out Info, Log))
-				{
-					ErrorMessage = "Unable to get Perforce info.";
-					return false;
-				}
-				if(!Perforce.FindClients(Info.UserName, out Clients, Log))
-				{
-					ErrorMessage = "Unable to enumerate Perforce clients.";
-					return false;
-				}
+				InfoRecord Info = await Perforce.GetInfoAsync(InfoOptions.ShortOutput, CancellationToken);
+				List<ClientsRecord> Clients = await Perforce.GetClientsAsync(ClientsOptions.None, Perforce.Settings.UserName, CancellationToken);
+
+				string? CurrentStream = null;
 				if(!String.IsNullOrEmpty(CurrentWorkspaceName))
 				{
-					PerforceConnection PerforceClient = new PerforceConnection(Perforce.UserName, CurrentWorkspaceName, Perforce.ServerAndPort);
-					PerforceClient.GetActiveStream(out CurrentStream, Log);
+					CurrentStream = await Perforce.GetCurrentStreamAsync(CancellationToken);
 				}
 
-				ErrorMessage = null;
-				return true;
+				return new FindWorkspaceSettingsTask(Info, Clients, CurrentStream);
 			}
 		}
 
-		class NewWorkspaceTask : IPerforceModalTask
+		static class NewWorkspaceTask
 		{
-			NewWorkspaceSettings Settings;
-			string Owner;
-			string HostName;
-
-			public NewWorkspaceTask(NewWorkspaceSettings Settings, string Owner, string HostName)
-			{
-				this.Settings = Settings;
-				this.Owner = Owner;
-				this.HostName = HostName;
-			}
-
-			public bool Run(PerforceConnection Perforce, TextWriter Log, out string ErrorMessage)
+			public static async Task RunAsync(IPerforceConnection Perforce, NewWorkspaceSettings Settings, string Owner, string HostName, CancellationToken CancellationToken)
 			{
-				bool bExists;
-				if(!Perforce.ClientExists(Settings.Name, out bExists, Log))
+				PerforceResponseList<ClientsRecord> Response = await Perforce.TryGetClientsAsync(ClientsOptions.None, Settings.Name, -1, null, null, CancellationToken);
+				if (!Response.Succeeded)
 				{
-					ErrorMessage = String.Format("Unable to determine if client already exists.\n\n{0}", Log.ToString());
-					return false;
+					throw new UserErrorException($"Unable to determine if client already exists.\n\n{Response[0]}");
 				}
-				if(bExists)
-				{
-					ErrorMessage = String.Format("Client '{0}' already exists.", Settings.Name);
-					return false;
-				}
-
-				PerforceSpec Client = new PerforceSpec();
-				Client.SetField("Client", Settings.Name);
-				Client.SetField("Owner", Owner);
-				Client.SetField("Host", HostName);
-				Client.SetField("Stream", Settings.Stream);
-				Client.SetField("Root", Settings.RootDir);
-				Client.SetField("Options", "rmdir");
-
-				string Message;
-				if(!Perforce.CreateClient(Client, out Message, Log))
+				if (Response.Data.Count > 0)
 				{
-					ErrorMessage = Message;
-					return false;
+					throw new UserErrorException($"Client '{Settings.Name}' already exists.");
 				}
 
-				ErrorMessage = null;
-				return true;
+				ClientRecord Client = new ClientRecord(Settings.Name, Owner, Settings.RootDir.FullName);
+				Client.Host = HostName;
+				Client.Stream = Settings.Stream;
+				Client.Options = ClientOptions.Rmdir;
+				await Perforce.CreateClientAsync(Client, CancellationToken);
 			}
 		}
 
-		PerforceConnection Perforce;
-		PerforceInfoRecord Info;
-		List<PerforceClientRecord> Clients;
-		TextWriter Log;
-		NewWorkspaceSettings Settings;
-		string DefaultRootPath;
+		IPerforceSettings PerforceSettings;
+		InfoRecord Info;
+		List<ClientsRecord> Clients;
+		IServiceProvider ServiceProvider;
+		NewWorkspaceSettings? Settings;
+		DirectoryReference? DefaultRootPath;
 
-		private NewWorkspaceWindow(PerforceConnection Perforce, string ForceStream, string DefaultStream, PerforceInfoRecord Info, List<PerforceClientRecord> Clients, TextWriter Log)
+		private NewWorkspaceWindow(IPerforceSettings PerforceSettings, string? ForceStream, string? DefaultStream, InfoRecord Info, List<ClientsRecord> Clients, IServiceProvider ServiceProvider)
 		{
 			InitializeComponent();
 
-			this.Perforce = Perforce;
+			this.PerforceSettings = PerforceSettings;
 			this.Info = Info;
 			this.Clients = Clients;
-			this.Log = Log;
+			this.ServiceProvider = ServiceProvider;
 
-			Dictionary<string, int> RootPathToCount = new Dictionary<string, int>(StringComparer.InvariantCultureIgnoreCase);
-			foreach(PerforceClientRecord Client in Clients)
+			Dictionary<DirectoryReference, int> RootPathToCount = new Dictionary<DirectoryReference, int>();
+			foreach(ClientsRecord Client in Clients)
 			{
-				if(Client.Host == null || String.Compare(Client.Host, Info.HostName, StringComparison.InvariantCultureIgnoreCase) == 0)
+				if(Client.Host == null || String.Compare(Client.Host, Info.ClientHost, StringComparison.OrdinalIgnoreCase) == 0)
 				{
 					if(!String.IsNullOrEmpty(Client.Root) && Client.Root != ".")
 					{
-						string ParentDir;
+						DirectoryReference? ParentDir;
 						try
 						{
-							ParentDir = Path.GetFullPath(Path.GetDirectoryName(Client.Root));
+							ParentDir = new DirectoryReference(Client.Root);
 						}
 						catch
 						{
@@ -148,7 +131,7 @@ namespace UnrealGameSync
 			}
 
 			int RootPathMaxCount = 0;
-			foreach(KeyValuePair<string, int> RootPathPair in RootPathToCount)
+			foreach(KeyValuePair<DirectoryReference, int> RootPathPair in RootPathToCount)
 			{
 				if(RootPathPair.Value > RootPathMaxCount)
 				{
@@ -178,26 +161,19 @@ namespace UnrealGameSync
 			UpdateRootDirCueBanner();
 		}
 
-		public static bool ShowModal(IWin32Window Owner, PerforceConnection Perforce, string ForceStreamName, string CurrentWorkspaceName, TextWriter Log, out string WorkspaceName)
+		public static bool ShowModal(IWin32Window Owner, IPerforceSettings PerforceSettings, string? ForceStreamName, string CurrentWorkspaceName, IServiceProvider ServiceProvider, [NotNullWhen(true)] out string? WorkspaceName)
 		{
-			FindWorkspaceSettingsTask FindSettings = new FindWorkspaceSettingsTask(CurrentWorkspaceName);
-
-			string ErrorMessage;
-			if(PerforceModalTask.Execute(Owner, Perforce, FindSettings, "Checking settings", "Checking settings, please wait...", Log, out ErrorMessage) != ModalTaskResult.Succeeded)
+			ModalTask<FindWorkspaceSettingsTask>? Task = PerforceModalTask.Execute(Owner, "Checking settings", "Checking settings, please wait...", PerforceSettings, (p, c) => FindWorkspaceSettingsTask.RunAsync(p, CurrentWorkspaceName, c), ServiceProvider.GetRequiredService<ILogger<FindWorkspaceSettingsTask>>());
+			if (Task == null || !Task.Succeeded)
 			{
-				if(!String.IsNullOrEmpty(ErrorMessage))
-				{
-					MessageBox.Show(ErrorMessage);
-				}
-
 				WorkspaceName = null;
 				return false;
 			}
 
-			NewWorkspaceWindow Window = new NewWorkspaceWindow(Perforce, ForceStreamName, FindSettings.CurrentStream, FindSettings.Info, FindSettings.Clients, Log);
+			NewWorkspaceWindow Window = new NewWorkspaceWindow(PerforceSettings, ForceStreamName, Task.Result.CurrentStream, Task.Result.Info, Task.Result.Clients, ServiceProvider);
 			if(Window.ShowDialog(Owner) == DialogResult.OK)
 			{
-				WorkspaceName = Window.Settings.Name;
+				WorkspaceName = Window.Settings!.Name;
 				return true;
 			}
 			else
@@ -221,7 +197,7 @@ namespace UnrealGameSync
 
 		private string GetDefaultWorkspaceName()
 		{
-			string BaseName = Sanitize(String.Format("{0}_{1}_{2}", Info.UserName, Info.HostName, StreamTextBox.Text.Replace('/', '_').Trim('_'))).Trim('_');
+			string BaseName = Sanitize(String.Format("{0}_{1}_{2}", Info.UserName, Info.ClientHost, StreamTextBox.Text.Replace('/', '_').Trim('_'))).Trim('_');
 
 			string Name = BaseName;
 			for(int Idx = 2; Clients.Any(x => x.Name != null && String.Compare(x.Name, Name, StringComparison.InvariantCultureIgnoreCase) == 0); Idx++)
@@ -234,12 +210,12 @@ namespace UnrealGameSync
 		private string GetDefaultWorkspaceRootDir()
 		{
 			string RootDir = "";
-			if(!String.IsNullOrEmpty(DefaultRootPath))
+			if(DefaultRootPath != null)
 			{
 				string Suffix = String.Join("_", StreamTextBox.Text.Split(new char[]{ '/' }, StringSplitOptions.RemoveEmptyEntries).Select(x => Sanitize(x)).Where(x => x.Length > 0));
 				if(Suffix.Length > 0)
 				{
-					RootDir = Path.Combine(DefaultRootPath, Suffix);
+					RootDir = DirectoryReference.Combine(DefaultRootPath, Suffix).FullName;
 				}
 			}
 			return RootDir;
@@ -270,11 +246,11 @@ namespace UnrealGameSync
 
 		private void UpdateOkButton()
 		{
-			NewWorkspaceSettings Settings;
+			NewWorkspaceSettings? Settings;
 			OkBtn.Enabled = TryGetWorkspaceSettings(out Settings);
 		}
 
-		private bool TryGetWorkspaceSettings(out NewWorkspaceSettings Settings)
+		private bool TryGetWorkspaceSettings([NotNullWhen(true)] out NewWorkspaceSettings? Settings)
 		{
 			string NewWorkspaceName = NameTextBox.Text.Trim();
 			if(NewWorkspaceName.Length == 0)
@@ -305,9 +281,10 @@ namespace UnrealGameSync
 				}
 			}
 
+			DirectoryReference NewRootDirRef;
 			try
 			{
-				NewRootDir = Path.GetFullPath(NewRootDir);
+				NewRootDirRef = new DirectoryReference(NewRootDir);
 			}
 			catch
 			{
@@ -315,7 +292,7 @@ namespace UnrealGameSync
 				return false;
 			}
 
-			Settings = new NewWorkspaceSettings{ Name = NewWorkspaceName, Stream = NewStream, RootDir = NewRootDir };
+			Settings = new NewWorkspaceSettings(NewWorkspaceName, NewStream, NewRootDirRef);
 			return true;
 		}
 
@@ -340,7 +317,7 @@ namespace UnrealGameSync
 		{
 			if(TryGetWorkspaceSettings(out Settings))
 			{
-				DirectoryInfo RootDir = new DirectoryInfo(Settings.RootDir);
+				DirectoryInfo RootDir = Settings.RootDir.ToDirectoryInfo();
 				if(RootDir.Exists && RootDir.EnumerateFileSystemInfos().Any(x => x.Name != "." && x.Name != ".."))
 				{
 					if(MessageBox.Show(this, String.Format("The directory '{0}' is not empty. Are you sure you want to create a workspace there?", RootDir.FullName), "Directory not empty", MessageBoxButtons.YesNo) != DialogResult.Yes)
@@ -349,12 +326,11 @@ namespace UnrealGameSync
 					}
 				}
 
-				NewWorkspaceTask NewWorkspace = new NewWorkspaceTask(Settings, Info.UserName, Info.HostName);
+				ILogger Logger = ServiceProvider.GetRequiredService<ILogger<NewWorkspaceWindow>>();
 
-				string ErrorMessage;
-				if(PerforceModalTask.Execute(Owner, Perforce, NewWorkspace, "Creating workspace", "Creating workspace, please wait...", Log, out ErrorMessage) != ModalTaskResult.Succeeded)
+				ModalTask? Result = PerforceModalTask.Execute(Owner, "Creating workspace", "Creating workspace, please wait...", PerforceSettings, (p, c) => NewWorkspaceTask.RunAsync(p, Settings, Info.UserName ?? "", Info.ClientHost ?? "", c), Logger);
+				if (Result == null || !Result.Succeeded)
 				{
-					MessageBox.Show(ErrorMessage);
 					return;
 				}
 
@@ -371,8 +347,8 @@ namespace UnrealGameSync
 
 		private void StreamBrowseBtn_Click(object sender, EventArgs e)
 		{
-			string StreamName = StreamTextBox.Text.Trim();
-			if(SelectStreamWindow.ShowModal(this, Perforce, StreamName, Log, out StreamName))
+			string? StreamName = StreamTextBox.Text.Trim();
+			if(SelectStreamWindow.ShowModal(this, PerforceSettings, StreamName, ServiceProvider, out StreamName))
 			{
 				StreamTextBox.Text = StreamName;
 			}
