UGS: Significant refactor of UGS internals.

diff --git a/Engine/Source/Programs/UnrealGameSync/UnrealGameSync/Forms/OpenProjectWindow.cs b/Engine/Source/Programs/UnrealGameSync/UnrealGameSync/Forms/OpenProjectWindow.cs
--- a/Engine/Source/Programs/UnrealGameSync/UnrealGameSync/Forms/OpenProjectWindow.cs
+++ b/Engine/Source/Programs/UnrealGameSync/UnrealGameSync/Forms/OpenProjectWindow.cs
@@ -1,42 +1,45 @@
 // Copyright Epic Games, Inc. All Rights Reserved.
 
 using EpicGames.Core;
+using EpicGames.Perforce;
+using Microsoft.Extensions.DependencyInjection;
+using Microsoft.Extensions.Logging;
 using System;
 using System.Collections.Generic;
 using System.ComponentModel;
 using System.Data;
+using System.Diagnostics.CodeAnalysis;
 using System.Drawing;
 using System.IO;
 using System.Linq;
 using System.Text;
+using System.Threading;
 using System.Threading.Tasks;
 using System.Windows.Forms;
 
+#nullable enable
+
 namespace UnrealGameSync
 {
 	partial class OpenProjectWindow : Form
 	{
-		string ServerAndPortOverride;
-		string UserNameOverride;
-		DetectProjectSettingsTask DetectedProjectSettings;
-		DirectoryReference DataFolder;
-		DirectoryReference CacheFolder;
-		PerforceConnection DefaultConnection;
-		TextWriter Log;
+		string? ServerAndPortOverride;
+		string? UserNameOverride;
+		WorkspaceSettings? DetectedProjectSettings;
+		IPerforceSettings DefaultSettings;
+		IServiceProvider ServiceProvider;
 		UserSettings Settings;
 
-		private OpenProjectWindow(UserSelectedProjectSettings Project, UserSettings Settings, DirectoryReference DataFolder, DirectoryReference CacheFolder, PerforceConnection DefaultConnection, TextWriter Log)
+		private OpenProjectWindow(UserSelectedProjectSettings Project, UserSettings Settings, IPerforceSettings DefaultSettings, IServiceProvider ServiceProvider)
 		{
 			InitializeComponent();
 
 			this.Settings = Settings;
 			this.DetectedProjectSettings = null;
-			this.DataFolder = DataFolder;
-			this.CacheFolder = CacheFolder;
-			this.DefaultConnection = DefaultConnection;
-			this.Log = Log;
+			this.DefaultSettings = DefaultSettings;
+			this.ServiceProvider = ServiceProvider;
 
-			if(Project == null)
+			if (Project == null)
 			{
 				LocalFileRadioBtn.Checked = true;
 			}
@@ -82,23 +85,21 @@ namespace UnrealGameSync
 			UpdateOkButton();
 		}
 
-		private PerforceConnection Perforce
+		private IPerforceSettings Perforce
 		{
-			get { return Utility.OverridePerforceSettings(DefaultConnection, ServerAndPortOverride, UserNameOverride); }
+			get => Utility.OverridePerforceSettings(DefaultSettings, ServerAndPortOverride, UserNameOverride);
 		}
 
-		public static bool ShowModal(IWin32Window Owner, UserSelectedProjectSettings Project, out DetectProjectSettingsTask NewDetectedProjectSettings, UserSettings Settings, DirectoryReference DataFolder, DirectoryReference CacheFolder, PerforceConnection DefaultConnection, TextWriter Log)
+		public static WorkspaceSettings? ShowModal(IWin32Window Owner, UserSelectedProjectSettings Project, UserSettings Settings, DirectoryReference DataFolder, DirectoryReference CacheFolder, IPerforceSettings DefaultPerforceSettings, IServiceProvider ServiceProvider)
 		{
-			OpenProjectWindow Window = new OpenProjectWindow(Project, Settings, DataFolder, CacheFolder, DefaultConnection, Log);
+			OpenProjectWindow Window = new OpenProjectWindow(Project, Settings, DefaultPerforceSettings, ServiceProvider);
 			if(Window.ShowDialog(Owner) == DialogResult.OK)
 			{
-				NewDetectedProjectSettings = Window.DetectedProjectSettings;
-				return true;
+				return Window.DetectedProjectSettings;
 			}
 			else
 			{
-				NewDetectedProjectSettings = null;
-				return false;
+				return null;
 			}
 		}
 
@@ -121,11 +122,11 @@ namespace UnrealGameSync
 			UpdateWorkspacePathBrowseButton();
 		}
 
-		public static string GetServerLabelText(PerforceConnection DefaultConnection, string ServerAndPort, string UserName)
+		public static string GetServerLabelText(IPerforceSettings DefaultSettings, string? ServerAndPort, string? UserName)
 		{
 			if(ServerAndPort == null && UserName == null)
 			{
-				return String.Format("Using default connection settings (user '{0}' on server '{1}').", DefaultConnection.UserName, DefaultConnection.ServerAndPort);
+				return String.Format("Using default connection settings (user '{0}' on server '{1}').", DefaultSettings.UserName, DefaultSettings.ServerAndPort);
 			}
 			else
 			{
@@ -153,18 +154,18 @@ namespace UnrealGameSync
 
 		private void UpdateServerLabel()
 		{
-			ServerLabel.Text = GetServerLabelText(DefaultConnection, ServerAndPortOverride, UserNameOverride);
+			ServerLabel.Text = GetServerLabelText(DefaultSettings, ServerAndPortOverride, UserNameOverride);
 		}
 
 		private void UpdateWorkspacePathBrowseButton()
 		{
-			string WorkspaceName;
+			string? WorkspaceName;
 			WorkspacePathBrowseBtn.Enabled = TryGetWorkspaceName(out WorkspaceName);
 		}
 
 		private void UpdateOkButton()
 		{
-			string ProjectPath;
+			string? ProjectPath;
 			OkBtn.Enabled = WorkspaceRadioBtn.Checked? TryGetClientPath(out ProjectPath) : TryGetLocalPath(out ProjectPath);
 		}
 
@@ -172,8 +173,8 @@ namespace UnrealGameSync
 		{
 			WorkspaceRadioBtn.Checked = true;
 			
-			string WorkspaceName;
-			if(NewWorkspaceWindow.ShowModal(this, Perforce, null, WorkspaceNameTextBox.Text, Log, out WorkspaceName))
+			string? WorkspaceName;
+			if(NewWorkspaceWindow.ShowModal(this, Perforce, null, WorkspaceNameTextBox.Text, ServiceProvider, out WorkspaceName))
 			{
 				WorkspaceNameTextBox.Text = WorkspaceName;
 				UpdateOkButton();
@@ -184,8 +185,8 @@ namespace UnrealGameSync
 		{
 			WorkspaceRadioBtn.Checked = true;
 
-			string WorkspaceName = WorkspaceNameTextBox.Text;
-			if(SelectWorkspaceWindow.ShowModal(this, Perforce, WorkspaceName, Log, out WorkspaceName))
+			string? WorkspaceName = WorkspaceNameTextBox.Text;
+			if(SelectWorkspaceWindow.ShowModal(this, Perforce, WorkspaceName, ServiceProvider, out WorkspaceName))
 			{
 				WorkspaceNameTextBox.Text = WorkspaceName;
 			}
@@ -195,11 +196,11 @@ namespace UnrealGameSync
 		{
 			WorkspaceRadioBtn.Checked = true;
 
-			string WorkspaceName;
+			string? WorkspaceName;
 			if(TryGetWorkspaceName(out WorkspaceName))
 			{
-				string WorkspacePath = WorkspacePathTextBox.Text.Trim();
-				if(SelectProjectFromWorkspaceWindow.ShowModal(this, Perforce, WorkspaceName, WorkspacePath, Log, out WorkspacePath))
+				string? WorkspacePath = WorkspacePathTextBox.Text.Trim();
+				if(SelectProjectFromWorkspaceWindow.ShowModal(this, Perforce, WorkspaceName, WorkspacePath, ServiceProvider, out WorkspacePath))
 				{
 					WorkspacePathTextBox.Text = WorkspacePath;
 					UpdateOkButton();
@@ -207,7 +208,7 @@ namespace UnrealGameSync
 			}
 		}
 
-		private bool TryGetWorkspaceName(out string WorkspaceName)
+		private bool TryGetWorkspaceName([NotNullWhen(true)] out string? WorkspaceName)
 		{
 			string Text = WorkspaceNameTextBox.Text.Trim();
 			if(Text.Length == 0)
@@ -220,9 +221,9 @@ namespace UnrealGameSync
 			return true;
 		}
 
-		private bool TryGetClientPath(out string ClientPath)
+		private bool TryGetClientPath([NotNullWhen(true)] out string? ClientPath)
 		{
-			string WorkspaceName;
+			string? WorkspaceName;
 			if(!TryGetWorkspaceName(out WorkspaceName))
 			{
 				ClientPath = null;
@@ -240,7 +241,7 @@ namespace UnrealGameSync
 			return true;
 		}
 
-		private bool TryGetLocalPath(out string LocalPath)
+		private bool TryGetLocalPath(out string? LocalPath)
 		{
 			string LocalFile = LocalFileTextBox.Text.Trim();
 			if(LocalFile.Length == 0)
@@ -253,11 +254,11 @@ namespace UnrealGameSync
 			return true;
 		}
 
-		private bool TryGetSelectedProject(out UserSelectedProjectSettings Project)
+		private bool TryGetSelectedProject([NotNullWhen(true)] out UserSelectedProjectSettings? Project)
 		{
 			if(WorkspaceRadioBtn.Checked)
 			{
-				string ClientPath;
+				string? ClientPath;
 				if(TryGetClientPath(out ClientPath))
 				{
 					Project = new UserSelectedProjectSettings(ServerAndPortOverride, UserNameOverride, UserSelectedProjectType.Client, ClientPath, null);
@@ -266,7 +267,7 @@ namespace UnrealGameSync
 			}
 			else
 			{
-				string LocalPath;
+				string? LocalPath;
 				if(TryGetLocalPath(out LocalPath))
 				{
 					Project = new UserSelectedProjectSettings(ServerAndPortOverride, UserNameOverride, UserSelectedProjectType.Local, null, LocalPath);
@@ -280,41 +281,24 @@ namespace UnrealGameSync
 
 		private void OkBtn_Click(object sender, EventArgs e)
 		{
-			UserSelectedProjectSettings SelectedProject;
+			UserSelectedProjectSettings? SelectedProject;
 			if(TryGetSelectedProject(out SelectedProject))
 			{
-				DetectProjectSettingsTask NewDetectedProjectSettings = new DetectProjectSettingsTask(SelectedProject, DataFolder, CacheFolder, Log);
-				try
-				{
-					string ErrorMessage;
-					if(PerforceModalTask.Execute(this, Perforce, NewDetectedProjectSettings, "Opening Project", "Opening project, please wait...", Log, out ErrorMessage) != ModalTaskResult.Succeeded)
-					{
-						if(!String.IsNullOrEmpty(ErrorMessage))
-						{
-							MessageBox.Show(ErrorMessage);
-						}
-						return;
-					}
+				ILogger<WorkspaceSettings> Logger = ServiceProvider.GetRequiredService<ILogger<WorkspaceSettings>>();
 
-					DetectedProjectSettings = NewDetectedProjectSettings;
-					NewDetectedProjectSettings = null;
+				ModalTask<WorkspaceSettings>? NewDetectedProjectSettings = PerforceModalTask.Execute(this, "Opening project", "Opening project, please wait...", DefaultSettings, (x, y) => WorkspaceSettings.CreateAsync(x, SelectedProject, Logger, y), Logger);
+				if (NewDetectedProjectSettings != null && NewDetectedProjectSettings.Succeeded)
+				{
+					DetectedProjectSettings = NewDetectedProjectSettings.Result;
 					DialogResult = DialogResult.OK;
 					Close();
 				}
-				finally
-				{
-					if(NewDetectedProjectSettings != null)
-					{
-						NewDetectedProjectSettings.Dispose();
-						NewDetectedProjectSettings = null;
-					}
-				}
 			}
 		}
 
 		private void ChangeLink_LinkClicked(object sender, LinkLabelLinkClickedEventArgs e)
 		{
-			if(ConnectWindow.ShowModal(this, DefaultConnection, ref ServerAndPortOverride, ref UserNameOverride, Log))
+			if(ConnectWindow.ShowModal(this, DefaultSettings, ref ServerAndPortOverride, ref UserNameOverride, ServiceProvider))
 			{
 				UpdateServerLabel();
 			}
@@ -339,7 +323,7 @@ namespace UnrealGameSync
 				}
 			}
 
-			if(Dialog.ShowDialog(this) == System.Windows.Forms.DialogResult.OK)
+			if (Dialog.ShowDialog(this) == DialogResult.OK)
 			{
 				string FullName = Path.GetFullPath(Dialog.FileName);
 
diff --git a/Engine/Source/Programs/UnrealGameSync/UnrealGameSync/Forms/SelectProjectFromWorkspaceWindow.cs b/Engine/Source/Programs/UnrealGameSync/UnrealGameSync/Forms/SelectProjectFromWorkspaceWindow.cs
--- a/Engine/Source/Programs/UnrealGameSync/UnrealGameSync/Forms/SelectProjectFromWorkspaceWindow.cs
+++ b/Engine/Source/Programs/UnrealGameSync/UnrealGameSync/Forms/SelectProjectFromWorkspaceWindow.cs
@@ -1,22 +1,29 @@
 // Copyright Epic Games, Inc. All Rights Reserved.
 
+using EpicGames.Perforce;
+using Microsoft.Extensions.DependencyInjection;
+using Microsoft.Extensions.Logging;
 using System;
 using System.Collections.Generic;
 using System.ComponentModel;
 using System.Data;
+using System.Diagnostics.CodeAnalysis;
 using System.Drawing;
 using System.IO;
 using System.Linq;
 using System.Runtime.InteropServices;
 using System.Text;
+using System.Threading;
 using System.Threading.Tasks;
 using System.Windows.Forms;
 
+#nullable enable
+
 namespace UnrealGameSync
 {
 	partial class SelectProjectFromWorkspaceWindow : Form
 	{
-		class EnumerateWorkspaceProjectsTask : IPerforceModalTask
+		static class EnumerateWorkspaceProjectsTask
 		{
 			static readonly string[] Patterns =
 			{
@@ -24,43 +31,23 @@ namespace UnrealGameSync
 				"....uprojectdirs",
 			};
 
-			public string ClientName;
-			public List<string> Paths;
-
-			public EnumerateWorkspaceProjectsTask(string ClientName)
+			public static async Task<List<string>> RunAsync(IPerforceConnection Perforce, string ClientName, CancellationToken CancellationToken)
 			{
-				this.ClientName = ClientName;
-			}
-
-			public bool Run(PerforceConnection Perforce, TextWriter Log, out string ErrorMessage)
-			{
-				PerforceConnection PerforceClient = new PerforceConnection(Perforce.UserName, ClientName, Perforce.ServerAndPort);
-
-				PerforceSpec ClientSpec;
-				if(!PerforceClient.TryGetClientSpec(Perforce.ClientName, out ClientSpec, Log))
-				{
-					ErrorMessage = String.Format("Unable to get client spec for {0}", ClientName);
-					return false;
-				}
+				ClientRecord ClientSpec = await Perforce.GetClientAsync(Perforce.Settings.ClientName, CancellationToken);
 
-				string ClientRoot = ClientSpec.GetField("Root")?.TrimEnd(Path.DirectorySeparatorChar);
-				if(String.IsNullOrEmpty(ClientRoot))
+				string ClientRoot = ClientSpec.Root.TrimEnd(Path.DirectorySeparatorChar);
+				if (String.IsNullOrEmpty(ClientRoot))
 				{
-					ErrorMessage = String.Format("Client '{0}' does not have a valid root directory.", ClientName);
-					return false;
+					throw new UserErrorException($"Client '{ClientName}' does not have a valid root directory.");
 				}
 
-				List<PerforceFileRecord> FileRecords = new List<PerforceFileRecord>();
+				List<FStatRecord> FileRecords = new List<FStatRecord>();
 				foreach(string Pattern in Patterns)
 				{
 					string Filter = String.Format("//{0}/{1}", ClientName, Pattern);
 
-					List<PerforceFileRecord> WildcardFileRecords = new List<PerforceFileRecord>();
-					if(!PerforceClient.FindFiles(Filter, out WildcardFileRecords, Log))
-					{
-						ErrorMessage = String.Format("Unable to enumerate files matching {0}", Filter);
-						return false;
-					}
+					List<FStatRecord> WildcardFileRecords = await Perforce.FStatAsync(Filter, CancellationToken).ToListAsync(CancellationToken);
+					WildcardFileRecords.RemoveAll(x => x.Action == FileAction.Delete || x.Action == FileAction.MoveDelete);
 
 					FileRecords.AddRange(WildcardFileRecords);
 				}
@@ -71,17 +58,16 @@ namespace UnrealGameSync
 					ClientPrefix += Path.DirectorySeparatorChar;
 				}
 
-				Paths = new List<string>();
-				foreach(PerforceFileRecord FileRecord in FileRecords)
+				List<string> Paths = new List<string>();
+				foreach(FStatRecord FileRecord in FileRecords)
 				{
-					if(FileRecord.ClientPath.StartsWith(ClientPrefix, StringComparison.InvariantCultureIgnoreCase))
+					if(FileRecord.ClientFile != null && FileRecord.ClientFile.StartsWith(ClientPrefix, StringComparison.OrdinalIgnoreCase))
 					{
-						Paths.Add(FileRecord.ClientPath.Substring(ClientRoot.Length).Replace(Path.DirectorySeparatorChar, '/'));
+						Paths.Add(FileRecord.ClientFile.Substring(ClientRoot.Length).Replace(Path.DirectorySeparatorChar, '/'));
 					}
 				}
 
-				ErrorMessage = null;
-				return true;
+				return Paths;
 			}
 		}
 
@@ -196,9 +182,9 @@ namespace UnrealGameSync
 
 		static TreeNode FindOrAddChildNode(TreeNode ParentNode, string Text, int ImageIndex)
 		{
-			foreach(TreeNode ChildNode in ParentNode.Nodes)
+			foreach(TreeNode? ChildNode in ParentNode.Nodes)
 			{
-				if(String.Compare(ChildNode.Text, Text, StringComparison.InvariantCultureIgnoreCase) == 0)
+				if(ChildNode != null && String.Compare(ChildNode.Text, Text, StringComparison.InvariantCultureIgnoreCase) == 0)
 				{
 					return ChildNode;
 				}
@@ -211,22 +197,20 @@ namespace UnrealGameSync
 			return NextNode;
 		}
 
-		public static bool ShowModal(IWin32Window Owner, PerforceConnection Perforce, string WorkspaceName, string WorkspacePath, TextWriter Log, out string NewWorkspacePath)
+		public static bool ShowModal(IWin32Window Owner, IPerforceSettings Perforce, string WorkspaceName, string WorkspacePath, IServiceProvider ServiceProvider, [NotNullWhen(true)] out string? NewWorkspacePath)
 		{
-			EnumerateWorkspaceProjectsTask EnumerateProjectsTask = new EnumerateWorkspaceProjectsTask(WorkspaceName);
+			Perforce = new PerforceSettings(Perforce) { ClientName = WorkspaceName };
 
-			string ErrorMessage;
-			if(PerforceModalTask.Execute(Owner, Perforce, EnumerateProjectsTask, "Finding Projects", "Finding projects, please wait...", Log, out ErrorMessage) != ModalTaskResult.Succeeded)
+			ILogger Logger = ServiceProvider.GetRequiredService<ILogger<SelectProjectFromWorkspaceWindow>>();
+
+			ModalTask<List<string>>? PathsTask = PerforceModalTask.Execute(Owner, "Finding Projects", "Finding projects, please wait...", Perforce, (p, c) => EnumerateWorkspaceProjectsTask.RunAsync(p, WorkspaceName, c), Logger);
+			if(PathsTask == null || !PathsTask.Succeeded)
 			{
-				if(!String.IsNullOrEmpty(ErrorMessage))
-				{
-					MessageBox.Show(Owner, ErrorMessage);
-				}
 				NewWorkspacePath = null;
 				return false;
 			}
 
-			SelectProjectFromWorkspaceWindow SelectProjectWindow = new SelectProjectFromWorkspaceWindow(WorkspaceName, EnumerateProjectsTask.Paths, WorkspacePath);
+			SelectProjectFromWorkspaceWindow SelectProjectWindow = new SelectProjectFromWorkspaceWindow(WorkspaceName, PathsTask.Result, WorkspacePath);
 			if(SelectProjectWindow.ShowDialog() == DialogResult.OK && !String.IsNullOrEmpty(SelectProjectWindow.SelectedProjectFile))
 			{
 				NewWorkspacePath = SelectProjectWindow.SelectedProjectFile;
diff --git a/Engine/Source/Programs/UnrealGameSync/UnrealGameSync/Forms/SelectStreamWindow.cs b/Engine/Source/Programs/UnrealGameSync/UnrealGameSync/Forms/SelectStreamWindow.cs
--- a/Engine/Source/Programs/UnrealGameSync/UnrealGameSync/Forms/SelectStreamWindow.cs
+++ b/Engine/Source/Programs/UnrealGameSync/UnrealGameSync/Forms/SelectStreamWindow.cs
@@ -1,43 +1,41 @@
 ﻿// Copyright Epic Games, Inc. All Rights Reserved.
 
+using EpicGames.Perforce;
+using Microsoft.Extensions.DependencyInjection;
+using Microsoft.Extensions.Logging;
 using System;
 using System.Collections.Generic;
 using System.ComponentModel;
 using System.Data;
+using System.Diagnostics.CodeAnalysis;
 using System.Drawing;
 using System.IO;
 using System.Linq;
 using System.Text;
+using System.Threading;
 using System.Threading.Tasks;
 using System.Windows.Forms;
 
+#nullable enable
+
 namespace UnrealGameSync
 {
 	partial class SelectStreamWindow : Form
 	{
-		class EnumerateStreamsTask : IPerforceModalTask
+		static class EnumerateStreamsTask
 		{
-			public List<PerforceStreamRecord> Streams;
-
-			public bool Run(PerforceConnection Perforce, TextWriter Log, out string ErrorMessage)
+			public static async Task<List<StreamsRecord>> RunAsync(IPerforceConnection Perforce, CancellationToken CancellationToken)
 			{
-				if(!Perforce.FindStreams(out Streams, Log))
-				{
-					ErrorMessage = "Unable to enumerate streams.";
-					return false;
-				}
-
-				ErrorMessage = null;
-				return true;
+				return await Perforce.GetStreamsAsync(null, CancellationToken);
 			}
 		}
 
 		class StreamNode
 		{
-			public PerforceStreamRecord Record;
+			public StreamsRecord Record;
 			public List<StreamNode> ChildNodes = new List<StreamNode>();
 
-			public StreamNode(PerforceStreamRecord Record)
+			public StreamNode(StreamsRecord Record)
 			{
 				this.Record = Record;
 			}
@@ -72,10 +70,10 @@ namespace UnrealGameSync
 			}
 		}
 
-		private string SelectedStream;
+		private string? SelectedStream;
 		private List<StreamDepot> Depots = new List<StreamDepot>();
 
-		private SelectStreamWindow(List<PerforceStreamRecord> Streams, string StreamName)
+		private SelectStreamWindow(List<StreamsRecord> Streams, string? StreamName)
 		{
 			InitializeComponent();
 
@@ -90,11 +88,11 @@ namespace UnrealGameSync
 
 			// Build a map of stream names to their nodes
 			Dictionary<string, StreamNode> IdentifierToNode = new Dictionary<string, StreamNode>(StringComparer.InvariantCultureIgnoreCase);
-			foreach(PerforceStreamRecord Stream in Streams)
+			foreach(StreamsRecord Stream in Streams)
 			{
-				if(Stream.Identifier != null && Stream.Name != null)
+				if(Stream.Stream != null && Stream.Name != null)
 				{
-					IdentifierToNode[Stream.Identifier] = new StreamNode(Stream);
+					IdentifierToNode[Stream.Stream] = new StreamNode(Stream);
 				}
 			}
 
@@ -102,12 +100,12 @@ namespace UnrealGameSync
 			Dictionary<string, StreamDepot> NameToDepot = new Dictionary<string, StreamDepot>(StringComparer.InvariantCultureIgnoreCase);
 			foreach(StreamNode Node in IdentifierToNode.Values)
 			{
-				if(Node.Record.Parent == null)
+				if(Node.Record.Parent == null || Node.Record.Parent.Equals("none", StringComparison.OrdinalIgnoreCase))
 				{
-					string DepotName;
-					if(PerforceUtils.TryGetDepotName(Node.Record.Identifier, out DepotName))
+					string? DepotName;
+					if(PerforceUtils.TryGetDepotName(Node.Record.Stream, out DepotName))
 					{
-						StreamDepot Depot;
+						StreamDepot? Depot;
 						if(!NameToDepot.TryGetValue(DepotName, out Depot))
 						{
 							Depot = new StreamDepot(DepotName);
@@ -118,7 +116,7 @@ namespace UnrealGameSync
 				}
 				else
 				{
-					StreamNode ParentNode;
+					StreamNode? ParentNode;
 					if(IdentifierToNode.TryGetValue(Node.Record.Parent, out ParentNode))
 					{
 						ParentNode.ChildNodes.Add(Node);
@@ -140,12 +138,15 @@ namespace UnrealGameSync
 
 		private void GetExpandedNodes(TreeNodeCollection Nodes, List<TreeNode> ExpandedNodes)
 		{
-			foreach(TreeNode Node in Nodes)
+			foreach(TreeNode? Node in Nodes)
 			{
-				ExpandedNodes.Add(Node);
-				if(Node.IsExpanded)
+				if (Node != null)
 				{
-					GetExpandedNodes(Node.Nodes, ExpandedNodes);
+					ExpandedNodes.Add(Node);
+					if (Node.IsExpanded)
+					{
+						GetExpandedNodes(Node.Nodes, ExpandedNodes);
+					}
 				}
 			}
 		}
@@ -191,15 +192,15 @@ namespace UnrealGameSync
 
 		private bool IncludeNodeInFilter(StreamNode Node, string[] Filter)
 		{
-			return Filter.All(x => Node.Record.Identifier.IndexOf(x, StringComparison.InvariantCultureIgnoreCase) != -1 || Node.Record.Name.IndexOf(x, StringComparison.InvariantCultureIgnoreCase) != -1);
+			return Filter.All(x => Node.Record.Stream.IndexOf(x, StringComparison.InvariantCultureIgnoreCase) != -1 || Node.Record.Name.IndexOf(x, StringComparison.InvariantCultureIgnoreCase) != -1);
 		}
 
-		private bool TryFilterTree(StreamNode Node, string[] Filter, out StreamNode NewNode)
+		private bool TryFilterTree(StreamNode Node, string[] Filter, [NotNullWhen(true)] out StreamNode? NewNode)
 		{
 			StreamNode FilteredNode = new StreamNode(Node.Record);
 			foreach(StreamNode ChildNode in Node.ChildNodes)
 			{
-				StreamNode FilteredChildNode;
+				StreamNode? FilteredChildNode;
 				if(TryFilterTree(ChildNode, Filter, out FilteredChildNode))
 				{
 					FilteredNode.ChildNodes.Add(FilteredChildNode);
@@ -234,7 +235,7 @@ namespace UnrealGameSync
 					StreamDepot FilteredDepot = new StreamDepot(Depot.Name);
 					foreach(StreamNode RootNode in Depot.RootNodes)
 					{
-						StreamNode FilteredRootNode;
+						StreamNode? FilteredRootNode;
 						if(TryFilterTree(RootNode, Filter, out FilteredRootNode))
 						{
 							FilteredDepot.RootNodes.Add(FilteredRootNode);
@@ -270,7 +271,7 @@ namespace UnrealGameSync
 			{
 				for(TreeNode Node = StreamsTreeView.Nodes[0];;Node = Node.Nodes[0])
 				{
-					StreamNode Stream = Node.Tag as StreamNode;
+					StreamNode? Stream = Node.Tag as StreamNode;
 					if(Stream != null && IncludeNodeInFilter(Stream, Filter))
 					{
 						StreamsTreeView.SelectedNode = Node;
@@ -304,7 +305,7 @@ namespace UnrealGameSync
 			StreamTreeNode.Tag = Stream;
 			ParentTreeNode.Nodes.Add(StreamTreeNode);
 
-			if(Stream.Record.Identifier == SelectedStream && IncludeNodeInFilter(Stream, Filter))
+			if(Stream.Record.Name == SelectedStream && IncludeNodeInFilter(Stream, Filter))
 			{
 				StreamsTreeView.SelectedNode = StreamTreeNode;
 			}
@@ -318,27 +319,21 @@ namespace UnrealGameSync
 			{
 				StreamTreeNode.Expand();
 			}
-			return bExpand || (SelectedStream == Stream.Record.Identifier);
+			return bExpand || (SelectedStream == Stream.Record.Stream);
 		}
 
-		public static bool ShowModal(IWin32Window Owner, PerforceConnection Perforce, string StreamName, TextWriter Log, out string NewStreamName)
+		public static bool ShowModal(IWin32Window Owner, IPerforceSettings Perforce, string StreamName, IServiceProvider ServiceProvider, out string? NewStreamName)
 		{
-			EnumerateStreamsTask Task = new EnumerateStreamsTask();
+			ILogger Logger = ServiceProvider.GetRequiredService<ILogger<SelectStreamWindow>>();
 
-			string ErrorMessage;
-			ModalTaskResult Result = PerforceModalTask.Execute(Owner, Perforce, Task, "Finding streams", "Finding streams, please wait...", Log, out ErrorMessage);
-			if(Result != ModalTaskResult.Succeeded)
+			ModalTask<List<StreamsRecord>>? StreamsTask = PerforceModalTask.Execute(Owner, "Finding streams", "Finding streams, please wait...", Perforce, EnumerateStreamsTask.RunAsync, Logger);
+			if(StreamsTask == null || !StreamsTask.Succeeded)
 			{
-				if(!String.IsNullOrEmpty(ErrorMessage))
-				{
-					MessageBox.Show(Owner, ErrorMessage);
-				}
-
 				NewStreamName = null;
 				return false;
 			}
 
-			SelectStreamWindow SelectStream = new SelectStreamWindow(Task.Streams, StreamName);
+			SelectStreamWindow SelectStream = new SelectStreamWindow(StreamsTask.Result, StreamName);
 			if(SelectStream.ShowDialog(Owner) == DialogResult.OK)
 			{
 				NewStreamName = SelectStream.SelectedStream;
@@ -361,15 +356,15 @@ namespace UnrealGameSync
 			OkBtn.Enabled = (GetSelectedStream() != null);
 		}
 
-		private string GetSelectedStream()
+		private string? GetSelectedStream()
 		{
-			string NewSelectedStream = null;
+			string? NewSelectedStream = null;
 			if(StreamsTreeView.SelectedNode != null)
 			{
 				StreamNode StreamNode = (StreamNode)StreamsTreeView.SelectedNode.Tag;
 				if(StreamNode != null)
 				{
-					NewSelectedStream = StreamNode.Record.Identifier;
+					NewSelectedStream = StreamNode.Record.Stream;
 				}
 			}
 			return NewSelectedStream;
diff --git a/Engine/Source/Programs/UnrealGameSync/UnrealGameSync/Forms/SelectUserWindow.cs b/Engine/Source/Programs/UnrealGameSync/UnrealGameSync/Forms/SelectUserWindow.cs
--- a/Engine/Source/Programs/UnrealGameSync/UnrealGameSync/Forms/SelectUserWindow.cs
+++ b/Engine/Source/Programs/UnrealGameSync/UnrealGameSync/Forms/SelectUserWindow.cs
@@ -1,41 +1,39 @@
 // Copyright Epic Games, Inc. All Rights Reserved.
 
+using EpicGames.Perforce;
+using Microsoft.Extensions.DependencyInjection;
+using Microsoft.Extensions.Logging;
 using System;
 using System.Collections.Generic;
 using System.ComponentModel;
 using System.Data;
+using System.Diagnostics.CodeAnalysis;
 using System.Drawing;
 using System.IO;
 using System.Linq;
 using System.Text;
+using System.Threading;
 using System.Threading.Tasks;
 using System.Windows.Forms;
 
+#nullable enable
+
 namespace UnrealGameSync
 {
 	partial class SelectUserWindow : Form
 	{
-		class EnumerateUsersTask : IPerforceModalTask
+		static class EnumerateUsersTask
 		{
-			public List<PerforceUserRecord> Users;
-
-			public bool Run(PerforceConnection Perforce, TextWriter Log, out string ErrorMessage)
+			public static async Task<List<UsersRecord>> RunAsync(IPerforceConnection Perforce, CancellationToken CancellationToken)
 			{
-				if(!Perforce.FindUsers(out Users, Log))
-				{
-					ErrorMessage = "Unable to enumerate users.";
-					return false;
-				}
-
-				ErrorMessage = null;
-				return true;
+				return await Perforce.GetUsersAsync(UsersOptions.None, -1, CancellationToken);
 			}
 		}
 
 		private int SelectedUserIndex;
-		private List<PerforceUserRecord> Users;
+		private List<UsersRecord> Users;
 		
-		private SelectUserWindow(List<PerforceUserRecord> Users, int SelectedUserIndex)
+		private SelectUserWindow(List<UsersRecord> Users, int SelectedUserIndex)
 		{
 			InitializeComponent();
 
@@ -94,11 +92,11 @@ namespace UnrealGameSync
 			return base.ProcessCmdKey(ref Msg, KeyData);
 		}
 
-		private bool IncludeInFilter(PerforceUserRecord User, string[] FilterWords)
+		private bool IncludeInFilter(UsersRecord User, string[] FilterWords)
 		{
 			foreach(string FilterWord in FilterWords)
 			{
-				if(User.Name.IndexOf(FilterWord, StringComparison.OrdinalIgnoreCase) == -1 
+				if(User.UserName.IndexOf(FilterWord, StringComparison.OrdinalIgnoreCase) == -1 
 					&& User.FullName.IndexOf(FilterWord, StringComparison.OrdinalIgnoreCase) == -1
 					&& User.Email.IndexOf(FilterWord, StringComparison.OrdinalIgnoreCase) == -1)
 				{
@@ -118,10 +116,10 @@ namespace UnrealGameSync
 			string[] Filter = FilterTextBox.Text.Split(new char[]{' '}, StringSplitOptions.RemoveEmptyEntries);
 			for(int Idx = 0; Idx < Users.Count; Idx++)
 			{
-				PerforceUserRecord User = Users[Idx];
+				UsersRecord User = Users[Idx];
 				if(IncludeInFilter(User, Filter))
 				{
-					ListViewItem Item = new ListViewItem(User.Name);
+					ListViewItem Item = new ListViewItem(User.UserName);
 					Item.SubItems.Add(new ListViewItem.ListViewSubItem(Item, User.FullName));
 					Item.SubItems.Add(new ListViewItem.ListViewSubItem(Item, User.Email));
 					Item.Tag = (int)Idx;
@@ -151,27 +149,21 @@ namespace UnrealGameSync
 			UpdateOkButton();
 		}
 
-		public static bool ShowModal(IWin32Window Owner, PerforceConnection Perforce, TextWriter Log, out string SelectedUserName)
+		public static bool ShowModal(IWin32Window Owner, IPerforceSettings PerforceSettings, IServiceProvider ServiceProvider, [NotNullWhen(true)] out string? SelectedUserName)
 		{
-			EnumerateUsersTask Task = new EnumerateUsersTask();
+			ILogger<SelectUserWindow> Logger = ServiceProvider.GetRequiredService<ILogger<SelectUserWindow>>();
 
-			string ErrorMessage;
-			ModalTaskResult Result = PerforceModalTask.Execute(Owner, Perforce, Task, "Finding users", "Finding users, please wait...", Log, out ErrorMessage);
-			if(Result != ModalTaskResult.Succeeded)
+			ModalTask<List<UsersRecord>>? UsersTask = PerforceModalTask.Execute(Owner, "Finding users", "Finding users, please wait...", PerforceSettings, EnumerateUsersTask.RunAsync, Logger);
+			if(UsersTask == null || !UsersTask.Succeeded)
 			{
-				if(!String.IsNullOrEmpty(ErrorMessage))
-				{
-					MessageBox.Show(Owner, ErrorMessage);
-				}
-
 				SelectedUserName = null;
 				return false;
 			}
 
-			SelectUserWindow SelectUser = new SelectUserWindow(Task.Users, 0);
+			SelectUserWindow SelectUser = new SelectUserWindow(UsersTask.Result, 0);
 			if(SelectUser.ShowDialog(Owner) == DialogResult.OK)
 			{
-				SelectedUserName = Task.Users[SelectUser.SelectedUserIndex].Name;
+				SelectedUserName = UsersTask.Result[SelectUser.SelectedUserIndex].UserName;
 				return true;
 			}
 			else
diff --git a/Engine/Source/Programs/UnrealGameSync/UnrealGameSync/Forms/SelectWorkspaceWindow.cs b/Engine/Source/Programs/UnrealGameSync/UnrealGameSync/Forms/SelectWorkspaceWindow.cs
--- a/Engine/Source/Programs/UnrealGameSync/UnrealGameSync/Forms/SelectWorkspaceWindow.cs
+++ b/Engine/Source/Programs/UnrealGameSync/UnrealGameSync/Forms/SelectWorkspaceWindow.cs
@@ -1,5 +1,8 @@
 ﻿// Copyright Epic Games, Inc. All Rights Reserved.
 
+using EpicGames.Perforce;
+using Microsoft.Extensions.DependencyInjection;
+using Microsoft.Extensions.Logging;
 using System;
 using System.Collections.Generic;
 using System.ComponentModel;
@@ -7,42 +10,42 @@ using System.Data;
 using System.Drawing;
 using System.IO;
 using System.Linq;
+using System.Net;
 using System.Text;
+using System.Threading;
 using System.Threading.Tasks;
 using System.Windows.Forms;
 
+#nullable enable
+
 namespace UnrealGameSync
 {
 	partial class SelectWorkspaceWindow : Form
 	{
-		class EnumerateWorkspacesTask : IPerforceModalTask
+		class EnumerateWorkspaces
 		{
-			public PerforceInfoRecord Info;
-			public List<PerforceClientRecord> Clients;
+			public InfoRecord Info { get; }
+			public List<ClientsRecord> Clients { get; }
 
-			public bool Run(PerforceConnection Perforce, TextWriter Log, out string ErrorMessage)
+			public EnumerateWorkspaces(InfoRecord Info, List<ClientsRecord> Clients)
 			{
-				if(!Perforce.Info(out Info, Log))
-				{
-					ErrorMessage = "Unable to query Perforce info.";
-					return false;
-				}
-				if(!Perforce.FindClients(Info.UserName, out Clients, Log))
-				{
-					ErrorMessage = "Unable to enumerate clients from Perforce";
-					return false;
-				}
+				this.Info = Info;
+				this.Clients = Clients;
+			}
 
-				ErrorMessage = null;
-				return true;
+			public static async Task<EnumerateWorkspaces> RunAsync(IPerforceConnection Perforce, CancellationToken CancellationToken)
+			{
+				InfoRecord Info = await Perforce.GetInfoAsync(InfoOptions.ShortOutput, CancellationToken);
+				List<ClientsRecord> Clients = await Perforce.GetClientsAsync(ClientsOptions.None, Perforce.Settings.UserName, CancellationToken);
+				return new EnumerateWorkspaces(Info, Clients);
 			}
 		}
 
-		PerforceInfoRecord Info;
-		List<PerforceClientRecord> Clients;
-		string WorkspaceName;
+		InfoRecord Info;
+		List<ClientsRecord> Clients;
+		string? WorkspaceName;
 
-		private SelectWorkspaceWindow(PerforceInfoRecord Info, List<PerforceClientRecord> Clients, string WorkspaceName)
+		private SelectWorkspaceWindow(InfoRecord Info, List<ClientsRecord> Clients, string? WorkspaceName)
 		{
 			InitializeComponent();
 
@@ -67,9 +70,9 @@ namespace UnrealGameSync
 
 			WorkspaceListView.Items.Clear();
 
-			foreach(PerforceClientRecord Client in Clients.OrderBy(x => x.Name))
+			foreach(ClientsRecord Client in Clients.OrderBy(x => x.Name))
 			{
-				if(!OnlyForThisComputer.Checked || String.Compare(Client.Host, Info.HostName, StringComparison.InvariantCultureIgnoreCase) == 0)
+				if(!OnlyForThisComputer.Checked || String.Compare(Client.Host, Info.ClientHost, StringComparison.InvariantCultureIgnoreCase) == 0)
 				{
 					ListViewItem Item = new ListViewItem(Client.Name);
 					Item.SubItems.Add(new ListViewItem.ListViewSubItem(Item, Client.Host));
@@ -86,24 +89,16 @@ namespace UnrealGameSync
 			OkBtn.Enabled = (WorkspaceListView.SelectedItems.Count == 1);
 		}
 
-		public static bool ShowModal(IWin32Window Owner, PerforceConnection Perforce, string WorkspaceName, TextWriter Log, out string NewWorkspaceName)
+		public static bool ShowModal(IWin32Window Owner, IPerforceSettings Perforce, string WorkspaceName, IServiceProvider ServiceProvider, out string? NewWorkspaceName)
 		{
-			EnumerateWorkspacesTask EnumerateWorkspaces = new EnumerateWorkspacesTask();
-
-			string ErrorMessage;
-			ModalTaskResult Result = PerforceModalTask.Execute(Owner, Perforce, EnumerateWorkspaces, "Finding workspaces", "Finding workspaces, please wait...", Log, out ErrorMessage);
-			if(Result != ModalTaskResult.Succeeded)
+			ModalTask<EnumerateWorkspaces>? Task = PerforceModalTask.Execute(Owner, "Finding workspaces", "Finding workspaces, please wait...", Perforce, EnumerateWorkspaces.RunAsync, ServiceProvider.GetRequiredService<ILogger<EnumerateWorkspaces>>());
+			if (Task == null || !Task.Succeeded)
 			{
-				if(!String.IsNullOrEmpty(ErrorMessage))
-				{
-					MessageBox.Show(Owner, ErrorMessage);
-				}
-
 				NewWorkspaceName = null;
 				return false;
 			}
 
-			SelectWorkspaceWindow SelectWorkspace = new SelectWorkspaceWindow(EnumerateWorkspaces.Info, EnumerateWorkspaces.Clients, WorkspaceName);
+			SelectWorkspaceWindow SelectWorkspace = new SelectWorkspaceWindow(Task.Result.Info, Task.Result.Clients, WorkspaceName);
 			if(SelectWorkspace.ShowDialog(Owner) == DialogResult.OK)
 			{
 				NewWorkspaceName = SelectWorkspace.WorkspaceName;
diff --git a/Engine/Source/Programs/UnrealGameSync/UnrealGameSync/Forms/VariablesWindow.cs b/Engine/Source/Programs/UnrealGameSync/UnrealGameSync/Forms/VariablesWindow.cs
--- a/Engine/Source/Programs/UnrealGameSync/UnrealGameSync/Forms/VariablesWindow.cs
+++ b/Engine/Source/Programs/UnrealGameSync/UnrealGameSync/Forms/VariablesWindow.cs
@@ -1,4 +1,4 @@
-﻿// Copyright Epic Games, Inc. All Rights Reserved.
+// Copyright Epic Games, Inc. All Rights Reserved.
 
 using System;
 using System.Collections;
@@ -11,13 +11,15 @@ using System.Text;
 using System.Threading.Tasks;
 using System.Windows.Forms;
 
+#nullable enable
+
 namespace UnrealGameSync
 {
 	public partial class VariablesWindow : Form
 	{
 		public delegate void InsertVariableDelegate(string Name);
 
-		public event InsertVariableDelegate OnInsertVariable;
+		public event InsertVariableDelegate? OnInsertVariable;
 
 		public VariablesWindow(IReadOnlyDictionary<string, string> Variables)
 		{
@@ -37,10 +39,10 @@ namespace UnrealGameSync
 				MacrosList.Items.Add(Item);
 			}
 
-			foreach(DictionaryEntry Entry in Environment.GetEnvironmentVariables())
+			foreach(DictionaryEntry Entry in Environment.GetEnvironmentVariables().OfType<DictionaryEntry>())
 			{
-				string Key = Entry.Key.ToString();
-				if(!Variables.ContainsKey(Key))
+				string? Key = Entry.Key?.ToString();
+				if(Key != null && Entry.Value != null && !Variables.ContainsKey(Key))
 				{
 					ListViewItem Item = new ListViewItem(String.Format("$({0})", Key));
 					Item.SubItems.Add(Entry.Value.ToString());
diff --git a/Engine/Source/Programs/UnrealGameSync/UnrealGameSync/IssueMonitor.cs b/Engine/Source/Programs/UnrealGameSync/UnrealGameSync/IssueMonitor.cs
--- a/Engine/Source/Programs/UnrealGameSync/UnrealGameSync/IssueMonitor.cs
+++ b/Engine/Source/Programs/UnrealGameSync/UnrealGameSync/IssueMonitor.cs
@@ -2,6 +2,8 @@
 
 using EnvDTE;
 using EpicGames.Core;
+using Microsoft.Extensions.DependencyInjection;
+using Microsoft.Extensions.Logging;
 using System;
 using System.Collections.Concurrent;
 using System.Collections.Generic;
@@ -16,6 +18,8 @@ using System.Threading.Tasks;
 
 using Thread = System.Threading.Thread;
 
+#nullable enable
+
 namespace UnrealGameSync
 {
 	public enum IssueBuildOutcome
@@ -29,21 +33,21 @@ namespace UnrealGameSync
 	public class IssueBuildData
 	{
 		public long Id { get; set; }
-		public string Stream { get; set; }
+		public string Stream { get; set; } = String.Empty;
 		public int Change { get; set; }
-		public string JobName { get; set; }
-		public string JobUrl { get; set; }
-		public string JobStepName { get; set; }
-		public string JobStepUrl { get; set; }
-		public string ErrorUrl { get; set; }
+		public string JobName { get; set; } = String.Empty;
+		public string JobUrl { get; set; } = String.Empty;
+		public string JobStepName { get; set; } = String.Empty;
+		public string JobStepUrl { get; set; } = String.Empty;
+		public string ErrorUrl { get; set; } = String.Empty;
 		public IssueBuildOutcome Outcome { get; set; }
 	}
 
 	public class IssueDiagnosticData
 	{
 		public long? BuildId { get; set; }
-		public string Message { get; set; }
-		public string Url { get; set; }
+		public string Message { get; set; } = String.Empty;
+		public string Url { get; set; } = String.Empty;
 	}
 
 	public class IssueData
@@ -52,20 +56,20 @@ namespace UnrealGameSync
 		public long Id { get; set; }
 		public DateTime CreatedAt { get; set; }
 		public DateTime RetrievedAt { get; set; }
-		public string Project { get; set; }
-		public string Summary { get; set; }
-		public string Details { get; set; }
-		public string Owner { get; set; }
-		public string NominatedBy { get; set; }
+		public string Project { get; set; } = String.Empty;
+		public string Summary { get; set; } = String.Empty;
+		public string Details { get; set; } = String.Empty;
+		public string Owner { get; set; } = String.Empty;
+		public string NominatedBy { get; set; } = String.Empty;
 		public DateTime? AcknowledgedAt { get; set; }
 		public int FixChange { get; set; }
 		public DateTime? ResolvedAt { get; set; }
 		public bool bNotify { get; set; }
 		public bool bIsWarning { get; set; }
-		public string BuildUrl { get; set; }
-		public List<string> Streams { get;set; }
+		public string BuildUrl { get; set; } = String.Empty;
+		public List<string> Streams { get; set; } = new List<string>();
 
-		HashSet<string> CachedProjects;
+		HashSet<string>? CachedProjects;
 
 		public HashSet<string> Projects
 		{
@@ -113,8 +117,8 @@ namespace UnrealGameSync
 	public class IssueUpdateData
 	{
 		public long Id { get; set; }
-		public string Owner { get; set; }
-		public string NominatedBy { get; set; }
+		public string Owner { get; set; } = String.Empty;
+		public string NominatedBy { get; set; } = String.Empty;
 		public bool? Acknowledged { get; set; }
 		public int? FixChange { get; set; }
 		public bool? Resolved { get; set; }
@@ -135,45 +139,48 @@ namespace UnrealGameSync
 		public readonly string ApiUrl;
 		public readonly string UserName;
 		int RefCount = 1;
-		Thread WorkerThread;
-		BoundedLogWriter LogWriter;
-		bool bDisposing;
-		AutoResetEvent RefreshEvent;
+		Task? WorkerTask;
+		ILogger Logger;
+		CancellationTokenSource CancellationSource;
+		AsyncEvent RefreshEvent;
 		int UpdateIntervalMs;
 		List<long> TrackingIssueIds = new List<long>();
 		List<IssueData> Issues = new List<IssueData>();
 		object LockObject = new object();
 		List<IssueUpdateData> PendingUpdates = new List<IssueUpdateData>();
+		IAsyncDisposer AsyncDisposer;
 
-		public event Action OnIssuesChanged;
+		public Action? OnIssuesChanged { get; set; }
 
 		// Only used by MainWindow, but easier to just store here
 		public Dictionary<long, IssueAlertReason> IssueIdToAlertReason = new Dictionary<long, IssueAlertReason>();
 
-		public IssueMonitor(string ApiUrl, string UserName, TimeSpan UpdateInterval, FileReference LogFileName)
+		public IssueMonitor(string ApiUrl, string UserName, TimeSpan UpdateInterval, IServiceProvider ServiceProvider)
 		{
 			this.ApiUrl = ApiUrl;
 			this.UserName = UserName;
 			this.UpdateIntervalMs = (int)UpdateInterval.TotalMilliseconds;
+			this.Logger = ServiceProvider.GetRequiredService<ILogger<IssueMonitor>>();
+			CancellationSource = new CancellationTokenSource();
+			this.AsyncDisposer = ServiceProvider.GetRequiredService<IAsyncDisposer>();
 
-			LogWriter = new BoundedLogWriter(LogFileName);
-			if(ApiUrl == null)
+			if (ApiUrl == null)
 			{
 				LastStatusMessage = "Database functionality disabled due to empty ApiUrl.";
 			}
 			else
 			{
-				LogWriter.WriteLine("Using connection string: {0}", this.ApiUrl);
+				Logger.LogInformation("Using connection string: {ApiUrl}", this.ApiUrl);
 			}
 
-			RefreshEvent = new AutoResetEvent(false);
+			RefreshEvent = new AsyncEvent();
 		}
 
 		public string LastStatusMessage
 		{
 			get;
 			private set;
-		}
+		} = String.Empty;
 
 		public List<IssueData> GetIssues()
 		{
@@ -272,8 +279,7 @@ namespace UnrealGameSync
 		{
 			if(ApiUrl != null)
 			{
-				WorkerThread = new Thread(() => PollForUpdates());
-				WorkerThread.Start();
+				WorkerTask = Task.Run(() => PollForUpdatesAsync(CancellationSource.Token));
 			}
 		}
 
@@ -306,30 +312,24 @@ namespace UnrealGameSync
 
 		void DisposeInternal()
 		{
-			bDisposing = true;
-			if(WorkerThread != null)
-			{
-				RefreshEvent.Set();
-				if(!WorkerThread.Join(100))
-				{
-					WorkerThread.Interrupt();
-					WorkerThread.Join();
-				}
-				WorkerThread = null;
-			}
-			if(LogWriter != null)
+			OnIssuesChanged = null;
+
+			if (WorkerTask != null)
 			{
-				LogWriter.Dispose();
-				LogWriter = null;
+				CancellationSource.Cancel();
+				AsyncDisposer.Add(WorkerTask.ContinueWith(_ => CancellationSource.Dispose()));
+				WorkerTask = null;
 			}
 		}
 
-		void PollForUpdates()
+		async Task PollForUpdatesAsync(CancellationToken CancellationToken)
 		{
-			while (!bDisposing)
+			while (CancellationToken.IsCancellationRequested)
 			{
+				Task RefreshTask = RefreshEvent.Task;
+
 				// Check if there's any pending update
-				IssueUpdateData PendingUpdate;
+				IssueUpdateData? PendingUpdate;
 				lock (LockObject)
 				{
 					if (PendingUpdates.Count > 0)
@@ -345,53 +345,52 @@ namespace UnrealGameSync
 				// If we have an update, try to post it to the backend and check for another
 				if (PendingUpdate != null)
 				{
-					if (SendUpdate(PendingUpdate))
+					if (await SendUpdateAsync(PendingUpdate, CancellationToken))
 					{
 						lock (LockObject) { PendingUpdates.RemoveAt(0); }
 					}
 					else
 					{
-						RefreshEvent.WaitOne(5 * 1000);
+						await Task.WhenAny(RefreshTask, Task.Delay(TimeSpan.FromSeconds(5.0), CancellationToken));
 					}
 					continue;
 				}
 
 				// Read all the current issues
-				ReadCurrentIssues();
+				await ReadCurrentIssuesAsync(CancellationToken);
 
 				// Wait for something else to do
-				RefreshEvent.WaitOne(UpdateIntervalMs);
+				await Task.WhenAny(RefreshTask, Task.Delay(UpdateIntervalMs, CancellationToken));
 			}
 		}
 
-		bool SendUpdate(IssueUpdateData Update)
+		async Task<bool> SendUpdateAsync(IssueUpdateData Update, CancellationToken CancellationToken)
 		{
 			try
 			{
-				RESTApi.PUT<IssueUpdateData>(ApiUrl, String.Format("issues/{0}", Update.Id), Update);
+				await RESTApi.PutAsync<IssueUpdateData>($"{ApiUrl}/api/issues/{Update.Id}", Update, CancellationToken);
 				return true;
 			}
 			catch(Exception Ex)
 			{
-				LogWriter.WriteException(Ex, "Failed with exception.");
+				Logger.LogError(Ex, "Failed with exception.");
 				LastStatusMessage = String.Format("Failed to send update: ({0})", Ex.ToString());
 				return false;
 			}
 		}
 
-		bool ReadCurrentIssues()
+		async Task<bool> ReadCurrentIssuesAsync(CancellationToken CancellationToken)
 		{
 			try
 			{
 				Stopwatch Timer = Stopwatch.StartNew();
-				LogWriter.WriteLine();
-				LogWriter.WriteLine("Polling for notifications at {0}...", DateTime.Now.ToString());
+				Logger.LogInformation("Polling for issues...");
 
 				// Get the initial number of issues. We won't post updates if this stays at zero.
 				int InitialNumIssues = Issues.Count;
 
 				// Fetch the new issues
-				List<IssueData> NewIssues = RESTApi.GET<List<IssueData>>(ApiUrl, String.Format("issues?user={0}", UserName));
+				List<IssueData> NewIssues = await RESTApi.GetAsync<List<IssueData>>($"{ApiUrl}/api/issues?user={UserName}", CancellationToken);
 
 				// Check if we're tracking a particular issue. If so, we want updates even when it's resolved.
 				long[] LocalTrackingIssueIds;
@@ -405,7 +404,7 @@ namespace UnrealGameSync
 					{
 						try
 						{
-							IssueData Issue = RESTApi.GET<IssueData>(ApiUrl, String.Format("issues/{0}", LocalTrackingIssueId));
+							IssueData Issue = await RESTApi.GetAsync<IssueData>($"{ApiUrl}/api/issues/{LocalTrackingIssueId}", CancellationToken);
 							if(Issue != null)
 							{
 								NewIssues.Add(Issue);
@@ -413,7 +412,7 @@ namespace UnrealGameSync
 						}
 						catch(Exception Ex)
 						{
-							LogWriter.WriteException(Ex, "Exception while fetching tracked issue");
+							Logger.LogError(Ex, "Exception while fetching tracked issue");
 						}
 					}
 				}
@@ -423,7 +422,7 @@ namespace UnrealGameSync
 				{
 					if (NewIssue.Version == 0)
 					{
-						List<IssueBuildData> Builds = RESTApi.GET<List<IssueBuildData>>(ApiUrl, String.Format("issues/{0}/builds", NewIssue.Id));
+						List<IssueBuildData> Builds = await RESTApi.GetAsync<List<IssueBuildData>>($"{ApiUrl}/api/issues/{NewIssue.Id}/builds", CancellationToken);
 						if (Builds != null && Builds.Count > 0)
 						{
 							NewIssue.bIsWarning = !Builds.Any(x => x.Outcome != IssueBuildOutcome.Warning);
@@ -458,12 +457,12 @@ namespace UnrealGameSync
 
 				// Update the stats
 				LastStatusMessage = String.Format("Last update took {0}ms", Timer.ElapsedMilliseconds);
-				LogWriter.WriteLine("Done in {0}ms.", Timer.ElapsedMilliseconds);
+				Logger.LogInformation("Done in {Time}ms.", Timer.ElapsedMilliseconds);
 				return true;
 			}
 			catch(Exception Ex)
 			{
-				LogWriter.WriteException(Ex, "Failed with exception.");
+				Logger.LogError(Ex, "Failed with exception.");
 				LastStatusMessage = String.Format("Last update failed: ({0})", Ex.ToString());
 				return false;
 			}
diff --git a/Engine/Source/Programs/UnrealGameSync/UnrealGameSync/Jupiter/JupiterMonitor.cs b/Engine/Source/Programs/UnrealGameSync/UnrealGameSync/Jupiter/JupiterMonitor.cs
--- a/Engine/Source/Programs/UnrealGameSync/UnrealGameSync/Jupiter/JupiterMonitor.cs
+++ b/Engine/Source/Programs/UnrealGameSync/UnrealGameSync/Jupiter/JupiterMonitor.cs
@@ -11,6 +11,8 @@ using System.Threading.Tasks;
 using IdentityModel.Client;
 using EpicGames.Core;
 using EpicGames.Jupiter;
+using EpicGames.Perforce;
+using Microsoft.Extensions.Logging;
 
 namespace UnrealGameSync
 {
@@ -26,7 +28,7 @@ namespace UnrealGameSync
 		private readonly Uri JupiterUrl;
 
 		private readonly Timer UpdateTimer;
-		private readonly BoundedLogWriter LogWriter;
+		private readonly ILogger Logger;
 
 		public IReadOnlyList<IArchiveInfo> AvailableArchives
 		{
@@ -39,7 +41,7 @@ namespace UnrealGameSync
 			}
 		}
 
-		private JupiterMonitor(OIDCTokenManager InTokenManager, FileReference InLogPath, string InNamespace, string InUrl,
+		private JupiterMonitor(OIDCTokenManager InTokenManager, ILogger InLogger, string InNamespace, string InUrl,
 			string InProviderIdentifier, string InExpectedBranch)
 		{
 			TokenManager = InTokenManager;
@@ -48,18 +50,17 @@ namespace UnrealGameSync
 			ExpectedBranch = InExpectedBranch;
 			JupiterUrl = new Uri(InUrl);
 
-			LogWriter = new BoundedLogWriter(InLogPath);
+			Logger = InLogger;
 			UpdateTimer = new Timer(DoUpdate, null, TimeSpan.Zero, TimeSpan.FromMinutes(5));
 		}
 		public void Dispose()
 		{
 			UpdateTimer?.Dispose();
-			LogWriter?.Dispose();
 		}
 
 		private async void DoUpdate(object State)
 		{
-			await LogWriter.WriteLineAsync($"Starting poll of JupiterMonitor for namespace {JupiterNamespace}");
+			Logger.LogInformation("Starting poll of JupiterMonitor for namespace {JupiterNamespace}", JupiterNamespace);
 			try
 			{
 				IReadOnlyList<IArchiveInfo> NewArchives = await GetAvailableArchives();
@@ -71,11 +72,11 @@ namespace UnrealGameSync
 			}
 			catch (Exception Exception)
 			{
-				await LogWriter.WriteLineAsync($"Exception occured during poll! {Exception}");
+				Logger.LogError(Exception, "Exception occured during poll!");
 			}
 		}
 
-		public static JupiterMonitor CreateFromConfigFile(OIDCTokenManager TokenManager, FileReference LogPath, ConfigFile ConfigFile, string SelectedProjectIdentifier)
+		public static JupiterMonitor CreateFromConfigFile(OIDCTokenManager TokenManager, ILogger<JupiterMonitor> Logger, ConfigFile ConfigFile, string SelectedProjectIdentifier)
 		{
 			ConfigSection JupiterConfigSection = ConfigFile.FindSection("Jupiter");
 			if (JupiterConfigSection == null)
@@ -107,7 +108,7 @@ namespace UnrealGameSync
 			// project specific overrides
 			JupiterUrl = ProjectConfigSection.GetValue("JupiterUrl") ?? JupiterUrl;
 
-			return new JupiterMonitor(TokenManager, LogPath, JupiterNamespace, JupiterUrl, OIDCProviderIdentifier, ExpectedBranch);
+			return new JupiterMonitor(TokenManager, Logger, JupiterNamespace, JupiterUrl, OIDCProviderIdentifier, ExpectedBranch);
 		}
 
 		private async Task<IReadOnlyList<IArchiveInfo>> GetAvailableArchives()
@@ -212,7 +213,7 @@ namespace UnrealGameSync
 				get { return null; } 
 			}
 
-			public string Target => throw new NotImplementedException();
+			public string Target => throw new NotSupportedException();
 
 			public bool Exists()
 			{
@@ -224,7 +225,7 @@ namespace UnrealGameSync
 				return ChangeToKey.TryGetValue(ChangeNumber, out ArchiveKey);
 			}
 
-			public bool DownloadArchive(string ArchiveKey, DirectoryReference LocalRootPath, FileReference ManifestFileName, TextWriter Log, ProgressValue Progress)
+			public Task<bool> DownloadArchive(IPerforceConnection _, string ArchiveKey, DirectoryReference LocalRootPath, FileReference ManifestFileName, ILogger Logger, ProgressValue Progress, CancellationToken CancellationToken)
 			{
 				try
 				{
@@ -232,10 +233,7 @@ namespace UnrealGameSync
 					{
 						(float ProgressFraction, FileReference FileReference) = ProgressUpdate;
 						Progress.Set(ProgressFraction);
-						lock (Log)
-						{
-							Log.WriteLine("Writing {0}", FileReference.FullName);
-						}
+						Logger.LogInformation("Writing {FileName}", FileReference.FullName);
 					});
 
 					// place the manifest for the Jupiter download next to the UGS manifest
@@ -262,13 +260,12 @@ namespace UnrealGameSync
 					}
 					FileReference.Move(TempManifestFileName, ManifestFileName);
 
-					return true;
+					return Task.FromResult(true);
 				}
 				catch (Exception Exception)
 				{
-					Log.WriteLine(
-						$"Exception occured when downloading build from Jupiter with key {ArchiveKey}. Exception {Exception}");
-					return false;
+					Logger.LogError(Exception, "Exception occured when downloading build from Jupiter with key {Key}.", ArchiveKey);
+					return Task.FromResult(false);
 				}
 			}
 
diff --git a/Engine/Source/Programs/UnrealGameSync/UnrealGameSync/Logging.cs b/Engine/Source/Programs/UnrealGameSync/UnrealGameSync/Logging.cs
new file mode 100644
--- /dev/null
+++ b/Engine/Source/Programs/UnrealGameSync/UnrealGameSync/Logging.cs
@@ -0,0 +1,25 @@
+// Copyright Epic Games, Inc. All Rights Reserved.
+
+using EpicGames.Core;
+using Microsoft.Extensions.Logging;
+using Serilog;
+using Serilog.Templates;
+using System;
+using System.Collections.Generic;
+using System.Text;
+
+namespace UnrealGameSync
+{
+	public static class Logging
+	{
+		public static ILoggerProvider CreateLoggerProvider(FileReference File)
+		{
+			Serilog.ILogger Logger = new LoggerConfiguration()
+				.Enrich.FromLogContext()
+				.WriteTo.File(new ExpressionTemplate("[{@t:HH:mm:ss.fff} {@l:w3}] {Concat('[', Substring(SourceContext, LastIndexOf(SourceContext, '.') + 1, 14), ']'),16} {Indent}{@m}{@x}\n"), File.FullName, rollingInterval: RollingInterval.Day, rollOnFileSizeLimit: true, fileSizeLimitBytes: 20 * 1024 * 1024, retainedFileCountLimit: 10)
+				.CreateLogger();
+
+			return new Serilog.Extensions.Logging.SerilogLoggerProvider(Logger, true);
+		}
+	}
+}
diff --git a/Engine/Source/Programs/UnrealGameSync/UnrealGameSync/ModalTask.cs b/Engine/Source/Programs/UnrealGameSync/UnrealGameSync/ModalTask.cs
--- a/Engine/Source/Programs/UnrealGameSync/UnrealGameSync/ModalTask.cs
+++ b/Engine/Source/Programs/UnrealGameSync/UnrealGameSync/ModalTask.cs
@@ -4,47 +4,98 @@ using System;
 using System.Collections.Generic;
 using System.Linq;
 using System.Text;
+using System.Threading;
 using System.Threading.Tasks;
 using System.Windows.Forms;
 
+#nullable enable
+
 namespace UnrealGameSync
 {
-	interface IModalTask
+	[Flags]
+	enum ModalTaskFlags
 	{
-		bool Run(out string ErrorMessage);
+		None = 0,
+		Quiet = 1,
 	}
 
-	public enum ModalTaskResult
+	class ModalTask
 	{
-		Succeeded,
-		Failed,
-		Aborted,
-	}
+		public Task Task { get; }
 
-	static class ModalTask
-	{
-		public static ModalTaskResult Execute(IWin32Window Owner, IModalTask Task, string InTitle, string InMessage, out string ErrorMessage)
+		public bool Failed => Task.IsFaulted;
+		public bool Succeeded => Task.IsCompletedSuccessfully;
+
+		public Exception? Exception => Task.Exception?.InnerException;
+
+		public ModalTask(Task Task)
+		{
+			this.Task = Task;
+		}
+
+		public string Error
+		{
+			get
+			{
+				if (Succeeded)
+				{
+					return "Succeeded.";
+				}
+
+				Exception? Ex = Exception;
+				if (Ex == null)
+				{
+					return "Failed.";
+				}
+				else if (Ex is UserErrorException UserEx)
+				{
+					return UserEx.Message;
+				}
+				else
+				{
+					return $"Unhandled exception ({Ex.Message})";
+				}
+			}
+		}
+
+		public static ModalTask? Execute(IWin32Window? Owner, string Title, string Message, Func<CancellationToken, Task> TaskFunc, ModalTaskFlags Flags = ModalTaskFlags.None)
 		{
-			ModalTaskWindow Window = new ModalTaskWindow(Task, InTitle, InMessage, (Owner == null)? FormStartPosition.CenterScreen : FormStartPosition.CenterParent);
-			Window.Complete += () => Window.Close();
-			Window.ShowDialog(Owner);
-			ErrorMessage = Window.ErrorMessage;
-			return Window.Result;
+			Func<CancellationToken, Task<int>> TypedTaskFunc = async x => { await TaskFunc(x); return 0; };
+			return Execute(Owner, Title, Message, TypedTaskFunc, Flags);
 		}
 
-		public static bool ExecuteAndShowError(IWin32Window Owner, IModalTask Task, string InTitle, string InMessage)
+		public static ModalTask<T>? Execute<T>(IWin32Window? Owner, string Title, string Message, Func<CancellationToken, Task<T>> TaskFunc, ModalTaskFlags Flags = ModalTaskFlags.None)
 		{
-			string ErrorMessage;
-			ModalTaskResult Result = Execute(Owner, Task, InTitle, InMessage, out ErrorMessage);
-			if(Result != ModalTaskResult.Succeeded)
+			using (CancellationTokenSource CancellationSource = new CancellationTokenSource())
 			{
-				if (!String.IsNullOrEmpty(ErrorMessage))
+				Task<T> BackgroundTask = Task.Run(() => TaskFunc(CancellationSource.Token));
+
+				ModalTaskWindow Window = new ModalTaskWindow(Title, Message, (Owner == null) ? FormStartPosition.CenterScreen : FormStartPosition.CenterParent, BackgroundTask, CancellationSource);
+				Window.ShowDialog(Owner);
+
+				if (BackgroundTask.IsCanceled || (BackgroundTask.Exception != null && BackgroundTask.Exception.InnerException is OperationCanceledException))
 				{
-					MessageBox.Show(ErrorMessage);
+					return null;
 				}
-				return false;
+
+				ModalTask<T> Result = new ModalTask<T>(BackgroundTask);
+				if (Result.Failed && (Flags & ModalTaskFlags.Quiet) == 0)
+				{
+					MessageBox.Show(Owner, Result.Error, Title);
+				}
+				return Result;
 			}
-			return true;
+		}
+	}
+
+	class ModalTask<T> : ModalTask
+	{
+		public new Task<T> Task => (Task<T>)base.Task;
+
+		public T Result => Task.Result;
+
+		public ModalTask(Task<T> Task) : base(Task)
+		{
 		}
 	}
 }
diff --git a/Engine/Source/Programs/UnrealGameSync/UnrealGameSync/OIDCTokenManager.cs b/Engine/Source/Programs/UnrealGameSync/UnrealGameSync/OIDCTokenManager.cs
--- a/Engine/Source/Programs/UnrealGameSync/UnrealGameSync/OIDCTokenManager.cs
+++ b/Engine/Source/Programs/UnrealGameSync/UnrealGameSync/OIDCTokenManager.cs
@@ -65,11 +65,11 @@ namespace UnrealGameSync
 			return TokenClients.Any(Pair => Pair.Value.GetStatusForProvider() == OIDCStatus.NotLoggedIn);
 		}
 
-		internal static OIDCTokenManager CreateFromConfigFile(UserSettings Settings, List<DetectProjectSettingsTask> ConfigFiles)
+		internal static OIDCTokenManager CreateFromConfigFile(UserSettings Settings, List<WorkspaceSettings> ConfigFiles)
 		{
 			// join the provider configuration from all projects
 			Dictionary<string, ProviderInfo> Providers = new Dictionary<string, ProviderInfo>();
-			foreach (DetectProjectSettingsTask DetectProjectSettingsTask in ConfigFiles)
+			foreach (WorkspaceSettings DetectProjectSettingsTask in ConfigFiles)
 			{
 				if(DetectProjectSettingsTask == null)
 				{
diff --git a/Engine/Source/Programs/UnrealGameSync/UnrealGameSync/OutputAdapters.cs b/Engine/Source/Programs/UnrealGameSync/UnrealGameSync/OutputAdapters.cs
--- a/Engine/Source/Programs/UnrealGameSync/UnrealGameSync/OutputAdapters.cs
+++ b/Engine/Source/Programs/UnrealGameSync/UnrealGameSync/OutputAdapters.cs
@@ -1,6 +1,7 @@
 // Copyright Epic Games, Inc. All Rights Reserved.
 
 using EpicGames.Core;
+using Microsoft.Extensions.Logging;
 using System;
 using System.Collections.Generic;
 using System.IO;
@@ -8,245 +9,34 @@ using System.Linq;
 using System.Text;
 using System.Threading.Tasks;
 
+#nullable enable
+
 namespace UnrealGameSync
 {
-	public abstract class LineBasedTextWriter : TextWriter
-	{
-		StringBuilder CurrentLine = new StringBuilder();
-
-		public LineBasedTextWriter()
-		{
-		}
-
-		public override void Write(char Character)
-		{
-			if(Character == '\n')
-			{
-				FlushLine(CurrentLine.ToString());
-				CurrentLine.Clear();
-			}
-			else if(Character != '\r')
-			{
-				CurrentLine.Append(Character);
-			}
-		}
-
-		public override void WriteLine(string Line)
-		{
-			foreach(string SubLine in Line.Split('\n'))
-			{
-				FlushLine(SubLine.TrimEnd());
-			}
-		}
-
-		protected abstract void FlushLine(string Line);
-
-		public override Encoding Encoding
-		{
-			get { return Encoding.UTF8; }
-		}
-	}
-
-	class ThreadSafeTextWriter : LineBasedTextWriter
-	{
-		object LockObject;
-		TextWriter Inner;
-
-		public ThreadSafeTextWriter(TextWriter Inner)
-		{
-			this.LockObject = new object();
-			this.Inner = Inner;
-		}
-
-		protected override void FlushLine(string Line)
-		{
-			lock (LockObject)
-			{
-				Inner.WriteLine(Line);
-			}
-		}
-	}
-
-	class NullTextWriter : LineBasedTextWriter
-	{
-		protected override void FlushLine(string Line)
-		{
-		}
-	}
-
-	class BufferedTextWriter : LineBasedTextWriter
-	{
-		public List<string> Lines = new List<string>();
-
-		protected override void FlushLine(string Line)
-		{
-			Lines.Add(Line);
-		}
-	}
-
-	class PrefixedTextWriter : LineBasedTextWriter
+	class PrefixedTextWriter : ILogger
 	{
 		string Prefix;
-		TextWriter Inner;
+		ILogger Inner;
 
-		public PrefixedTextWriter(string InPrefix, TextWriter InInner)
+		public PrefixedTextWriter(string InPrefix, ILogger InInner)
 		{
 			Prefix = InPrefix;
 			Inner = InInner;
 		}
 
-		protected override void Dispose(bool bDisposing)
-		{
-			if(bDisposing)
-			{
-				Inner.Dispose();
-			}
-		}
-
-		protected override void FlushLine(string Line)
-		{
-			Inner.WriteLine(Prefix + Line);
-		}
-	}
-
-	class TimestampLogWriter : LineBasedTextWriter
-	{
-		LineBasedTextWriter Inner;
-
-		public TimestampLogWriter(LineBasedTextWriter Inner)
-		{
-			this.Inner = Inner;
-		}
+		public IDisposable BeginScope<TState>(TState State) => Inner.BeginScope(State);
 
-		protected override void Dispose(bool bDisposing)
-		{
-			if(Inner != null)
-			{
-				Inner.Dispose();
-				Inner = null;
-			}
-		}
+		public bool IsEnabled(LogLevel LogLevel) => Inner.IsEnabled(LogLevel);
 
-		protected override void FlushLine(string Line)
+		public void Log<TState>(LogLevel LogLevel, EventId EventId, TState State, Exception Exception, Func<TState, Exception, string> Formatter)
 		{
-			Inner.WriteLine("[{0}] {1}", DateTime.Now, Line);
+			Inner.Log(LogLevel, EventId, State, Exception, (State, Exception) => Prefix + Formatter(State, Exception));
 		}
 	}
-	
-	class BoundedLogWriter : LineBasedTextWriter
-	{
-		FileReference FileName;
-		FileReference BackupFileName;
-		int MaxSize;
-		StreamWriter Inner;
 
-		public BoundedLogWriter(FileReference InFileName, int InMaxSize = 128 * 1024)
-		{
-			FileName = InFileName;
-			BackupFileName = FileName + ".bak";
-			MaxSize = InMaxSize;
-			OpenInner();
-		}
-
-		protected override void Dispose(bool bDisposing)
-		{
-			if(bDisposing)
-			{
-				CloseInner();
-			}
-		}
-
-		protected override void FlushLine(string Line)
-		{
-			if(Inner != null)
-			{
-				Inner.WriteLine(Line);
-				if(Inner.BaseStream.Position > MaxSize)
-				{
-					CloseInner();
-					OpenInner();
-				}
-			}
-		}
-
-		private void OpenInner()
-		{
-			CloseInner();
-			try
-			{
-				if(FileReference.Exists(BackupFileName))
-				{
-					FileReference.Delete(BackupFileName);
-				}
-				if(FileReference.Exists(FileName))
-				{
-					FileReference.Move(FileName, BackupFileName);
-				}
-
-				Inner = new StreamWriter(FileName.FullName);
-				Inner.AutoFlush = true;
-			}
-			catch(Exception)
-			{
-				Inner = null;
-			}
-		}
-
-		private void CloseInner()
-		{
-			if(Inner != null)
-			{
-				Inner.Close();
-				Inner = null;
-			}
-		}
-	}
-
-	class ComposedTextWriter : TextWriter
-	{
-		TextWriter[] Inners;
-
-		public ComposedTextWriter(params TextWriter[] InInners)
-		{
-			Inners = InInners;
-		}
-
-		protected override void Dispose(bool bDisposing)
-		{
-			if(bDisposing)
-			{
-				foreach(TextWriter Inner in Inners)
-				{
-					Inner.Dispose();
-				}
-			}
-		}
-
-		public override void Write(char Character)
-		{
-			foreach(TextWriter Inner in Inners)
-			{
-				Inner.Write(Character);
-			}
-		}
-
-		public override void WriteLine(string Line)
-		{
-			foreach(TextWriter Inner in Inners)
-			{
-				Inner.WriteLine(Line);
-			}
-		}
-
-		public override Encoding Encoding
-		{
-			get { return Encoding.UTF8; }
-		}
-	}
-
-	class ProgressValue
+	public class ProgressValue
 	{
-		Tuple<string, float> State;
+		Tuple<string, float> State = null!;
 		Stack<Tuple<float, float>> Ranges = new Stack<Tuple<float,float>>();
 
 		public ProgressValue()
@@ -317,27 +107,19 @@ namespace UnrealGameSync
 		}
 	}
 
-	class ProgressTextWriter : LineBasedTextWriter
+	static class ProgressTextWriter
 	{
 		const string DirectivePrefix = "@progress ";
 
-		ProgressValue Value;
-		TextWriter Inner;
-
-		public ProgressTextWriter(ProgressValue InValue, TextWriter InInner)
-		{ 
-			Value = InValue;
-			Inner = InInner;
-		}
-
-		protected override void FlushLine(string Line)
+		public static string? ParseLine(string Line, ProgressValue Value)
 		{
 			string TrimLine = Line.Trim();
 			if(TrimLine.StartsWith(DirectivePrefix))
 			{
 				// Line that just contains a progress directive
 				bool bSkipLine = false;
-				ProcessInternal(TrimLine.Substring(DirectivePrefix.Length), ref bSkipLine);
+				ProcessInternal(TrimLine.Substring(DirectivePrefix.Length), ref bSkipLine, Value);
+				return null;
 			}
 			else
 			{
@@ -354,7 +136,7 @@ namespace UnrealGameSync
 							string DirectiveSubstring = TrimLine.Substring(LastIdx + 1, TrimLine.Length - LastIdx - 2);
 							if(DirectiveSubstring.StartsWith(DirectivePrefix))
 							{
-								ProcessInternal(DirectiveSubstring.Substring(DirectivePrefix.Length), ref bSkipLine);
+								ProcessInternal(DirectiveSubstring.Substring(DirectivePrefix.Length), ref bSkipLine, Value);
 								RemainingLine = Line.Substring(0, LastIdx).TrimEnd();
 							}
 							break;
@@ -362,14 +144,18 @@ namespace UnrealGameSync
 					}
 				}
 
-				if(!bSkipLine)
+				if (bSkipLine)
+				{
+					return null;
+				}
+				else
 				{
-					Inner.WriteLine(RemainingLine);
+					return RemainingLine;
 				}
 			}
 		}
 
-		void ProcessInternal(string Line, ref bool bSkipLine)
+		static void ProcessInternal(string Line, ref bool bSkipLine, ProgressValue Value)
 		{
 			List<string> Tokens = ParseTokens(Line);
 			for(int TokenIdx = 0; TokenIdx < Tokens.Count; )
@@ -459,7 +245,7 @@ namespace UnrealGameSync
 			return Tokens;
 		}
 
-		bool ReadFraction(List<string> Tokens, ref int TokenIdx, out float Fraction)
+		static bool ReadFraction(List<string> Tokens, ref int TokenIdx, out float Fraction)
 		{
 			// Read a fraction in the form x%
 			if(TokenIdx + 2 <= Tokens.Count && Tokens[TokenIdx + 1] == "%")
diff --git a/Engine/Source/Programs/UnrealGameSync/UnrealGameSync/PerforceModalTask.cs b/Engine/Source/Programs/UnrealGameSync/UnrealGameSync/PerforceModalTask.cs
--- a/Engine/Source/Programs/UnrealGameSync/UnrealGameSync/PerforceModalTask.cs
+++ b/Engine/Source/Programs/UnrealGameSync/UnrealGameSync/PerforceModalTask.cs
@@ -1,135 +1,127 @@
 // Copyright Epic Games, Inc. All Rights Reserved.
 
+using EpicGames.Perforce;
+using Microsoft.Extensions.Logging;
 using System;
 using System.Collections.Generic;
 using System.IO;
 using System.Linq;
 using System.Text;
+using System.Threading;
 using System.Threading.Tasks;
 using System.Windows.Forms;
 
+#nullable enable
+
 namespace UnrealGameSync
 {
-	interface IPerforceModalTask
+	public class PerforceLoginException : Exception
 	{
-		bool Run(PerforceConnection Perforce, TextWriter Log, out string ErrorMessage);
+		public PerforceResponse<LoginRecord> Response { get; }
+
+		public PerforceLoginException(PerforceResponse<LoginRecord> Response) : base($"Login failed: {Response}")
+		{
+			this.Response = Response;
+		}
 	}
 
-	class PerforceModalTask : IModalTask
+	static class PerforceModalTask
 	{
-		public PerforceConnection Perforce;
-		public string Password;
-		public LoginResult LoginResult;
-		public IPerforceModalTask InnerTask;
-		public TextWriter Log;
-
-		public PerforceModalTask(PerforceConnection Perforce, IPerforceModalTask InnerTask, TextWriter Log)
+		public static ModalTask? Execute(IWin32Window? Owner, string Title, string Message, IPerforceSettings PerforceSettings, Func<IPerforceConnection, CancellationToken, Task> ExecuteAsync, ILogger Logger, ModalTaskFlags Flags = ModalTaskFlags.None)
 		{
-			this.Perforce = Perforce;
-			this.InnerTask = InnerTask;
-			this.Log = Log;
+			Func<IPerforceConnection, CancellationToken, Task<int>> ExecuteTypedAsync = async (p, c) => { await ExecuteAsync(p, c); return 0; };
+			return Execute(Owner, Title, Message, PerforceSettings, ExecuteTypedAsync, Logger, Flags);
 		}
 
-		public bool Run(out string ErrorMessage)
+		public static ModalTask<T>? Execute<T>(IWin32Window? Owner, string Title, string Message, IPerforceSettings PerforceSettings, Func<IPerforceConnection, CancellationToken, Task<T>> ExecuteAsync, ILogger Logger, ModalTaskFlags Flags = ModalTaskFlags.None)
 		{
-			// Set the default login state to failed
-			LoginResult = LoginResult.Failed;
-
-			// If we've got a password, execute the login command
-			if(Password != null)
+			IPerforceConnection? Connection = null;
+			try
 			{
-				string PasswordErrorMessage;
-				LoginResult = Perforce.Login(Password, out PasswordErrorMessage, Log);
-				if(LoginResult != LoginResult.Succeded)
+				Func<Task<IPerforceConnection>> ConnectAsync = async () =>
 				{
-					Log.WriteLine(PasswordErrorMessage);
-					ErrorMessage = String.Format("Unable to login: {0}", PasswordErrorMessage);
-					return false;
-				}
+					Connection ??= await PerforceConnection.CreateAsync(PerforceSettings, Logger);
+					return Connection;
+				};
+				return ExecuteInternal(Owner, Title, Message, PerforceSettings, ConnectAsync, ExecuteAsync, Flags);
 			}
-
-			// Check that we're logged in
-			bool bLoggedIn;
-			if(!Perforce.GetLoggedInState(out bLoggedIn, Log))
+			finally
 			{
-				ErrorMessage = "Unable to get login status.";
-				return false;
-			}
-			if(!bLoggedIn)
-			{
-				LoginResult = LoginResult.MissingPassword;
-				ErrorMessage = "User is not logged in to Perforce.";
-				Log.WriteLine(ErrorMessage);
-				return false;
+				Connection?.Dispose();
 			}
+		}
 
-			// Execute the inner task
-			LoginResult = LoginResult.Succeded;
-			return InnerTask.Run(Perforce, Log, out ErrorMessage);
+		public static ModalTask? Execute(IWin32Window? Owner, string Title, string Message, IPerforceConnection PerforceConnection, Func<IPerforceConnection, CancellationToken, Task> ExecuteAsync, ILogger Logger, ModalTaskFlags Flags = ModalTaskFlags.None)
+		{
+			Func<IPerforceConnection, CancellationToken, Task<int>> ExecuteTypedAsync = async (p, c) => { await ExecuteAsync(p, c); return 0; };
+			return Execute(Owner, Title, Message, PerforceConnection, ExecuteTypedAsync, Logger, Flags);
 		}
 
-		public static ModalTaskResult Execute(IWin32Window Owner, PerforceConnection Perforce, IPerforceModalTask PerforceTask, string Title, string Message, TextWriter Log, out string ErrorMessage)
+		public static ModalTask<T>? Execute<T>(IWin32Window? Owner, string Title, string Message, IPerforceConnection PerforceConnection, Func<IPerforceConnection, CancellationToken, Task<T>> ExecuteAsync, ILogger Logger, ModalTaskFlags Flags = ModalTaskFlags.None)
 		{
-			PerforceModalTask Task = new PerforceModalTask(Perforce, PerforceTask, Log);
+			return ExecuteInternal(Owner, Title, Message, PerforceConnection.Settings, () => Task.FromResult(PerforceConnection), ExecuteAsync, Flags);
+		}
+
+		private static ModalTask<T>? ExecuteInternal<T>(IWin32Window? Owner, string Title, string Message, IPerforceSettings PerforceSettings, Func<Task<IPerforceConnection>> ConnectAsync, Func<IPerforceConnection, CancellationToken, Task<T>> ExecuteAsync, ModalTaskFlags Flags = ModalTaskFlags.None)
+		{
+			string? Password = PerforceSettings.Password;
 			for(;;)
 			{
-				string TaskErrorMessage;
-				ModalTaskResult TaskResult = ModalTask.Execute(Owner, Task, Title, Message, out TaskErrorMessage);
+				Func<CancellationToken, Task<T>> RunAsync = CancellationToken => LoginAndExecuteAsync(Password, ConnectAsync, ExecuteAsync, CancellationToken);
 
-				if(Task.LoginResult == LoginResult.Succeded)
-				{
-					ErrorMessage = TaskErrorMessage;
-					return TaskResult;
-				}
-				else if(Task.LoginResult == LoginResult.MissingPassword)
+				ModalTask<T>? Result = ModalTask.Execute(Owner, Title, Message, RunAsync, Flags);
+				if (Result != null && Result.Failed && Result.Exception is PerforceLoginException && (Flags & ModalTaskFlags.Quiet) == 0)
 				{
-					PasswordWindow PasswordWindow = new PasswordWindow(String.Format("Enter the password for user '{0}' on server '{1}'.", Perforce.UserName, Perforce.ServerAndPort), Task.Password);
-					if(Owner == null)
+					string PasswordPrompt;
+					if (String.IsNullOrEmpty(Password))
 					{
-						PasswordWindow.StartPosition = FormStartPosition.CenterScreen;
+						PasswordPrompt = $"Enter the password for user '{PerforceSettings.UserName}' on server '{PerforceSettings.ServerAndPort}'.";
 					}
-					if(PasswordWindow.ShowDialog() != DialogResult.OK)
+					else
 					{
-						ErrorMessage = null;
-						return ModalTaskResult.Aborted;
+						PasswordPrompt = $"Authentication failed. Enter the password for user '{PerforceSettings.UserName}' on server '{PerforceSettings.ServerAndPort}'.";
 					}
-					Task.Password = PasswordWindow.Password;
-				}
-				else if(Task.LoginResult == LoginResult.IncorrectPassword)
-				{
-					PasswordWindow PasswordWindow = new PasswordWindow(String.Format("Authentication failed. Enter the password for user '{0}' on server '{1}'.", Perforce.UserName, Perforce.ServerAndPort), Task.Password);
+
+					PasswordWindow PasswordWindow = new PasswordWindow(PasswordPrompt, Password);
 					if (Owner == null)
 					{
 						PasswordWindow.StartPosition = FormStartPosition.CenterScreen;
 					}
-					if (PasswordWindow.ShowDialog() != DialogResult.OK)
+					if (PasswordWindow.ShowDialog(Owner) != DialogResult.OK)
 					{
-						ErrorMessage = null;
-						return ModalTaskResult.Aborted;
+						return null;
 					}
-					Task.Password = PasswordWindow.Password;
-				}
-				else
-				{
-					ErrorMessage = TaskErrorMessage;
-					return ModalTaskResult.Failed;
+
+					Password = PasswordWindow.Password;
+					continue;
 				}
+
+				return Result;
 			}
 		}
 
-		public static bool ExecuteAndShowError(IWin32Window Owner, PerforceConnection Perforce, IPerforceModalTask Task, string Title, string Message, TextWriter Log)
+		private static async Task<T> LoginAndExecuteAsync<T>(string? Password, Func<Task<IPerforceConnection>> ConnectAsync, Func<IPerforceConnection, CancellationToken, Task<T>> ExecuteAsync, CancellationToken CancellationToken)
 		{
-			string ErrorMessage;
-			ModalTaskResult Result = Execute(Owner, Perforce, Task, Title, Message, Log, out ErrorMessage);
-			if (Result != ModalTaskResult.Succeeded)
+			IPerforceConnection Perforce = await ConnectAsync();
+
+			// If we've got a password, execute the login command
+			PerforceResponse<LoginRecord> Response;
+			if (String.IsNullOrEmpty(Password))
 			{
-				if (!String.IsNullOrEmpty(ErrorMessage))
-				{
-					MessageBox.Show(ErrorMessage);
-				}
-				return false;
+				Response = await Perforce.TryGetLoginStateAsync(CancellationToken);
 			}
-			return true;
+			else
+			{
+				Response = await Perforce.TryLoginAsync(Password, CancellationToken);
+			}
+
+			if (!Response.Succeeded)
+			{
+				throw new PerforceLoginException(Response);
+			}
+
+			// Execute the inner task
+			return await ExecuteAsync(Perforce, CancellationToken);
 		}
 	}
 }
