UGS: Significant refactor of UGS internals.

diff --git a/Engine/Source/Programs/UnrealGameSync/UnrealGameSync/PerforceMonitor.cs b/Engine/Source/Programs/UnrealGameSync/UnrealGameSync/PerforceMonitor.cs
--- a/Engine/Source/Programs/UnrealGameSync/UnrealGameSync/PerforceMonitor.cs
+++ b/Engine/Source/Programs/UnrealGameSync/UnrealGameSync/PerforceMonitor.cs
@@ -1,16 +1,23 @@
 // Copyright Epic Games, Inc. All Rights Reserved.
 
 using EpicGames.Core;
+using EpicGames.Perforce;
+using Microsoft.Extensions.DependencyInjection;
+using Microsoft.Extensions.Logging;
 using System;
 using System.Collections.Concurrent;
 using System.Collections.Generic;
 using System.Diagnostics;
+using System.Diagnostics.CodeAnalysis;
 using System.IO;
 using System.Linq;
+using System.Runtime.CompilerServices;
 using System.Text;
 using System.Threading;
 using System.Threading.Tasks;
 
+#nullable enable
+
 namespace UnrealGameSync
 {
 	class PerforceChangeDetails
@@ -42,14 +49,14 @@ namespace UnrealGameSync
 		public bool bContainsCode;
 		public bool bContainsContent;
 
-		public PerforceChangeDetails(PerforceDescribeRecord DescribeRecord)
+		public PerforceChangeDetails(DescribeRecord DescribeRecord)
 		{
 			Description = DescribeRecord.Description;
 
 			// Check whether the files are code or content
-			foreach(PerforceDescribeFileRecord File in DescribeRecord.Files)
+			foreach(DescribeFileRecord File in DescribeRecord.Files)
 			{
-				if(CodeExtensions.Any(Extension => File.DepotFile.EndsWith(Extension, StringComparison.InvariantCultureIgnoreCase)))
+				if(CodeExtensions.Any(Extension => File.DepotFile.EndsWith(Extension, StringComparison.OrdinalIgnoreCase)))
 				{
 					bContainsCode = true;
 				}
@@ -71,15 +78,15 @@ namespace UnrealGameSync
 		IReadOnlyList<IArchiveInfo> AvailableArchives { get; }
 	}
 
-	interface IArchiveInfo
+	public interface IArchiveInfo
 	{
 		string Name { get; }
 		string Type { get; }
 		string BasePath { get; }
-		string Target { get;  }
+		string? Target { get;  }
 		bool Exists();
-		bool TryGetArchiveKeyForChangeNumber(int ChangeNumber, out string ArchiveKey);
-		bool DownloadArchive(string ArchiveKey, DirectoryReference LocalRootPath, FileReference ManifestFileName, TextWriter Log, ProgressValue Progress);
+		bool TryGetArchiveKeyForChangeNumber(int ChangeNumber, [NotNullWhen(true)] out string? ArchiveKey);
+		Task<bool> DownloadArchive(IPerforceConnection Perforce, string ArchiveKey, DirectoryReference LocalRootPath, FileReference ManifestFileName, ILogger Logger, ProgressValue Progress, CancellationToken CancellationToken);
 	}
 
 	class PerforceMonitor : IDisposable, IArchiveInfoSource
@@ -89,8 +96,8 @@ namespace UnrealGameSync
 			public string Name { get; }
 			public string Type { get; }
 			public string DepotPath { get; }
-			public string Target { get; }
-			public PerforceConnection Perforce { get; }
+			public string? Target { get; }
+			public PerforceMonitor Outer { get; }
 
 			public string BasePath
 			{
@@ -100,24 +107,24 @@ namespace UnrealGameSync
 			// TODO: executable/configuration?
 			public SortedList<int, string> ChangeNumberToFileRevision = new SortedList<int, string>();
 
-			public PerforceArchiveInfo(string Name, string Type, string DepotPath, string Target, PerforceConnection Perforce)
+			public PerforceArchiveInfo(string Name, string Type, string DepotPath, string? Target, PerforceMonitor Outer)
 			{
 				this.Name = Name;
 				this.Type = Type;
 				this.DepotPath = DepotPath;
 				this.Target = Target;
-				this.Perforce = Perforce;
+				this.Outer = Outer;
 			}
 
-			public override bool Equals(object Other)
+			public override bool Equals(object? Other)
 			{
-				PerforceArchiveInfo OtherArchive = Other as PerforceArchiveInfo;
+				PerforceArchiveInfo? OtherArchive = Other as PerforceArchiveInfo;
 				return OtherArchive != null && Name == OtherArchive.Name && Type == OtherArchive.Type && DepotPath == OtherArchive.DepotPath && Target == OtherArchive.Target && Enumerable.SequenceEqual(ChangeNumberToFileRevision, OtherArchive.ChangeNumberToFileRevision);
 			}
 
 			public override int GetHashCode()
 			{
-				throw new NotImplementedException();
+				throw new NotSupportedException();
 			}
 
 			public bool Exists()
@@ -125,49 +132,48 @@ namespace UnrealGameSync
 				return ChangeNumberToFileRevision.Count > 0;
 			}
 
-			public static bool TryParseConfigEntry(string Text, PerforceConnection PerforceConnection, out PerforceArchiveInfo Info)
+			public static bool TryParseConfigEntry(string Text, PerforceMonitor Outer, [NotNullWhen(true)] out PerforceArchiveInfo? Info)
 			{
 				ConfigObject Object = new ConfigObject(Text);
 
-				string Name = Object.GetValue("Name", null);
+				string? Name = Object.GetValue("Name", null);
 				if (Name == null)
 				{
 					Info = null;
 					return false;
 				}
 
-				string DepotPath = Object.GetValue("DepotPath", null);
+				string? DepotPath = Object.GetValue("DepotPath", null);
 				if (DepotPath == null)
 				{
 					Info = null;
 					return false;
 				}
 
-				string Target = Object.GetValue("Target", null);
+				string? Target = Object.GetValue("Target", null);
 
 				string Type = Object.GetValue("Type", null) ?? Name;
 
-				Info = new PerforceArchiveInfo(Name, Type, DepotPath, Target, PerforceConnection);
+				Info = new PerforceArchiveInfo(Name, Type, DepotPath, Target, Outer);
 				return true;
 			}
 
-			public bool TryGetArchiveKeyForChangeNumber(int ChangeNumber, out string ArchiveKey)
+			public bool TryGetArchiveKeyForChangeNumber(int ChangeNumber, [NotNullWhen(true)] out string? ArchiveKey)
 			{
 				return ChangeNumberToFileRevision.TryGetValue(ChangeNumber, out ArchiveKey);
 			}
 
-			public bool DownloadArchive(string ArchiveKey, DirectoryReference LocalRootPath, FileReference ManifestFileName, TextWriter Log, ProgressValue Progress)
+			public async Task<bool> DownloadArchive(IPerforceConnection Perforce, string ArchiveKey, DirectoryReference LocalRootPath, FileReference ManifestFileName, ILogger Logger, ProgressValue Progress, CancellationToken CancellationToken)
 			{
 				FileReference TempZipFileName = new FileReference(Path.GetTempFileName());
 				try
 				{
-					if (!Perforce.PrintToFile(ArchiveKey, TempZipFileName.FullName, Log) || TempZipFileName.ToFileInfo().Length == 0)
+					PrintRecord Record = await Perforce.PrintAsync(TempZipFileName.FullName, ArchiveKey, CancellationToken);
+					if (TempZipFileName.ToFileInfo().Length == 0)
 					{
 						return false;
 					}
-
-					ArchiveUtils.ExtractFiles(TempZipFileName, LocalRootPath, ManifestFileName, Progress, Log);
-
+					ArchiveUtils.ExtractFiles(TempZipFileName, LocalRootPath, ManifestFileName, Progress, Logger);
 				}
 				finally
 				{
@@ -184,37 +190,39 @@ namespace UnrealGameSync
 			}
 		}
 
-		class PerforceChangeSorter : IComparer<PerforceChangeSummary>
+		class PerforceChangeSorter : IComparer<ChangesRecord>
 		{
-			public int Compare(PerforceChangeSummary SummaryA, PerforceChangeSummary SummaryB)
+			public int Compare(ChangesRecord? SummaryA, ChangesRecord? SummaryB)
 			{
-				return SummaryB.Number - SummaryA.Number;
+				return SummaryB!.Number - SummaryA!.Number;
 			}
 		}
 
 		public int InitialMaxChangesValue = 100;
 
-		PerforceConnection Perforce;
+		IPerforceSettings PerforceSettings;
 		readonly string BranchClientPath;
 		readonly string SelectedClientFileName;
 		readonly string SelectedProjectIdentifier;
-		Thread WorkerThread;
+		Task? WorkerTask;
+		CancellationTokenSource CancellationSource;
 		int PendingMaxChangesValue;
-		SortedSet<PerforceChangeSummary> Changes = new SortedSet<PerforceChangeSummary>(new PerforceChangeSorter());
+		SortedSet<ChangesRecord> Changes = new SortedSet<ChangesRecord>(new PerforceChangeSorter());
 		SortedDictionary<int, PerforceChangeDetails> ChangeDetails = new SortedDictionary<int,PerforceChangeDetails>();
 		SortedSet<int> PromotedChangeNumbers = new SortedSet<int>();
 		List<PerforceArchiveInfo> Archives = new List<PerforceArchiveInfo>();
-		AutoResetEvent RefreshEvent = new AutoResetEvent(false);
-		BoundedLogWriter LogWriter;
+		AsyncEvent RefreshEvent = new AsyncEvent();
+		ILogger Logger;
 		bool bIsEnterpriseProject;
-		bool bDisposing;
 		DirectoryReference CacheFolder;
 		List<KeyValuePair<FileReference, DateTime>> LocalConfigFiles;
+		IAsyncDisposer AsyncDisposeTasks;
 
-		public event Action OnUpdate;
-		public event Action OnUpdateMetadata;
-		public event Action OnStreamChange;
-		public event Action OnLoginExpired;
+		SynchronizationContext SynchronizationContext;
+		public event Action? OnUpdate;
+		public event Action? OnUpdateMetadata;
+		public event Action? OnStreamChange;
+		public event Action? OnLoginExpired;
 
 		public TimeSpan ServerTimeZone
 		{
@@ -222,46 +230,47 @@ namespace UnrealGameSync
 			private set;
 		}
 
-		public PerforceMonitor(PerforceConnection InPerforce, string InBranchClientPath, string InSelectedClientFileName, string InSelectedProjectIdentifier, FileReference InLogPath, bool bInIsEnterpriseProject, ConfigFile InProjectConfigFile, DirectoryReference InCacheFolder, List<KeyValuePair<FileReference, DateTime>> InLocalConfigFiles)
+		public PerforceMonitor(IPerforceSettings InPerforceSettings, ProjectInfo ProjectInfo, ConfigFile InProjectConfigFile, DirectoryReference InCacheFolder, List<KeyValuePair<FileReference, DateTime>> InLocalConfigFiles, IServiceProvider InServiceProvider)
 		{
-			Perforce = InPerforce;
-			BranchClientPath = InBranchClientPath;
-			SelectedClientFileName = InSelectedClientFileName;
-			SelectedProjectIdentifier = InSelectedProjectIdentifier;
+			PerforceSettings = InPerforceSettings;
+			BranchClientPath = ProjectInfo.ClientRootPath;
+			SelectedClientFileName = ProjectInfo.ClientFileName;
+			SelectedProjectIdentifier = ProjectInfo.ProjectIdentifier;
 			PendingMaxChangesValue = InitialMaxChangesValue;
 			LastChangeByCurrentUser = -1;
 			LastCodeChangeByCurrentUser = -1;
-			bIsEnterpriseProject = bInIsEnterpriseProject;
+			Logger = InServiceProvider.GetRequiredService<ILogger<PerforceMonitor>>();
+			bIsEnterpriseProject = ProjectInfo.bIsEnterpriseProject;
 			LatestProjectConfigFile = InProjectConfigFile;
 			CacheFolder = InCacheFolder;
 			LocalConfigFiles = InLocalConfigFiles;
+			AsyncDisposeTasks = InServiceProvider.GetRequiredService<IAsyncDisposer>();
+			SynchronizationContext = SynchronizationContext.Current!;
+			CancellationSource = new CancellationTokenSource();
 
-			LogWriter = new BoundedLogWriter(InLogPath);
 			AvailableArchives = (new List<IArchiveInfo>()).AsReadOnly();
 		}
 
 		public void Start()
 		{
-			WorkerThread = new Thread(() => PollForUpdates());
-			WorkerThread.Start();
+			if (WorkerTask == null)
+			{
+				WorkerTask = Task.Run(() => PollForUpdates(CancellationSource.Token));
+			}
 		}
 
 		public void Dispose()
 		{
-			bDisposing = true;
-			if(WorkerThread != null)
-			{
-				RefreshEvent.Set();
-				if(!WorkerThread.Join(100))
-				{
-					WorkerThread.Interrupt();
-				}
-				WorkerThread = null;
-			}
-			if(LogWriter != null)
+			OnUpdate = null;
+			OnUpdateMetadata = null;
+			OnStreamChange = null;
+			OnLoginExpired = null;
+
+			if (WorkerTask != null)
 			{
-				LogWriter.Dispose();
-				LogWriter = null;
+				CancellationSource.Cancel();
+				AsyncDisposeTasks.Add(WorkerTask.ContinueWith(_ => CancellationSource.Dispose()));
+				WorkerTask = null;
 			}
 		}
 
@@ -275,7 +284,7 @@ namespace UnrealGameSync
 		{
 			get;
 			private set;
-		}
+		} = "";
 
 		public int CurrentMaxChanges
 		{
@@ -289,94 +298,100 @@ namespace UnrealGameSync
 			set { lock(this){ if(value != PendingMaxChangesValue){ PendingMaxChangesValue = value; RefreshEvent.Set(); } } }
 		}
 
-		void PollForUpdates()
+		async Task PollForUpdates(CancellationToken CancellationToken)
 		{
-			while (!bDisposing)
+			bool bCoolDown = false;
+			while (!CancellationToken.IsCancellationRequested)
 			{
 				try
 				{
-					PollForUpdatesInner();
+					if (bCoolDown)
+					{
+						await Task.Delay(TimeSpan.FromSeconds(5.0), CancellationToken);
+					}
+					else
+					{
+						await PollForUpdatesInner(CancellationToken);
+					}
 				}
-				catch (ThreadInterruptedException)
+				catch (TaskCanceledException) when (CancellationToken.IsCancellationRequested)
 				{
 					break;
 				}
 				catch (Exception Ex)
 				{
-					if(!bDisposing)
-					{
-						LogWriter.WriteException(Ex, "Unhandled exception in PollForUpdatesInner()");
-					}
+					Logger.LogError(Ex, "Unhandled exception in PollForUpdatesInner()");
+					bCoolDown = true;
 				}
 			}
 		}
 
-		void PollForUpdatesInner()
+		async Task PollForUpdatesInner(CancellationToken CancellationToken)
 		{
-			string StreamName;
-			if(!Perforce.GetActiveStream(out StreamName, LogWriter))
-			{
-				StreamName = null;
-			}
+			using IPerforceConnection Perforce = await PerforceConnection.CreateAsync(PerforceSettings, Logger);
+			string? StreamName = await Perforce.GetCurrentStreamAsync(CancellationToken);
 
 			// Get the perforce server settings
-			PerforceInfoRecord PerforceInfo;
-			if(Perforce.Info(out PerforceInfo, LogWriter))
+			PerforceResponse<InfoRecord> InfoResponse = await Perforce.TryGetInfoAsync(InfoOptions.ShortOutput, CancellationToken);
+			if (InfoResponse.Succeeded)
 			{
-				ServerTimeZone = PerforceInfo.ServerTimeZone;
+				DateTimeOffset? ServerDate = InfoResponse.Data.ServerDate;
+				if (ServerDate.HasValue)
+				{
+					ServerTimeZone = ServerDate.Value.Offset;
+				}
 			}
 
 			// Try to update the zipped binaries list before anything else, because it causes a state change in the UI
-			UpdateArchives();
+			await UpdateArchivesAsync(Perforce, CancellationToken);
 
-			while(!bDisposing)
+			while(!CancellationToken.IsCancellationRequested)
 			{
 				Stopwatch Timer = Stopwatch.StartNew();
+				Task NextRefreshTask = RefreshEvent.Task;
 
 				// Check we still have a valid login ticket
-				bool bLoggedIn;
-				if(Perforce.GetLoggedInState(out bLoggedIn, LogWriter))
+				PerforceResponse<LoginRecord> LoginState = await Perforce.TryGetLoginStateAsync(CancellationToken);
+				if(!LoginState.Succeeded)
 				{
-					if(!bLoggedIn)
+					LastStatusMessage = "User is not logged in";
+					SynchronizationContext.Post(_ => OnLoginExpired?.Invoke(), null);
+				}
+				else
+				{
+					// Check we haven't switched streams
+					string? NewStreamName = await Perforce.GetCurrentStreamAsync(CancellationToken);
+					if (NewStreamName != StreamName)
 					{
-						LastStatusMessage = "User is not logged in";
-						OnLoginExpired();
+						SynchronizationContext.Post(_ => OnStreamChange?.Invoke(), null);
+					}
+
+					// Check for any p4 changes
+					if(!await UpdateChangesAsync(Perforce, CancellationToken))
+					{
+						LastStatusMessage = "Failed to update changes";
+					}
+					else if(!await UpdateChangeTypesAsync(Perforce, CancellationToken))
+					{
+						LastStatusMessage = "Failed to update change types";
+					}
+					else if(!await UpdateArchivesAsync(Perforce, CancellationToken))
+					{
+						LastStatusMessage = "Failed to update zipped binaries list";
 					}
 					else
 					{
-						// Check we haven't switched streams
-						string NewStreamName;
-						if(Perforce.GetActiveStream(out NewStreamName, LogWriter) && NewStreamName != StreamName)
-						{
-							OnStreamChange();
-						}
-
-						// Check for any p4 changes
-						if(!UpdateChanges())
-						{
-							LastStatusMessage = "Failed to update changes";
-						}
-						else if(!UpdateChangeTypes())
-						{
-							LastStatusMessage = "Failed to update change types";
-						}
-						else if(!UpdateArchives())
-						{
-							LastStatusMessage = "Failed to update zipped binaries list";
-						}
-						else
-						{
-							LastStatusMessage = String.Format("Last update took {0}ms", Timer.ElapsedMilliseconds);
-						}
+						LastStatusMessage = String.Format("Last update took {0}ms", Timer.ElapsedMilliseconds);
 					}
 				}
 
 				// Wait for another request, or scan for new builds after a timeout
-				RefreshEvent.WaitOne((IsActive? 2 : 10) * 60 * 1000);
+				Task DelayTask = Task.Delay(TimeSpan.FromMinutes(IsActive ? 2 : 10), CancellationToken);
+				await Task.WhenAny(NextRefreshTask, DelayTask);
 			}
 		}
 
-		bool UpdateChanges()
+		async Task<bool> UpdateChangesAsync(IPerforceConnection Perforce, CancellationToken CancellationToken)
 		{
 			// Get the current status of the build
 			int MaxChanges;
@@ -424,20 +439,14 @@ namespace UnrealGameSync
 			}
 
 			// Read any new changes
-			List<PerforceChangeSummary> NewChanges;
+			List<ChangesRecord> NewChanges;
 			if(MaxChanges > CurrentMaxChanges)
 			{
-				if(!Perforce.FindChanges(DepotPaths, MaxChanges, out NewChanges, LogWriter))
-				{
-					return false;
-				}
+				NewChanges = await Perforce.GetChangesAsync(ChangesOptions.IncludeTimes | ChangesOptions.LongOutput, MaxChanges, ChangeStatus.Submitted, DepotPaths, CancellationToken);
 			}
 			else
 			{
-				if(!Perforce.FindChanges(DepotPaths.Select(DepotPath => String.Format("{0}@>{1}", DepotPath, NewestChangeNumber)), -1, out NewChanges, LogWriter))
-				{
-					return false;
-				}
+				NewChanges = await Perforce.GetChangesAsync(ChangesOptions.IncludeTimes | ChangesOptions.LongOutput, -1, ChangeStatus.Submitted, DepotPaths.Select(x => $"{x}@>{NewestChangeNumber}").ToArray(), CancellationToken);
 			}
 
 			// Remove anything we already have
@@ -465,18 +474,16 @@ namespace UnrealGameSync
 			}
 			if(MinZippedChangeNumber != -1 && MinZippedChangeNumber < OldestChangeNumber)
 			{
-				List<PerforceChangeSummary> ZipChanges;
-				if(Perforce.FindChanges(DepotPaths.Select(DepotPath => String.Format("{0}@{1},{2}", DepotPath, MinZippedChangeNumber, OldestChangeNumber - 1)), -1, out ZipChanges, LogWriter))
-				{
-					NewChanges.AddRange(ZipChanges);
-				}
+				string[] FilteredPaths = DepotPaths.Select(x => $"{x}@{MinZippedChangeNumber},{OldestChangeNumber - 1}").ToArray();
+				List<ChangesRecord> ZipChanges = await Perforce.GetChangesAsync(ChangesOptions.None, -1, ChangeStatus.Submitted, FilteredPaths, CancellationToken);
+				NewChanges.AddRange(ZipChanges);
 			}
 
 			// Fixup any ROBOMERGE authors
 			const string RoboMergePrefix = "#ROBOMERGE-AUTHOR:";
-			foreach(PerforceChangeSummary Change in NewChanges)
+			foreach (ChangesRecord Change in NewChanges)
 			{
-				if(Change.Description.StartsWith(RoboMergePrefix))
+				if(Change.Description != null && Change.Description.StartsWith(RoboMergePrefix))
 				{
 					int StartIdx = RoboMergePrefix.Length;
 					while(StartIdx < Change.Description.Length && Change.Description[StartIdx] == ' ')
@@ -509,8 +516,8 @@ namespace UnrealGameSync
 					{
 						// Remove changes to shrink it to the max requested size, being careful to avoid removing changes that would affect our ability to correctly
 						// show the availability for content changes using zipped binaries.
-						SortedSet<PerforceChangeSummary> TrimmedChanges = new SortedSet<PerforceChangeSummary>(new PerforceChangeSorter());
-						foreach(PerforceChangeSummary Change in Changes)
+						SortedSet<ChangesRecord> TrimmedChanges = new SortedSet<ChangesRecord>(new PerforceChangeSorter());
+						foreach(ChangesRecord Change in Changes)
 						{
 							TrimmedChanges.Add(Change);
 							if(TrimmedChanges.Count >= MaxChanges && Archives.Any(x => x.ChangeNumberToFileRevision.Count == 0 || x.ChangeNumberToFileRevision.ContainsKey(Change.Number) || x.ChangeNumberToFileRevision.First().Key > Change.Number))
@@ -525,9 +532,9 @@ namespace UnrealGameSync
 
 				// Find the last submitted change by the current user
 				int NewLastChangeByCurrentUser = -1;
-				foreach(PerforceChangeSummary Change in Changes)
+				foreach(ChangesRecord Change in Changes)
 				{
-					if(String.Compare(Change.User, Perforce.UserName, StringComparison.InvariantCultureIgnoreCase) == 0)
+					if(String.Compare(Change.User, Perforce.Settings.UserName, StringComparison.InvariantCultureIgnoreCase) == 0)
 					{
 						NewLastChangeByCurrentUser = Math.Max(NewLastChangeByCurrentUser, Change.Number);
 					}
@@ -535,21 +542,18 @@ namespace UnrealGameSync
 				LastChangeByCurrentUser = NewLastChangeByCurrentUser;
 
 				// Notify the main window that we've got more data
-				if(OnUpdate != null)
-				{
-					OnUpdate();
-				}
+				SynchronizationContext.Post(_ => OnUpdate?.Invoke(), null);
 			}
 			return true;
 		}
 
-		public bool UpdateChangeTypes()
+		public async Task<bool> UpdateChangeTypesAsync(IPerforceConnection Perforce, CancellationToken CancellationToken)
 		{
 			// Find the changes we need to query
 			List<int> QueryChangeNumbers = new List<int>();
 			lock(this)
 			{
-				foreach(PerforceChangeSummary Change in Changes)
+				foreach(ChangesRecord Change in Changes)
 				{
 					if(!ChangeDetails.ContainsKey(Change.Number))
 					{
@@ -559,70 +563,65 @@ namespace UnrealGameSync
 			}
 
 			// Update them in batches
-			foreach(int QueryChangeNumber in QueryChangeNumbers)
+			foreach (IReadOnlyList<int> QueryChangeNumberBatch in QueryChangeNumbers.Batch(10))
 			{
 				// Skip this stuff if the user wants us to query for more changes
-				if(PendingMaxChanges != CurrentMaxChanges)
+				if (PendingMaxChanges != CurrentMaxChanges)
 				{
 					break;
 				}
 
 				// If there's something to check for, find all the content changes after this changelist
-				PerforceDescribeRecord DescribeRecord;
-				if(Perforce.Describe(QueryChangeNumber, out DescribeRecord, LogWriter))
+				List<DescribeRecord> DescribeRecords = await Perforce.DescribeAsync(QueryChangeNumberBatch.ToArray(), CancellationToken);
+				foreach (DescribeRecord DescribeRecord in DescribeRecords.OrderByDescending(x => x.Number))
 				{
-					// Create the details object
+					int QueryChangeNumber = DescribeRecord.Number;
+
 					PerforceChangeDetails Details = new PerforceChangeDetails(DescribeRecord);
-					lock(this)
+					lock (this)
 					{
-						if(!ChangeDetails.ContainsKey(QueryChangeNumber))
+						if (!ChangeDetails.ContainsKey(QueryChangeNumber))
 						{
 							ChangeDetails.Add(QueryChangeNumber, Details);
 						}
 					}
 
 					// Reload the config file if it changes
-					if(DescribeRecord.Files.Any(x => x.DepotFile.EndsWith("/UnrealGameSync.ini", StringComparison.InvariantCultureIgnoreCase)))
+					if (DescribeRecord.Files.Any(x => x.DepotFile.EndsWith("/UnrealGameSync.ini", StringComparison.OrdinalIgnoreCase)))
 					{
-						UpdateProjectConfigFile();
+						await UpdateProjectConfigFileAsync(Perforce, CancellationToken);
 					}
 				}
+			}
 
-				// Find the last submitted code change by the current user
-				int NewLastCodeChangeByCurrentUser = -1;
-				foreach(PerforceChangeSummary Change in Changes)
+			// Find the last submitted code change by the current user
+			int NewLastCodeChangeByCurrentUser = -1;
+			foreach(ChangesRecord Change in Changes)
+			{
+				if(String.Compare(Change.User, Perforce.Settings.UserName, StringComparison.InvariantCultureIgnoreCase) == 0)
 				{
-					if(String.Compare(Change.User, Perforce.UserName, StringComparison.InvariantCultureIgnoreCase) == 0)
+					PerforceChangeDetails? OtherDetails;
+					if(ChangeDetails.TryGetValue(Change.Number, out OtherDetails) && OtherDetails.bContainsCode)
 					{
-						PerforceChangeDetails Details;
-						if(ChangeDetails.TryGetValue(Change.Number, out Details) && Details.bContainsCode)
-						{
-							NewLastCodeChangeByCurrentUser = Math.Max(NewLastCodeChangeByCurrentUser, Change.Number);
-						}
+						NewLastCodeChangeByCurrentUser = Math.Max(NewLastCodeChangeByCurrentUser, Change.Number);
 					}
 				}
-				LastCodeChangeByCurrentUser = NewLastCodeChangeByCurrentUser;
-
-				// Notify the main window that we've got an update
-				if(OnUpdateMetadata != null)
-				{
-					OnUpdateMetadata();
-				}
 			}
+			LastCodeChangeByCurrentUser = NewLastCodeChangeByCurrentUser;
+
+			// Notify the main window that we've got an update
+			SynchronizationContext.Post(_ => OnUpdateMetadata?.Invoke(), null);
 
 			if(LocalConfigFiles.Any(x => FileReference.GetLastWriteTimeUtc(x.Key) != x.Value))
 			{
-				UpdateProjectConfigFile();
-				if(OnUpdateMetadata != null)
-				{
-					OnUpdateMetadata();
-				}
+				await UpdateProjectConfigFileAsync(Perforce, CancellationToken);
+				SynchronizationContext.Post(_ => OnUpdateMetadata?.Invoke(), null);
 			}
 
 			return true;
 		}
 
-		bool UpdateArchives()
+		async Task<bool> UpdateArchivesAsync(IPerforceConnection Perforce, CancellationToken CancellationToken)
 		{
 			List<PerforceArchiveInfo> NewArchives = new List<PerforceArchiveInfo>();
 
@@ -631,17 +630,17 @@ namespace UnrealGameSync
 			if (ProjectConfigSection != null)
 			{
 				// Legacy
-				string LegacyEditorArchivePath = ProjectConfigSection.GetValue("ZippedBinariesPath", null);
+				string? LegacyEditorArchivePath = ProjectConfigSection.GetValue("ZippedBinariesPath", null);
 				if (LegacyEditorArchivePath != null)
 				{
-					NewArchives.Add(new PerforceArchiveInfo("Editor", "Editor", LegacyEditorArchivePath, null, Perforce));
+					NewArchives.Add(new PerforceArchiveInfo("Editor", "Editor", LegacyEditorArchivePath, null, this));
 				}
 
 				// New style
 				foreach (string ArchiveValue in ProjectConfigSection.GetValues("Archives", new string[0]))
 				{
-					PerforceArchiveInfo Archive;
-					if (PerforceArchiveInfo.TryParseConfigEntry(ArchiveValue, Perforce, out Archive))
+					PerforceArchiveInfo? Archive;
+					if (PerforceArchiveInfo.TryParseConfigEntry(ArchiveValue, this, out Archive))
 					{
 						NewArchives.Add(Archive);
 					}
@@ -650,32 +649,24 @@ namespace UnrealGameSync
 				// Make sure the zipped binaries path exists
 				foreach (PerforceArchiveInfo NewArchive in NewArchives)
 				{
-					bool bExists;
-					if (!Perforce.FileExists(NewArchive.DepotPath, out bExists, LogWriter))
-					{
-						return false;
-					}
-					if (bExists)
+					PerforceResponseList<FileLogRecord> Response = await Perforce.TryFileLogAsync(30, FileLogOptions.FullDescriptions, NewArchive.DepotPath, CancellationToken);
+					if (Response.Succeeded)
 					{
-						// Query all the changes to this file
-						List<PerforceFileChangeSummary> Changes;
-						if (!Perforce.FindFileChanges(NewArchive.DepotPath, 100, out Changes, LogWriter))
-						{
-							return false;
-						}
-
-						// Build a new list of zipped binaries
-						foreach (PerforceFileChangeSummary Change in Changes)
+ 						// Build a new list of zipped binaries
+						foreach (FileLogRecord File in Response.Data)
 						{
-							if (Change.Action != "purge")
+							foreach(RevisionRecord Revision in File.Revisions)
 							{
-								string[] Tokens = Change.Description.Split(' ');
-								if (Tokens[0].StartsWith("[CL") && Tokens[1].EndsWith("]"))
+								if (Revision.Action != FileAction.Purge)
 								{
-									int OriginalChangeNumber;
-									if (int.TryParse(Tokens[1].Substring(0, Tokens[1].Length - 1), out OriginalChangeNumber) && !NewArchive.ChangeNumberToFileRevision.ContainsKey(OriginalChangeNumber))
+									string[] Tokens = Revision.Description.Split(' ');
+									if (Tokens[0].StartsWith("[CL") && Tokens[1].EndsWith("]"))
 									{
-										NewArchive.ChangeNumberToFileRevision[OriginalChangeNumber] = String.Format("{0}#{1}", NewArchive.DepotPath, Change.Revision);
+										int OriginalChangeNumber;
+										if (int.TryParse(Tokens[1].Substring(0, Tokens[1].Length - 1), out OriginalChangeNumber) && !NewArchive.ChangeNumberToFileRevision.ContainsKey(OriginalChangeNumber))
+										{
+											NewArchive.ChangeNumberToFileRevision[OriginalChangeNumber] = $"{NewArchive.DepotPath}#{Revision.RevisionNumber}";
+										}
 									}
 								}
 							}
@@ -690,60 +681,61 @@ namespace UnrealGameSync
 				Archives = NewArchives;
 				AvailableArchives = Archives.Select(x => (IArchiveInfo)x).ToList();
 
-				if (OnUpdateMetadata != null && Changes.Count > 0)
+				if (Changes.Count > 0)
 				{
-					OnUpdateMetadata();
+					SynchronizationContext.Post(_ => OnUpdateMetadata?.Invoke(), null);
 				}
 			}
 
 			return true;
 		}
 
-		void UpdateProjectConfigFile()
+		async Task UpdateProjectConfigFileAsync(IPerforceConnection Perforce, CancellationToken CancellationToken)
 		{
 			LocalConfigFiles.Clear();
-			LatestProjectConfigFile = ReadProjectConfigFile(Perforce, BranchClientPath, SelectedClientFileName, CacheFolder, LocalConfigFiles, LogWriter);
+			LatestProjectConfigFile = await ReadProjectConfigFileAsync(Perforce, BranchClientPath, SelectedClientFileName, CacheFolder, LocalConfigFiles, Logger, CancellationToken);
 		}
 
-		public static ConfigFile ReadProjectConfigFile(PerforceConnection Perforce, string BranchClientPath, string SelectedClientFileName, DirectoryReference CacheFolder, List<KeyValuePair<FileReference, DateTime>> LocalConfigFiles, TextWriter Log)
+		public static async Task<ConfigFile> ReadProjectConfigFileAsync(IPerforceConnection Perforce, string BranchClientPath, string SelectedClientFileName, DirectoryReference CacheFolder, List<KeyValuePair<FileReference, DateTime>> LocalConfigFiles, ILogger Logger, CancellationToken CancellationToken)
 		{
 			List<string> ConfigFilePaths = Utility.GetDepotConfigPaths(BranchClientPath + "/Engine", SelectedClientFileName);
 
 			ConfigFile ProjectConfig = new ConfigFile();
 
-			List<PerforceFileRecord> FileRecords;
-			if(Perforce.Stat(new List<string> { "-Ol" }, ConfigFilePaths, out FileRecords, Log))
+			List<PerforceResponse<FStatRecord>> Responses = await Perforce.TryFStatAsync(FStatOptions.IncludeFileSizes, ConfigFilePaths, CancellationToken).ToListAsync(CancellationToken);
+			foreach(PerforceResponse<FStatRecord> Response in Responses)
 			{
-				foreach(PerforceFileRecord FileRecord in FileRecords)
+				if (Response.Succeeded)
 				{
-					List<string> Lines = null;
+					string[]? Lines = null;
 
 					// Skip file records which are still in the workspace, but were synced from a different branch. For these files, the action seems to be empty, so filter against that.
-					if(FileRecord.Action == null)
+					FStatRecord FileRecord = Response.Data;
+					if (FileRecord.HeadAction == FileAction.None)
 					{
 						continue;
 					}
 
 					// If this file is open for edit, read the local version
-					string LocalFileName = FileRecord.ClientPath;
+					string? LocalFileName = FileRecord.ClientFile;
 					if(LocalFileName != null && File.Exists(LocalFileName) && (File.GetAttributes(LocalFileName) & FileAttributes.ReadOnly) == 0)
 					{
 						try
 						{
 							DateTime LastModifiedTime = File.GetLastWriteTimeUtc(LocalFileName);
 							LocalConfigFiles.Add(new KeyValuePair<FileReference, DateTime>(new FileReference(LocalFileName), LastModifiedTime));
-							Lines = File.ReadAllLines(LocalFileName).ToList();
+							Lines = await File.ReadAllLinesAsync(LocalFileName, CancellationToken);
 						}
 						catch(Exception Ex)
 						{
-							Log.WriteLine("Failed to read local config file for {0}: {1}", LocalFileName, Ex.ToString());
+							Logger.LogInformation(Ex, "Failed to read local config file for {Path}", LocalFileName);
 						}
 					}
 
 					// Otherwise try to get it from perforce
-					if(Lines == null)
+					if(Lines == null && FileRecord.DepotFile != null)
 					{
-						Utility.TryPrintFileUsingCache(Perforce, FileRecord.DepotPath, CacheFolder, FileRecord.Digest, out Lines, Log);
+						Lines = await Utility.TryPrintFileUsingCacheAsync(Perforce, FileRecord.DepotFile, CacheFolder, FileRecord.Digest, Logger, CancellationToken);
 					}
 
 					// Merge the text with the config file
@@ -752,11 +744,11 @@ namespace UnrealGameSync
 						try
 						{
 							ProjectConfig.Parse(Lines.ToArray());
-							Log.WriteLine("Read config file from {0}", FileRecord.DepotPath);
+							Logger.LogInformation("Read config file from {DepotFile}", FileRecord.DepotFile);
 						}
 						catch(Exception Ex)
 						{
-							Log.WriteLine("Failed to read config file from {0}: {1}", FileRecord.DepotPath, Ex.ToString());
+							Logger.LogInformation(Ex, "Failed to read config file from {DepotFile}", FileRecord.DepotFile);
 						}
 					}
 				}
@@ -764,15 +756,15 @@ namespace UnrealGameSync
 			return ProjectConfig;
 		}
 
-		public List<PerforceChangeSummary> GetChanges()
+		public List<ChangesRecord> GetChanges()
 		{
 			lock(this)
 			{
-				return new List<PerforceChangeSummary>(Changes);
+				return new List<ChangesRecord>(Changes);
 			}
 		}
 
-		public bool TryGetChangeDetails(int Number, out PerforceChangeDetails Details)
+		public bool TryGetChangeDetails(int Number, [NotNullWhen(true)] out PerforceChangeDetails? Details)
 		{
 			lock(this)
 			{
diff --git a/Engine/Source/Programs/UnrealGameSync/UnrealGameSync/Program.cs b/Engine/Source/Programs/UnrealGameSync/UnrealGameSync/Program.cs
--- a/Engine/Source/Programs/UnrealGameSync/UnrealGameSync/Program.cs
+++ b/Engine/Source/Programs/UnrealGameSync/UnrealGameSync/Program.cs
@@ -1,6 +1,9 @@
 // Copyright Epic Games, Inc. All Rights Reserved.
 
 using EpicGames.Core;
+using EpicGames.Perforce;
+using Microsoft.Extensions.DependencyInjection;
+using Microsoft.Extensions.Logging;
 using System;
 using System.Collections.Generic;
 using System.Data.SqlClient;
@@ -40,7 +43,7 @@ namespace UnrealGameSync
 		public static string SyncVersion = null;
 
 		[STAThread]
-		static void Main(string[] Args)
+		static async Task Main(string[] Args)
 		{
 			bool bFirstInstance;
 			using (Mutex InstanceMutex = new Mutex(true, "UnrealGameSyncRunning", out bFirstInstance))
@@ -61,7 +64,7 @@ namespace UnrealGameSync
 
 					if (bFirstInstance)
 					{
-						InnerMain(InstanceMutex, ActivateEvent, Args);
+						await InnerMainAsync(InstanceMutex, ActivateEvent, Args);
 					}
 					else
 					{
@@ -71,7 +74,7 @@ namespace UnrealGameSync
 			}
 		}
 
-		static void InnerMain(Mutex InstanceMutex, EventWaitHandle ActivateEvent, string[] Args)
+		static async Task InnerMainAsync(Mutex InstanceMutex, EventWaitHandle ActivateEvent, string[] Args)
 		{
 			string ServerAndPort = null;
 			string UserName = null;
@@ -126,28 +129,37 @@ namespace UnrealGameSync
 			// Enable TLS 1.1 and 1.2. TLS 1.0 is now deprecated and not allowed by default in NET Core servers.
 			ServicePointManager.SecurityProtocol |= SecurityProtocolType.Tls11 | SecurityProtocolType.Tls12;
 
-			// Create the log file
-			using (TimestampLogWriter Log = new TimestampLogWriter(new BoundedLogWriter(FileReference.Combine(DataFolder, "UnrealGameSync.log"))))
+			// Create a new logger
+			using (ILoggerProvider LoggerProvider = Logging.CreateLoggerProvider(FileReference.Combine(DataFolder, "UnrealGameSync.log")))
 			{
-				Log.WriteLine("Application version: {0}", Assembly.GetExecutingAssembly().GetName().Version);
-				Log.WriteLine("Started at {0}", DateTime.Now.ToString());
+				ServiceCollection Services = new ServiceCollection();
+				Services.AddLogging(Builder => Builder.AddProvider(LoggerProvider));
+				Services.AddSingleton<IAsyncDisposer, AsyncDisposer>();
 
-				string SessionId = Guid.NewGuid().ToString();
-				Log.WriteLine("SessionId: {0}", SessionId);
-
-				if (ServerAndPort == null || UserName == null)
+				await using (ServiceProvider ServiceProvider = Services.BuildServiceProvider())
 				{
-					Log.WriteLine("Missing server settings; finding defaults.");
-					GetDefaultServerSettings(ref ServerAndPort, ref UserName, Log);
-					Utility.SaveGlobalPerforceSettings(ServerAndPort, UserName, BaseUpdatePath);
-				}
+					ILoggerFactory LoggerFactory = ServiceProvider.GetRequiredService<ILoggerFactory>();
 
-				using (BoundedLogWriter TelemetryLog = new BoundedLogWriter(FileReference.Combine(DataFolder, "Telemetry.log")))
-				{
-					TelemetryLog.WriteLine("Creating telemetry sink for session {0}", SessionId);
+					ILogger Logger = LoggerFactory.CreateLogger("Startup");
+					Logger.LogInformation("Application version: {Version}", Assembly.GetExecutingAssembly().GetName().Version);
+					Logger.LogInformation("Started at {Time}", DateTime.Now.ToString());
+
+					string SessionId = Guid.NewGuid().ToString();
+					Logger.LogInformation("SessionId: {SessionId}", SessionId);
+
+					if (ServerAndPort == null || UserName == null)
+					{
+						Logger.LogInformation("Missing server settings; finding defaults.");
+						ServerAndPort ??= PerforceSettings.Default.ServerAndPort;
+						UserName ??= PerforceSettings.Default.UserName;
+						Utility.SaveGlobalPerforceSettings(ServerAndPort, UserName, BaseUpdatePath);
+					}
+
+					ILogger TelemetryLogger = LoggerProvider.CreateLogger("Telemetry");
+					TelemetryLogger.LogInformation("Creating telemetry sink for session {SessionId}", SessionId);
 
 					ITelemetrySink PrevTelemetrySink = Telemetry.ActiveSink;
-					using (ITelemetrySink TelemetrySink = DeploymentSettings.CreateTelemetrySink(UserName, SessionId, TelemetryLog))
+					using (ITelemetrySink TelemetrySink = DeploymentSettings.CreateTelemetrySink(UserName, SessionId, TelemetryLogger))
 					{
 						Telemetry.ActiveSink = TelemetrySink;
 
@@ -155,10 +167,11 @@ namespace UnrealGameSync
 
 						AppDomain.CurrentDomain.UnhandledException += CurrentDomain_UnhandledException;
 
-						PerforceConnection DefaultConnection = new PerforceConnection(UserName, null, ServerAndPort);
-						using (UpdateMonitor UpdateMonitor = new UpdateMonitor(DefaultConnection, UpdatePath))
+						IPerforceSettings DefaultSettings = new PerforceSettings(ServerAndPort, UserName) { PreferNativeClient = true };
+
+						using (UpdateMonitor UpdateMonitor = new UpdateMonitor(DefaultSettings, UpdatePath, ServiceProvider))
 						{
-							ProgramApplicationContext Context = new ProgramApplicationContext(DefaultConnection, UpdateMonitor, DeploymentSettings.ApiUrl, DataFolder, ActivateEvent, bRestoreState, UpdateSpawn, ProjectFileName, bUnstable, Log, Uri);
+							ProgramApplicationContext Context = new ProgramApplicationContext(DefaultSettings, UpdateMonitor, DeploymentSettings.ApiUrl, DataFolder, ActivateEvent, bRestoreState, UpdateSpawn, ProjectFileName, bUnstable, ServiceProvider, Uri);
 							Application.Run(Context);
 
 							if (UpdateMonitor.IsUpdateAvailable && UpdateSpawn != null)
@@ -174,41 +187,6 @@ namespace UnrealGameSync
 			}
 		}
 
-		public static void GetDefaultServerSettings(ref string ServerAndPort, ref string UserName, TextWriter Log)
-		{
-			// Read the P4PORT setting for the server, if necessary. Change to the project folder if set, so we can respect the contents of any P4CONFIG file.
-			if(ServerAndPort == null)
-			{
-				PerforceConnection Perforce = new PerforceConnection(UserName, null, null);
-
-				string NewServerAndPort;
-				if (Perforce.GetSetting("P4PORT", out NewServerAndPort, Log))
-				{
-					ServerAndPort = NewServerAndPort;
-				}
-				else
-				{
-					ServerAndPort = PerforceConnection.DefaultServerAndPort;
-				}
-			}
-
-			// Update the server and username from the reported server info if it's not set
-			if(UserName == null)
-			{
-				PerforceConnection Perforce = new PerforceConnection(UserName, null, ServerAndPort);
-
-				PerforceInfoRecord PerforceInfo;
-				if(Perforce.Info(out PerforceInfo, Log) && !String.IsNullOrEmpty(PerforceInfo.UserName))
-				{
-					UserName = PerforceInfo.UserName;
-				}
-				else
-				{
-					UserName = Environment.UserName;
-				}
-			}
-		}
-
 		private static void CurrentDomain_UnhandledException(object Sender, UnhandledExceptionEventArgs Args)
 		{
 			Exception Ex = Args.ExceptionObject as Exception;
diff --git a/Engine/Source/Programs/UnrealGameSync/UnrealGameSync/ProgramApplicationContext.cs b/Engine/Source/Programs/UnrealGameSync/UnrealGameSync/ProgramApplicationContext.cs
--- a/Engine/Source/Programs/UnrealGameSync/UnrealGameSync/ProgramApplicationContext.cs
+++ b/Engine/Source/Programs/UnrealGameSync/UnrealGameSync/ProgramApplicationContext.cs
@@ -1,6 +1,9 @@
 ﻿// Copyright Epic Games, Inc. All Rights Reserved.
 
 using EpicGames.Core;
+using EpicGames.Perforce;
+using Microsoft.Extensions.DependencyInjection;
+using Microsoft.Extensions.Logging;
 using System;
 using System.Collections.Generic;
 using System.ComponentModel;
@@ -14,13 +17,15 @@ using System.Threading.Tasks;
 using System.Windows.Forms;
 using UnrealGameSync.Forms;
 
+#nullable enable
+
 namespace UnrealGameSync
 {
 	class ProgramApplicationContext : ApplicationContext
 	{
 		SynchronizationContext MainThreadSynchronizationContext;
 
-		PerforceConnection DefaultConnection;
+		IPerforceSettings DefaultPerforceSettings;
 		UpdateMonitor UpdateMonitor;
 		string ApiUrl;
 		DirectoryReference DataFolder;
@@ -31,7 +36,8 @@ namespace UnrealGameSync
 		bool bIsClosing;
 		string Uri;
 
-		TimestampLogWriter Log;
+		IServiceProvider ServiceProvider;
+		ILogger Logger;
 		UserSettings Settings;
 		ActivationListener ActivationListener;
 
@@ -45,16 +51,16 @@ namespace UnrealGameSync
 		ToolStripSeparator NotifyMenu_ExitSeparator;
 		ToolStripMenuItem NotifyMenu_Exit;
 
-		List<BufferedTextWriter> StartupLogs = new List<BufferedTextWriter>();
-		DetectMultipleProjectSettingsTask DetectStartupProjectSettingsTask;
-		ModalTaskWindow DetectStartupProjectSettingsWindow;
-		MainWindow MainWindowInstance;
+		CancellationTokenSource StartupCancellationSource = new CancellationTokenSource();
+		Task StartupTask;
+		ModalTaskWindow? StartupWindow;
+		MainWindow? MainWindowInstance;
 
-		OIDCTokenManager OIDCTokenManager;
+		OIDCTokenManager? OIDCTokenManager;
 
-		public ProgramApplicationContext(PerforceConnection DefaultConnection, UpdateMonitor UpdateMonitor, string ApiUrl, DirectoryReference DataFolder, EventWaitHandle ActivateEvent, bool bRestoreState, string UpdateSpawn, string ProjectFileName, bool bUnstable, TimestampLogWriter Log, string Uri)
+		public ProgramApplicationContext(IPerforceSettings DefaultPerforceSettings, UpdateMonitor UpdateMonitor, string ApiUrl, DirectoryReference DataFolder, EventWaitHandle ActivateEvent, bool bRestoreState, string UpdateSpawn, string ProjectFileName, bool bUnstable, IServiceProvider ServiceProvider, string Uri)
 		{
-			this.DefaultConnection = DefaultConnection;
+			this.DefaultPerforceSettings = DefaultPerforceSettings;
 			this.UpdateMonitor = UpdateMonitor;
 			this.ApiUrl = ApiUrl;
 			this.DataFolder = DataFolder;
@@ -62,7 +68,8 @@ namespace UnrealGameSync
 			this.bRestoreState = bRestoreState;
 			this.UpdateSpawn = UpdateSpawn;
 			this.bUnstable = bUnstable;
-			this.Log = Log;
+			this.ServiceProvider = ServiceProvider;
+			this.Logger = ServiceProvider.GetRequiredService<ILogger<ProgramApplicationContext>>();
 			this.Uri = Uri;
 
 			// Create the directories
@@ -77,10 +84,10 @@ namespace UnrealGameSync
 			}
 
 			// Capture the main thread's synchronization context for callbacks
-			MainThreadSynchronizationContext = WindowsFormsSynchronizationContext.Current;
+			MainThreadSynchronizationContext = SynchronizationContext.Current!;
 
 			// Read the user's settings
-			Settings = new UserSettings(FileReference.Combine(DataFolder, "UnrealGameSync.ini"), Log);
+			Settings = new UserSettings(FileReference.Combine(DataFolder, "UnrealGameSync.ini"), ServiceProvider.GetRequiredService<ILogger<UserSettings>>());
 			if(!String.IsNullOrEmpty(ProjectFileName))
 			{
 				string FullProjectFileName = Path.GetFullPath(ProjectFileName);
@@ -96,7 +103,7 @@ namespace UnrealGameSync
 				// Clear out the server settings for anything using the default server
 				if(Settings.Version < UserSettingsVersion.DefaultServerSettings)
 				{
-					Log.WriteLine("Clearing project settings for default server");
+					Logger.LogInformation("Clearing project settings for default server");
 					for (int Idx = 0; Idx < Settings.OpenProjects.Count; Idx++)
 					{
 						Settings.OpenProjects[Idx] = UpgradeSelectedProjectSettings(Settings.OpenProjects[Idx]);
@@ -176,43 +183,39 @@ namespace UnrealGameSync
 			NotifyIcon.DoubleClick += new EventHandler(NotifyIcon_DoubleClick);
 			NotifyIcon.MouseDown += new MouseEventHandler(NotifyIcon_MouseDown);
 
-			// Find the initial list of projects to attempt to reopen
-			List<DetectProjectSettingsTask> Tasks = new List<DetectProjectSettingsTask>();
-			foreach(UserSelectedProjectSettings OpenProject in Settings.OpenProjects)
+			// Create the startup tasks
+			List<(UserSelectedProjectSettings, ModalTask<WorkspaceSettings>)> StartupTasks = new List<(UserSelectedProjectSettings, ModalTask<WorkspaceSettings>)>();
+			foreach (UserSelectedProjectSettings ProjectSettings in Settings.OpenProjects)
 			{
-				Log.WriteLine("Opening existing project {0}", OpenProject);
-
-				BufferedTextWriter StartupLog = new BufferedTextWriter();
-				StartupLog.WriteLine("Detecting settings for {0}", OpenProject);
-				StartupLogs.Add(StartupLog);
-
-				Tasks.Add(new DetectProjectSettingsTask(OpenProject, DataFolder, CacheFolder, new TimestampLogWriter(new PrefixedTextWriter("  ", StartupLog))));
+				ILogger<WorkspaceSettings> Logger = ServiceProvider.GetRequiredService<ILogger<WorkspaceSettings>>();
+				Task<WorkspaceSettings> WorkspaceSettingsTask = Task.Run(() => WorkspaceSettings.CreateAsync(DefaultPerforceSettings, ProjectSettings, Logger, StartupCancellationSource.Token), StartupCancellationSource.Token);
+				StartupTasks.Add((ProjectSettings, new ModalTask<WorkspaceSettings>(WorkspaceSettingsTask)));
 			}
+			StartupTask = Task.WhenAll(StartupTasks.Select(x => x.Item2.Task));
 
-			// Detect settings for the project we want to open
-			DetectStartupProjectSettingsTask = new DetectMultipleProjectSettingsTask(DefaultConnection, Tasks);
+			StartupWindow = new ModalTaskWindow("Opening Projects", "Opening projects, please wait...", FormStartPosition.CenterScreen, StartupTask, StartupCancellationSource);
+			Components.Add(StartupWindow);
 
-			DetectStartupProjectSettingsWindow = new ModalTaskWindow(DetectStartupProjectSettingsTask, "Opening Projects", "Opening projects, please wait...", FormStartPosition.CenterScreen);
 			if(bRestoreState)
 			{
 				if(Settings.bWindowVisible)
 				{
-					DetectStartupProjectSettingsWindow.Show();
+					StartupWindow.Show();
 				}
 			}
 			else
 			{
-				DetectStartupProjectSettingsWindow.Show();
-				DetectStartupProjectSettingsWindow.Activate();
+				StartupWindow.Show();
+				StartupWindow.Activate();
 			}
-			DetectStartupProjectSettingsWindow.Complete += OnDetectStartupProjectsComplete;
+			StartupWindow.FormClosed += (s, e) => OnStartupComplete(StartupTasks);
 		}
 
 		private UserSelectedProjectSettings UpgradeSelectedProjectSettings(UserSelectedProjectSettings Project)
 		{
-			if (Project.ServerAndPort == null || String.Compare(Project.ServerAndPort, DefaultConnection.ServerAndPort, StringComparison.OrdinalIgnoreCase) == 0)
+			if (Project.ServerAndPort == null || String.Compare(Project.ServerAndPort, DefaultPerforceSettings.ServerAndPort, StringComparison.OrdinalIgnoreCase) == 0)
 			{
-				if (Project.UserName == null || String.Compare(Project.UserName, DefaultConnection.UserName, StringComparison.OrdinalIgnoreCase) == 0)
+				if (Project.UserName == null || String.Compare(Project.UserName, DefaultPerforceSettings.UserName, StringComparison.OrdinalIgnoreCase) == 0)
 				{
 					Project = new UserSelectedProjectSettings(null, null, Project.Type, Project.ClientPath, Project.LocalPath);
 				}
@@ -220,29 +223,16 @@ namespace UnrealGameSync
 			return Project;
 		}
 
-		private void OnDetectStartupProjectsComplete()
+		private void OnStartupComplete(List<(UserSelectedProjectSettings, ModalTask<WorkspaceSettings>)> StartupTasks)
 		{
 			// Close the startup window
-			bool bVisible = DetectStartupProjectSettingsWindow.Visible;
-			DetectStartupProjectSettingsWindow.Close();
-			DetectStartupProjectSettingsWindow = null;
-
-			// Copy all the logs to the main log
-			foreach(BufferedTextWriter StartupLog in StartupLogs)
-			{
-				foreach(string Line in StartupLog.Lines)
-				{
-					Log.WriteLine("{0}", Line);
-				}
-			}
+			bool bVisible = StartupWindow!.Visible;
+			StartupWindow = null;
 
 			// Clear out the cache folder
 			Utility.ClearPrintCache(CacheFolder);
 
-			// Get a list of all the valid projects to open
-			DetectProjectSettingsResult[] StartupProjects = DetectStartupProjectSettingsTask.Results.Where(x => x != null).ToArray();
-
-			OIDCTokenManager = OIDCTokenManager.CreateFromConfigFile(Settings, StartupProjects.Select(Result => Result.Task).ToList());
+			OIDCTokenManager = OIDCTokenManager.CreateFromConfigFile(Settings, StartupTasks.Where(x => x.Item2.Succeeded).Select(x => x.Item2.Result).ToList());
 			// Verify that none of the projects we are opening needs a OIDC login, if they do prompt for the login
 			if (OIDCTokenManager?.HasUnfinishedLogin() ?? false)
 			{
@@ -258,7 +248,7 @@ namespace UnrealGameSync
 			}
 
 			// Create the main window
-			MainWindowInstance = new MainWindow(UpdateMonitor, ApiUrl, DataFolder, CacheFolder, bRestoreState, OriginalExe, bUnstable, StartupProjects, DefaultConnection, Log, Settings, Uri, OIDCTokenManager);
+			MainWindowInstance = new MainWindow(UpdateMonitor, ApiUrl, DataFolder, CacheFolder, bRestoreState, OriginalExe, bUnstable, StartupTasks, DefaultPerforceSettings, ServiceProvider, Settings, Uri, OIDCTokenManager);
 			if(bVisible)
 			{
 				MainWindowInstance.Show();
@@ -268,10 +258,6 @@ namespace UnrealGameSync
 				}
 			}
 			MainWindowInstance.FormClosed += MainWindowInstance_FormClosed;
-
-			// Delete the project settings task
-			DetectStartupProjectSettingsTask.Dispose();
-			DetectStartupProjectSettingsTask = null;
 		}
 
 		private void MainWindowInstance_FormClosed(object sender, FormClosedEventArgs e)
@@ -316,9 +302,8 @@ namespace UnrealGameSync
 
 			if(UpdateMonitor != null)
 			{
-				UpdateMonitor.OnUpdateAvailable -= OnUpdateAvailableCallback;
-				UpdateMonitor.Close(); // prevent race condition
-				UpdateMonitor = null;
+				UpdateMonitor.Dispose();
+				UpdateMonitor = null!;
 			}
 
 			if(ActivationListener != null)
@@ -326,37 +311,37 @@ namespace UnrealGameSync
 				ActivationListener.OnActivate -= OnActivationListenerAsyncCallback;
 				ActivationListener.Stop();
 				ActivationListener.Dispose();
-				ActivationListener = null;
+				ActivationListener = null!;
 			}
 
 			if (Components != null)
 			{
 				Components.Dispose();
-				Components = null;
+				Components = null!;
 			}
 
 			if (NotifyIcon != null)
 			{
 				NotifyIcon.Dispose();
-				NotifyIcon = null;
+				NotifyIcon = null!;
 			}
 
-			if (Log != null)
+			if(MainWindowInstance != null)
 			{
-				Log.Dispose();
-				Log = null;
+				MainWindowInstance.ForceClose();
+				MainWindowInstance = null!;
 			}
 
-			if(MainWindowInstance != null)
+			if(StartupWindow != null)
 			{
-				MainWindowInstance.ForceClose();
-				MainWindowInstance = null;
+				StartupWindow.Close();
+				StartupWindow = null;
 			}
 
-			if(DetectStartupProjectSettingsWindow != null)
+			if (StartupCancellationSource != null)
 			{
-				DetectStartupProjectSettingsWindow.Close();
-				DetectStartupProjectSettingsWindow = null;
+				StartupCancellationSource.Dispose();
+				StartupCancellationSource = null!;
 			}
 		}
 
@@ -370,13 +355,13 @@ namespace UnrealGameSync
 			NotifyMenu_ExitSeparator.Visible = bCanSyncNow || bCanLaunchEditor;
 
 			// Show the startup window, if not already visible
-			if(DetectStartupProjectSettingsWindow != null)
+			if(StartupWindow != null)
 			{
-				DetectStartupProjectSettingsWindow.Show();
+				StartupWindow.Show();
 			}
 		}
 
-		private void NotifyIcon_DoubleClick(object sender, EventArgs e)
+		private void NotifyIcon_DoubleClick(object? sender, EventArgs e)
 		{
 			if(MainWindowInstance != null)
 			{
@@ -384,11 +369,11 @@ namespace UnrealGameSync
 			}
 		}
 
-		private void NotifyMenu_OpenUnrealGameSync_Click(object sender, EventArgs e)
+		private void NotifyMenu_OpenUnrealGameSync_Click(object? sender, EventArgs e)
 		{
-			if(DetectStartupProjectSettingsWindow != null)
+			if(StartupWindow != null)
 			{
-				DetectStartupProjectSettingsWindow.ShowAndActivate();
+				StartupWindow.ShowAndActivate();
 			}
 			if(MainWindowInstance != null)
 			{
@@ -396,7 +381,7 @@ namespace UnrealGameSync
 			}
 		}
 
-		private void NotifyMenu_SyncNow_Click(object sender, EventArgs e)
+		private void NotifyMenu_SyncNow_Click(object? sender, EventArgs e)
 		{
 			if(MainWindowInstance != null)
 			{
@@ -404,7 +389,7 @@ namespace UnrealGameSync
 			}
 		}
 
-		private void NotifyMenu_LaunchEditor_Click(object sender, EventArgs e)
+		private void NotifyMenu_LaunchEditor_Click(object? sender, EventArgs e)
 		{
 			if(MainWindowInstance != null)
 			{
@@ -412,17 +397,17 @@ namespace UnrealGameSync
 			}
 		}
 
-		private void NotifyMenu_Exit_Click(object sender, EventArgs e)
+		private void NotifyMenu_Exit_Click(object? sender, EventArgs e)
 		{
 			if (MainWindowInstance != null && !MainWindowInstance.ConfirmClose())
 			{
 				return;
 			}
 
-			if (DetectStartupProjectSettingsWindow != null)
+			if (StartupWindow != null)
 			{
-				DetectStartupProjectSettingsWindow.Close();
-				DetectStartupProjectSettingsWindow = null;
+				StartupWindow.Close();
+				StartupWindow = null;
 			}
 
 			if(MainWindowInstance != null)
diff --git a/Engine/Source/Programs/UnrealGameSync/UnrealGameSync/RESTApi.cs b/Engine/Source/Programs/UnrealGameSync/UnrealGameSync/RESTApi.cs
--- a/Engine/Source/Programs/UnrealGameSync/UnrealGameSync/RESTApi.cs
+++ b/Engine/Source/Programs/UnrealGameSync/UnrealGameSync/RESTApi.cs
@@ -7,6 +7,10 @@ using System.Net;
 using System.Net.Mime;
 using System.Text;
 using System.Text.Json;
+using System.Threading;
+using System.Threading.Tasks;
+
+#nullable enable
 
 namespace UnrealGameSync
 {
@@ -19,30 +23,15 @@ namespace UnrealGameSync
 
 		public override string ToString()
 		{
-			return String.Format("{0}\n\n{1}", Message, InnerException.ToString());
+			return String.Format("{0}\n\n{1}", Message, InnerException!.ToString());
 		}
 	}
 
 	public static class RESTApi
 	{
-		private static string SendRequestInternal(string URI, string Resource, string Method, string RequestBody = null, params string[] QueryParams)
+		private static async Task<string> SendRequestInternal(string Url, string Method, string? RequestBody, CancellationToken CancellationToken)
 		{
-			// set up the query string
-			StringBuilder TargetURI = new StringBuilder(string.Format("{0}/api/{1}", URI, Resource));
-			if (QueryParams.Length != 0)
-			{
-				TargetURI.Append("?");
-				for (int Idx = 0; Idx < QueryParams.Length; Idx++)
-				{
-					TargetURI.Append(QueryParams[Idx]);
-					if (Idx != QueryParams.Length - 1)
-					{
-						TargetURI.Append("&");
-					}
-				}
-			}
-
-			HttpWebRequest Request = (HttpWebRequest)WebRequest.Create(TargetURI.ToString());
+			HttpWebRequest Request = (HttpWebRequest)WebRequest.Create(Url);
 			Request.ContentType = "application/json";
 			Request.Method = Method;
 
@@ -54,7 +43,7 @@ namespace UnrealGameSync
 					byte[] bytes = Encoding.UTF8.GetBytes(RequestBody);
 					using (Stream RequestStream = Request.GetRequestStream())
 					{
-						RequestStream.Write(bytes, 0, bytes.Length);
+						await RequestStream.WriteAsync(bytes, 0, bytes.Length, CancellationToken);
 					}
 				}
 			}
@@ -62,12 +51,13 @@ namespace UnrealGameSync
 			{
 				using (WebResponse Response = Request.GetResponse())
 				{
-					string ResponseContent;
-					using (StreamReader ResponseReader = new StreamReader(Response.GetResponseStream(), Encoding.UTF8))
+					byte[] Data;
+					using (MemoryStream Buffer = new MemoryStream())
 					{
-						ResponseContent = ResponseReader.ReadToEnd();
+						await Response.GetResponseStream().CopyToAsync(Buffer, CancellationToken);
+						Data = Buffer.ToArray();
 					}
-					return ResponseContent;
+					return Encoding.UTF8.GetString(Data);
 				}
 			}
 			catch (Exception Ex)
@@ -76,24 +66,24 @@ namespace UnrealGameSync
 			}
 		}
 
-		public static string POST(string URI, string Resource, string RequestBody = null, params string[] QueryParams)
+		public static Task<string> PostAsync(string Url, string RequestBody, CancellationToken CancellationToken)
 		{
-			return SendRequestInternal(URI, Resource, "POST", RequestBody, QueryParams);
+			return SendRequestInternal(Url, "POST", RequestBody, CancellationToken);
 		}
 
-		public static string GET(string URI, string Resource, params string[] QueryParams)
+		public static Task<string> GetAsync(string Url, CancellationToken CancellationToken)
 		{
-			return SendRequestInternal(URI, Resource, "GET", null, QueryParams);
+			return SendRequestInternal(Url, "GET", null, CancellationToken);
 		}
 
-		public static T GET<T>(string URI, string Resource, params string[] QueryParams)
+		public static async Task<T> GetAsync<T>(string Url, CancellationToken CancellationToken)
 		{
-			return JsonSerializer.Deserialize<T>(SendRequestInternal(URI, Resource, "GET", null, QueryParams), Program.DefaultJsonSerializerOptions);
+			return JsonSerializer.Deserialize<T>(await GetAsync(Url, CancellationToken), Program.DefaultJsonSerializerOptions);
 		}
 
-		public static string PUT<T>(string URI, string Resource, T Object, params string[] QueryParams)
+		public static Task<string> PutAsync<T>(string Url, T Object, CancellationToken CancellationToken)
 		{
-			return SendRequestInternal(URI, Resource, "PUT", JsonSerializer.Serialize(Object), QueryParams);
+			return SendRequestInternal(Url, "PUT", JsonSerializer.Serialize(Object), CancellationToken);
 		}
 	}
 }
diff --git a/Engine/Source/Programs/UnrealGameSync/UnrealGameSync/Telemetry.cs b/Engine/Source/Programs/UnrealGameSync/UnrealGameSync/Telemetry.cs
--- a/Engine/Source/Programs/UnrealGameSync/UnrealGameSync/Telemetry.cs
+++ b/Engine/Source/Programs/UnrealGameSync/UnrealGameSync/Telemetry.cs
@@ -1,5 +1,6 @@
 // Copyright Epic Games, Inc. All Rights Reserved.
 
+using Microsoft.Extensions.Logging;
 using System;
 using System.Collections.Generic;
 using System.IO;
@@ -68,17 +69,17 @@ namespace UnrealGameSync
 		/// <summary>
 		/// The log writer to use
 		/// </summary>
-		TextWriter Log;
+		ILogger Logger;
 
 		/// <summary>
 		/// Constructor
 		/// </summary>
-		public EpicTelemetrySink(string Url, TextWriter Log)
+		public EpicTelemetrySink(string Url, ILogger Logger)
 		{
 			this.Url = Url;
-			this.Log = Log;
+			this.Logger = Logger;
 
-			Log.WriteLine("Posting to URL: {0}", Url);
+			Logger.LogInformation("Posting to URL: {Url}", Url);
 		}
 
 		/// <inheritdoc/>
@@ -150,7 +151,7 @@ namespace UnrealGameSync
 					// Print all the events we're sending
 					foreach (string Event in Events)
 					{
-						Log.WriteLine("Sending Event: {0}", Event);
+						Logger.LogInformation("Sending Event: {0}", Event);
 					}
 
 					// Convert the content to UTF8
@@ -173,7 +174,7 @@ namespace UnrealGameSync
 					// Wait for the response and dispose of it immediately
 					using (HttpWebResponse Response = (HttpWebResponse)Request.GetResponse())
 					{
-						Log.WriteLine("Response: {0}", (int)Response.StatusCode);
+						Logger.LogInformation("Response: {StatusCode}", (int)Response.StatusCode);
 					}
 				}
 				catch (WebException Ex)
@@ -182,7 +183,7 @@ namespace UnrealGameSync
 					HttpWebResponse Response = (HttpWebResponse)Ex.Response;
 					if (Response == null)
 					{
-						Log.WriteLine("Exception while attempting to send event: {0}", Ex.ToString());
+						Logger.LogError(Ex, "Exception while attempting to send event");
 					}
 					else
 					{
@@ -193,12 +194,12 @@ namespace UnrealGameSync
 							ResponseStream.CopyTo(MemoryStream);
 							ResponseText = Encoding.UTF8.GetString(MemoryStream.ToArray());
 						}
-						Log.WriteLine("EpicTelemetrySink: Failed to send analytics event. Code = {0}. Desc = {1}. Response = {2}.", (int)Response.StatusCode, Response.StatusDescription, ResponseText);
+						Logger.LogError("Failed to send analytics event. Code = {Code}. Desc = {Dec}. Response = {Response}.", (int)Response.StatusCode, Response.StatusDescription, ResponseText);
 					}
 				}
 				catch (Exception Ex)
 				{
-					Log.WriteLine("Exception while attempting to send event: {0}", Ex.ToString());
+					Logger.LogError(Ex, "Exception while attempting to send event");
 				}
 			}
 		}
diff --git a/Engine/Source/Programs/UnrealGameSync/UnrealGameSync/ToolUpdateMonitor.cs b/Engine/Source/Programs/UnrealGameSync/UnrealGameSync/ToolUpdateMonitor.cs
--- a/Engine/Source/Programs/UnrealGameSync/UnrealGameSync/ToolUpdateMonitor.cs
+++ b/Engine/Source/Programs/UnrealGameSync/UnrealGameSync/ToolUpdateMonitor.cs
@@ -1,6 +1,9 @@
 // Copyright Epic Games, Inc. All Rights Reserved.
 
 using EpicGames.Core;
+using EpicGames.Perforce;
+using Microsoft.Extensions.DependencyInjection;
+using Microsoft.Extensions.Logging;
 using System;
 using System.Collections.Generic;
 using System.Diagnostics;
@@ -11,14 +14,22 @@ using System.Threading;
 using System.Threading.Tasks;
 using System.Windows.Forms;
 
+#nullable enable
+
 namespace UnrealGameSync
 {
 	class ToolLink
 	{
 		public string Label;
 		public string FileName;
-		public string Arguments;
-		public string WorkingDir;
+		public string? Arguments;
+		public string? WorkingDir;
+
+		public ToolLink(string Label, string FileName)
+		{
+			this.Label = Label;
+			this.FileName = FileName;
+		}
 	}
 
 	class ToolDefinition
@@ -27,79 +38,84 @@ namespace UnrealGameSync
 		public string Name;
 		public string Description;
 		public bool Enabled;
-		public Action<TextWriter> InstallAction;
-		public Action<TextWriter> UninstallAction;
+		public Func<ILogger, CancellationToken, Task>? InstallAction;
+		public Func<ILogger, CancellationToken, Task>? UninstallAction;
 		public List<ToolLink> StatusPanelLinks = new List<ToolLink>();
 		public string ZipPath;
 		public int ZipChange;
 		public string ConfigPath;
 		public int ConfigChange;
+
+		public ToolDefinition(Guid Id, string Name, string Description, string ZipPath, int ZipChange, string ConfigPath, int ConfigChange)
+		{
+			this.Id = Id;
+			this.Name = Name;
+			this.Description = Description;
+			this.ZipPath = ZipPath;
+			this.ZipChange = ZipChange;
+			this.ConfigPath = ConfigPath;
+			this.ConfigChange = ConfigChange;
+		}
 	}
 
 	class ToolUpdateMonitor : IDisposable
 	{
-		Thread WorkerThread;
-		AutoResetEvent WakeEvent;
-		bool bQuit;
-		FileReference LogFile;
+		CancellationTokenSource CancellationSource;
+		SynchronizationContext SynchronizationContext;
+		Task? WorkerTask;
+		AsyncEvent WakeEvent;
+		ILogger Logger;
 		public List<ToolDefinition> Tools { get; private set; } = new List<ToolDefinition>();
 		int LastChange = -1;
+		IAsyncDisposer AsyncDisposer;
 
-		public PerforceConnection Perforce { get; }
+		IPerforceSettings PerforceSettings { get; }
 		DirectoryReference ToolsDir { get; }
 		UserSettings Settings { get; }
 
-		public event Action OnChange;
+		public Action? OnChange;
 
-		public ToolUpdateMonitor(PerforceConnection InPerforce, DirectoryReference DataDir, UserSettings Settings)
+		public ToolUpdateMonitor(IPerforceSettings PerforceSettings, DirectoryReference DataDir, UserSettings Settings, IServiceProvider ServiceProvider)
 		{
-			this.Perforce = InPerforce;
+			CancellationSource = new CancellationTokenSource();
+			SynchronizationContext = SynchronizationContext.Current!;
 			this.ToolsDir = DirectoryReference.Combine(DataDir, "Tools");
+			this.PerforceSettings = PerforceSettings;
 			this.Settings = Settings;
-
-			LogFile = FileReference.Combine(DataDir, "Tools.log");
+			this.Logger = ServiceProvider.GetRequiredService<ILogger<ToolUpdateMonitor>>();
+			this.AsyncDisposer = ServiceProvider.GetRequiredService<IAsyncDisposer>();
 
 			DirectoryReference.CreateDirectory(ToolsDir);
 
-			WakeEvent = new AutoResetEvent(false);
+			WakeEvent = new AsyncEvent();
 		}
 
 		public void Start()
 		{
 			if (DeploymentSettings.ToolsDepotPath != null)
 			{
-				WorkerThread = new Thread(() => PollForUpdates());
-				WorkerThread.Start();
+				WorkerTask = Task.Run(() => PollForUpdatesAsync(CancellationSource.Token));
 			}
 		}
 
-		public void Close()
+		public void Dispose()
 		{
-			bQuit = true;
-			WakeEvent.Set();
+			OnChange = null;
 
-			if (WorkerThread != null)
+			if (WorkerTask != null)
 			{
-				if (!WorkerThread.Join(30))
-				{
-					WorkerThread.Abort();
-					WorkerThread.Join();
-				}
-				WorkerThread = null;
+				CancellationSource.Cancel();
+				AsyncDisposer.Add(WorkerTask.ContinueWith(_ => CancellationSource.Dispose()));
+				WorkerTask = null;
 			}
 		}
 
-		public void Dispose()
-		{
-			Close();
-		}
-
 		DirectoryReference GetToolPathInternal(string ToolName)
 		{
 			return DirectoryReference.Combine(ToolsDir, ToolName, "Current");
 		}
 
-		public string GetToolName(Guid ToolId)
+		public string? GetToolName(Guid ToolId)
 		{
 			foreach (ToolDefinition Tool in Tools)
 			{
@@ -111,7 +127,7 @@ namespace UnrealGameSync
 			return null;
 		}
 
-		public DirectoryReference GetToolPath(string ToolName)
+		public DirectoryReference? GetToolPath(string ToolName)
 		{
 			if (GetToolChange(ToolName) != 0)
 			{
@@ -128,49 +144,48 @@ namespace UnrealGameSync
 			WakeEvent.Set();
 		}
 
-		void PollForUpdates()
+		async Task PollForUpdatesAsync(CancellationToken CancellationToken)
 		{
-			using (BoundedLogWriter Log = new BoundedLogWriter(LogFile))
+			while (!CancellationToken.IsCancellationRequested)
 			{
-				while (!bQuit)
-				{
-					try
-					{
-						PollForUpdatesOnce(Log);
-					}
-					catch
-					{
-					}
+				Task WakeTask = WakeEvent.Task;
 
-					WakeEvent.WaitOne(TimeSpan.FromMinutes(60.0));
+				try
+				{
+					await PollForUpdatesOnce(Logger, CancellationToken);
+				}
+				catch(Exception Ex)
+				{
+					Logger.LogError(Ex, "Exception while checking for tool updates");
 				}
+
+				Task DelayTask = Task.Delay(TimeSpan.FromMinutes(60.0), CancellationToken);
+				await Task.WhenAny(DelayTask, WakeTask);
 			}
 		}
 
-		void PollForUpdatesOnce(TextWriter Log)
+		async Task PollForUpdatesOnce(ILogger Logger, CancellationToken CancellationToken)
 		{
-			List<PerforceChangeSummary> Changes;
-			if (!Perforce.FindChanges(DeploymentSettings.ToolsDepotPath + "/...", 1, out Changes, Log) || Changes.Count == 0 || Changes[0].Number == LastChange)
-			{
-				return;
-			}
+			using IPerforceConnection Perforce = await PerforceConnection.CreateAsync(PerforceSettings, Logger);
 
-			List<PerforceFileRecord> FileRecords;
-			if (!Perforce.Stat(DeploymentSettings.ToolsDepotPath + "/...", out FileRecords, Log))
+			List<ChangesRecord> Changes = await Perforce.GetChangesAsync(ChangesOptions.None, 1, ChangeStatus.Submitted, $"{DeploymentSettings.ToolsDepotPath}/...", CancellationToken);
+			if (Changes.Count == 0 || Changes[0].Number == LastChange)
 			{
 				return;
 			}
 
+			List<FStatRecord> FileRecords = await Perforce.FStatAsync($"{DeploymentSettings.ToolsDepotPath}/...", CancellationToken).ToListAsync(CancellationToken);
+
 			// Update the tools list
 			List<ToolDefinition> NewTools = new List<ToolDefinition>();
-			foreach (PerforceFileRecord FileRecord in FileRecords)
+			foreach (FStatRecord FileRecord in FileRecords)
 			{
-				if (FileRecord.DepotPath.EndsWith(".ini"))
+				if (FileRecord.DepotFile != null && FileRecord.DepotFile.EndsWith(".ini"))
 				{
-					ToolDefinition Tool = Tools.FirstOrDefault(x => x.ConfigPath.Equals(FileRecord.DepotPath, StringComparison.Ordinal));
+					ToolDefinition? Tool = Tools.FirstOrDefault(x => x.ConfigPath.Equals(FileRecord.DepotFile, StringComparison.Ordinal));
 					if (Tool == null || Tool.ConfigChange != FileRecord.HeadChange)
 					{
-						Tool = ReadToolDefinition(FileRecord.DepotPath, FileRecord.HeadChange, Log);
+						Tool = await ReadToolDefinitionAsync(Perforce, FileRecord.DepotFile, FileRecord.HeadChange, Logger, CancellationToken);
 					}
 					if (Tool != null)
 					{
@@ -189,7 +204,7 @@ namespace UnrealGameSync
 					continue;
 				}
 
-				List<PerforceFileRecord> ToolFileRecords = FileRecords.Where(x => x.DepotPath.Equals(Tool.ZipPath, StringComparison.OrdinalIgnoreCase)).ToList();
+				List<FStatRecord> ToolFileRecords = FileRecords.Where(x => String.Equals(x.DepotFile, Tool.ZipPath, StringComparison.OrdinalIgnoreCase)).ToList();
 				if (ToolFileRecords.Count == 0)
 				{
 					continue;
@@ -201,14 +216,14 @@ namespace UnrealGameSync
 					continue;
 				}
 
-				List<PerforceFileRecord> SyncFileRecords = ToolFileRecords.Where(x => x.Action != "delete").ToList();
+				List<FStatRecord> SyncFileRecords = ToolFileRecords.Where(x => x.Action != FileAction.Delete && x.Action != FileAction.MoveDelete).ToList();
 				try
 				{
-					UpdateTool(Tool.Name, HeadChange, SyncFileRecords, Tool.InstallAction, Log);
+					await UpdateToolAsync(Perforce, Tool.Name, HeadChange, SyncFileRecords, Tool.InstallAction, Logger, CancellationToken);
 				}
 				catch (Exception Ex)
 				{
-					Log.WriteLine("Exception while updating tool: {0}", Ex.ToString());
+					Logger.LogError(Ex, "Exception while updating tool");
 				}
 			}
 
@@ -218,22 +233,22 @@ namespace UnrealGameSync
 				{
 					try
 					{
-						RemoveTool(Tool.Name, Tool.UninstallAction, Log);
+						await RemoveToolAsync(Tool.Name, Tool.UninstallAction, Logger, CancellationToken);
 					}
 					catch (Exception Ex)
 					{
-						Log.WriteLine("Exception while removing tool: {0}", Ex.ToString());
+						Logger.LogError(Ex, "Exception while removing tool");
 					}
 				}
 			}
 
-			OnChange();
+			SynchronizationContext.Post(_ => OnChange?.Invoke(), null);
 		}
 
-		ToolDefinition ReadToolDefinition(string DepotPath, int Change, TextWriter Log)
+		async Task<ToolDefinition?> ReadToolDefinitionAsync(IPerforceConnection Perforce, string DepotPath, int Change, ILogger Logger, CancellationToken CancellationToken)
 		{
-			List<string> Lines;
-			if (!Perforce.Print(String.Format("{0}@{1}", DepotPath, Change), out Lines, Log))
+			PerforceResponse<PrintRecord<string[]>> Response = await Perforce.TryPrintLinesAsync($"{DepotPath}@{Change}", CancellationToken);
+			if (!Response.Succeeded || Response.Data.Contents == null)
 			{
 				return null;
 			}
@@ -242,33 +257,41 @@ namespace UnrealGameSync
 			int ExtensionIdx = DepotPath.LastIndexOf('.');
 
 			ConfigFile ConfigFile = new ConfigFile();
-			ConfigFile.Parse(Lines.ToArray());
-
-			ToolDefinition Tool = new ToolDefinition();
+			ConfigFile.Parse(Response.Data.Contents);
 
-			string Id = ConfigFile.GetValue("Settings.Id", null);
-			if (Id == null || !Guid.TryParse(Id, out Tool.Id))
+			string? Id = ConfigFile.GetValue("Settings.Id", null);
+			if (Id == null || !Guid.TryParse(Id, out Guid ToolId))
 			{
 				return null;
 			}
 
-			Tool.Name = ConfigFile.GetValue("Settings.Name", DepotPath.Substring(NameIdx, ExtensionIdx - NameIdx));
-			Tool.Description = ConfigFile.GetValue("Settings.Description", Tool.Name);
-			Tool.ZipPath = DepotPath.Substring(0, ExtensionIdx) + ".zip";
-			Tool.ZipChange = GetToolChange(Tool.Name);
-			Tool.ConfigPath = DepotPath;
-			Tool.ConfigChange = Change;
+			string ToolName = ConfigFile.GetValue("Settings.Name", DepotPath.Substring(NameIdx, ExtensionIdx - NameIdx));
+			string ToolDescription = ConfigFile.GetValue("Settings.Description", ToolName);
+			string ToolZipPath = DepotPath.Substring(0, ExtensionIdx) + ".zip";
+			int ToolZipChange = GetToolChange(ToolName);
+			string ToolConfigPath = DepotPath;
+			int ToolConfigChange = Change;
+
+			ToolDefinition Tool = new ToolDefinition(ToolId, ToolName, ToolDescription, ToolZipPath, ToolZipChange, ToolConfigPath, ToolConfigChange);
 
-			string InstallCommand = ConfigFile.GetValue("Settings.InstallCommand", null);
+			string? InstallCommand = ConfigFile.GetValue("Settings.InstallCommand", null);
 			if (!String.IsNullOrEmpty(InstallCommand))
 			{
-				Tool.InstallAction = NewLog => RunCommand(Tool.Name, InstallCommand, NewLog);
+				Tool.InstallAction = async (Logger, CancellationToken) =>
+				{
+					Logger.LogInformation("Running install action: {Command}", InstallCommand);
+					await RunCommandAsync(Tool.Name, InstallCommand, Logger, CancellationToken);
+				};
 			}
 
-			string UninstallCommand = ConfigFile.GetValue("Settings.UninstallCommand", null);
+			string? UninstallCommand = ConfigFile.GetValue("Settings.UninstallCommand", null);
 			if (!String.IsNullOrEmpty(UninstallCommand))
 			{
-				Tool.UninstallAction = NewLog => RunCommand(Tool.Name, UninstallCommand, NewLog);
+				Tool.UninstallAction = async (Logger, CancellationToken) =>
+				{
+					Logger.LogInformation("Running unininstall action: {Command}", UninstallCommand);
+					await RunCommandAsync(Tool.Name, UninstallCommand, Logger, CancellationToken);
+				};
 			}
 
 			string[] StatusPanelLinks = ConfigFile.GetValues("Settings.StatusPanelLinks", new string[0]);
@@ -276,14 +299,12 @@ namespace UnrealGameSync
 			{
 				ConfigObject Object = new ConfigObject(StatusPanelLink);
 
-				string Label = Object.GetValue("Label", null);
-				string FileName = Object.GetValue("FileName", null);
+				string? Label = Object.GetValue("Label", null);
+				string? FileName = Object.GetValue("FileName", null);
 
 				if (Label != null && FileName != null)
 				{
-					ToolLink Link = new ToolLink();
-					Link.Label = Label;
-					Link.FileName = FileName;
+					ToolLink Link = new ToolLink(Label, FileName);
 					Link.Arguments = Object.GetValue("Arguments", null);
 					Link.WorkingDir = Object.GetValue("WorkingDir", null);
 					Tool.StatusPanelLinks.Add(Link);
@@ -293,7 +314,7 @@ namespace UnrealGameSync
 			return Tool;
 		}
 
-		void RunCommand(string ToolName, string Command, TextWriter Log)
+		async Task RunCommandAsync(string ToolName, string Command, ILogger Logger, CancellationToken CancellationToken)
 		{
 			DirectoryReference ToolPath = GetToolPathInternal(ToolName);
 
@@ -307,28 +328,30 @@ namespace UnrealGameSync
 				CommandArgs = Command.Substring(SpaceIdx + 1);
 			}
 
-			int ExitCode = Utility.ExecuteProcess(FileReference.Combine(ToolPath, CommandExe).FullName, ToolPath.FullName, CommandArgs, null, Log);
-			Log.WriteLine("(Exit: {0})", ExitCode);
+			int ExitCode = await Utility.ExecuteProcessAsync(FileReference.Combine(ToolPath, CommandExe).FullName, ToolPath.FullName, CommandArgs, Line => Logger.LogInformation("{ToolName}> {Line}", ToolName, Line), CancellationToken);
+			Logger.LogInformation("{ToolName}> Exit code {ExitCode})", ToolName, ExitCode);
 		}
 
-		void RemoveTool(string ToolName, Action<TextWriter> UninstallAction, TextWriter Log)
+		async Task RemoveToolAsync(string ToolName, Func<ILogger, CancellationToken, Task>? UninstallAction, ILogger Logger, CancellationToken CancellationToken)
 		{
-			Log.WriteLine("Removing {0}", ToolName);
+			Logger.LogInformation("Removing {0}", ToolName);
+			DirectoryReference? ToolPath = GetToolPath(ToolName);
 
 			if (UninstallAction != null)
 			{
-				Log.WriteLine("Running uninstall...");
-				UninstallAction?.Invoke(Log);
+				Logger.LogInformation("Running uninstall...");
+				await UninstallAction(Logger, CancellationToken);
 			}
 
 			SetToolChange(ToolName, null);
 
-			DirectoryReference ToolPath = GetToolPath(ToolName);
 			if (ToolPath != null)
 			{
-				Log.WriteLine("Removing {0}", ToolPath);
-				TryDeleteDirectory(ToolPath);
+				Logger.LogInformation("Removing {ToolPath}", ToolPath);
+				TryDeleteDirectory(ToolPath, Logger);
 			}
+
+			Logger.LogInformation("{ToolName} has been removed successfully", ToolName);
 		}
 
 		static void ForceDeleteDirectory(DirectoryReference DirectoryName)
@@ -344,27 +367,28 @@ namespace UnrealGameSync
 			}
 		}
 
-		static bool TryDeleteDirectory(DirectoryReference DirectoryName)
+		static bool TryDeleteDirectory(DirectoryReference DirectoryName, ILogger Logger)
 		{
 			try
 			{
 				ForceDeleteDirectory(DirectoryName);
 				return true;
 			}
-			catch
+			catch(Exception Ex)
 			{
+				Logger.LogWarning(Ex, "Unable to delete directory {DirectoryName}", DirectoryName);
 				return false;
 			}
 		}
 
-		bool UpdateTool(string ToolName, int Change, List<PerforceFileRecord> Records, Action<TextWriter> InstallAction, TextWriter Log)
+		async Task<bool> UpdateToolAsync(IPerforceConnection Perforce, string ToolName, int Change, List<FStatRecord> Records, Func<ILogger, CancellationToken, Task>? InstallAction, ILogger Logger, CancellationToken CancellationToken)
 		{
 			DirectoryReference ToolDir = DirectoryReference.Combine(ToolsDir, ToolName);
 			DirectoryReference.CreateDirectory(ToolDir);
 
 			foreach (DirectoryReference ExistingDir in DirectoryReference.EnumerateDirectories(ToolDir, "Prev-*"))
 			{
-				TryDeleteDirectory(ExistingDir);
+				TryDeleteDirectory(ExistingDir, Logger);
 			}
 
 			DirectoryReference NextToolDir = DirectoryReference.Combine(ToolDir, "Next");
@@ -377,12 +401,15 @@ namespace UnrealGameSync
 			for (int Idx = 0; Idx < Records.Count; Idx++)
 			{
 				FileReference ZipFile = FileReference.Combine(NextToolZipsDir, String.Format("{0}.{1}.zip", ToolName, Idx));
-				if (!Perforce.PrintToFile(String.Format("{0}#{1}", Records[Idx].DepotPath, Records[Idx].HeadRevision), ZipFile.FullName, Log) || !FileReference.Exists(ZipFile))
+
+				PerforceResponse<PrintRecord> Response = await Perforce.TryPrintAsync(ZipFile.FullName, $"{Records[Idx].DepotFile}#{Records[Idx].HeadRevision}", CancellationToken);
+				if(!Response.Succeeded || !FileReference.Exists(ZipFile))
 				{
-					Log.WriteLine("Unable to print {0}", Records[Idx].DepotPath);
+					Logger.LogError("Unable to print {0}", Records[Idx].DepotFile);
 					return false;
 				}
-				ArchiveUtils.ExtractFiles(ZipFile, NextToolDir, null, new ProgressValue(), Log);
+
+				ArchiveUtils.ExtractFiles(ZipFile, NextToolDir, null, new ProgressValue(), Logger);
 			}
 
 			SetToolChange(ToolName, null);
@@ -392,18 +419,19 @@ namespace UnrealGameSync
 			{
 				DirectoryReference PrevDirectoryName = DirectoryReference.Combine(ToolDir, String.Format("Prev-{0:X16}", Stopwatch.GetTimestamp()));
 				Directory.Move(CurrentToolDir.FullName, PrevDirectoryName.FullName);
-				TryDeleteDirectory(PrevDirectoryName);
+				TryDeleteDirectory(PrevDirectoryName, Logger);
 			}
 
 			Directory.Move(NextToolDir.FullName, CurrentToolDir.FullName);
 
 			if (InstallAction != null)
 			{
-				Log.WriteLine("Running installer...");
-				InstallAction.Invoke(Log);
+				Logger.LogInformation("Running installer...");
+				await InstallAction.Invoke(Logger, CancellationToken);
 			}
 
 			SetToolChange(ToolName, Change);
+			Logger.LogInformation("Updated {ToolName} to change {Change}", ToolName, Change);
 			return true;
 		}
 
diff --git a/Engine/Source/Programs/UnrealGameSync/UnrealGameSync/UnrealGameSync.csproj b/Engine/Source/Programs/UnrealGameSync/UnrealGameSync/UnrealGameSync.csproj
--- a/Engine/Source/Programs/UnrealGameSync/UnrealGameSync/UnrealGameSync.csproj
+++ b/Engine/Source/Programs/UnrealGameSync/UnrealGameSync/UnrealGameSync.csproj
@@ -22,10 +22,16 @@
   <ItemGroup>
     <PackageReference Include="EnvDTE" Version="16.7.30328.74" />
     <PackageReference Include="IdentityModel.OidcClient" Version="3.1.2" />
+    <PackageReference Include="Serilog.Expressions" Version="3.2.1" />
+    <PackageReference Include="Serilog.Extensions.Hosting" Version="4.2.0" />
+    <PackageReference Include="Serilog.Extensions.Logging" Version="3.1.0" />
+    <PackageReference Include="Serilog.Sinks.File" Version="5.0.0" />
+    <PackageReference Include="System.Linq.Async" Version="5.1.0" />
   </ItemGroup>
   <ItemGroup>
     <ProjectReference Include="..\..\Shared\EpicGames.Core\EpicGames.Core.csproj" />
     <ProjectReference Include="..\..\Shared\EpicGames.Jupiter\EpicGames.Jupiter.csproj" />
+    <ProjectReference Include="..\..\Shared\EpicGames.Perforce\EpicGames.Perforce.csproj" />
   </ItemGroup>
   <ItemGroup>
     <Reference Include="Microsoft.VisualStudio.Setup.Configuration.Interop">
diff --git a/Engine/Source/Programs/UnrealGameSync/UnrealGameSync/UpdateMonitor.cs b/Engine/Source/Programs/UnrealGameSync/UnrealGameSync/UpdateMonitor.cs
--- a/Engine/Source/Programs/UnrealGameSync/UnrealGameSync/UpdateMonitor.cs
+++ b/Engine/Source/Programs/UnrealGameSync/UnrealGameSync/UpdateMonitor.cs
@@ -1,5 +1,9 @@
-﻿// Copyright Epic Games, Inc. All Rights Reserved.
+// Copyright Epic Games, Inc. All Rights Reserved.
 
+using EpicGames.Core;
+using EpicGames.Perforce;
+using Microsoft.Extensions.DependencyInjection;
+using Microsoft.Extensions.Logging;
 using System;
 using System.Collections.Generic;
 using System.IO;
@@ -8,6 +12,8 @@ using System.Text;
 using System.Threading;
 using System.Threading.Tasks;
 
+#nullable enable
+
 namespace UnrealGameSync
 {
 	enum UpdateType
@@ -18,17 +24,12 @@ namespace UnrealGameSync
 
 	class UpdateMonitor : IDisposable
 	{
-		string WatchPath;
-		Thread WorkerThread;
-		ManualResetEvent QuitEvent;
+		Task? WorkerTask;
+		CancellationTokenSource CancellationSource = new CancellationTokenSource();
+		ILogger Logger;
+		IAsyncDisposer AsyncDisposer;
 
-		public event Action<UpdateType> OnUpdateAvailable;
-
-		public PerforceConnection Perforce
-		{
-			get;
-			private set;
-		}
+		public Action<UpdateType>? OnUpdateAvailable;
 
 		public bool? RelaunchUnstable
 		{
@@ -36,54 +37,45 @@ namespace UnrealGameSync
 			private set;
 		}
 
-		public UpdateMonitor(PerforceConnection InPerforce, string InWatchPath)
+		public UpdateMonitor(IPerforceSettings PerforceSettings, string WatchPath, IServiceProvider ServiceProvider)
 		{
-			Perforce = InPerforce;
-			WatchPath = InWatchPath;
-
-			QuitEvent = new ManualResetEvent(false);
+			this.Logger = ServiceProvider.GetRequiredService<ILogger<UpdateMonitor>>();
+			this.AsyncDisposer = ServiceProvider.GetRequiredService<IAsyncDisposer>();
 
 			if(WatchPath != null)
 			{
-				WorkerThread = new Thread(() => PollForUpdates());
-				WorkerThread.Start();
+				Logger.LogInformation("Watching for updates on {WatchPath}", WatchPath);
+				WorkerTask = Task.Run(() => PollForUpdatesAsync(PerforceSettings, WatchPath, CancellationSource.Token));
 			}
 		}
 
-		public void Close()
+		public void Dispose()
 		{
-			QuitEvent.Set();
+			OnUpdateAvailable = null;
 
-			if(WorkerThread != null)
+			if (WorkerTask != null)
 			{
-				if(!WorkerThread.Join(30))
-				{
-					WorkerThread.Interrupt();
-					WorkerThread.Join();
-				}
-				WorkerThread = null;
+				CancellationSource.Cancel();
+				AsyncDisposer.Add(WorkerTask.ContinueWith(_ => CancellationSource.Dispose()));
+				WorkerTask = null;
 			}
 		}
 
-		public void Dispose()
-		{
-			Close();
-		}
-
 		public bool IsUpdateAvailable
 		{
 			get;
 			private set;
 		}
 
-		void PollForUpdates()
+		async Task PollForUpdatesAsync(IPerforceSettings PerforceSettings, string WatchPath, CancellationToken CancellationToken)
 		{
-			while(!QuitEvent.WaitOne(5 * 60 * 1000))
+			using IPerforceConnection Perforce = await PerforceConnection.CreateAsync(PerforceSettings, Logger);
+			for (; ; )
 			{
-				StringWriter Log = new StringWriter();
+				await Task.Delay(TimeSpan.FromMinutes(5.0), CancellationToken);
 
-				List<PerforceChangeSummary> Changes;
-				if(Perforce.FindChanges(WatchPath, 1, out Changes, Log) && Changes.Count > 0)
+				PerforceResponseList<ChangesRecord> Changes = await Perforce.TryGetChangesAsync(ChangesOptions.None, -1, ChangeStatus.Submitted, WatchPath, CancellationToken);
+				if (Changes.Succeeded && Changes.Data.Count > 0)
 				{
 					TriggerUpdate(UpdateType.Background, null);
 				}
diff --git a/Engine/Source/Programs/UnrealGameSync/UnrealGameSync/UriHandlers/VisualStudioHandler.cs b/Engine/Source/Programs/UnrealGameSync/UnrealGameSync/UriHandlers/VisualStudioHandler.cs
--- a/Engine/Source/Programs/UnrealGameSync/UnrealGameSync/UriHandlers/VisualStudioHandler.cs
+++ b/Engine/Source/Programs/UnrealGameSync/UnrealGameSync/UriHandlers/VisualStudioHandler.cs
@@ -1,7 +1,10 @@
 // Copyright Epic Games, Inc. All Rights Reserved.
 
 using System.IO;
+using System.Threading;
+using System.Threading.Tasks;
 using EnvDTE;
+using Microsoft.Extensions.Logging.Abstractions;
 
 namespace UnrealGameSync
 {
@@ -14,14 +17,9 @@ namespace UnrealGameSync
 		[UriHandler(true)]
 		public static UriResult VSOpen(string DepotPath, int Line = -1)
 		{
-			string ErrorMessage;
-			string TempFileName;
-
-			if (!P4Automation.PrintToTempFile(null, DepotPath, out TempFileName, out ErrorMessage))
-			{
-				return new UriResult() { Error = ErrorMessage ?? "Unknown P4 Error" };
-			}
+			string TempFileName = P4Automation.PrintToTempFile(null, DepotPath, NullLogger.Instance).GetAwaiter().GetResult();
 
+			string ErrorMessage;
 			if (!VisualStudioAutomation.OpenFile(TempFileName, out ErrorMessage, Line))
 			{
 				return new UriResult() { Error = ErrorMessage ?? "Unknown Visual Studio Error" };
diff --git a/Engine/Source/Programs/UnrealGameSync/UnrealGameSync/UserSettings.cs b/Engine/Source/Programs/UnrealGameSync/UnrealGameSync/UserSettings.cs
--- a/Engine/Source/Programs/UnrealGameSync/UnrealGameSync/UserSettings.cs
+++ b/Engine/Source/Programs/UnrealGameSync/UnrealGameSync/UserSettings.cs
@@ -1,8 +1,10 @@
 // Copyright Epic Games, Inc. All Rights Reserved.
 
 using EpicGames.Core;
+using Microsoft.Extensions.Logging;
 using System;
 using System.Collections.Generic;
+using System.Diagnostics.CodeAnalysis;
 using System.Drawing;
 using System.IO;
 using System.Linq;
@@ -13,16 +15,18 @@ using System.Text.RegularExpressions;
 using System.Threading.Tasks;
 using System.Windows.Forms;
 
+#nullable enable
+
 namespace UnrealGameSync
 {
-	enum BuildConfig
+	public enum BuildConfig
 	{
 		Debug,
 		DebugGame,
 		Development,
 	}
 
-	enum TabLabels
+	public enum TabLabels
 	{
 		Stream,
 		WorkspaceName,
@@ -30,7 +34,7 @@ namespace UnrealGameSync
 		ProjectFile,
 	}
 
-	enum BisectState
+	public enum BisectState
 	{
 		Include,
 		Exclude,
@@ -38,20 +42,20 @@ namespace UnrealGameSync
 		Fail,
 	}
 
-	enum UserSelectedProjectType
+	public enum UserSelectedProjectType
 	{
 		Client,
 		Local
 	}
 
-	enum FilterType
+	public enum FilterType
 	{
 		None,
 		Code,
 		Content
 	}
 
-	enum UserSettingsVersion
+	public enum UserSettingsVersion
 	{
 		Initial = 0,
 		DefaultServerSettings = 1,
@@ -73,11 +77,11 @@ namespace UnrealGameSync
 			this.Order = new List<string>(Order);
 		}
 
-		public static bool TryParseConfigEntry(string Text, out ArchiveSettings Settings)
+		public static bool TryParseConfigEntry(string Text, [NotNullWhen(true)] out ArchiveSettings? Settings)
 		{
 			ConfigObject Object = new ConfigObject(Text);
 
-			string Type = Object.GetValue("Type", null);
+			string? Type = Object.GetValue("Type", null);
 			if (Type == null)
 			{
 				Settings = null;
@@ -110,15 +114,15 @@ namespace UnrealGameSync
 		}
 	}
 
-	class UserSelectedProjectSettings
+	public class UserSelectedProjectSettings
 	{
-		public readonly string ServerAndPort;
-		public readonly string UserName;
+		public readonly string? ServerAndPort;
+		public readonly string? UserName;
 		public readonly UserSelectedProjectType Type;
-		public readonly string ClientPath;
-		public readonly string LocalPath;
+		public readonly string? ClientPath;
+		public readonly string? LocalPath;
 
-		public UserSelectedProjectSettings(string ServerAndPort, string UserName, UserSelectedProjectType Type, string ClientPath, string LocalPath)
+		public UserSelectedProjectSettings(string? ServerAndPort, string? UserName, UserSelectedProjectType Type, string? ClientPath, string? LocalPath)
 		{
 			this.ServerAndPort = ServerAndPort;
 			this.UserName = UserName;
@@ -127,14 +131,14 @@ namespace UnrealGameSync
 			this.LocalPath = LocalPath;
 		}
 
-		public static bool TryParseConfigEntry(string Text, out UserSelectedProjectSettings Project)
+		public static bool TryParseConfigEntry(string Text, [NotNullWhen(true)] out UserSelectedProjectSettings? Project)
 		{
 			ConfigObject Object = new ConfigObject(Text);
 
 			UserSelectedProjectType Type;
 			if(Enum.TryParse(Object.GetValue("Type", ""), out Type))
 			{
-				string ServerAndPort = Object.GetValue("ServerAndPort", null);
+				string? ServerAndPort = Object.GetValue("ServerAndPort", null);
 				if(String.IsNullOrWhiteSpace(ServerAndPort))
 				{
 					ServerAndPort = null;
@@ -146,19 +150,19 @@ namespace UnrealGameSync
 					ServerAndPort = "perforce:1666";
 				}
 
-				string UserName = Object.GetValue("UserName", null);
+				string? UserName = Object.GetValue("UserName", null);
 				if(String.IsNullOrWhiteSpace(UserName))
 				{
 					UserName = null;
 				}
 
-				string LocalPath = Object.GetValue("LocalPath", null);
+				string? LocalPath = Object.GetValue("LocalPath", null);
 				if(String.IsNullOrWhiteSpace(LocalPath))
 				{
 					LocalPath = null;
 				}
 
-				string ClientPath = Object.GetValue("ClientPath", null);
+				string? ClientPath = Object.GetValue("ClientPath", null);
 				if(String.IsNullOrWhiteSpace(ClientPath))
 				{
 					ClientPath = null;
@@ -202,31 +206,31 @@ namespace UnrealGameSync
 			return Object.ToString();
 		}
 
-		public override string ToString()
+		public override string? ToString()
 		{
 			return LocalPath ?? ClientPath;
 		}
 	}
 
-	class SyncCategory
+	public class SyncCategory
 	{
 		public Guid Id { get; set; }
 		public bool Enable { get; set; }
 	}
 
-	class BisectEntry
+	public class BisectEntry
 	{
 		public int Change { get; set; }
 		public BisectState State { get; set; }
 	}
 
-	class UserWorkspaceSettings
+	public class UserWorkspaceSettings
 	{
 		[JsonIgnore]
-		public FileReference File { get; set; }
+		public FileReference? File { get; set; }
 
 		// Workspace specific SyncFilters
-		public string[] SyncView { get; set; }
+		public string[]? SyncView { get; set; }
 		public List<SyncCategory> SyncCategories { get; set; } = new List<SyncCategory>();
 		public bool? bSyncAllProjects { get; set; }
 		public bool? bIncludeAllProjectsInSolution { get; set; }
@@ -238,27 +242,27 @@ namespace UnrealGameSync
 			set => SyncCategories = value.Select(x => new SyncCategory { Id = x.Key, Enable = x.Value }).ToList();
 		}
 
-		public void Save() => Utility.SaveJson(File, this);
+		public void Save() => Utility.SaveJson(File!, this);
 	}
 
-	class UserWorkspaceState
+	public class UserWorkspaceState
 	{
 		[JsonIgnore]
-		public FileReference File { get; set; }
+		public FileReference? File { get; set; }
 
 		// Client path to the current project
-		public string SelectedClientFileName { get; set; }
+		public string? SelectedClientFileName { get; set; }
 
 		// Settings for the currently synced project in this workspace. CurrentChangeNumber is only valid for this workspace if CurrentProjectPath is the current project.
-		public string CurrentProjectIdentifier { get; set; }
+		public string? CurrentProjectIdentifier { get; set; }
 		public int CurrentChangeNumber { get; set; } = -1;
-		public string CurrentSyncFilterHash { get; set; }
+		public string? CurrentSyncFilterHash { get; set; }
 		public List<int> AdditionalChangeNumbers { get; set; } = new List<int>();
 
 		// Settings for the last attempted sync. These values are set to persist error messages between runs.
 		public int LastSyncChangeNumber { get; set; }
 		public WorkspaceUpdateResult LastSyncResult { get; set; }
-		public string LastSyncResultMessage { get; set; }
+		public string? LastSyncResultMessage { get; set; }
 		public DateTime? LastSyncTime { get; set; }
 		public int LastSyncDurationSeconds { get; set; }
 
@@ -266,7 +270,7 @@ namespace UnrealGameSync
 		public int LastBuiltChangeNumber { get; set; }
 
 		// Expanded archives in the workspace
-		public string[] ExpandedArchiveTypes { get; set; }
+		public string[]? ExpandedArchiveTypes { get; set; }
 
 		// The changes that we're regressing at the moment
 		public List<BisectEntry> BisectChanges { get; set; } = new List<BisectEntry>();
@@ -285,7 +289,7 @@ namespace UnrealGameSync
 
 		public void Save()
 		{
-			lock(File)
+			lock(File!)
 			{
 				Utility.SaveJson(File, this);
 			}
@@ -295,13 +299,13 @@ namespace UnrealGameSync
 	class UserProjectSettings
 	{
 		[JsonIgnore]
-		public FileReference File { get; set; }
+		public FileReference? File { get; set; }
 
 		public List<ConfigObject> BuildSteps { get; set; } = new List<ConfigObject>();
 		public FilterType FilterType { get; set; }
 		public HashSet<string> FilterBadges { get; set; } = new HashSet<string>(StringComparer.OrdinalIgnoreCase);
 
-		public void Save() => Utility.SaveJson(File, this);
+		public void Save() => Utility.SaveJson(File!, this);
 	}
 
 	class UserSettings
@@ -335,7 +339,7 @@ namespace UnrealGameSync
 		public Guid[] EnabledTools;
 		public bool bShowNetCoreInfo;
 		public int FilterIndex;
-		public UserSelectedProjectSettings LastProject;
+		public UserSelectedProjectSettings? LastProject;
 		public List<UserSelectedProjectSettings> OpenProjects;
 		public List<UserSelectedProjectSettings> RecentProjects;
 		public string[] SyncView;
@@ -386,12 +390,12 @@ namespace UnrealGameSync
 		{
 			List<UserSelectedProjectSettings> Projects = new List<UserSelectedProjectSettings>();
 
-			string[] ProjectStrings = ConfigFile.GetValues(SettingName, null);
+			string[]? ProjectStrings = ConfigFile.GetValues(SettingName, null);
 			if(ProjectStrings != null)
 			{
 				foreach(string ProjectString in ProjectStrings)
 				{
-					UserSelectedProjectSettings Project;
+					UserSelectedProjectSettings? Project;
 					if(UserSelectedProjectSettings.TryParseConfigEntry(ProjectString, out Project))
 					{
 						Projects.Add(Project);
@@ -400,7 +404,7 @@ namespace UnrealGameSync
 			}
 			else if(LegacySettingName != null)
 			{
-				string[] LegacyProjectStrings = ConfigFile.GetValues(LegacySettingName, null);
+				string[]? LegacyProjectStrings = ConfigFile.GetValues(LegacySettingName, null);
 				if(LegacyProjectStrings != null)
 				{
 					foreach(string LegacyProjectString in LegacyProjectStrings)
@@ -416,10 +420,10 @@ namespace UnrealGameSync
 			return Projects;
 		}
 
-		public UserSettings(FileReference InFileName, TextWriter Log)
+		public UserSettings(FileReference InFileName, ILogger Logger)
 		{
 			FileName = InFileName;
-			ConfigFile.TryLoad(FileName, Log);
+			ConfigFile.TryLoad(FileName, Logger);
 
 			// General settings
 			Version = (UserSettingsVersion)ConfigFile.GetValue("General.Version", (int)UserSettingsVersion.Initial);
@@ -454,14 +458,14 @@ namespace UnrealGameSync
 			bShowNetCoreInfo = ConfigFile.GetValue("General.ShowNetCoreInfo", true);
 			int.TryParse(ConfigFile.GetValue("General.FilterIndex", "0"), out FilterIndex);
 
-			string LastProjectString = ConfigFile.GetValue("General.LastProject", null);
+			string? LastProjectString = ConfigFile.GetValue("General.LastProject", null);
 			if(LastProjectString != null)
 			{
 				UserSelectedProjectSettings.TryParseConfigEntry(LastProjectString, out LastProject);
 			}
 			else
 			{
-				string LastProjectFileName = ConfigFile.GetValue("General.LastProjectFileName", null);
+				string? LastProjectFileName = ConfigFile.GetValue("General.LastProjectFileName", null);
 				if(LastProjectFileName != null)
 				{
 					LastProject = new UserSelectedProjectSettings(null, null, UserSelectedProjectType.Local, null, LastProjectFileName);
@@ -521,7 +525,7 @@ namespace UnrealGameSync
 			string[] ArchiveValues = ConfigFile.GetValues("PrecompiledBinaries.Archives", new string[0]);
 			foreach (string ArchiveValue in ArchiveValues)
 			{
-				ArchiveSettings Settings;
+				ArchiveSettings? Settings;
 				if (ArchiveSettings.TryParseConfigEntry(ArchiveValue, out Settings))
 				{
 					Archives.Add(Settings);
@@ -538,9 +542,12 @@ namespace UnrealGameSync
 			foreach (string TokenValue in Tokens)
 			{
 				ConfigObject O = new ConfigObject(TokenValue);
-				string Provider = O.GetValue("Provider");
-				string Token = O.GetValue("Token");
-				ProviderToRefreshTokens.TryAdd(Provider, Token);
+				string? Provider = O.GetValue("Provider");
+				string? Token = O.GetValue("Token");
+				if (Provider != null && Token != null)
+				{
+					ProviderToRefreshTokens.TryAdd(Provider, Token);
+				}
 			}
 
 			// Window settings
@@ -671,7 +678,7 @@ namespace UnrealGameSync
 		public UserWorkspaceState FindOrAddWorkspaceState(DirectoryReference RootDir, string SelectedClientFileName, string SelectedProjectIdentifier)
 		{
 			FileReference ConfigFile = FileReference.Combine(GetConfigDir(RootDir), "global.json");
-			if (!WorkspaceKeyToState.TryGetValue(ConfigFile, out UserWorkspaceState State))
+			if (!WorkspaceKeyToState.TryGetValue(ConfigFile, out UserWorkspaceState? State))
 			{
 				if (!Utility.TryLoadJson(ConfigFile, out State))
 				{
@@ -698,7 +705,7 @@ namespace UnrealGameSync
 		public UserWorkspaceSettings FindOrAddWorkspaceSettings(DirectoryReference LocalRoot, string ClientBranchPath)
 		{
 			FileReference ConfigFile = FileReference.Combine(GetConfigDir(LocalRoot), $"client_{ClientBranchPath.Trim('/').Replace('/', '+')}.json");
-			if (!WorkspaceKeyToSettings.TryGetValue(ConfigFile, out UserWorkspaceSettings Settings))
+			if (!WorkspaceKeyToSettings.TryGetValue(ConfigFile, out UserWorkspaceSettings? Settings))
 			{
 				if (!Utility.TryLoadJson(ConfigFile, out Settings))
 				{
@@ -728,7 +735,7 @@ namespace UnrealGameSync
 					LegacyBranchAndClientKey = LegacyBranchAndClientKey.Substring(0, SlashIdx) + "$" + LegacyBranchAndClientKey.Substring(SlashIdx + 1);
 				}
 
-				string CurrentSync = ConfigFile.GetValue("Clients." + LegacyBranchAndClientKey, null);
+				string? CurrentSync = ConfigFile.GetValue("Clients." + LegacyBranchAndClientKey, null);
 				if(CurrentSync != null)
 				{
 					int AtIdx = CurrentSync.LastIndexOf('@');
@@ -743,7 +750,7 @@ namespace UnrealGameSync
 					}
 				}
 
-				string LastUpdateResultText = ConfigFile.GetValue("Clients." + LegacyBranchAndClientKey + "$LastUpdate", null);
+				string? LastUpdateResultText = ConfigFile.GetValue("Clients." + LegacyBranchAndClientKey + "$LastUpdate", null);
 				if(LastUpdateResultText != null)
 				{
 					int ColonIdx = LastUpdateResultText.LastIndexOf(':');
@@ -846,7 +853,7 @@ namespace UnrealGameSync
 		public UserProjectSettings FindOrAddProjectSettings(DirectoryReference LocalRoot, string ClientProjectFileName)
 		{
 			FileReference ConfigFile = FileReference.Combine(GetConfigDir(LocalRoot), $"project_{ClientProjectFileName.Trim('/').Replace('/', '+')}.json");
-			if (!ProjectKeyToSettings.TryGetValue(ConfigFile, out UserProjectSettings Settings))
+			if (!ProjectKeyToSettings.TryGetValue(ConfigFile, out UserProjectSettings? Settings))
 			{
 				if (!Utility.TryLoadJson(ConfigFile, out Settings))
 				{
@@ -1061,7 +1068,7 @@ namespace UnrealGameSync
 		{
 			if(Enabled.Add(UniqueId))
 			{
-				WorkspaceSyncCategory Category;
+				WorkspaceSyncCategory? Category;
 				if(UniqueIdToFilter.TryGetValue(UniqueId, out Category))
 				{
 					foreach(Guid RequiresUniqueId in Category.Requires)
@@ -1072,7 +1079,8 @@ namespace UnrealGameSync
 			}
 		}
 
-		static string EscapeText(string Text)
+		[return: NotNullIfNotNull("Text")]
+		static string? EscapeText(string? Text)
 		{
 			if(Text == null)
 			{
@@ -1110,7 +1118,8 @@ namespace UnrealGameSync
 			return Result.ToString();
 		}
 
-		static string UnescapeText(string Text)
+		[return: NotNullIfNotNull("Text")]
+		static string? UnescapeText(string? Text)
 		{
 			if(Text == null)
 			{
