UGS: Significant refactor of UGS internals.

diff --git a/Engine/Source/Programs/UnrealGameSync/UnrealGameSync/Utility.cs b/Engine/Source/Programs/UnrealGameSync/UnrealGameSync/Utility.cs
--- a/Engine/Source/Programs/UnrealGameSync/UnrealGameSync/Utility.cs
+++ b/Engine/Source/Programs/UnrealGameSync/UnrealGameSync/Utility.cs
@@ -1,12 +1,13 @@
 // Copyright Epic Games, Inc. All Rights Reserved.
 
-#define USE_NEW_PROCESS_JOBS
-
 using EpicGames.Core;
+using EpicGames.Perforce;
+using Microsoft.Extensions.Logging;
 using Microsoft.Win32;
 using System;
 using System.Collections.Generic;
 using System.Diagnostics;
+using System.Diagnostics.CodeAnalysis;
 using System.Drawing;
 using System.IO;
 using System.Linq;	
@@ -14,13 +15,27 @@ using System.Text;
 using System.Text.Json;
 using System.Text.Json.Serialization;
 using System.Text.RegularExpressions;
+using System.Threading;
+using System.Threading.Tasks;
 using System.Windows.Forms;
 
+#nullable enable
+
 namespace UnrealGameSync
 {
+	sealed class UserErrorException : Exception
+	{
+		public int Code { get; }
+
+		public UserErrorException(string Message, int Code = 1) : base(Message)
+		{
+			this.Code = Code;
+		}
+	}
+
 	static class Utility
 	{
-		public static bool TryLoadJson<T>(FileReference File, out T Object) where T : class
+		public static bool TryLoadJson<T>(FileReference File, [NotNullWhen(true)] out T? Object) where T : class
 		{
 			if (!FileReference.Exists(File))
 			{
@@ -143,80 +158,21 @@ namespace UnrealGameSync
 			}
 		}
 
-		public static int ExecuteProcess(string FileName, string WorkingDir, string CommandLine, string Input, TextWriter Log)
-		{
-			return ExecuteProcess(FileName, WorkingDir, CommandLine, Input, Line => Log.WriteLine(Line));
-		}
-
-		public static int ExecuteProcess(string FileName, string WorkingDir, string CommandLine, string Input, out List<string> OutputLines)
-		{
-			using(ChildProcess NewProcess = new ChildProcess(FileName, WorkingDir, CommandLine, Input))
-			{
-				OutputLines = NewProcess.ReadAllLines();
-				return NewProcess.ExitCode;
-			}
-		}
-
-		public static int ExecuteProcess(string FileName, string WorkingDir, string CommandLine, string Input, Action<string> OutputLine)
+		public static async Task<int> ExecuteProcessAsync(string FileName, string? WorkingDir, string CommandLine, Action<string> OutputLine, CancellationToken CancellationToken)
 		{
-#if USE_NEW_PROCESS_JOBS
-			using(ChildProcess NewProcess = new ChildProcess(FileName, WorkingDir, CommandLine, Input))
+			using (ManagedProcess NewProcess = new ManagedProcess(null, FileName, CommandLine, WorkingDir, null, ProcessPriorityClass.Normal))
 			{
-				string Line;
-				while(NewProcess.TryReadLine(out Line))
+				for (; ; )
 				{
-					OutputLine(Line);
-				}
-				return NewProcess.ExitCode;
-			}
-#else
-			using(Process ChildProcess = new Process())
-			{
-				try
-				{
-					object LockObject = new object();
-
-					DataReceivedEventHandler OutputHandler = (x, y) => { if(y.Data != null){ lock(LockObject){ OutputLine(y.Data.TrimEnd()); } } };
-
-					ChildProcess.StartInfo.FileName = FileName;
-					ChildProcess.StartInfo.Arguments = String.IsNullOrEmpty(CommandLine) ? "" : CommandLine;
-					ChildProcess.StartInfo.UseShellExecute = false;
-					ChildProcess.StartInfo.RedirectStandardOutput = true;
-					ChildProcess.StartInfo.RedirectStandardError = true;
-					ChildProcess.OutputDataReceived += OutputHandler;
-					ChildProcess.ErrorDataReceived += OutputHandler;
-					ChildProcess.StartInfo.RedirectStandardInput = Input != null;
-					ChildProcess.StartInfo.CreateNoWindow = true;
-					ChildProcess.StartInfo.StandardOutputEncoding = new System.Text.UTF8Encoding(false, false);
-					ChildProcess.Start();
-					ChildProcess.BeginOutputReadLine();
-					ChildProcess.BeginErrorReadLine();
-
-					if(!String.IsNullOrEmpty(Input))
-					{
-						ChildProcess.StandardInput.WriteLine(Input);
-						ChildProcess.StandardInput.Close();
-					}
-
-					// Busy wait for the process to exit so we can get a ThreadAbortException if the thread is terminated. It won't wait until we enter managed code
-					// again before it throws otherwise.
-					for(;;)
+					string? Line = await NewProcess.ReadLineAsync(CancellationToken);
+					if (Line == null)
 					{
-						if(ChildProcess.WaitForExit(20))
-						{
-							ChildProcess.WaitForExit();
-							break;
-						}
+						NewProcess.WaitForExit();
+						return NewProcess.ExitCode;
 					}
-
-					return ChildProcess.ExitCode;
-				}
-				finally
-				{
-					CloseHandle(JobObject);
+					OutputLine(Line);
 				}
 			}
-#endif
 		}
 
 		public static bool SafeIsFileUnderDirectory(string FileName, string DirectoryName)
@@ -240,7 +196,7 @@ namespace UnrealGameSync
 		/// <param name="InputString">String to search for variable names</param>
 		/// <param name="Variables">Lookup of variable names to values</param>
 		/// <returns>String with all variables replaced</returns>
-		public static string ExpandVariables(string InputString, Dictionary<string, string> AdditionalVariables = null)
+		public static string ExpandVariables(string InputString, Dictionary<string, string>? AdditionalVariables = null)
 		{
 			string Result = InputString;
 			for (int Idx = Result.IndexOf("$("); Idx != -1; Idx = Result.IndexOf("$(", Idx))
@@ -256,7 +212,7 @@ namespace UnrealGameSync
 				string Name = Result.Substring(Idx + 2, EndIdx - (Idx + 2));
 
 				// Strip the format from the name
-				string Format = null;
+				string? Format = null;
 				int FormatIdx = Name.IndexOf(':');
 				if(FormatIdx != -1)
 				{ 
@@ -265,7 +221,7 @@ namespace UnrealGameSync
 				}
 
 				// Find the value for it, either from the dictionary or the environment block
-				string Value;
+				string? Value;
 				if (AdditionalVariables == null || !AdditionalVariables.TryGetValue(Name, out Value))
 				{
 					Value = Environment.GetEnvironmentVariable(Name);
@@ -417,7 +373,7 @@ namespace UnrealGameSync
 
 		/******/
 
-		public static void ReadGlobalPerforceSettings(ref string ServerAndPort, ref string UserName, ref string DepotPath)
+		public static void ReadGlobalPerforceSettings(ref string? ServerAndPort, ref string? UserName, ref string? DepotPath)
 		{
 			using (RegistryKey Key = Registry.CurrentUser.OpenSubKey("SOFTWARE\\Epic Games\\UnrealGameSync", false))
 			{
@@ -467,7 +423,7 @@ namespace UnrealGameSync
 			}
 		}
 
-		public static void SaveGlobalPerforceSettings(string ServerAndPort, string UserName, string DepotPath)
+		public static void SaveGlobalPerforceSettings(string? ServerAndPort, string? UserName, string? DepotPath)
 		{
 			try
 			{
@@ -510,39 +466,47 @@ namespace UnrealGameSync
 			}
 		}
 
-		public static bool TryPrintFileUsingCache(PerforceConnection Perforce, string DepotPath, DirectoryReference CacheFolder, string Digest, out List<string> Lines, TextWriter Log)
+		public static async Task<string[]?> TryPrintFileUsingCacheAsync(IPerforceConnection Perforce, string DepotPath, DirectoryReference CacheFolder, string? Digest, ILogger Logger, CancellationToken CancellationToken)
 		{
 			if(Digest == null)
 			{
-				return Perforce.Print(DepotPath, out Lines, Log);
+				PerforceResponse<PrintRecord<string[]>> Response = await Perforce.TryPrintLinesAsync(DepotPath, CancellationToken);
+				if (Response.Succeeded)
+				{
+					return Response.Data.Contents;
+				}
+				else
+				{
+					return null;
+				}
 			}
 
 			FileReference CacheFile = FileReference.Combine(CacheFolder, Digest);
 			if(FileReference.Exists(CacheFile))
 			{
-				Log.WriteLine("Reading cached copy of {0} from {1}", DepotPath, CacheFile);
-				Lines = new List<string>(FileReference.ReadAllLines(CacheFile));
+				Logger.LogInformation("Reading cached copy of {DepotFile} from {LocalFile}", DepotPath, CacheFile);
+				string[] Lines = FileReference.ReadAllLines(CacheFile);
 				try
 				{
 					FileReference.SetLastWriteTimeUtc(CacheFile, DateTime.UtcNow);
 				}
 				catch(Exception Ex)
 				{
-					Log.WriteLine("Exception touching cache file {0}: {1}", CacheFile, Ex.ToString());
+					Logger.LogWarning(Ex, "Exception touching cache file {LocalFile}", CacheFile);
 				}
-				return true;
+				return Lines;
 			}
 			else
 			{
 				FileReference TempFile = new FileReference(String.Format("{0}.{1}.temp", CacheFile.FullName, Guid.NewGuid()));
-				if(!Perforce.PrintToFile(DepotPath, TempFile.FullName, Log))
+				PerforceResponse<PrintRecord> Response = await Perforce.TryPrintAsync(TempFile.FullName, DepotPath, CancellationToken);
+				if (!Response.Succeeded)
 				{
-					Lines = null;
-					return false;
+					return null;
 				}
 				else
 				{
-					Lines = new List<string>(FileReference.ReadAllLines(TempFile));
+					string[] Lines = await FileReference.ReadAllLinesAsync(TempFile);
 					try
 					{
 						FileReference.SetAttributes(TempFile, FileAttributes.Normal);
@@ -559,7 +523,7 @@ namespace UnrealGameSync
 						{
 						}
 					}
-					return true;
+					return Lines;
 				}
 			}
 		}
@@ -592,11 +556,18 @@ namespace UnrealGameSync
 			return Color.FromArgb((int)(First.R + (Second.R - First.R) * T), (int)(First.G + (Second.G - First.G) * T), (int)(First.B + (Second.B - First.B) * T));
 		}
 
-		public static PerforceConnection OverridePerforceSettings(PerforceConnection DefaultConnection, string ServerAndPort, string UserName)
+		public static PerforceSettings OverridePerforceSettings(IPerforceSettings DefaultConnection, string? ServerAndPort, string? UserName)
 		{
-			string ResolvedServerAndPort = String.IsNullOrWhiteSpace(ServerAndPort) ? DefaultConnection.ServerAndPort : ServerAndPort;
-			string ResolvedUserName = String.IsNullOrWhiteSpace(UserName) ? DefaultConnection.UserName : UserName;
-			return new PerforceConnection(ResolvedUserName, null, ResolvedServerAndPort);
+			PerforceSettings NewSettings = new PerforceSettings(DefaultConnection);
+			if(!String.IsNullOrWhiteSpace(ServerAndPort))
+			{
+				NewSettings.ServerAndPort = ServerAndPort;
+			}
+			if (!String.IsNullOrWhiteSpace(UserName))
+			{
+				NewSettings.UserName = UserName;
+			}
+			return NewSettings;
 		}
 
 		public static string FormatRecentDateTime(DateTime Date)
@@ -672,12 +643,12 @@ namespace UnrealGameSync
 
 		public static IEnumerable<string> GetPerforcePaths()
 		{
-			string PathList = Environment.GetEnvironmentVariable("PATH");
+			string? PathList = Environment.GetEnvironmentVariable("PATH");
 			if (!String.IsNullOrEmpty(PathList))
 			{
 				foreach (string PathEntry in PathList.Split(Path.PathSeparator))
 				{
-					string PerforcePath = null;
+					string? PerforcePath = null;
 					try
 					{
 						string TestPerforcePath = Path.Combine(PathEntry, "p4.exe");
@@ -704,8 +675,8 @@ namespace UnrealGameSync
 
 			foreach (string PerforcePath in GetPerforcePaths())
 			{
-				string PerforceDir = Path.GetDirectoryName(PerforcePath);
-				if (File.Exists(Path.Combine(PerforceDir, "p4vc.bat")) && !File.Exists(Path.Combine(PerforceDir, "p4vc.exe")))
+				string? PerforceDir = Path.GetDirectoryName(PerforcePath);
+				if (PerforceDir != null && File.Exists(Path.Combine(PerforceDir, "p4vc.bat")) && !File.Exists(Path.Combine(PerforceDir, "p4vc.exe")))
 				{
 					Executable = Path.Combine(PerforceDir, "p4v.exe");
 					Arguments = "-p4vc " + Arguments;
