UGS: Significant refactor of UGS internals.

diff --git a/Engine/Source/Programs/UnrealGameSync/UnrealGameSync/Workspace.cs b/Engine/Source/Programs/UnrealGameSync/UnrealGameSync/Workspace.cs
--- a/Engine/Source/Programs/UnrealGameSync/UnrealGameSync/Workspace.cs
+++ b/Engine/Source/Programs/UnrealGameSync/UnrealGameSync/Workspace.cs
@@ -1,6 +1,9 @@
 // Copyright Epic Games, Inc. All Rights Reserved.
 
 using EpicGames.Core;
+using EpicGames.Perforce;
+using Microsoft.Extensions.DependencyInjection;
+using Microsoft.Extensions.Logging;
 using System;
 using System.Collections.Concurrent;
 using System.Collections.Generic;
@@ -16,10 +19,12 @@ using System.Text.Json;
 using System.Threading;
 using System.Threading.Tasks;
 
+#nullable enable
+
 namespace UnrealGameSync
 {
 	[Flags]
-	enum WorkspaceUpdateOptions
+	public enum WorkspaceUpdateOptions
 	{
 		Sync = 0x01,
 		SyncSingleChange = 0x02,
@@ -37,7 +42,7 @@ namespace UnrealGameSync
 		IncludeAllProjectsInSolution = 0x2000,
 	}
 
-	enum WorkspaceUpdateResult
+	public enum WorkspaceUpdateResult
 	{
 		Canceled,
 		FailedToSync,
@@ -50,7 +55,32 @@ namespace UnrealGameSync
 		Success,
 	}
 
-	class WorkspaceUpdateContext
+	public class PerforceSyncOptions
+	{
+		public const int DefaultNumRetries = 0;
+		public const int DefaultNumThreads = 2;
+		public const int DefaultTcpBufferSize = 0;
+		public const int DefaultFileBufferSize = 0;
+
+		public const int DefaultMaxCommandsPerBatch = 200;
+		public const int DefaultMaxSizePerBatch = 128 * 1024 * 1024;
+
+		public int NumRetries = DefaultNumRetries;
+		public int NumThreads = DefaultNumThreads;
+		public int TcpBufferSize = DefaultTcpBufferSize;
+		public int FileBufferSize = DefaultFileBufferSize;
+
+		public int MaxCommandsPerBatch = DefaultMaxCommandsPerBatch;
+		public int MaxSizePerBatch = DefaultMaxSizePerBatch;
+
+
+		public PerforceSyncOptions Clone()
+		{
+			return (PerforceSyncOptions)MemberwiseClone();
+		}
+	}
+
+	public class WorkspaceUpdateContext
 	{
 		public DateTime StartTime = DateTime.UtcNow;
 		public int ChangeNumber;
@@ -63,12 +93,12 @@ namespace UnrealGameSync
 		public List<ConfigObject> UserBuildStepObjects;
 		public HashSet<Guid> CustomBuildSteps;
 		public Dictionary<string, string> Variables;
-		public PerforceSyncOptions PerforceSyncOptions;
-		public List<PerforceFileRecord> HaveFiles; // Cached when sync filter has changed
+		public PerforceSyncOptions? PerforceSyncOptions;
+		public List<HaveRecord>? HaveFiles; // Cached when sync filter has changed
 
 		// May be updated during sync
-		public ConfigFile ProjectConfigFile;
-		public IReadOnlyList<string> ProjectStreamFilter;
+		public ConfigFile? ProjectConfigFile;
+		public IReadOnlyList<string>? ProjectStreamFilter;
 
 		public WorkspaceUpdateContext(int InChangeNumber, WorkspaceUpdateOptions InOptions, string[] InSyncFilter, Dictionary<Guid, ConfigObject> InDefaultBuildSteps, List<ConfigObject> InUserBuildSteps, HashSet<Guid> InCustomBuildSteps, Dictionary<string, string> InVariables)
 		{
@@ -140,24 +170,26 @@ namespace UnrealGameSync
 	public class ProjectInfo
 	{
 		public DirectoryReference LocalRootPath { get; }
-		public FileReference SelectedLocalFileName { get; }
+		public FileReference LocalFileName { get; }
 		public string ClientRootPath { get; }
-		public string SelectedClientFileName { get; }
-		public string TelemetryProjectPath { get; }
+		public string ClientFileName { get; }
+		public string ProjectIdentifier { get; }
+		public string TelemetryProjectIdentifier { get; }
 		public bool bIsEnterpriseProject { get; }
 
-		public ProjectInfo(DirectoryReference InLocalRootPath, FileReference InSelectedLocalFileName, string InClientRootPath, string InSelectedClientFileName, string InTelemetryProjectPath, bool bInIsEnterpriseProject)
+		public ProjectInfo(DirectoryReference LocalRootPath, FileReference LocalFileName, string ClientRootPath, string ClientFileName, string ProjectIdentifier, bool bIsEnterpriseProject)
 		{
-			LocalRootPath = InLocalRootPath;
-			SelectedLocalFileName = InSelectedLocalFileName;
-			ClientRootPath = InClientRootPath;
-			SelectedClientFileName = InSelectedClientFileName;
-			TelemetryProjectPath = InTelemetryProjectPath;
-			bIsEnterpriseProject = bInIsEnterpriseProject;
+			this.LocalRootPath = LocalRootPath;
+			this.LocalFileName = LocalFileName;
+			this.ClientRootPath = ClientRootPath;
+			this.ClientFileName = ClientFileName;
+			this.ProjectIdentifier = ProjectIdentifier;
+			this.TelemetryProjectIdentifier = PerforceUtils.GetClientOrDepotDirectoryName(ProjectIdentifier);
+			this.bIsEnterpriseProject = bIsEnterpriseProject;
 		}
 	}
 
-	class Workspace : IDisposable
+	public sealed class Workspace : IDisposable
 	{
 		const string BuildVersionFileName = "/Engine/Build/Build.version";
 		const string VersionHeaderFileName = "/Engine/Source/Runtime/Launch/Resources/Version.h";
@@ -166,17 +198,19 @@ namespace UnrealGameSync
 		static readonly string LocalVersionHeaderFileName = VersionHeaderFileName.Replace('/', '\\');
 		static readonly string LocalObjectVersionFileName = ObjectVersionFileName.Replace('/', '\\');
 
-		public readonly PerforceConnection Perforce;
-		public ProjectInfo Project { get; }
+		public IPerforceSettings PerforceSettings { get; }
 		public UserWorkspaceState State { get; }
-		Thread WorkerThread;
-		TextWriter Log;
-		bool bSyncing;
-		ProgressValue Progress = new ProgressValue();
+		public ProjectInfo Project { get; }
+		public SynchronizationContext SynchronizationContext { get; }
+		ILogger Logger;
+
+		bool bSyncing => CurrentUpdate != null;
+		ProgressValue Progress => CurrentUpdate?.Progress ?? new ProgressValue();
 
-		static UserWorkspaceState ActiveWorkspaceState;
+		WorkspaceUpdate? CurrentUpdate;
+		static SemaphoreSlim UpdateSemaphore = new SemaphoreSlim(1);
 
-		public event Action<WorkspaceUpdateContext, WorkspaceUpdateResult, string> OnUpdateComplete;
+		public event Action<WorkspaceUpdateContext, WorkspaceUpdateResult, string>? OnUpdateComplete;
 
 		class RecordCounter : IDisposable
 		{
@@ -201,7 +235,7 @@ namespace UnrealGameSync
 			public void Increment()
 			{
 				Count++;
-				if(Timer.ElapsedMilliseconds > 250)
+				if (Timer.ElapsedMilliseconds > 250)
 				{
 					UpdateMessage();
 				}
@@ -220,14 +254,14 @@ namespace UnrealGameSync
 			public long MaxSizePerList { get; }
 			public Queue<List<string>> Batches { get; }
 
-			List<string> Commands;
-			List<string> DeleteCommands;
+			List<string>? Commands;
+			List<string>? DeleteCommands;
 			long Size;
 
-			public SyncBatchBuilder(int MaxCommandsPerList, long MaxSizePerList)
+			public SyncBatchBuilder(int? MaxCommandsPerList, long? MaxSizePerList)
 			{
-				this.MaxCommandsPerList = MaxCommandsPerList;
-				this.MaxSizePerList = MaxSizePerList;
+				this.MaxCommandsPerList = MaxCommandsPerList ?? PerforceSyncOptions.DefaultMaxCommandsPerBatch;
+				this.MaxSizePerList = MaxSizePerList ?? PerforceSyncOptions.DefaultMaxSizePerBatch;
 				this.Batches = new Queue<List<string>>();
 			}
 
@@ -274,7 +308,7 @@ namespace UnrealGameSync
 
 			public SyncTree FindOrAddSubTree(string Name)
 			{
-				SyncTree Result;
+				SyncTree? Result;
 				if (!NameToSubTree.TryGetValue(Name, out Result))
 				{
 					Result = new SyncTree(bCanUseWildcard);
@@ -330,20 +364,28 @@ namespace UnrealGameSync
 			}
 		}
 
-		public Workspace(PerforceConnection InPerforce, ProjectInfo InProject, UserWorkspaceState InState, TextWriter InLog)
+		IAsyncDisposer AsyncDisposer;
+
+		public Workspace(IPerforceSettings InPerfoceSettings, ProjectInfo InProject, UserWorkspaceState InState, ConfigFile ProjectConfigFile, IReadOnlyList<string> ProjectStreamFilter, ILogger Logger, IServiceProvider ServiceProvider)
 		{
-			Perforce = InPerforce;
+			PerforceSettings = InPerfoceSettings;
 			Project = InProject;
 			State = InState;
-			Log = InLog;
+			this.SynchronizationContext = SynchronizationContext.Current!;
+			this.Logger = Logger;
+			this.AsyncDisposer = ServiceProvider.GetRequiredService<IAsyncDisposer>();
 
-			ProjectConfigFile = ReadProjectConfigFile(Project.LocalRootPath, Project.SelectedLocalFileName, Log);
-			ProjectStreamFilter = ReadProjectStreamFilter(Perforce, ProjectConfigFile, Log);
+			this.ProjectConfigFile = ProjectConfigFile;
+			this.ProjectStreamFilter = ProjectStreamFilter;
 		}
 
 		public void Dispose()
 		{
 			CancelUpdate();
+			if (PrevUpdateTask != null)
+			{
+				AsyncDisposer.Add(PrevUpdateTask);
+			}
 		}
 
 		public ConfigFile ProjectConfigFile
@@ -356,125 +398,187 @@ namespace UnrealGameSync
 			get; private set;
 		}
 
-		public void Update(WorkspaceUpdateContext Context)
+		class WorkspaceUpdate : IAsyncDisposable
 		{
-			// Kill any existing sync
-			CancelUpdate();
+			CancellationTokenSource CancellationSource;
+			public WorkspaceUpdateContext Context { get; }
+			public ProgressValue Progress { get; } = new ProgressValue();
+			public Task UpdateTask { get; }
+
+			public WorkspaceUpdate(WorkspaceUpdateContext Context, Func<WorkspaceUpdate, CancellationToken, Task> TaskFunc)
+			{
+				this.CancellationSource = new CancellationTokenSource();
+				this.Context = Context;
+				this.UpdateTask = Task.Run(() => TaskFunc(this, CancellationSource.Token));
+			}
 
-			// Set the initial progress message
-			if(CurrentChangeNumber != Context.ChangeNumber)
+			public void Cancel()
 			{
-				PendingChangeNumber = Context.ChangeNumber;
+				CancellationSource.Cancel();
 			}
-			Progress.Clear();
-			bSyncing = true;
 
-			// Spawn the new thread
-			WorkerThread = new Thread(x => UpdateWorkspace(Context));
-			WorkerThread.Name = "Workspace Update";
-			WorkerThread.Start();
+			public ValueTask DisposeAsync()
+			{
+				Task DisposeTask = Task.CompletedTask;
+				if (UpdateTask != null)
+				{
+					CancellationSource.Cancel();
+					DisposeTask = UpdateTask.ContinueWith(Task => CancellationSource.Dispose());
+				}
+				return new ValueTask(DisposeTask);
+			}
+		}
+
+		Task PrevUpdateTask = Task.CompletedTask;
+
+		public void Update(WorkspaceUpdateContext Context)
+		{
+			CancelUpdate();
+
+			Task PrevUpdateTaskCopy = PrevUpdateTask;
+			CurrentUpdate = new WorkspaceUpdate(Context, (Update, CancellationToken) => UpdateWorkspaceMini(Update, PrevUpdateTaskCopy, CancellationToken));
 		}
 
 		public void CancelUpdate()
 		{
-			if(bSyncing)
+			// Cancel the current task. We actually terminate the operation asynchronously, but we can signal the cancellation and 
+			// send a cancelled event, then wait for the heavy lifting to finish in the new update task.
+			if (CurrentUpdate != null)
 			{
-				Log.WriteLine("OPERATION ABORTED");
-				if(WorkerThread != null)
-				{
-					WorkerThread.Interrupt();
-					WorkerThread.Join();
-					WorkerThread = null;
-				}
-				PendingChangeNumber = CurrentChangeNumber;
-				bSyncing = false;
-				Interlocked.CompareExchange(ref ActiveWorkspaceState, null, State);
+				PrevUpdateTask = CurrentUpdate.DisposeAsync().AsTask();
+				CompleteUpdate(CurrentUpdate, WorkspaceUpdateResult.Canceled, "Cancelled");
 			}
 		}
 
-		void UpdateWorkspace(WorkspaceUpdateContext Context)
+		async Task UpdateWorkspaceMini(WorkspaceUpdate Update, Task PrevUpdateTask, CancellationToken CancellationToken)
 		{
-			string StatusMessage;
+			if (PrevUpdateTask != null)
+			{
+				await PrevUpdateTask;
+			}
 
+			WorkspaceUpdateContext Context = Update.Context;
 			Context.ProjectConfigFile = ProjectConfigFile;
 			Context.ProjectStreamFilter = ProjectStreamFilter;
 
+			string StatusMessage;
 			WorkspaceUpdateResult Result = WorkspaceUpdateResult.FailedToSync;
+
 			try
 			{
-				Result = UpdateWorkspaceInternal(Perforce, Project, State, Context, Progress, Log, out StatusMessage);
-				if(Result != WorkspaceUpdateResult.Success)
+				(Result, StatusMessage) = await UpdateWorkspaceInternal(PerforceSettings, Project, State, Context, Progress, Logger, CancellationToken);
+				if (Result != WorkspaceUpdateResult.Success)
 				{
-					Log.WriteLine("{0}", StatusMessage);
+					Logger.LogError("{Message}", StatusMessage);
 				}
 			}
-			catch (ThreadAbortException)
-			{
-				StatusMessage = "Canceled.";
-				Log.WriteLine("Canceled.");
-			}
-			catch (ThreadInterruptedException)
+			catch (OperationCanceledException)
 			{
 				StatusMessage = "Canceled.";
-				Log.WriteLine("Canceled.");
+				Logger.LogError("Canceled.");
 			}
 			catch (Exception Ex)
 			{
 				StatusMessage = "Failed with exception - " + Ex.ToString();
-				Log.WriteException(Ex, "Failed with exception");
+				Logger.LogError(Ex, "Failed with exception");
 			}
 
 			ProjectConfigFile = Context.ProjectConfigFile;
 			ProjectStreamFilter = Context.ProjectStreamFilter;
 
-			State.LastSyncChangeNumber = Context.ChangeNumber;
-			State.LastSyncResult = Result;
-			State.LastSyncResultMessage = StatusMessage;
-			State.LastSyncTime = DateTime.UtcNow;
-			State.LastSyncDurationSeconds = (int)(State.LastSyncTime.Value - Context.StartTime).TotalSeconds;
-			State.Save();
+			SynchronizationContext.Post(x => CompleteUpdate(Update, Result, StatusMessage), null);
+		}
+
+		void CompleteUpdate(WorkspaceUpdate Update, WorkspaceUpdateResult Result, string StatusMessage)
+		{
+			if (CurrentUpdate == Update)
+			{
+				WorkspaceUpdateContext Context = Update.Context;
+
+				State.LastSyncChangeNumber = Context.ChangeNumber;
+				State.LastSyncResult = Result;
+				State.LastSyncResultMessage = StatusMessage;
+				State.LastSyncTime = DateTime.UtcNow;
+				State.LastSyncDurationSeconds = (int)(State.LastSyncTime.Value - Context.StartTime).TotalSeconds;
+				State.Save();
+
+				OnUpdateComplete?.Invoke(Context, Result, StatusMessage);
+				CurrentUpdate = null;
+			}
+		}
+
+		class SyncFile
+		{
+			public string DepotFile;
+			public string RelativePath;
+			public long Size;
+
+			public SyncFile(string DepotFile, string RelativePath, long Size)
+			{
+				this.DepotFile = DepotFile;
+				this.RelativePath = RelativePath;
+				this.Size = Size;
+			}
+		};
+
+		class SemaphoreScope : IDisposable
+		{
+			SemaphoreSlim Semaphore;
+			public bool HasLock { get; private set; }
 
-			bSyncing = false;
-			PendingChangeNumber = CurrentChangeNumber;
-			Interlocked.CompareExchange(ref ActiveWorkspaceState, null, State);
+			public SemaphoreScope(SemaphoreSlim Semaphore)
+			{
+				this.Semaphore = Semaphore;
+			}
+
+			public bool TryAcquire()
+			{
+				HasLock = HasLock || Semaphore.Wait(0);
+				return HasLock;
+			}
+
+			public async Task AcquireAsync(CancellationToken CancellationToken)
+			{
+				await Semaphore.WaitAsync(CancellationToken);
+				HasLock = true;
+			}
 
-			if(OnUpdateComplete != null)
+			public void Release()
 			{
-				OnUpdateComplete(Context, Result, StatusMessage);
+				if (HasLock)
+				{
+					Semaphore.Release();
+					HasLock = false;
+				}
 			}
+
+			public void Dispose() => Release();
 		}
 
-		static WorkspaceUpdateResult UpdateWorkspaceInternal(PerforceConnection Perforce, ProjectInfo Project, UserWorkspaceState State, WorkspaceUpdateContext Context, ProgressValue Progress, TextWriter Log, out string StatusMessage)
+		static async Task<(WorkspaceUpdateResult, string)> UpdateWorkspaceInternal(IPerforceSettings PerforceSettings, ProjectInfo Project, UserWorkspaceState State, WorkspaceUpdateContext Context, ProgressValue Progress, ILogger Logger, CancellationToken CancellationToken)
 		{
-//			string LocalRootPrefix = Project.LocalRootPath.FullName.TrimEnd(Path.DirectorySeparatorChar) + Path.DirectorySeparatorChar;
+			using IPerforceConnection Perforce = await PerforceConnection.CreateAsync(PerforceSettings, Logger);
 
 			string CmdExe = Path.Combine(Environment.GetFolderPath(Environment.SpecialFolder.System), "cmd.exe");
-			if(!File.Exists(CmdExe))
+			if (!File.Exists(CmdExe))
 			{
-				StatusMessage = String.Format("Missing {0}.", CmdExe);
-				return WorkspaceUpdateResult.FailedToSync;
+				return (WorkspaceUpdateResult.FailedToSync, $"Missing {CmdExe}.");
 			}
 
-			List<Tuple<string, TimeSpan>> Times = new List<Tuple<string,TimeSpan>>();
+			List<Tuple<string, TimeSpan>> Times = new List<Tuple<string, TimeSpan>>();
 
 			int NumFilesSynced = 0;
-			if(Context.Options.HasFlag(WorkspaceUpdateOptions.Sync) || Context.Options.HasFlag(WorkspaceUpdateOptions.SyncSingleChange))
+			if (Context.Options.HasFlag(WorkspaceUpdateOptions.Sync) || Context.Options.HasFlag(WorkspaceUpdateOptions.SyncSingleChange))
 			{
-				using(TelemetryStopwatch SyncTelemetryStopwatch = new TelemetryStopwatch("Workspace_Sync", Project.TelemetryProjectPath))
+				using (TelemetryStopwatch SyncTelemetryStopwatch = new TelemetryStopwatch("Workspace_Sync", Project.TelemetryProjectIdentifier))
 				{
-					Log.WriteLine("Syncing to {0}...", Context.ChangeNumber);
+					Logger.LogInformation("Syncing to {Change}...", Context.ChangeNumber);
 
 					// Make sure we're logged in
-					bool bLoggedIn;
-					if(!Perforce.GetLoggedInState(out bLoggedIn, Log))
-					{
-						StatusMessage = "Unable to get login status.";
-						return WorkspaceUpdateResult.FailedToSync;
-					}
-					if(!bLoggedIn)
+					PerforceResponse<LoginRecord> LoginResponse = await Perforce.TryGetLoginStateAsync(CancellationToken);
+					if (!LoginResponse.Succeeded)
 					{
-						StatusMessage = "User is not logged in.";
-						return WorkspaceUpdateResult.FailedToSyncLoginExpired;
+						return (WorkspaceUpdateResult.FailedToSyncLoginExpired, "User is not logged in.");
 					}
 
 					// Figure out which paths to sync
@@ -483,21 +587,21 @@ namespace UnrealGameSync
 
 					// Get the user's sync filter
 					FileFilter UserFilter = new FileFilter(FileFilterType.Include);
-					if(Context.SyncFilter != null)
+					if (Context.SyncFilter != null)
 					{
 						UserFilter.AddRules(Context.SyncFilter.Select(x => x.Trim()).Where(x => x.Length > 0 && !x.StartsWith(";") && !x.StartsWith("#")));
 					}
 
 					// Check if the new sync filter matches the previous one. If not, we'll enumerate all files in the workspace and make sure there's nothing extra there.
-					string NextSyncFilterHash = null;
+					string? NextSyncFilterHash = null;
 					using (SHA1Managed SHA = new SHA1Managed())
 					{
 						StringBuilder CombinedFilter = new StringBuilder();
-						foreach(string RelativeSyncPath in RelativeSyncPaths)
+						foreach (string RelativeSyncPath in RelativeSyncPaths)
 						{
 							CombinedFilter.AppendFormat("{0}\n", RelativeSyncPath);
 						}
-						if(Context.SyncFilter != null)
+						if (Context.SyncFilter != null)
 						{
 							CombinedFilter.Append("--FROM--\n");
 							CombinedFilter.Append(String.Join("\n", Context.SyncFilter));
@@ -508,21 +612,28 @@ namespace UnrealGameSync
 					// If the hash differs, enumerate everything in the workspace to find what needs to be removed
 					if (NextSyncFilterHash != State.CurrentSyncFilterHash)
 					{
-						using (TelemetryStopwatch FilterStopwatch = new TelemetryStopwatch("Workspace_Sync_FilterChanged", Project.TelemetryProjectPath))
+						using (TelemetryStopwatch FilterStopwatch = new TelemetryStopwatch("Workspace_Sync_FilterChanged", Project.TelemetryProjectIdentifier))
 						{
-							Log.WriteLine("Filter has changed ({0} -> {1}); finding files in workspace that need to be removed.", (String.IsNullOrEmpty(State.CurrentSyncFilterHash)) ? "None" : State.CurrentSyncFilterHash, NextSyncFilterHash);
+							Logger.LogInformation("Filter has changed ({PrevHash} -> {NextHash}); finding files in workspace that need to be removed.", (String.IsNullOrEmpty(State.CurrentSyncFilterHash)) ? "None" : State.CurrentSyncFilterHash, NextSyncFilterHash);
 
 							// Find all the files that are in this workspace
-							List<PerforceFileRecord> HaveFiles = Context.HaveFiles;
+							List<HaveRecord>? HaveFiles = Context.HaveFiles;
 							if (HaveFiles == null)
 							{
-								HaveFiles = new List<PerforceFileRecord>();
+								HaveFiles = new List<HaveRecord>();
 								using (RecordCounter HaveCounter = new RecordCounter(Progress, "Sync filter changed; checking workspace..."))
 								{
-									if (!Perforce.Have("//...", Record => { HaveFiles.Add(Record); HaveCounter.Increment(); }, Log))
+									await foreach (PerforceResponse<HaveRecord> Record in Perforce.TryHaveAsync(FileSpecList.Any, CancellationToken))
 									{
-										StatusMessage = "Unable to query files.";
-										return WorkspaceUpdateResult.FailedToSync;
+										if (Record.Succeeded)
+										{
+											HaveFiles.Add(Record.Data);
+											HaveCounter.Increment();
+										}
+										else
+										{
+											return (WorkspaceUpdateResult.FailedToSync, $"Unable to query files ({Record}).");
+										}
 									}
 								}
 								Context.HaveFiles = HaveFiles;
@@ -537,15 +648,15 @@ namespace UnrealGameSync
 
 							// Remove all the files that are not included by the filter
 							List<string> RemoveDepotPaths = new List<string>();
-							foreach (PerforceFileRecord HaveFile in HaveFiles)
+							foreach (HaveRecord HaveFile in HaveFiles)
 							{
 								try
 								{
 									FileReference FullPath = new FileReference(HaveFile.Path);
 									if (MatchFilter(Project, FullPath, SyncPathsFilter) && !MatchFilter(Project, FullPath, UserFilter))
 									{
-										Log.WriteLine("  {0}", HaveFile.DepotPath);
-										RemoveDepotPaths.Add(HaveFile.DepotPath);
+										Logger.LogInformation("  {DepotFile}", HaveFile.DepotFile);
+										RemoveDepotPaths.Add(HaveFile.DepotFile);
 									}
 								}
 								catch (PathTooLongException)
@@ -574,8 +685,7 @@ namespace UnrealGameSync
 
 								if (!bDeleteListMatches)
 								{
-									StatusMessage = String.Format("Cancelled after finding {0} files excluded by filter", NewDeleteFiles.Count);
-									return WorkspaceUpdateResult.FilesToDelete;
+									return (WorkspaceUpdateResult.FilesToDelete, $"Cancelled after finding {NewDeleteFiles.Count} files excluded by filter");
 								}
 
 								RemoveDepotPaths.RemoveAll(x => !Context.DeleteFiles[x]);
@@ -596,8 +706,8 @@ namespace UnrealGameSync
 								List<string> RevisionsToRemove = new List<string>();
 								RevisionsToRemove.AddRange(RemoveDepotPaths.Select(x => String.Format("{0}#0", x)));
 
-								WorkspaceUpdateResult RemoveResult = SyncFileRevisions(Perforce, "Removing files...", Context, RevisionsToRemove, RemainingDepotPathsToRemove, Progress, Log, out StatusMessage);
-								if (RemoveResult != WorkspaceUpdateResult.Success)
+								(WorkspaceUpdateResult, string) RemoveResult = await SyncFileRevisions(Perforce, "Removing files...", Context, RevisionsToRemove, RemainingDepotPathsToRemove, Progress, Logger, CancellationToken);
+								if (RemoveResult.Item1 != WorkspaceUpdateResult.Success)
 								{
 									return RemoveResult;
 								}
@@ -641,82 +751,73 @@ namespace UnrealGameSync
 					}
 
 					// Find all the server changes, and anything that's opened for edit locally. We need to sync files we have open to schedule a resolve.
-					SyncBatchBuilder BatchBuilder = new SyncBatchBuilder(Context.PerforceSyncOptions.MaxCommandsPerBatch, Context.PerforceSyncOptions.MaxSizePerBatch);
+					SyncBatchBuilder BatchBuilder = new SyncBatchBuilder(Context.PerforceSyncOptions?.MaxCommandsPerBatch, Context.PerforceSyncOptions?.MaxSizePerBatch);
 					List<string> SyncDepotPaths = new List<string>();
-					using(RecordCounter Counter = new RecordCounter(Progress, "Filtering files..."))
+					using (RecordCounter Counter = new RecordCounter(Progress, "Filtering files..."))
 					{
 						// Track the total new bytes that will be required on disk when syncing. Add an extra 100MB for padding.
 						long RequiredFreeSpace = 100 * 1024 * 1024;
 
-						foreach(string SyncPath in SyncPaths)
+						foreach (string SyncPath in SyncPaths)
 						{
-							List<PerforceFileRecord> SyncRecords = new List<PerforceFileRecord>();
-							if(!Perforce.SyncPreview(SyncPath, Context.ChangeNumber, !Context.Options.HasFlag(WorkspaceUpdateOptions.Sync), Record => { SyncRecords.Add(Record); Counter.Increment(); }, Log))
-							{
-								StatusMessage = String.Format("Couldn't enumerate changes matching {0}.", SyncPath);
-								return WorkspaceUpdateResult.FailedToSync;
-							}
+							string SyncFilter = Context.Options.HasFlag(WorkspaceUpdateOptions.Sync) ? $"{Context.ChangeNumber}" : $"={Context.ChangeNumber}";
 
-							List<PerforceFileRecord> OpenRecords;
-							if (!Perforce.GetOpenFiles(SyncPath, out OpenRecords, Log))
+							List<SyncFile> SyncFiles = new List<SyncFile>();
+							await foreach (PerforceResponse<SyncRecord> Response in Perforce.TrySyncAsync(SyncOptions.PreviewOnly, -1, $"{SyncPath}@{SyncFilter}", CancellationToken))
 							{
-								StatusMessage = String.Format("Couldn't find open files matching {0}.", SyncPath);
-								return WorkspaceUpdateResult.FailedToSync;
-							}
-
-							SyncRecords.AddRange(OpenRecords.Where(x => x.Action != "add" && x.Action != "branch" && x.Action != "move/add"));
-
-							// Enumerate all the files to be synced. NOTE: depotPath is escaped, whereas clientPath is not.
-							foreach (PerforceFileRecord SyncRecord in SyncRecords)
-							{
-								// If it doesn't exist locally, just add a sync command for it
-								if (String.IsNullOrEmpty(SyncRecord.ClientPath))
+								if (!Response.Succeeded)
 								{
-									BatchBuilder.Add(String.Format("{0}@{1}", SyncRecord.DepotPath, Context.ChangeNumber), SyncRecord.FileSize);
-									SyncDepotPaths.Add(SyncRecord.DepotPath);
-									RequiredFreeSpace += SyncRecord.FileSize;
-									continue;
+									return (WorkspaceUpdateResult.FailedToSync, $"Couldn't enumerate changes matching {SyncPath}.");
 								}
 
+								SyncRecord Record = Response.Data;
+
 								// Get the full local path
-								FileReference SyncFile;
+								string RelativePath;
 								try
 								{
-									SyncFile = new FileReference(SyncRecord.ClientPath);
+									FileReference SyncFile = new FileReference(Record.Path);
+									RelativePath = PerforceUtils.GetClientRelativePath(Project.LocalRootPath, SyncFile);
 								}
-								catch(PathTooLongException)
+								catch (PathTooLongException)
 								{
-									Log.WriteLine("The local path for {0} exceeds the maximum allowed by Windows. Re-sync your workspace to a directory with a shorter name, or delete the file from the server.", SyncRecord.ClientPath);
-									StatusMessage = "File exceeds maximum path length allowed by Windows.";
-									return WorkspaceUpdateResult.FailedToSync;
+									Logger.LogInformation("The local path for {Path} exceeds the maximum allowed by Windows. Re-sync your workspace to a directory with a shorter name, or delete the file from the server.", Record.Path);
+									return (WorkspaceUpdateResult.FailedToSync, "File exceeds maximum path length allowed by Windows.");
 								}
 
-								// Make sure it's under the current directory. Not sure why this would happen, just being safe.
-								// This occurs for files returned from GetOpenFiles as SyncRecord.ClientPath is the client workspace path not local path.
-								if (!SyncFile.IsUnderDirectory(Project.LocalRootPath))
+								// Create the sync record
+								long SyncSize = (Record.Action == SyncAction.Deleted) ? 0 : Record.FileSize;
+								SyncFiles.Add(new SyncFile(Record.DepotFile, RelativePath, SyncSize));
+								Counter.Increment();
+							}
+
+							// Also sync the currently open files
+							await foreach (PerforceResponse<OpenedRecord> Response in Perforce.TryOpenedAsync(OpenedOptions.None, FileSpecList.Any, CancellationToken))
+							{
+								if (!Response.Succeeded)
 								{
-									BatchBuilder.Add(String.Format("{0}@{1}", SyncRecord.DepotPath, Context.ChangeNumber), SyncRecord.FileSize);
-									SyncDepotPaths.Add(SyncRecord.DepotPath);
-									RequiredFreeSpace += SyncRecord.FileSize;
-									continue;
+									return (WorkspaceUpdateResult.FailedToSync, $"Couldn't enumerate changes matching {SyncPath}.");
 								}
 
-								// Check that it matches the filter
-								string RelativePath = SyncFile.MakeRelativeTo(Project.LocalRootPath).Replace('\\', '/');
-								if (Filter.Matches(RelativePath))
+								OpenedRecord Record = Response.Data;
+								if (Record.Action != FileAction.Add || Record.Action != FileAction.Branch || Record.Action != FileAction.MoveAdd)
 								{
-									long FileSize = SyncRecord.FileSize;
-									if (SyncRecord.Action == "deleted" || SyncRecord.Action == "move/delete")
-									{
-										FileSize = 0;
-									}
+									string RelativePath = PerforceUtils.GetClientRelativePath(Record.ClientFile);
+									SyncFiles.Add(new SyncFile(Record.DepotFile, RelativePath, 0));
+								}
+							}
 
-									SyncTree.IncludeFile(PerforceUtils.EscapePath(RelativePath), FileSize);
-									SyncDepotPaths.Add(SyncRecord.DepotPath);
-									RequiredFreeSpace += FileSize;
-									FileInfo LocalFileInfo = SyncFile.ToFileInfo();
+							// Enumerate all the files to be synced. NOTE: depotPath is escaped, whereas clientPath is not.
+							foreach (SyncFile SyncRecord in SyncFiles)
+							{
+								if (Filter.Matches(SyncRecord.RelativePath))
+								{
+									SyncTree.IncludeFile(PerforceUtils.EscapePath(SyncRecord.RelativePath), SyncRecord.Size);
+									SyncDepotPaths.Add(SyncRecord.DepotFile);
+									RequiredFreeSpace += SyncRecord.Size;
 
 									// If the file exists the required free space can be reduced as those bytes will be replaced.
+									FileInfo LocalFileInfo = FileReference.Combine(Project.LocalRootPath, SyncRecord.RelativePath).ToFileInfo();
 									if (LocalFileInfo.Exists)
 									{
 										RequiredFreeSpace -= LocalFileInfo.Length;
@@ -724,7 +825,7 @@ namespace UnrealGameSync
 								}
 								else
 								{
-									SyncTree.ExcludeFile(PerforceUtils.EscapePath(RelativePath));
+									SyncTree.ExcludeFile(PerforceUtils.EscapePath(SyncRecord.RelativePath));
 								}
 							}
 						}
@@ -736,14 +837,13 @@ namespace UnrealGameSync
 
 							if (Drive.AvailableFreeSpace < RequiredFreeSpace)
 							{
-								Log.WriteLine("Syncing requires {0} which exceeds the {1} available free space on {2}.", EpicGames.Core.StringUtils.FormatBytesString(RequiredFreeSpace), EpicGames.Core.StringUtils.FormatBytesString(Drive.AvailableFreeSpace), Drive.Name);
-								StatusMessage = "Not enough available free space.";
-								return WorkspaceUpdateResult.FailedToSync;
+								Logger.LogInformation("Syncing requires {RequiredSpace} which exceeds the {AvailableSpace} available free space on {Drive}.", StringUtils.FormatBytesString(RequiredFreeSpace), StringUtils.FormatBytesString(Drive.AvailableFreeSpace), Drive.Name);
+								return (WorkspaceUpdateResult.FailedToSync, "Not enough available free space.");
 							}
 						}
 						catch (SystemException)
 						{
-							Log.WriteLine("Unable to check available free space for {0}.", Project.LocalRootPath);
+							Logger.LogInformation("Unable to check available free space for {RootPath}.", Project.LocalRootPath);
 						}
 					}
 					SyncTree.GetOptimizedSyncCommands(Project.ClientRootPath, Context.ChangeNumber, BatchBuilder);
@@ -760,14 +860,14 @@ namespace UnrealGameSync
 					HashSet<string> RemainingDepotPaths = new HashSet<string>(StringComparer.OrdinalIgnoreCase);
 					RemainingDepotPaths.UnionWith(SyncDepotPaths);
 
-					using (TelemetryStopwatch TransferStopwatch = new TelemetryStopwatch("Workspace_Sync_TransferFiles", Project.TelemetryProjectPath))
+					using (TelemetryStopwatch TransferStopwatch = new TelemetryStopwatch("Workspace_Sync_TransferFiles", Project.TelemetryProjectIdentifier))
 					{
-						TransferStopwatch.AddData(new { MachineName = Environment.MachineName, DomainName = Environment.UserDomainName, ServerAndPort = Perforce.ServerAndPort, UserName = Perforce.UserName, IncludedFiles = SyncTree.TotalIncludedFiles, ExcludedFiles = SyncTree.TotalExcludedFiles, Size = SyncTree.TotalSize, NumThreads = Context.PerforceSyncOptions.NumThreads });
+						TransferStopwatch.AddData(new { MachineName = Environment.MachineName, DomainName = Environment.UserDomainName, ServerAndPort = Perforce.Settings.ServerAndPort, UserName = Perforce.Settings.UserName, IncludedFiles = SyncTree.TotalIncludedFiles, ExcludedFiles = SyncTree.TotalExcludedFiles, Size = SyncTree.TotalSize, NumThreads = Context.PerforceSyncOptions?.NumThreads ?? PerforceSyncOptions.DefaultNumThreads });
 
-						WorkspaceUpdateResult SyncResult = SyncFileRevisions(Perforce, "Syncing files...", Context, BatchBuilder.Batches, RemainingDepotPaths, Progress, Log, out StatusMessage);
-						if (SyncResult != WorkspaceUpdateResult.Success)
+						(WorkspaceUpdateResult, string) SyncResult = await SyncFileRevisions(Perforce, "Syncing files...", Context, BatchBuilder.Batches, RemainingDepotPaths, Progress, Logger, CancellationToken);
+						if (SyncResult.Item1 != WorkspaceUpdateResult.Success)
 						{
-							TransferStopwatch.AddData(new { SyncResult = SyncResult.ToString(), CompletedFilesFiles = SyncDepotPaths.Count - RemainingDepotPaths.Count });
+							TransferStopwatch.AddData(new { SyncResult = SyncResult.Item2.ToString(), CompletedFilesFiles = SyncDepotPaths.Count - RemainingDepotPaths.Count });
 							return SyncResult;
 						}
 
@@ -776,61 +876,55 @@ namespace UnrealGameSync
 					}
 
 					int VersionChangeNumber = -1;
-					if(Context.Options.HasFlag(WorkspaceUpdateOptions.Sync) && !Context.Options.HasFlag(WorkspaceUpdateOptions.UpdateFilter))
+					if (Context.Options.HasFlag(WorkspaceUpdateOptions.Sync) && !Context.Options.HasFlag(WorkspaceUpdateOptions.UpdateFilter))
 					{
 						// Read the new config file
-						Context.ProjectConfigFile = ReadProjectConfigFile(Project.LocalRootPath, Project.SelectedLocalFileName, Log);
-						Context.ProjectStreamFilter = ReadProjectStreamFilter(Perforce, Context.ProjectConfigFile, Log);
+						Context.ProjectConfigFile = await ReadProjectConfigFile(Project.LocalRootPath, Project.LocalFileName, Logger);
+						Context.ProjectStreamFilter = await ReadProjectStreamFilter(Perforce, Context.ProjectConfigFile, Logger, CancellationToken);
 
 						// Get the branch name
-						string BranchOrStreamName;
-						if(Perforce.GetActiveStream(out BranchOrStreamName, Log))
+						string? BranchOrStreamName = await Perforce.GetCurrentStreamAsync(CancellationToken);
+						if (BranchOrStreamName != null)
 						{
 							// If it's a virtual stream, take the concrete parent stream instead
-							for (;;)
+							for (; ; )
 							{
-								PerforceSpec StreamSpec;
-								if (!Perforce.TryGetStreamSpec(BranchOrStreamName, out StreamSpec, Log))
-								{
-									StatusMessage = String.Format("Unable to get stream spec for {0}.", BranchOrStreamName);
-									return WorkspaceUpdateResult.FailedToSync;
-								}
-								if (StreamSpec.GetField("Type") != "virtual")
+								StreamRecord StreamSpec = await Perforce.GetStreamAsync(BranchOrStreamName, false, CancellationToken);
+								if (StreamSpec.Type != "virtual" || StreamSpec.Parent == "none" || StreamSpec.Parent == null)
 								{
 									break;
 								}
-								BranchOrStreamName = StreamSpec.GetField("Parent");
+								BranchOrStreamName = StreamSpec.Parent;
 							}
 						}
 						else
 						{
 							// Otherwise use the depot path for GenerateProjectFiles.bat in the root of the workspace
-							string DepotFileName;
-							if(!Perforce.ConvertToDepotPath(Project.ClientRootPath + "/GenerateProjectFiles.bat", out DepotFileName, Log))
+							List<WhereRecord> Files = await Perforce.WhereAsync(Project.ClientRootPath + "/GenerateProjectFiles.bat").Where(x => !x.Unmap).ToListAsync(CancellationToken);
+							if (Files.Count != 1)
 							{
-								StatusMessage = String.Format("Couldn't determine branch name for {0}.", Project.SelectedClientFileName);
-								return WorkspaceUpdateResult.FailedToSync;
+								return (WorkspaceUpdateResult.FailedToSync, $"Couldn't determine branch name for {Project.ClientFileName}.");
 							}
-							BranchOrStreamName = PerforceUtils.GetClientOrDepotDirectoryName(DepotFileName);
+							BranchOrStreamName = PerforceUtils.GetClientOrDepotDirectoryName(Files[0].DepotFile);
 						}
 
 						// Find the last code change before this changelist. For consistency in versioning between local builds and precompiled binaries, we need to use the last submitted code changelist as our version number.
-						List<PerforceChangeSummary> CodeChanges;
-						if(!Perforce.FindChanges(new string[]{ ".cs", ".h", ".cpp", ".usf", ".ush", ".uproject", ".uplugin" }.SelectMany(x => SyncPaths.Select(y => String.Format("{0}{1}@<={2}", y, x, Context.ChangeNumber))), 1, out CodeChanges, Log))
+						string[] CodeFilter = new string[] { ".cs", ".h", ".cpp", ".usf", ".ush", ".uproject", ".uplugin" }.SelectMany(x => SyncPaths.Select(y => String.Format("{0}{1}@<={2}", y, x, Context.ChangeNumber))).ToArray();
+
+						PerforceResponseList<ChangesRecord> CodeChanges = await Perforce.TryGetChangesAsync(ChangesOptions.None, 1, ChangeStatus.Submitted, CodeFilter, CancellationToken);
+						if (!CodeChanges.Succeeded)
 						{
-							StatusMessage = String.Format("Couldn't determine last code changelist before CL {0}.", Context.ChangeNumber);
-							return WorkspaceUpdateResult.FailedToSync;
+							return (WorkspaceUpdateResult.FailedToSync, $"Couldn't determine last code changelist before CL {Context.ChangeNumber}.");
 						}
-						if(CodeChanges.Count == 0)
+						if (CodeChanges.Count == 0)
 						{
-							StatusMessage = String.Format("Could not find any code changes before CL {0}.", Context.ChangeNumber);
-							return WorkspaceUpdateResult.FailedToSync;
+							return (WorkspaceUpdateResult.FailedToSync, $"Could not find any code changes before CL {Context.ChangeNumber}.");
 						}
 
 						// Get the last code change
-						if(Context.ProjectConfigFile.GetValue("Options.VersionToLastCodeChange", true))
+						if (Context.ProjectConfigFile.GetValue("Options.VersionToLastCodeChange", true))
 						{
-							VersionChangeNumber = CodeChanges.Max(x => x.Number);
+							VersionChangeNumber = CodeChanges.Max(x => x.Data.Number);
 						}
 						else
 						{
@@ -838,110 +932,98 @@ namespace UnrealGameSync
 						}
 
 						// Update the version files
-						if(Context.ProjectConfigFile.GetValue("Options.UseFastModularVersioningV2", false))
+						if (Context.ProjectConfigFile.GetValue("Options.UseFastModularVersioningV2", false))
 						{
-							bool bIsLicenseeVersion = IsLicenseeVersion(Perforce, Project, Log);
-							if (!UpdateVersionFile(Perforce, Project.ClientRootPath + BuildVersionFileName, Context.ChangeNumber, Text => UpdateBuildVersion(Text, Context.ChangeNumber, VersionChangeNumber, BranchOrStreamName, bIsLicenseeVersion), Log))
+							bool bIsLicenseeVersion = await IsLicenseeVersion(Perforce, Project, Logger, CancellationToken);
+							if (!await UpdateVersionFile(Perforce, Project.ClientRootPath + BuildVersionFileName, Context.ChangeNumber, Text => UpdateBuildVersion(Text, Context.ChangeNumber, VersionChangeNumber, BranchOrStreamName, bIsLicenseeVersion), Logger, CancellationToken))
 							{
-								StatusMessage = String.Format("Failed to update {0}.", BuildVersionFileName);
-								return WorkspaceUpdateResult.FailedToSync;
+								return (WorkspaceUpdateResult.FailedToSync, $"Failed to update {BuildVersionFileName}.");
 							}
 						}
-						else if(Context.ProjectConfigFile.GetValue("Options.UseFastModularVersioning", false))
+						else if (Context.ProjectConfigFile.GetValue("Options.UseFastModularVersioning", false))
 						{
-							bool bIsLicenseeVersion = IsLicenseeVersion(Perforce, Project, Log);
-							if (!UpdateVersionFile(Perforce, Project.ClientRootPath + BuildVersionFileName, Context.ChangeNumber, Text => UpdateBuildVersion(Text, Context.ChangeNumber, VersionChangeNumber, BranchOrStreamName, bIsLicenseeVersion), Log))
+							bool bIsLicenseeVersion = await IsLicenseeVersion(Perforce, Project, Logger, CancellationToken);
+							if (!await UpdateVersionFile(Perforce, Project.ClientRootPath + BuildVersionFileName, Context.ChangeNumber, Text => UpdateBuildVersion(Text, Context.ChangeNumber, VersionChangeNumber, BranchOrStreamName, bIsLicenseeVersion), Logger, CancellationToken))
 							{
-								StatusMessage = String.Format("Failed to update {0}.", BuildVersionFileName);
-								return WorkspaceUpdateResult.FailedToSync;
+								return (WorkspaceUpdateResult.FailedToSync, $"Failed to update {BuildVersionFileName}");
 							}
 
-							Dictionary<string, string> VersionHeaderStrings = new Dictionary<string,string>();
+							Dictionary<string, string> VersionHeaderStrings = new Dictionary<string, string>();
 							VersionHeaderStrings["#define ENGINE_IS_PROMOTED_BUILD"] = " (0)";
 							VersionHeaderStrings["#define BUILT_FROM_CHANGELIST"] = " 0";
 							VersionHeaderStrings["#define BRANCH_NAME"] = " \"" + BranchOrStreamName.Replace('/', '+') + "\"";
-							if(!UpdateVersionFile(Perforce, Project.ClientRootPath + VersionHeaderFileName, VersionHeaderStrings, Context.ChangeNumber, Log))
+							if (!await UpdateVersionFile(Perforce, Project.ClientRootPath + VersionHeaderFileName, VersionHeaderStrings, Context.ChangeNumber, Logger, CancellationToken))
 							{
-								StatusMessage = String.Format("Failed to update {0}.", VersionHeaderFileName);
-								return WorkspaceUpdateResult.FailedToSync;
+								return (WorkspaceUpdateResult.FailedToSync, $"Failed to update {VersionHeaderFileName}.");
 							}
-							if(!UpdateVersionFile(Perforce, Project.ClientRootPath + ObjectVersionFileName, new Dictionary<string,string>(), Context.ChangeNumber, Log))
+							if (!await UpdateVersionFile(Perforce, Project.ClientRootPath + ObjectVersionFileName, new Dictionary<string, string>(), Context.ChangeNumber, Logger, CancellationToken))
 							{
-								StatusMessage = String.Format("Failed to update {0}.", ObjectVersionFileName);
-								return WorkspaceUpdateResult.FailedToSync;
+								return (WorkspaceUpdateResult.FailedToSync, $"Failed to update {ObjectVersionFileName}.");
 							}
 						}
 						else
 						{
-							if(!UpdateVersionFile(Perforce, Project.ClientRootPath + BuildVersionFileName, new Dictionary<string, string>(), Context.ChangeNumber, Log))
+							if (!await UpdateVersionFile(Perforce, Project.ClientRootPath + BuildVersionFileName, new Dictionary<string, string>(), Context.ChangeNumber, Logger, CancellationToken))
 							{
-								StatusMessage = String.Format("Failed to update {0}", BuildVersionFileName);
-								return WorkspaceUpdateResult.FailedToSync;
+								return (WorkspaceUpdateResult.FailedToSync, $"Failed to update {BuildVersionFileName}");
 							}
 
-							Dictionary<string, string> VersionStrings = new Dictionary<string,string>();
+							Dictionary<string, string> VersionStrings = new Dictionary<string, string>();
 							VersionStrings["#define ENGINE_VERSION"] = " " + VersionChangeNumber.ToString();
 							VersionStrings["#define ENGINE_IS_PROMOTED_BUILD"] = " (0)";
 							VersionStrings["#define BUILT_FROM_CHANGELIST"] = " " + VersionChangeNumber.ToString();
 							VersionStrings["#define BRANCH_NAME"] = " \"" + BranchOrStreamName.Replace('/', '+') + "\"";
-							if(!UpdateVersionFile(Perforce, Project.ClientRootPath + VersionHeaderFileName, VersionStrings, Context.ChangeNumber, Log))
+							if (!await UpdateVersionFile(Perforce, Project.ClientRootPath + VersionHeaderFileName, VersionStrings, Context.ChangeNumber, Logger, CancellationToken))
 							{
-								StatusMessage = String.Format("Failed to update {0}", VersionHeaderFileName);
-								return WorkspaceUpdateResult.FailedToSync;
+								return (WorkspaceUpdateResult.FailedToSync, $"Failed to update {VersionHeaderFileName}");
 							}
-							if(!UpdateVersionFile(Perforce, Project.ClientRootPath + ObjectVersionFileName, VersionStrings, Context.ChangeNumber, Log))
+							if (!await UpdateVersionFile(Perforce, Project.ClientRootPath + ObjectVersionFileName, VersionStrings, Context.ChangeNumber, Logger, CancellationToken))
 							{
-								StatusMessage = String.Format("Failed to update {0}", ObjectVersionFileName);
-								return WorkspaceUpdateResult.FailedToSync;
+								return (WorkspaceUpdateResult.FailedToSync, $"Failed to update {ObjectVersionFileName}");
 							}
 						}
 
 						// Remove all the receipts for build targets in this branch
-						if(Project.SelectedClientFileName.EndsWith(".uproject", StringComparison.OrdinalIgnoreCase))
+						if (Project.ClientFileName.EndsWith(".uproject", StringComparison.OrdinalIgnoreCase))
 						{
-							Perforce.Sync(PerforceUtils.GetClientOrDepotDirectoryName(Project.SelectedClientFileName) + "/Build/Receipts/...#0", Log);
+							await Perforce.SyncQuietAsync(SyncOptions.None, -1, PerforceUtils.GetClientOrDepotDirectoryName(Project.ClientFileName) + "/Build/Receipts/...#0", CancellationToken);
 						}
 					}
 
 					// Check if there are any files which need resolving
-					List<PerforceFileRecord> UnresolvedFiles;
-					if(!FindUnresolvedFiles(Perforce, SyncPaths, Log, out UnresolvedFiles))
-					{
-						StatusMessage = "Couldn't get list of unresolved files.";
-						return WorkspaceUpdateResult.FailedToSync;
-					}
-					if(UnresolvedFiles.Count > 0 && Context.Options.HasFlag(WorkspaceUpdateOptions.AutoResolveChanges))
+					List<FStatRecord> UnresolvedFiles = await FindUnresolvedFiles(Perforce, SyncPaths, CancellationToken);
+					if (UnresolvedFiles.Count > 0 && Context.Options.HasFlag(WorkspaceUpdateOptions.AutoResolveChanges))
 					{
-						foreach (PerforceFileRecord UnresolvedFile in UnresolvedFiles)
-						{
-							Perforce.AutoResolveFile(UnresolvedFile.DepotPath, Log);
-						}
-						if(!FindUnresolvedFiles(Perforce, SyncPaths, Log, out UnresolvedFiles))
+						foreach (FStatRecord UnresolvedFile in UnresolvedFiles)
 						{
-							StatusMessage = "Couldn't get list of unresolved files.";
-							return WorkspaceUpdateResult.FailedToSync;
+							if (UnresolvedFile.DepotFile != null)
+							{
+								await Perforce.ResolveAsync(-1, ResolveOptions.Automatic, UnresolvedFile.DepotFile, CancellationToken);
+							}
 						}
+						UnresolvedFiles = await FindUnresolvedFiles(Perforce, SyncPaths, CancellationToken);
 					}
-					if(UnresolvedFiles.Count > 0)
+					if (UnresolvedFiles.Count > 0)
 					{
-						Log.WriteLine("{0} files need resolving:", UnresolvedFiles.Count);
-						foreach(PerforceFileRecord UnresolvedFile in UnresolvedFiles)
+						Logger.LogInformation("{NumFiles} files need resolving:", UnresolvedFiles.Count);
+						foreach (FStatRecord UnresolvedFile in UnresolvedFiles)
 						{
-							Log.WriteLine("  {0}", UnresolvedFile.ClientPath);
+							Logger.LogInformation("  {ClientFile}", UnresolvedFile.ClientFile);
 						}
-						StatusMessage = "Files need resolving.";
-						return WorkspaceUpdateResult.FilesToResolve;
+						return (WorkspaceUpdateResult.FilesToResolve, "Files need resolving.");
 					}
 
 					// Continue processing sync-only actions
 					if (Context.Options.HasFlag(WorkspaceUpdateOptions.Sync) && !Context.Options.HasFlag(WorkspaceUpdateOptions.UpdateFilter))
 					{
+						Context.ProjectConfigFile ??= await ReadProjectConfigFile(Project.LocalRootPath, Project.LocalFileName, Logger);
+
 						// Execute any project specific post-sync steps
-						string[] PostSyncSteps = Context.ProjectConfigFile.GetValues("Sync.Step", null);
+						string[]? PostSyncSteps = Context.ProjectConfigFile.GetValues("Sync.Step", null);
 						if (PostSyncSteps != null)
 						{
-							Log.WriteLine();
-							Log.WriteLine("Executing post-sync steps...");
+							Logger.LogInformation("");
+							Logger.LogInformation("Executing post-sync steps...");
 
 							Dictionary<string, string> PostSyncVariables = new Dictionary<string, string>(Context.Variables);
 							PostSyncVariables["Change"] = Context.ChangeNumber.ToString();
@@ -956,13 +1038,12 @@ namespace UnrealGameSync
 								{
 									string ToolArguments = Utility.ExpandVariables(PostSyncStepObject.GetValue("Arguments", ""), PostSyncVariables);
 
-									Log.WriteLine("post-sync> Running {0} {1}", ToolFileName, ToolArguments);
+									Logger.LogInformation("post-sync> Running {FileName} {Arguments}", ToolFileName, ToolArguments);
 
-									int ResultFromTool = Utility.ExecuteProcess(ToolFileName, null, ToolArguments, null, new ProgressTextWriter(Progress, new PrefixedTextWriter("post-sync> ", Log)));
+									int ResultFromTool = await Utility.ExecuteProcessAsync(ToolFileName, null, ToolArguments, Line => ProcessOutput(Line, "post-sync> ", Progress, Logger), CancellationToken);
 									if (ResultFromTool != 0)
 									{
-										StatusMessage = String.Format("Post-sync step terminated with exit code {0}.", ResultFromTool);
-										return WorkspaceUpdateResult.FailedToSync;
+										return (WorkspaceUpdateResult.FailedToSync, $"Post-sync step terminated with exit code {ResultFromTool}.");
 									}
 								}
 							}
@@ -981,59 +1062,58 @@ namespace UnrealGameSync
 					State.Save();
 
 					// Update the timing info
-					Times.Add(new Tuple<string,TimeSpan>("Sync", SyncTelemetryStopwatch.Stop("Success")));
+					Times.Add(new Tuple<string, TimeSpan>("Sync", SyncTelemetryStopwatch.Stop("Success")));
 
 					// Save the number of files synced
 					NumFilesSynced = SyncDepotPaths.Count;
-					Log.WriteLine();
+					Logger.LogInformation("");
 				}
 			}
 
 			// Extract an archive from the depot path
-			if(Context.Options.HasFlag(WorkspaceUpdateOptions.SyncArchives))
+			if (Context.Options.HasFlag(WorkspaceUpdateOptions.SyncArchives))
 			{
-				using(TelemetryStopwatch Stopwatch = new TelemetryStopwatch("Workspace_SyncArchives", Project.TelemetryProjectPath))
+				using (TelemetryStopwatch Stopwatch = new TelemetryStopwatch("Workspace_SyncArchives", Project.TelemetryProjectIdentifier))
 				{
 					// Create the directory for extracted archive manifests
 					DirectoryReference ManifestDirectoryName;
-					if(Project.SelectedLocalFileName.HasExtension(".uproject"))
+					if (Project.LocalFileName.HasExtension(".uproject"))
 					{
-						ManifestDirectoryName = DirectoryReference.Combine(Project.SelectedLocalFileName.Directory, "Saved", "UnrealGameSync");
+						ManifestDirectoryName = DirectoryReference.Combine(Project.LocalFileName.Directory, "Saved", "UnrealGameSync");
 					}
 					else
 					{
-						ManifestDirectoryName = DirectoryReference.Combine(Project.SelectedLocalFileName.Directory, "Engine", "Saved", "UnrealGameSync");
+						ManifestDirectoryName = DirectoryReference.Combine(Project.LocalFileName.Directory, "Engine", "Saved", "UnrealGameSync");
 					}
 					DirectoryReference.CreateDirectory(ManifestDirectoryName);
 
 					// Sync and extract (or just remove) the given archives
-					foreach(KeyValuePair<string, Tuple<IArchiveInfo, string>> ArchiveTypeAndArchive in Context.ArchiveTypeToArchive)
+					foreach (KeyValuePair<string, Tuple<IArchiveInfo, string>> ArchiveTypeAndArchive in Context.ArchiveTypeToArchive)
 					{
 						string ArchiveType = ArchiveTypeAndArchive.Key;
 
 						// Remove any existing binaries
 						FileReference ManifestFileName = FileReference.Combine(ManifestDirectoryName, String.Format("{0}.zipmanifest", ArchiveType));
-						if(FileReference.Exists(ManifestFileName))
+						if (FileReference.Exists(ManifestFileName))
 						{
-							Log.WriteLine("Removing {0} binaries...", ArchiveType);
+							Logger.LogInformation("Removing {ArchiveType} binaries...", ArchiveType);
 							Progress.Set(String.Format("Removing {0} binaries...", ArchiveType), 0.0f);
-							ArchiveUtils.RemoveExtractedFiles(Project.LocalRootPath, ManifestFileName, Progress, Log);
+							ArchiveUtils.RemoveExtractedFiles(Project.LocalRootPath, ManifestFileName, Progress, Logger);
 							FileReference.Delete(ManifestFileName);
-							Log.WriteLine();
+							Logger.LogInformation("");
 						}
 
 						// If we have a new depot path, sync it down and extract it
-						if(ArchiveTypeAndArchive.Value != null)
+						if (ArchiveTypeAndArchive.Value != null)
 						{
 							IArchiveInfo ArchiveInfo = ArchiveTypeAndArchive.Value.Item1;
 							string ArchiveKey = ArchiveTypeAndArchive.Value.Item2;
 
-							Log.WriteLine("Syncing {0} binaries...", ArchiveType.ToLowerInvariant());
+							Logger.LogInformation("Syncing {ArchiveType} binaries...", ArchiveType.ToLowerInvariant());
 							Progress.Set(String.Format("Syncing {0} binaries...", ArchiveType.ToLowerInvariant()), 0.0f);
-							if (!ArchiveInfo.DownloadArchive(ArchiveKey, Project.LocalRootPath, ManifestFileName, Log, Progress))
+							if (!ArchiveInfo.DownloadArchive(Perforce, ArchiveKey, Project.LocalRootPath, ManifestFileName, Logger, Progress, CancellationToken.None).Result)
 							{
-								StatusMessage = String.Format("Couldn't read {0}", ArchiveKey);
-								return WorkspaceUpdateResult.FailedToSync;
+								return (WorkspaceUpdateResult.FailedToSync, $"Couldn't read {ArchiveKey}");
 							}
 						}
 					}
@@ -1043,59 +1123,58 @@ namespace UnrealGameSync
 					State.Save();
 
 					// Add the finish time
-					Times.Add(new Tuple<string,TimeSpan>("Archive", Stopwatch.Stop("Success")));
+					Times.Add(new Tuple<string, TimeSpan>("Archive", Stopwatch.Stop("Success")));
 				}
 			}
 
 			// Take the lock before doing anything else. Building and generating project files can only be done on one workspace at a time.
-			if(Context.Options.HasFlag(WorkspaceUpdateOptions.GenerateProjectFiles) || Context.Options.HasFlag(WorkspaceUpdateOptions.Build))
+			using SemaphoreScope Scope = new SemaphoreScope(UpdateSemaphore);
+			if (Context.Options.HasFlag(WorkspaceUpdateOptions.GenerateProjectFiles) || Context.Options.HasFlag(WorkspaceUpdateOptions.Build))
 			{
-				if(Interlocked.CompareExchange(ref ActiveWorkspaceState, State, null) != null)
+				if (!Scope.TryAcquire())
 				{
-					Log.WriteLine("Waiting for other workspaces to finish...");
-					while(Interlocked.CompareExchange(ref ActiveWorkspaceState, State, null) != null)
-					{
-						Thread.Sleep(100);
-					}
+					Logger.LogInformation("Waiting for other workspaces to finish...");
+					await Scope.AcquireAsync(CancellationToken);
 				}
 			}
 
 			// Generate project files in the workspace
-			if(Context.Options.HasFlag(WorkspaceUpdateOptions.GenerateProjectFiles))
+			if (Context.Options.HasFlag(WorkspaceUpdateOptions.GenerateProjectFiles))
 			{
-				using(TelemetryStopwatch Stopwatch = new TelemetryStopwatch("Workspace_GenerateProjectFiles", Project.TelemetryProjectPath))
+				using (TelemetryStopwatch Stopwatch = new TelemetryStopwatch("Workspace_GenerateProjectFiles", Project.TelemetryProjectIdentifier))
 				{
 					Progress.Set("Generating project files...", 0.0f);
 
 					StringBuilder CommandLine = new StringBuilder();
 					CommandLine.AppendFormat("/C \"\"{0}\"", FileReference.Combine(Project.LocalRootPath, "GenerateProjectFiles.bat"));
-					if((Context.Options & WorkspaceUpdateOptions.SyncAllProjects) == 0 && (Context.Options & WorkspaceUpdateOptions.IncludeAllProjectsInSolution) == 0)
+					if ((Context.Options & WorkspaceUpdateOptions.SyncAllProjects) == 0 && (Context.Options & WorkspaceUpdateOptions.IncludeAllProjectsInSolution) == 0)
 					{
-						if(Project.SelectedLocalFileName.HasExtension(".uproject"))
+						if (Project.LocalFileName.HasExtension(".uproject"))
 						{
-							CommandLine.AppendFormat(" \"{0}\"", Project.SelectedLocalFileName);
+							CommandLine.AppendFormat(" \"{0}\"", Project.LocalFileName);
 						}
 					}
 					CommandLine.Append(" -progress\"");
 
-					Log.WriteLine("Generating project files...");
-					Log.WriteLine("gpf> Running {0} {1}", CmdExe, CommandLine);
+					Logger.LogInformation("Generating project files...");
+					Logger.LogInformation("gpf> Running {FileName} {Arguments}", CmdExe, CommandLine);
 
-					int GenerateProjectFilesResult = Utility.ExecuteProcess(CmdExe, null, CommandLine.ToString(), null, new ProgressTextWriter(Progress, new PrefixedTextWriter("gpf> ", Log)));
-					if(GenerateProjectFilesResult != 0)
+					int GenerateProjectFilesResult = await Utility.ExecuteProcessAsync(CmdExe, null, CommandLine.ToString(), Line => ProcessOutput(Line, "gpf> ", Progress, Logger), CancellationToken);
+					if (GenerateProjectFilesResult != 0)
 					{
-						StatusMessage = String.Format("Failed to generate project files (exit code {0}).", GenerateProjectFilesResult);
-						return WorkspaceUpdateResult.FailedToCompile;
+						return (WorkspaceUpdateResult.FailedToCompile, $"Failed to generate project files (exit code {GenerateProjectFilesResult}).");
 					}
 
-					Log.WriteLine();
-					Times.Add(new Tuple<string,TimeSpan>("Prj gen", Stopwatch.Stop("Success")));
+					Logger.LogInformation("");
+					Times.Add(new Tuple<string, TimeSpan>("Prj gen", Stopwatch.Stop("Success")));
 				}
 			}
 
 			// Build everything using MegaXGE
-			if(Context.Options.HasFlag(WorkspaceUpdateOptions.Build))
+			if (Context.Options.HasFlag(WorkspaceUpdateOptions.Build))
 			{
+				Context.ProjectConfigFile ??= await ReadProjectConfigFile(Project.LocalRootPath, Project.LocalFileName, Logger);
+
 				// Compile all the build steps together
 				Dictionary<Guid, ConfigObject> BuildStepObjects = Context.DefaultBuildSteps.ToDictionary(x => x.Key, x => new ConfigObject(x.Value));
 				BuildStep.MergeBuildStepObjects(BuildStepObjects, Context.ProjectConfigFile.GetValues("Build.Step", new string[0]).Select(x => new ConfigObject(x)));
@@ -1103,11 +1182,11 @@ namespace UnrealGameSync
 
 				// Construct build steps from them
 				List<BuildStep> BuildSteps = BuildStepObjects.Values.Select(x => new BuildStep(x)).OrderBy(x => (x.OrderIndex == -1) ? 10000 : x.OrderIndex).ToList();
-				if(Context.CustomBuildSteps != null && Context.CustomBuildSteps.Count > 0)
+				if (Context.CustomBuildSteps != null && Context.CustomBuildSteps.Count > 0)
 				{
 					BuildSteps.RemoveAll(x => !Context.CustomBuildSteps.Contains(x.UniqueId));
 				}
-				else if(Context.Options.HasFlag(WorkspaceUpdateOptions.ScheduledBuild))
+				else if (Context.Options.HasFlag(WorkspaceUpdateOptions.ScheduledBuild))
 				{
 					BuildSteps.RemoveAll(x => !x.bScheduledSync);
 				}
@@ -1118,17 +1197,17 @@ namespace UnrealGameSync
 
 				// Check if the last successful build was before a change that we need to force a clean for
 				bool bForceClean = false;
-				if(State.LastBuiltChangeNumber != 0)
+				if (State.LastBuiltChangeNumber != 0)
 				{
-					foreach(string CleanBuildChange in Context.ProjectConfigFile.GetValues("ForceClean.Changelist", new string[0]))
+					foreach (string CleanBuildChange in Context.ProjectConfigFile.GetValues("ForceClean.Changelist", new string[0]))
 					{
 						int ChangeNumber;
-						if(int.TryParse(CleanBuildChange, out ChangeNumber))
+						if (int.TryParse(CleanBuildChange, out ChangeNumber))
 						{
-							if((State.LastBuiltChangeNumber >= ChangeNumber) != (State.CurrentChangeNumber >= ChangeNumber))
+							if ((State.LastBuiltChangeNumber >= ChangeNumber) != (State.CurrentChangeNumber >= ChangeNumber))
 							{
-								Log.WriteLine("Forcing clean build due to changelist {0}.", ChangeNumber);
-								Log.WriteLine();
+								Logger.LogInformation("Forcing clean build due to changelist {Change}.", ChangeNumber);
+								Logger.LogInformation("");
 								bForceClean = true;
 								break;
 							}
@@ -1137,7 +1216,7 @@ namespace UnrealGameSync
 				}
 
 				// Execute them all
-				using(TelemetryStopwatch Stopwatch = new TelemetryStopwatch("Workspace_Build", Project.TelemetryProjectPath))
+				using (TelemetryStopwatch Stopwatch = new TelemetryStopwatch("Workspace_Build", Project.TelemetryProjectIdentifier))
 				{
 					Progress.Set("Starting build...", 0.0f);
 
@@ -1150,83 +1229,91 @@ namespace UnrealGameSync
 						Progress.Set(Step.StatusText);
 						Progress.Push(MaxProgressFraction);
 
-						Log.WriteLine(Step.StatusText);
+						Logger.LogInformation("{Status}", Step.StatusText);
 
-						if(Step.IsValid())
+						if (Step.IsValid())
 						{
-							switch(Step.Type)
+							switch (Step.Type)
 							{
 								case BuildStepType.Compile:
-									using(TelemetryStopwatch StepStopwatch = new TelemetryStopwatch("Workspace_Execute_Compile", Project.TelemetryProjectPath))
+									using (TelemetryStopwatch StepStopwatch = new TelemetryStopwatch("Workspace_Execute_Compile", Project.TelemetryProjectIdentifier))
 									{
 										StepStopwatch.AddData(new { Target = Step.Target });
 
 										FileReference BuildBat = FileReference.Combine(Project.LocalRootPath, "Engine", "Build", "BatchFiles", "Build.bat");
 										string CommandLine = String.Format("{0} {1} {2} {3} -NoHotReloadFromIDE", Step.Target, Step.Platform, Step.Configuration, Utility.ExpandVariables(Step.Arguments ?? "", Context.Variables));
-										if(Context.Options.HasFlag(WorkspaceUpdateOptions.Clean) || bForceClean)
+										if (Context.Options.HasFlag(WorkspaceUpdateOptions.Clean) || bForceClean)
 										{
-											Log.WriteLine("ubt> Running {0} {1} -clean", BuildBat, CommandLine);
-											Utility.ExecuteProcess(CmdExe, null, "/C \"\"" + BuildBat + "\" " + CommandLine + " -clean\"", null, new ProgressTextWriter(Progress, new PrefixedTextWriter("ubt> ", Log)));
+											Logger.LogInformation("ubt> Running {FileName} {Arguments}", BuildBat, CommandLine + " -clean");
+											await Utility.ExecuteProcessAsync(CmdExe, null, "/C \"\"" + BuildBat + "\" " + CommandLine + " -clean\"", Line => ProcessOutput(Line, "ubt> ", Progress, Logger), CancellationToken);
 										}
 
-										Log.WriteLine("ubt> Running {0} {1} -progress", BuildBat, CommandLine);
+										Logger.LogInformation("ubt> Running {FileName} {Arguments}", BuildBat, CommandLine + " -progress");
 
-										int ResultFromBuild = Utility.ExecuteProcess(CmdExe, null, "/C \"\"" + BuildBat + "\" "+ CommandLine + " -progress\"", null, new ProgressTextWriter(Progress, new PrefixedTextWriter("ubt> ", Log)));
-										if(ResultFromBuild != 0)
+										int ResultFromBuild = await Utility.ExecuteProcessAsync(CmdExe, null, "/C \"\"" + BuildBat + "\" " + CommandLine + " -progress\"", Line => ProcessOutput(Line, "ubt> ", Progress, Logger), CancellationToken);
+										if (ResultFromBuild != 0)
 										{
 											StepStopwatch.Stop("Failed");
-											StatusMessage = String.Format("Failed to compile {0}.", Step.Target);
-											return (HasModifiedSourceFiles(Perforce, Project, Log) || Context.UserBuildStepObjects.Count > 0)? WorkspaceUpdateResult.FailedToCompile : WorkspaceUpdateResult.FailedToCompileWithCleanWorkspace;
+
+											WorkspaceUpdateResult Result;
+											if (await HasModifiedSourceFiles(Perforce, Project, CancellationToken) || Context.UserBuildStepObjects.Count > 0)
+											{
+												Result = WorkspaceUpdateResult.FailedToCompile;
+											}
+											else
+											{
+												Result = WorkspaceUpdateResult.FailedToCompileWithCleanWorkspace;
+											}
+
+											return (Result, $"Failed to compile {Step.Target}");
 										}
 
 										StepStopwatch.Stop("Success");
 									}
 									break;
 								case BuildStepType.Cook:
-									using(TelemetryStopwatch StepStopwatch = new TelemetryStopwatch("Workspace_Execute_Cook", Project.TelemetryProjectPath))
+									using (TelemetryStopwatch StepStopwatch = new TelemetryStopwatch("Workspace_Execute_Cook", Project.TelemetryProjectIdentifier))
 									{
 										StepStopwatch.AddData(new { Project = Path.GetFileNameWithoutExtension(Step.FileName) });
 
 										FileReference LocalRunUAT = FileReference.Combine(Project.LocalRootPath, "Engine", "Build", "BatchFiles", "RunUAT.bat");
 										string Arguments = String.Format("/C \"\"{0}\" -profile=\"{1}\"\"", LocalRunUAT, FileReference.Combine(Project.LocalRootPath, Step.FileName));
-										Log.WriteLine("uat> Running {0} {1}", LocalRunUAT, Arguments);
+										Logger.LogInformation("uat> Running {FileName} {Argument}", LocalRunUAT, Arguments);
 
-										int ResultFromUAT = Utility.ExecuteProcess(CmdExe, null, Arguments, null, new ProgressTextWriter(Progress, new PrefixedTextWriter("uat> ", Log)));
-										if(ResultFromUAT != 0)
+										int ResultFromUAT = await Utility.ExecuteProcessAsync(CmdExe, null, Arguments, Line => ProcessOutput(Line, "uat> ", Progress, Logger), CancellationToken);
+										if (ResultFromUAT != 0)
 										{
 											StepStopwatch.Stop("Failed");
-											StatusMessage = String.Format("Cook failed. ({0})", ResultFromUAT);
-											return WorkspaceUpdateResult.FailedToCompile;
+											return (WorkspaceUpdateResult.FailedToCompile, $"Cook failed. ({ResultFromUAT})");
 										}
 
 										StepStopwatch.Stop("Success");
 									}
 									break;
 								case BuildStepType.Other:
-									using(TelemetryStopwatch StepStopwatch = new TelemetryStopwatch("Workspace_Execute_Custom", Project.TelemetryProjectPath))
+									using (TelemetryStopwatch StepStopwatch = new TelemetryStopwatch("Workspace_Execute_Custom", Project.TelemetryProjectIdentifier))
 									{
 										StepStopwatch.AddData(new { FileName = Path.GetFileNameWithoutExtension(Step.FileName) });
 
 										FileReference ToolFileName = FileReference.Combine(Project.LocalRootPath, Utility.ExpandVariables(Step.FileName, Context.Variables));
 										string ToolWorkingDir = String.IsNullOrWhiteSpace(Step.WorkingDir) ? ToolFileName.Directory.FullName : Utility.ExpandVariables(Step.WorkingDir, Context.Variables);
 										string ToolArguments = Utility.ExpandVariables(Step.Arguments ?? "", Context.Variables);
-										Log.WriteLine("tool> Running {0} {1}", ToolFileName, ToolArguments);
+										Logger.LogInformation("tool> Running {0} {1}", ToolFileName, ToolArguments);
 
-										if(Step.bUseLogWindow)
+										if (Step.bUseLogWindow)
 										{
-											int ResultFromTool = Utility.ExecuteProcess(ToolFileName.FullName, ToolWorkingDir, ToolArguments, null, new ProgressTextWriter(Progress, new PrefixedTextWriter("tool> ", Log)));
-											if(ResultFromTool != 0)
+											int ResultFromTool = await Utility.ExecuteProcessAsync(ToolFileName.FullName, ToolWorkingDir, ToolArguments, Line => ProcessOutput(Line, "tool> ", Progress, Logger), CancellationToken);
+											if (ResultFromTool != 0)
 											{
 												StepStopwatch.Stop("Failed");
-												StatusMessage = String.Format("Tool terminated with exit code {0}.", ResultFromTool);
-												return WorkspaceUpdateResult.FailedToCompile;
+												return (WorkspaceUpdateResult.FailedToCompile, $"Tool terminated with exit code {ResultFromTool}.");
 											}
 										}
 										else
 										{
 											ProcessStartInfo StartInfo = new ProcessStartInfo(ToolFileName.FullName, ToolArguments);
 											StartInfo.WorkingDirectory = ToolWorkingDir;
-											using(Process.Start(StartInfo))
+											using (Process.Start(StartInfo))
 											{
 											}
 										}
@@ -1237,15 +1324,15 @@ namespace UnrealGameSync
 							}
 						}
 
-						Log.WriteLine();
+						Logger.LogInformation("");
 						Progress.Pop();
 					}
 
-					Times.Add(new Tuple<string,TimeSpan>("Build", Stopwatch.Stop("Success")));
+					Times.Add(new Tuple<string, TimeSpan>("Build", Stopwatch.Stop("Success")));
 				}
 
 				// Update the last successful build change number
-				if(Context.CustomBuildSteps == null || Context.CustomBuildSteps.Count == 0)
+				if (Context.CustomBuildSteps == null || Context.CustomBuildSteps.Count == 0)
 				{
 					State.LastBuiltChangeNumber = State.CurrentChangeNumber;
 					State.Save();
@@ -1253,36 +1340,42 @@ namespace UnrealGameSync
 			}
 
 			// Write out all the timing information
-			Log.WriteLine("Total time : " + FormatTime(Times.Sum(x => (long)(x.Item2.TotalMilliseconds / 1000))));
-			foreach(Tuple<string, TimeSpan> Time in Times)
+			Logger.LogInformation("Total time: {TotalTime}", FormatTime(Times.Sum(x => (long)(x.Item2.TotalMilliseconds / 1000))));
+			foreach (Tuple<string, TimeSpan> Time in Times)
 			{
-				Log.WriteLine("   {0,-8}: {1}", Time.Item1, FormatTime((long)(Time.Item2.TotalMilliseconds / 1000)));
+				Logger.LogInformation("   {Name,-8}: {TimeSeconds}", Time.Item1, FormatTime((long)(Time.Item2.TotalMilliseconds / 1000)));
 			}
-			if(NumFilesSynced > 0)
+			if (NumFilesSynced > 0)
 			{
-				Log.WriteLine("{0} files synced.", NumFilesSynced);
+				Logger.LogInformation("{NumFiles} files synced.", NumFilesSynced);
 			}
 
 			DateTime FinishTime = DateTime.Now;
-			Log.WriteLine();
-			Log.WriteLine("UPDATE SUCCEEDED ({0} {1})", FinishTime.ToShortDateString(), FinishTime.ToShortTimeString());
+			Logger.LogInformation("");
+			Logger.LogInformation("UPDATE SUCCEEDED ({FinishDate} {FinishTime})", FinishTime.ToShortDateString(), FinishTime.ToShortTimeString());
 
-			StatusMessage = "Update succeeded";
-			return WorkspaceUpdateResult.Success;
+			return (WorkspaceUpdateResult.Success, "Update succeeded");
 		}
 
-		static bool IsLicenseeVersion(PerforceConnection Perforce, ProjectInfo Project, TextWriter Log)
+		static void ProcessOutput(string Line, string Prefix, ProgressValue Progress, ILogger Logger)
 		{
-			bool bIsEpicInternal;
-			if (Perforce.FileExists(Project.ClientRootPath + "/Engine/Build/NotForLicensees/EpicInternal.txt", out bIsEpicInternal, Log) && bIsEpicInternal)
+			string? ParsedLine = ProgressTextWriter.ParseLine(Line, Progress);
+			if (ParsedLine != null)
 			{
-				return false;
+				Logger.LogInformation("{Prefix}{Message}", Prefix, Line);
 			}
-			if (Perforce.FileExists(Project.ClientRootPath + "/Engine/Restricted/NotForLicensees/Build/EpicInternal.txt", out bIsEpicInternal, Log) && bIsEpicInternal)
+		}
+
+		static async Task<bool> IsLicenseeVersion(IPerforceConnection Perforce, ProjectInfo Project, ILogger Logger, CancellationToken CancellationToken)
+		{
+			string[] Files =
 			{
-				return false;
-			}
-			return true;
+				Project.ClientRootPath + "/Engine/Build/NotForLicensees/EpicInternal.txt",
+				Project.ClientRootPath + "/Engine/Restricted/NotForLicensees/Build/EpicInternal.txt"
+			};
+
+			List<FStatRecord> Records = await Perforce.FStatAsync(Files, CancellationToken).ToListAsync(CancellationToken);
+			return !Records.Any(x => x.IsMapped);
 		}
 
 		public static List<string> GetSyncPaths(ProjectInfo Project, bool bSyncAllProjects, string[] SyncFilter)
@@ -1296,14 +1389,14 @@ namespace UnrealGameSync
 			List<string> SyncPaths = new List<string>();
 
 			// Check the client path is formatted correctly
-			if (!Project.SelectedClientFileName.StartsWith(Project.ClientRootPath + "/"))
+			if (!Project.ClientFileName.StartsWith(Project.ClientRootPath + "/"))
 			{
-				throw new Exception(String.Format("Expected '{0}' to start with '{1}'", Project.SelectedClientFileName, Project.ClientRootPath));
+				throw new Exception(String.Format("Expected '{0}' to start with '{1}'", Project.ClientFileName, Project.ClientRootPath));
 			}
 
 			// Add the default project paths
-			int LastSlashIdx = Project.SelectedClientFileName.LastIndexOf('/');
-			if (bSyncAllProjects || !Project.SelectedClientFileName.EndsWith(".uproject", StringComparison.OrdinalIgnoreCase) || LastSlashIdx <= Project.ClientRootPath.Length)
+			int LastSlashIdx = Project.ClientFileName.LastIndexOf('/');
+			if (bSyncAllProjects || !Project.ClientFileName.EndsWith(".uproject", StringComparison.OrdinalIgnoreCase) || LastSlashIdx <= Project.ClientRootPath.Length)
 			{
 				SyncPaths.Add("/...");
 			}
@@ -1311,24 +1404,24 @@ namespace UnrealGameSync
 			{
 				SyncPaths.Add("/*");
 				SyncPaths.Add("/Engine/...");
-				if(Project.bIsEnterpriseProject)
+				if (Project.bIsEnterpriseProject)
 				{
 					SyncPaths.Add("/Enterprise/...");
 				}
-				SyncPaths.Add(Project.SelectedClientFileName.Substring(Project.ClientRootPath.Length, LastSlashIdx - Project.ClientRootPath.Length) + "/...");
+				SyncPaths.Add(Project.ClientFileName.Substring(Project.ClientRootPath.Length, LastSlashIdx - Project.ClientRootPath.Length) + "/...");
 			}
 
 			// Apply the sync filter to that list. We only want inclusive rules in the output list, but we can manually apply exclusions to previous entries.
-			if(SyncFilter != null)
+			if (SyncFilter != null)
 			{
-				foreach(string SyncPath in SyncFilter)
+				foreach (string SyncPath in SyncFilter)
 				{
 					string TrimSyncPath = SyncPath.Trim();
-					if(TrimSyncPath.StartsWith("/"))
+					if (TrimSyncPath.StartsWith("/"))
 					{
 						SyncPaths.Add(TrimSyncPath);
 					}
-					else if(TrimSyncPath.StartsWith("-/") && TrimSyncPath.EndsWith("..."))
+					else if (TrimSyncPath.StartsWith("-/") && TrimSyncPath.EndsWith("..."))
 					{
 						SyncPaths.RemoveAll(x => x.StartsWith(TrimSyncPath.Substring(1, TrimSyncPath.Length - 4)));
 					}
@@ -1337,15 +1430,15 @@ namespace UnrealGameSync
 
 			// Sort the remaining paths by length, and remove any paths which are included twice
 			SyncPaths = SyncPaths.OrderBy(x => x.Length).ToList();
-			for(int Idx = 0; Idx < SyncPaths.Count; Idx++)
+			for (int Idx = 0; Idx < SyncPaths.Count; Idx++)
 			{
 				string SyncPath = SyncPaths[Idx];
-				if(SyncPath.EndsWith("..."))
+				if (SyncPath.EndsWith("..."))
 				{
 					string SyncPathPrefix = SyncPath.Substring(0, SyncPath.Length - 3);
-					for(int OtherIdx = SyncPaths.Count - 1; OtherIdx > Idx; OtherIdx--)
+					for (int OtherIdx = SyncPaths.Count - 1; OtherIdx > Idx; OtherIdx--)
 					{
-						if(SyncPaths[OtherIdx].StartsWith(SyncPathPrefix))
+						if (SyncPaths[OtherIdx].StartsWith(SyncPathPrefix))
 						{
 							SyncPaths.RemoveAt(OtherIdx);
 						}
@@ -1359,10 +1452,10 @@ namespace UnrealGameSync
 		public static bool MatchFilter(ProjectInfo Project, FileReference FileName, FileFilter Filter)
 		{
 			bool bMatch = true;
-			if(FileName.IsUnderDirectory(Project.LocalRootPath))
+			if (FileName.IsUnderDirectory(Project.LocalRootPath))
 			{
 				string RelativePath = FileName.MakeRelativeTo(Project.LocalRootPath);
-				if(!Filter.Matches(RelativePath))
+				if (!Filter.Matches(RelativePath))
 				{
 					bMatch = false;
 				}
@@ -1377,70 +1470,76 @@ namespace UnrealGameSync
 			public Queue<List<string>> SyncCommandLists;
 			public string StatusMessage;
 			public WorkspaceUpdateResult Result = WorkspaceUpdateResult.Success;
+
+			public SyncState(HashSet<string> RemainingDepotPaths, Queue<List<string>> SyncCommandLists)
+			{
+				this.TotalDepotPaths = RemainingDepotPaths.Count;
+				this.RemainingDepotPaths = RemainingDepotPaths;
+				this.SyncCommandLists = SyncCommandLists;
+				this.StatusMessage = "Succeeded.";
+			}
 		}
 
-		static WorkspaceUpdateResult SyncFileRevisions(PerforceConnection Perforce, string Prefix, WorkspaceUpdateContext Context, List<string> SyncCommands, HashSet<string> RemainingDepotPaths, ProgressValue Progress, TextWriter Log, out string StatusMessage)
+		static Task<(WorkspaceUpdateResult, string)> SyncFileRevisions(IPerforceConnection Perforce, string Prefix, WorkspaceUpdateContext Context, List<string> SyncCommands, HashSet<string> RemainingDepotPaths, ProgressValue Progress, ILogger Logger, CancellationToken CancellationToken)
 		{
 			Queue<List<string>> SyncCommandLists = new Queue<List<string>>();
 			SyncCommandLists.Enqueue(SyncCommands);
-			return SyncFileRevisions(Perforce, Prefix, Context, SyncCommandLists, RemainingDepotPaths, Progress, Log, out StatusMessage);
+			return SyncFileRevisions(Perforce, Prefix, Context, SyncCommandLists, RemainingDepotPaths, Progress, Logger, CancellationToken);
 		}
 
-		static WorkspaceUpdateResult SyncFileRevisions(PerforceConnection Perforce, string Prefix, WorkspaceUpdateContext Context, Queue<List<string>> SyncCommandLists, HashSet<string> RemainingDepotPaths, ProgressValue Progress, TextWriter Log, out string StatusMessage)
+		static async Task<(WorkspaceUpdateResult, string)> SyncFileRevisions(IPerforceConnection Perforce, string Prefix, WorkspaceUpdateContext Context, Queue<List<string>> SyncCommandLists, HashSet<string> RemainingDepotPaths, ProgressValue Progress, ILogger Logger, CancellationToken CancellationToken)
 		{
 			// Figure out the number of additional background threads we want to run with. We can run worker on the current thread.
-			int NumExtraThreads = Math.Max(Math.Min(SyncCommandLists.Count, Context.PerforceSyncOptions.NumThreads) - 1, 0);
+			int NumThreads = Context.PerforceSyncOptions?.NumThreads ?? PerforceSyncOptions.DefaultNumThreads;
+			int NumExtraThreads = Math.Max(Math.Min(SyncCommandLists.Count, NumThreads) - 1, 0);
 
-			List<Thread> ChildThreads = new List<Thread>(NumExtraThreads);
+			List<IPerforceConnection> ChildConnections = new List<IPerforceConnection>();
+			List<Task> ChildTasks = new List<Task>(NumExtraThreads);
 			try
 			{
 				// Create the state object shared by all the worker threads
-				SyncState State = new SyncState();
-				State.TotalDepotPaths = RemainingDepotPaths.Count;
-				State.RemainingDepotPaths = RemainingDepotPaths;
-				State.SyncCommandLists = SyncCommandLists;
+				SyncState State = new SyncState(RemainingDepotPaths, SyncCommandLists);
 
 				// Wrapper writer around the log class to prevent multiple threads writing to it at once
-				ThreadSafeTextWriter LogWrapper = new ThreadSafeTextWriter(Log);
+				ILogger LogWrapper = Logger;
 
 				// Initialize Sync Progress
 				UpdateSyncState(Prefix, State, Progress);
 
 				// Delegate for updating the sync state after a file has been synced
-				Action<PerforceFileRecord, LineBasedTextWriter> SyncOutput = (Record, LocalLog) => { UpdateSyncState(Prefix, Record, State, Progress, LocalLog); };
+				Action<SyncRecord, ILogger> SyncOutput = (Record, LocalLog) => { UpdateSyncState(Prefix, Record, State, Progress, LocalLog); };
 
 				// Create all the child threads
 				for (int ThreadIdx = 0; ThreadIdx < NumExtraThreads; ThreadIdx++)
 				{
 					int ThreadNumber = ThreadIdx + 2;
-					Thread ChildThread = new Thread(() => StaticSyncWorker(ThreadNumber, Perforce, Context, State, SyncOutput, LogWrapper));
-					ChildThread.Name = "Sync File Revisions";
-					ChildThreads.Add(ChildThread);
-					ChildThread.Start();
-				}
 
-				// Run one worker on the current thread
-				StaticSyncWorker(1, Perforce, Context, State, SyncOutput, LogWrapper);
+					// Create connection
+					IPerforceConnection ChildConnection = await PerforceConnection.CreateAsync(Perforce.Settings, Perforce.Logger);
+					ChildConnections.Add(ChildConnection);
 
-				// Wait for all the background threads to finish
-				foreach (Thread ChildThread in ChildThreads)
-				{
-					ChildThread.Join();
+					Task ChildTask = Task.Run(() => StaticSyncWorker(ThreadNumber, ChildConnection, Context, State, SyncOutput, LogWrapper, CancellationToken));
+					ChildTasks.Add(ChildTask);
 				}
 
+				// Run one worker on the current thread
+				await StaticSyncWorker(1, Perforce, Context, State, SyncOutput, LogWrapper, CancellationToken);
+
 				// Return the result that was set on the state object
-				StatusMessage = State.StatusMessage;
-				return State.Result;
+				return (State.Result, State.StatusMessage);
 			}
 			finally
 			{
-				foreach (Thread ChildThread in ChildThreads)
+				await Task.WhenAll(ChildTasks);
+
+				foreach (IPerforceConnection ChildConnection in ChildConnections)
 				{
-					ChildThread.Interrupt();
+					ChildConnection.Dispose();
 				}
-				foreach (Thread ChildThread in ChildThreads)
+
+				foreach (Task ChildTask in ChildTasks)
 				{
-					ChildThread.Join();
+					await ChildTask;
 				}
 			}
 		}
@@ -1455,21 +1554,21 @@ namespace UnrealGameSync
 			}
 		}
 
-		static void UpdateSyncState(string Prefix, PerforceFileRecord Record, SyncState State, ProgressValue Progress, TextWriter Log)
+		static void UpdateSyncState(string Prefix, SyncRecord Record, SyncState State, ProgressValue Progress, ILogger Logger)
 		{
 			lock (State)
 			{
-				State.RemainingDepotPaths.Remove(Record.DepotPath);
+				State.RemainingDepotPaths.Remove(Record.DepotFile);
 
 				string Message = String.Format("{0} ({1}/{2})", Prefix, State.TotalDepotPaths - State.RemainingDepotPaths.Count, State.TotalDepotPaths);
 				float Fraction = Math.Min((float)(State.TotalDepotPaths - State.RemainingDepotPaths.Count) / (float)State.TotalDepotPaths, 1.0f);
 				Progress.Set(Message, Fraction);
 
-				Log.WriteLine("p4>   {0} {1}", Record.Action, Record.ClientPath);
+				Logger.LogInformation("p4>   {Action} {Path}", Record.Action, Record.Path);
 			}
 		}
 
-		static void StaticSyncWorker(int ThreadNumber, PerforceConnection Perforce, WorkspaceUpdateContext Context, SyncState State, Action<PerforceFileRecord, LineBasedTextWriter> SyncOutput, LineBasedTextWriter GlobalLog)
+		static async Task StaticSyncWorker(int ThreadNumber, IPerforceConnection Perforce, WorkspaceUpdateContext Context, SyncState State, Action<SyncRecord, ILogger> SyncOutput, ILogger GlobalLog, CancellationToken CancellationToken)
 		{
 			PrefixedTextWriter ThreadLog = new PrefixedTextWriter(String.Format("{0}:", ThreadNumber), GlobalLog);
 			for (; ; )
@@ -1489,17 +1588,7 @@ namespace UnrealGameSync
 				}
 
 				// Sync the files
-				string StatusMessage;
-				WorkspaceUpdateResult Result;
-				try
-				{
-					Result = StaticSyncFileRevisions(Perforce, Context, SyncCommands, Record => SyncOutput(Record, ThreadLog), ThreadLog, out StatusMessage);
-				}
-				catch (ThreadInterruptedException)
-				{
-					StatusMessage = "Sync cancelled";
-					Result = WorkspaceUpdateResult.Canceled;
-				}
+				(WorkspaceUpdateResult Result, string StatusMessage) = await StaticSyncFileRevisions(Perforce, Context, SyncCommands, Record => SyncOutput(Record, ThreadLog), CancellationToken);
 
 				// If it failed, try to set it on the state if nothing else has failed first
 				if (Result != WorkspaceUpdateResult.Success)
@@ -1517,26 +1606,39 @@ namespace UnrealGameSync
 			}
 		}
 
-		static WorkspaceUpdateResult StaticSyncFileRevisions(PerforceConnection Perforce, WorkspaceUpdateContext Context, List<string> SyncCommands, Action<PerforceFileRecord> SyncOutput, TextWriter Log, out string StatusMessage)
+		static async Task<(WorkspaceUpdateResult, string)> StaticSyncFileRevisions(IPerforceConnection Perforce, WorkspaceUpdateContext Context, List<string> SyncCommands, Action<SyncRecord> SyncOutput, CancellationToken CancellationToken)
 		{
 			// Sync them all
+			List<PerforceResponse<SyncRecord>> Responses = await Perforce.TrySyncAsync(SyncOptions.None, -1, SyncCommands, CancellationToken).ToListAsync(CancellationToken);
+
 			List<string> TamperedFiles = new List<string>();
-			if(!Perforce.Sync(SyncCommands, SyncOutput, TamperedFiles, false, Context.PerforceSyncOptions, Log))
+			foreach (PerforceResponse<SyncRecord> Response in Responses)
 			{
-				StatusMessage = "Aborted sync due to errors.";
-				return WorkspaceUpdateResult.FailedToSync;
+				const string NoClobberPrefix = "Can't clobber writable file ";
+				if (Response.Succeeded)
+				{
+					SyncOutput(Response.Data);
+				}
+				else if (Response.Error != null && Response.Error.Generic == PerforceGenericCode.Client && Response.Error.Data.StartsWith(NoClobberPrefix, StringComparison.OrdinalIgnoreCase))
+				{
+					TamperedFiles.Add(Response.Error.Data.Substring(NoClobberPrefix.Length).Trim());
+				}
+				else
+				{
+					return (WorkspaceUpdateResult.FailedToSync, "Aborted sync due to errors.");
+				}
 			}
 
 			// If any files need to be clobbered, defer to the main thread to figure out which ones
-			if(TamperedFiles.Count > 0)
+			if (TamperedFiles.Count > 0)
 			{
 				int NumNewFilesToClobber = 0;
-				foreach(string TamperedFile in TamperedFiles)
+				foreach (string TamperedFile in TamperedFiles)
 				{
-					if(!Context.ClobberFiles.ContainsKey(TamperedFile))
+					if (!Context.ClobberFiles.ContainsKey(TamperedFile))
 					{
 						Context.ClobberFiles[TamperedFile] = true;
-						if(TamperedFile.EndsWith(LocalObjectVersionFileName, StringComparison.OrdinalIgnoreCase) || TamperedFile.EndsWith(LocalVersionHeaderFileName, StringComparison.OrdinalIgnoreCase))
+						if (TamperedFile.EndsWith(LocalObjectVersionFileName, StringComparison.OrdinalIgnoreCase) || TamperedFile.EndsWith(LocalVersionHeaderFileName, StringComparison.OrdinalIgnoreCase))
 						{
 							// Hack for UseFastModularVersioningV2; we don't need to update these files any more.
 							continue;
@@ -1544,27 +1646,28 @@ namespace UnrealGameSync
 						NumNewFilesToClobber++;
 					}
 				}
-				if(NumNewFilesToClobber > 0)
+				if (NumNewFilesToClobber > 0)
 				{
-					StatusMessage = String.Format("Cancelled sync after checking files to clobber ({0} new files).", NumNewFilesToClobber);
-					return WorkspaceUpdateResult.FilesToClobber;
+					return (WorkspaceUpdateResult.FilesToClobber, $"Cancelled sync after checking files to clobber ({NumNewFilesToClobber} new files).");
 				}
-				foreach(string TamperedFile in TamperedFiles)
+				foreach (string TamperedFile in TamperedFiles)
 				{
-					if(Context.ClobberFiles[TamperedFile] && !Perforce.ForceSync(TamperedFile, Context.ChangeNumber, Log))
+					if (Context.ClobberFiles[TamperedFile])
 					{
-						StatusMessage = String.Format("Couldn't sync {0}.", TamperedFile);
-						return WorkspaceUpdateResult.FailedToSync;
+						List<PerforceResponse<SyncRecord>> Response = await Perforce.TrySyncAsync(SyncOptions.Force, -1, TamperedFile, CancellationToken).ToListAsync();
+						if (!Response.Succeeded())
+						{
+							return (WorkspaceUpdateResult.FailedToSync, $"Couldn't sync {TamperedFile}.");
+						}
 					}
 				}
 			}
 
 			// All succeeded
-			StatusMessage = null;
-			return WorkspaceUpdateResult.Success;
+			return (WorkspaceUpdateResult.Success, "Succeeded.");
 		}
 
-		static ConfigFile ReadProjectConfigFile(DirectoryReference LocalRootPath, FileReference SelectedLocalFileName, TextWriter Log)
+		public static async Task<ConfigFile> ReadProjectConfigFile(DirectoryReference LocalRootPath, FileReference SelectedLocalFileName, ILogger Logger)
 		{
 			// Find the valid config file paths
 			DirectoryInfo EngineDir = DirectoryReference.Combine(LocalRootPath, "Engine").ToDirectoryInfo();
@@ -1572,42 +1675,42 @@ namespace UnrealGameSync
 
 			// Read them in
 			ConfigFile ProjectConfig = new ConfigFile();
-			foreach(FileInfo LocalConfigFile in LocalConfigFiles)
+			foreach (FileInfo LocalConfigFile in LocalConfigFiles)
 			{
 				try
 				{
-					string[] Lines = File.ReadAllLines(LocalConfigFile.FullName);
+					string[] Lines = await File.ReadAllLinesAsync(LocalConfigFile.FullName);
 					ProjectConfig.Parse(Lines);
-					Log.WriteLine("Read config file from {0}", LocalConfigFile.FullName);
+					Logger.LogInformation("Read config file from {FileName}", LocalConfigFile.FullName);
 				}
-				catch(Exception Ex)
+				catch (Exception Ex)
 				{
-					Log.WriteLine("Failed to read config file from {0}: {1}", LocalConfigFile.FullName, Ex.ToString());
+					Logger.LogInformation(Ex, "Failed to read config file from {FileName}", LocalConfigFile.FullName);
 				}
 			}
 			return ProjectConfig;
 		}
 
-		static IReadOnlyList<string> ReadProjectStreamFilter(PerforceConnection Perforce, ConfigFile ProjectConfigFile, TextWriter Log)
+		public static async Task<IReadOnlyList<string>?> ReadProjectStreamFilter(IPerforceConnection Perforce, ConfigFile ProjectConfigFile, ILogger Logger, CancellationToken CancellationToken)
 		{
-			string StreamListDepotPath = ProjectConfigFile.GetValue("Options.QuickSelectStreamList", null);
-			if(StreamListDepotPath == null)
+			string? StreamListDepotPath = ProjectConfigFile.GetValue("Options.QuickSelectStreamList", null);
+			if (StreamListDepotPath == null)
 			{
 				return null;
 			}
 
-			List<string> Lines;
-			if(!Perforce.Print(StreamListDepotPath, out Lines, Log))
+			PerforceResponse<PrintRecord<string[]>> Response = await Perforce.TryPrintLinesAsync(StreamListDepotPath, CancellationToken);
+			if (!Response.Succeeded)
 			{
 				return null;
 			}
 
-			return Lines.Select(x => x.Trim()).Where(x => x.Length > 0).ToList().AsReadOnly();
+			return Response.Data.Contents.Select(x => x.Trim()).Where(x => x.Length > 0).ToList().AsReadOnly();
 		}
 
 		static string FormatTime(long Seconds)
 		{
-			if(Seconds >= 60)
+			if (Seconds >= 60)
 			{
 				return String.Format("{0,3}m {1:00}s", Seconds / 60, Seconds % 60);
 			}
@@ -1617,80 +1720,90 @@ namespace UnrealGameSync
 			}
 		}
 
-		static bool HasModifiedSourceFiles(PerforceConnection Perforce, ProjectInfo Project, TextWriter Log)
+		static async Task<bool> HasModifiedSourceFiles(IPerforceConnection Perforce, ProjectInfo Project, CancellationToken CancellationToken)
 		{
-			List<PerforceFileRecord> OpenFiles;
-			if(!Perforce.GetOpenFiles(Project.ClientRootPath + "/...", out OpenFiles, Log))
-			{
-				return true;
-			}
-			if(OpenFiles.Any(x => x.DepotPath.IndexOf("/Source/", StringComparison.OrdinalIgnoreCase) != -1))
+			List<OpenedRecord> OpenFiles = await Perforce.OpenedAsync(OpenedOptions.None, Project.ClientRootPath + "/...", CancellationToken).ToListAsync(CancellationToken);
+			if (OpenFiles.Any(x => x.DepotFile.IndexOf("/Source/", StringComparison.OrdinalIgnoreCase) != -1))
 			{
 				return true;
 			}
 			return false;
 		}
 
-		static bool FindUnresolvedFiles(PerforceConnection Perforce, IEnumerable<string> SyncPaths, TextWriter Log, out List<PerforceFileRecord> UnresolvedFiles)
+		static async Task<List<FStatRecord>> FindUnresolvedFiles(IPerforceConnection Perforce, IEnumerable<string> SyncPaths, CancellationToken CancellationToken)
 		{
-			UnresolvedFiles = new List<PerforceFileRecord>();
-			foreach(string SyncPath in SyncPaths)
+			List<FStatRecord> UnresolvedFiles = new List<FStatRecord>();
+			foreach (string SyncPath in SyncPaths)
 			{
-				List<PerforceFileRecord> Records;
-				if(!Perforce.GetUnresolvedFiles(SyncPath, out Records, Log))
-				{
-					Log.WriteLine("Couldn't find open files matching {0}", SyncPath);
-					return false;
-				}
+				List<FStatRecord> Records = await Perforce.FStatAsync(FStatOptions.OnlyUnresolved, SyncPath, CancellationToken).ToListAsync(CancellationToken);
 				UnresolvedFiles.AddRange(Records);
 			}
-			return true;
+			return UnresolvedFiles;
 		}
 
-		static bool UpdateVersionFile(PerforceConnection Perforce, string ClientPath, Dictionary<string, string> VersionStrings, int ChangeNumber, TextWriter Log)
+		static Task<bool> UpdateVersionFile(IPerforceConnection Perforce, string ClientPath, Dictionary<string, string> VersionStrings, int ChangeNumber, ILogger Logger, CancellationToken CancellationToken)
 		{
-			return UpdateVersionFile(Perforce, ClientPath, ChangeNumber, Text => UpdateVersionStrings(Text, VersionStrings), Log);
+			return UpdateVersionFile(Perforce, ClientPath, ChangeNumber, Text => UpdateVersionStrings(Text, VersionStrings), Logger, CancellationToken);
 		}
 
-		static bool UpdateVersionFile(PerforceConnection Perforce, string ClientPath, int ChangeNumber, Func<string, string> Update, TextWriter Log)
+		static async Task<bool> UpdateVersionFile(IPerforceConnection Perforce, string ClientPath, int ChangeNumber, Func<string, string> Update, ILogger Logger, CancellationToken CancellationToken)
 		{
-			List<PerforceFileRecord> Records;
-			if(!Perforce.Stat(ClientPath, out Records, Log))
+			List<PerforceResponse<FStatRecord>> Records = await Perforce.TryFStatAsync(FStatOptions.None, ClientPath, CancellationToken).ToListAsync(CancellationToken);
+			if (!Records.Succeeded())
 			{
-				Log.WriteLine("Failed to query records for {0}", ClientPath);
+				Logger.LogInformation("Failed to query records for {ClientPath}", ClientPath);
 				return false;
 			}
 			if (Records.Count > 1)
 			{
 				// Attempt to remove any existing file which is synced
-				Perforce.ForceSync(String.Format("{0}#0", ClientPath), Log);
+				await Perforce.SyncAsync(SyncOptions.Force, -1, $"{ClientPath}#0", CancellationToken).ToListAsync(CancellationToken);
 
 				// Try to get the mapped files again
-				if (!Perforce.Stat(ClientPath, out Records, Log))
+				Records = await Perforce.TryFStatAsync(FStatOptions.None, ClientPath, CancellationToken).ToListAsync(CancellationToken);
+				if (!Records.Succeeded())
 				{
-					Log.WriteLine("Failed to query records for {0}", ClientPath);
+					Logger.LogInformation("Failed to query records for {ClientPath}", ClientPath);
 					return false;
 				}
 			}
 			if (Records.Count == 0)
 			{
-				Log.WriteLine("Ignoring {0}; not found on server.", ClientPath);
+				Logger.LogInformation("Ignoring {ClientPath}; not found on server.", ClientPath);
 				return true;
 			}
 
-			string LocalPath = Records[0].ClientPath; // Actually a filesystem path
-			string DepotPath = Records[0].DepotPath;
+			FStatRecord Record = Records[0].Data;
+			string? LocalPath = Record.ClientFile; // Actually a filesystem path
+			if (LocalPath == null)
+			{
+				Logger.LogInformation("Version file is not mapped to workspace ({ClientFile})", ClientPath);
+				return false;
+			}
+			string? DepotPath = Record.DepotFile;
+			if (DepotPath == null)
+			{
+				Logger.LogInformation("Version file does not exist in depot ({ClientFile})", ClientPath);
+				return false;
+			}
 
-			List<string> Lines;
-			if(!Perforce.Print(String.Format("{0}@{1}", DepotPath, ChangeNumber), out Lines, Log))
+			PerforceResponse<PrintRecord<string[]>> Response = await Perforce.TryPrintLinesAsync($"{DepotPath}@{ChangeNumber}", CancellationToken);
+			if (!Response.Succeeded)
 			{
-				Log.WriteLine("Couldn't get default contents of {0}", DepotPath);
+				Logger.LogInformation("Couldn't get default contents of {DepotPath}", DepotPath);
 				return false;
 			}
 
-			string Text = String.Join("\n", Lines);
+			string[]? Contents = Response.Data.Contents;
+			if (Contents == null)
+			{
+				Logger.LogInformation("No data returned for {DepotPath}", DepotPath);
+				return false;
+			}
+
+			string Text = String.Join("\n", Contents);
 			Text = Update(Text);
-			return WriteVersionFile(Perforce, LocalPath, DepotPath, Text, Log);
+			return await WriteVersionFile(Perforce, LocalPath, DepotPath, Text, Logger, CancellationToken);
 		}
 
 		static string UpdateVersionStrings(string Text, Dictionary<string, string> VersionStrings)
@@ -1715,14 +1828,20 @@ namespace UnrealGameSync
 		{
 			Dictionary<string, object> Object = JsonSerializer.Deserialize<Dictionary<string, object>>(Text, Program.DefaultJsonSerializerOptions);
 
-			object PrevCompatibleChangelistObj;
-			int PrevCompatibleChangelist = Object.TryGetValue("CompatibleChangelist", out PrevCompatibleChangelistObj) ? (int)Convert.ChangeType(PrevCompatibleChangelistObj.ToString(), typeof(int)) : 0;
+			int PrevCompatibleChangelist = 0;
+			if (Object.TryGetValue("CompatibleChangelist", out object? PrevCompatibleChangelistObj))
+			{
+				int.TryParse(PrevCompatibleChangelistObj?.ToString(), out PrevCompatibleChangelist);
+			}
 
-			object PrevIsLicenseeVersionObj;
-			bool PrevIsLicenseeVersion = Object.TryGetValue("IsLicenseeVersion", out PrevIsLicenseeVersionObj)? ((int)Convert.ChangeType(PrevIsLicenseeVersionObj.ToString(), typeof(int)) != 0) : false;
+			int PrevIsLicenseeVersion = 0;
+			if (Object.TryGetValue("IsLicenseeVersion", out object? PrevIsLicenseeVersionObj))
+			{
+				int.TryParse(PrevIsLicenseeVersionObj?.ToString(), out PrevIsLicenseeVersion);
+			}
 
 			Object["Changelist"] = Changelist;
-			if(PrevCompatibleChangelist == 0 || PrevIsLicenseeVersion != bIsLicenseeVersion)
+			if (PrevCompatibleChangelist == 0 || (PrevIsLicenseeVersion != 0) != bIsLicenseeVersion)
 			{
 				// Don't overwrite the compatible changelist if we're in a hotfix release
 				Object["CompatibleChangelist"] = CodeChangelist;
@@ -1733,36 +1852,36 @@ namespace UnrealGameSync
 
 			return JsonSerializer.Serialize(Object, new JsonSerializerOptions
 			{
-				WriteIndented = true, 
+				WriteIndented = true,
 				// do not escape +
 				Encoder = JavaScriptEncoder.UnsafeRelaxedJsonEscaping
 			});
 		}
 
-		static bool WriteVersionFile(PerforceConnection Perforce, string LocalPath, string DepotPath, string NewText, TextWriter Log)
+		static async Task<bool> WriteVersionFile(IPerforceConnection Perforce, string LocalPath, string DepotPath, string NewText, ILogger Logger, CancellationToken CancellationToken)
 		{
 			try
 			{
-				if(File.Exists(LocalPath) && File.ReadAllText(LocalPath) == NewText)
+				if (File.Exists(LocalPath) && File.ReadAllText(LocalPath) == NewText)
 				{
-					Log.WriteLine("Ignored {0}; contents haven't changed", LocalPath);
+					Logger.LogInformation("Ignored {FileName}; contents haven't changed", LocalPath);
 				}
 				else
 				{
 					Directory.CreateDirectory(Path.GetDirectoryName(LocalPath));
 					Utility.ForceDeleteFile(LocalPath);
-					if(DepotPath != null)
+					if (DepotPath != null)
 					{
-						Perforce.Sync(DepotPath + "#0", Log);
+						await Perforce.SyncAsync(DepotPath + "#0", CancellationToken).ToListAsync(CancellationToken);
 					}
 					File.WriteAllText(LocalPath, NewText);
-					Log.WriteLine("Written {0}", LocalPath);
+					Logger.LogInformation("Written {FileName}", LocalPath);
 				}
 				return true;
 			}
-			catch(Exception Ex)
+			catch (Exception Ex)
 			{
-				Log.WriteException(Ex, "Failed to write to {0}.", LocalPath);
+				Logger.LogError(Ex, "Failed to write to {FileName}.", LocalPath);
 				return false;
 			}
 		}
@@ -1771,16 +1890,16 @@ namespace UnrealGameSync
 		{
 			int LineIdx = 0;
 			int PrefixIdx = 0;
-			for(;;)
+			for (; ; )
 			{
-				string PrefixToken = ReadToken(Prefix, ref PrefixIdx);
-				if(PrefixToken == null)
+				string? PrefixToken = ReadToken(Prefix, ref PrefixIdx);
+				if (PrefixToken == null)
 				{
 					break;
 				}
 
-				string LineToken = ReadToken(Line, ref LineIdx);
-				if(LineToken == null || LineToken != PrefixToken)
+				string? LineToken = ReadToken(Line, ref LineIdx);
+				if (LineToken == null || LineToken != PrefixToken)
 				{
 					return false;
 				}
@@ -1789,24 +1908,24 @@ namespace UnrealGameSync
 			return true;
 		}
 
-		static string ReadToken(string Line, ref int LineIdx)
+		static string? ReadToken(string Line, ref int LineIdx)
 		{
-			for(;; LineIdx++)
+			for (; ; LineIdx++)
 			{
-				if(LineIdx == Line.Length)
+				if (LineIdx == Line.Length)
 				{
 					return null;
 				}
-				else if(!Char.IsWhiteSpace(Line[LineIdx]))
+				else if (!Char.IsWhiteSpace(Line[LineIdx]))
 				{
 					break;
 				}
 			}
 
 			int StartIdx = LineIdx++;
-			if(Char.IsLetterOrDigit(Line[StartIdx]) || Line[StartIdx] == '_')
+			if (Char.IsLetterOrDigit(Line[StartIdx]) || Line[StartIdx] == '_')
 			{
-				while(LineIdx < Line.Length && (Char.IsLetterOrDigit(Line[LineIdx]) || Line[LineIdx] == '_'))
+				while (LineIdx < Line.Length && (Char.IsLetterOrDigit(Line[LineIdx]) || Line[LineIdx] == '_'))
 				{
 					LineIdx++;
 				}
@@ -1827,15 +1946,11 @@ namespace UnrealGameSync
 
 		public int CurrentChangeNumber => State.CurrentChangeNumber;
 
-		public int PendingChangeNumber
-		{
-			get;
-			private set;
-		}
+		public int PendingChangeNumber => CurrentUpdate?.Context?.ChangeNumber ?? CurrentChangeNumber;
 
 		public string ClientName
 		{
-			get { return Perforce.ClientName; }
+			get { return PerforceSettings.ClientName!; }
 		}
 	}
 }
