UGS: Add some more #nullable annotations.

diff --git a/Engine/Source/Programs/Shared/EpicGames.Perforce/Records/ChangesRecord.cs b/Engine/Source/Programs/Shared/EpicGames.Perforce/Records/ChangesRecord.cs
--- a/Engine/Source/Programs/Shared/EpicGames.Perforce/Records/ChangesRecord.cs
+++ b/Engine/Source/Programs/Shared/EpicGames.Perforce/Records/ChangesRecord.cs
@@ -29,7 +29,7 @@ namespace EpicGames.Perforce
 		/// The user that owns or submitted the change
 		/// </summary>
 		[PerforceTag("user")]
-		public string? User { get; set; }
+		public string User { get; set; } = String.Empty;
 
 		/// <summary>
 		/// The client that owns the change
@@ -59,6 +59,6 @@ namespace EpicGames.Perforce
 		/// Description for the changelist
 		/// </summary>
 		[PerforceTag("desc")]
-		public string? Description { get; set; }
+		public string Description { get; set; } = String.Empty;
 	}
 }
diff --git a/Engine/Source/Programs/UnrealGameSync/UnrealGameSync/Controls/WorkspaceControl.cs b/Engine/Source/Programs/UnrealGameSync/UnrealGameSync/Controls/WorkspaceControl.cs
--- a/Engine/Source/Programs/UnrealGameSync/UnrealGameSync/Controls/WorkspaceControl.cs
+++ b/Engine/Source/Programs/UnrealGameSync/UnrealGameSync/Controls/WorkspaceControl.cs
@@ -20,20 +20,23 @@ using EpicGames.Perforce;
 using System.Threading.Tasks;
 using Microsoft.Extensions.Logging;
 using Microsoft.Extensions.DependencyInjection;
+using System.Diagnostics.CodeAnalysis;
+
+#nullable enable
 
 namespace UnrealGameSync
 {
 	public class UncontrolledChangelist
 	{
-		public string GUID { get; set; }
-		public string Name { get; set; }
-		public List<string> Files { get; set; }
+		public string GUID { get; set; } = String.Empty;
+		public string Name { get; set; } = String.Empty;
+		public List<string> Files { get; set; } = new List<string>();
 	}
 
 	public class UncontrolledChangelistPersistency
 	{
 		public int Version { get; set; }
-		public List<UncontrolledChangelist> Changelists { get; set; }
+		public List<UncontrolledChangelist> Changelists { get; set; } = new List<UncontrolledChangelist>();
 	}
 
 	public enum LatestChangeType
@@ -84,22 +87,22 @@ namespace UnrealGameSync
 		class BadgeInfo
 		{
 			public string Label;
-			public string Group;
-			public string UniqueId;
+			public string? Group;
+			public string? UniqueId;
 			public int Offset;
 			public int Width;
 			public int Height;
 			public Color BackgroundColor;
 			public Color HoverBackgroundColor;
-			public Action ClickHandler;
-			public string ToolTip;
+			public Action? ClickHandler;
+			public string? ToolTip;
 
 			public BadgeInfo(string Label, string Group, Color BadgeColor)
 				: this(Label, Group, null, BadgeColor, BadgeColor, null)
 			{
 			}
 
-			public BadgeInfo(string Label, string Group, string UniqueId, Color BackgroundColor, Color HoverBackgroundColor, Action ClickHandler)
+			public BadgeInfo(string Label, string? Group, string? UniqueId, Color BackgroundColor, Color HoverBackgroundColor, Action? ClickHandler)
 			{
 				this.Label = Label;
 				this.Group = Group;
@@ -126,22 +129,22 @@ namespace UnrealGameSync
 		class ServiceBadgeInfo : IEquatable<ServiceBadgeInfo>
 		{
 			public string Name;
-			public string Url;
+			public string? Url;
 			public BadgeResult Result;
 
-			public ServiceBadgeInfo(string Name, string Url, BadgeResult Result)
+			public ServiceBadgeInfo(string Name, string? Url, BadgeResult Result)
 			{
 				this.Name = Name;
 				this.Url = Url;
 				this.Result = Result;
 			}
 
-			public bool Equals(ServiceBadgeInfo OtherBadge)
+			public bool Equals(ServiceBadgeInfo? OtherBadge)
 			{
 				return OtherBadge != null && Name == OtherBadge.Name && Url == OtherBadge.Url && Result == OtherBadge.Result;
 			}
 
-			public override bool Equals(object Other)
+			public override bool Equals(object? Other)
 			{
 				return Equals(Other as ServiceBadgeInfo);
 			}
@@ -154,18 +157,18 @@ namespace UnrealGameSync
 
 		class ChangeLayoutInfo
 		{
-			public List<BadgeInfo> DescriptionBadges;
-			public List<BadgeInfo> TypeBadges;
-			public List<BadgeInfo> BuildBadges;
-			public Dictionary<string, List<BadgeInfo>> CustomBadges;
+			public List<BadgeInfo> DescriptionBadges = new List<BadgeInfo>();
+			public List<BadgeInfo> TypeBadges = new List<BadgeInfo>();
+			public List<BadgeInfo> BuildBadges = new List<BadgeInfo>();
+			public Dictionary<string, List<BadgeInfo>> CustomBadges = new Dictionary<string, List<BadgeInfo>>();
 		}
 
 		class BuildListItemComparer : System.Collections.IComparer
 		{
-			public int Compare(object A, object B)
+			public int Compare(object? A, object? B)
 			{
-				ChangesRecord ChangeA = ((ListViewItem)A).Tag as ChangesRecord;
-				ChangesRecord ChangeB = ((ListViewItem)B).Tag as ChangesRecord;
+				ChangesRecord? ChangeA = ((ListViewItem?)A)?.Tag as ChangesRecord;
+				ChangesRecord? ChangeB = ((ListViewItem?)B)?.Tag as ChangesRecord;
 				if (ChangeA == null)
 				{
 					return +1;
@@ -270,7 +273,7 @@ namespace UnrealGameSync
 		HashSet<int> PromotedChangeNumbers = new HashSet<int>();
 		List<int> ListIndexToChangeIndex = new List<int>();
 		List<int> SortedChangeNumbers = new List<int>();
-		Dictionary<string, Dictionary<int, string>> ArchiveToChangeNumberToArchiveKey = new Dictionary<string, Dictionary<int, string>>();
+		Dictionary<string, Dictionary<int, string?>> ArchiveToChangeNumberToArchiveKey = new Dictionary<string, Dictionary<int, string?>>();
 		Dictionary<int, ChangeLayoutInfo> ChangeNumberToLayoutInfo = new Dictionary<int, ChangeLayoutInfo>();
 		List<ContextMenuStrip> CustomStatusPanelMenus = new List<ContextMenuStrip>();
 		List<(string, Action<Point, Rectangle>)> CustomStatusPanelLinks = new List<(string, Action<Point, Rectangle>)>();
@@ -282,12 +285,12 @@ namespace UnrealGameSync
 
 		bool bMouseOverExpandLink;
 
-		string HoverBadgeUniqueId = null;
+		string? HoverBadgeUniqueId = null;
 		bool bHoverSync;
-		ChangesRecord ContextMenuChange;
-		Font BuildFont;
-		Font SelectedBuildFont;
-		Font BadgeFont;
+		ChangesRecord? ContextMenuChange;
+		Font? BuildFont;
+		Font? SelectedBuildFont;
+		Font? BadgeFont;
 		List<KeyValuePair<string, string>> BadgeNameAndGroupPairs = new List<KeyValuePair<string, string>>();
 		Dictionary<string, Size> BadgeLabelToSize = new Dictionary<string, Size>();
 		List<ServiceBadgeInfo> ServiceBadges = new List<ServiceBadgeInfo>();
@@ -303,23 +306,23 @@ namespace UnrealGameSync
 		}
 
 		int BuildListWidth;
-		float[] ColumnWidths;
-		float[] ColumnWeights;
-		int[] MinColumnWidths;
-		int[] DesiredColumnWidths;
-		string LastColumnSettings;
-		List<ColumnHeader> CustomColumns;
+		float[] ColumnWidths = Array.Empty<float>();
+		float[] ColumnWeights = Array.Empty<float>();
+		int[] MinColumnWidths = Array.Empty<int>();
+		int[] DesiredColumnWidths = Array.Empty<int>();
+		string? LastColumnSettings;
+		List<ColumnHeader> CustomColumns = new List<ColumnHeader>();
 		int MaxBuildBadgeChars;
-		ListViewItem ExpandItem;
+		ListViewItem? ExpandItem;
 
 		bool bUpdateBuildListPosted;
 		bool bUpdateBuildMetadataPosted;
 		bool bUpdateReviewsPosted;
 
-		WorkspaceUpdateCallback UpdateCallback;
+		WorkspaceUpdateCallback? UpdateCallback;
 
-		System.Threading.Timer StartupTimer;
-		List<WorkspaceStartupCallback> StartupCallbacks;
+		System.Threading.Timer? StartupTimer;
+		List<WorkspaceStartupCallback>? StartupCallbacks;
 
 		// When an author filter is applied this will be non-empty and != AuthorFilterPlaceholderText
 		string AuthorFilterText = "";
@@ -331,7 +334,7 @@ namespace UnrealGameSync
 		{
 			InitializeComponent();
 
-			MainThreadSynchronizationContext = SynchronizationContext.Current;
+			MainThreadSynchronizationContext = SynchronizationContext.Current!;
 
 			Owner = InOwner;
 			ApiUrl = InApiUrl;
@@ -351,12 +354,12 @@ namespace UnrealGameSync
 
 			DesiredTaskbarState = Tuple.Create(TaskbarState.NoProgress, 0.0f);
 
-			System.Reflection.PropertyInfo DoubleBufferedProperty = typeof(Control).GetProperty("DoubleBuffered", System.Reflection.BindingFlags.NonPublic | System.Reflection.BindingFlags.Instance);
+			System.Reflection.PropertyInfo DoubleBufferedProperty = typeof(Control).GetProperty("DoubleBuffered", System.Reflection.BindingFlags.NonPublic | System.Reflection.BindingFlags.Instance)!;
 			DoubleBufferedProperty.SetValue(BuildList, true, null);
 
 			// force the height of the rows
 			BuildList.SmallImageList = new ImageList() { ImageSize = new Size(1, 20) };
-			BuildList_FontChanged(null, null);
+			BuildList_FontChanged(null, null!);
 			BuildList.OnScroll += BuildList_OnScroll;
 			BuildList.ListViewItemSorter = new BuildListItemComparer();
 
@@ -393,7 +396,7 @@ namespace UnrealGameSync
 
 			// Commit all the new project info
 			IPerforceSettings PerforceClientSettings = WorkspaceSettings.PerforceSettings;
-			ClientName = PerforceUtils.GetClientOrDepotDirectoryName(WorkspaceSettings.ProjectInfo.ClientFileName);
+			ClientName = PerforceClientSettings.ClientName!;
 			SelectedProject = WorkspaceSettings.SelectedProject;
 			SelectedProjectIdentifier = WorkspaceSettings.NewSelectedProjectIdentifier;
 			SelectedFileName = WorkspaceSettings.ProjectInfo.LocalFileName;
@@ -406,7 +409,7 @@ namespace UnrealGameSync
 
 			// Check if we've the project we've got open in this workspace is the one we're actually synced to
 			int CurrentChangeNumber = -1;
-			string CurrentSyncFilterHash = null;
+			string? CurrentSyncFilterHash = null;
 			if (String.Compare(WorkspaceState.CurrentProjectIdentifier, SelectedProjectIdentifier, true) == 0)
 			{
 				CurrentChangeNumber = WorkspaceState.CurrentChangeNumber;
@@ -414,7 +417,7 @@ namespace UnrealGameSync
 			}
 
 			// Figure out which API server to use
-			string NewApiUrl;
+			string? NewApiUrl;
 			if (TryGetProjectSetting(WorkspaceSettings.LatestProjectConfigFile, "ApiUrl", out NewApiUrl))
 			{
 				ApiUrl = NewApiUrl;
@@ -425,7 +428,7 @@ namespace UnrealGameSync
 			Workspace = new Workspace(PerforceClientSettings, Project, WorkspaceState, WorkspaceSettings.WorkspaceProjectConfigFile, WorkspaceSettings.WorkspaceProjectStreamFilter, new LogControlTextWriter(SyncLog), ServiceProvider);
 			Workspace.OnUpdateComplete += UpdateCompleteCallback;
 
-			FileReference ProjectLogBaseName = FileReference.Combine(DataFolder, "sync.log");
+			FileReference ProjectLogBaseName = FileReference.Combine(DataFolder, String.Format("sync_{0}.log", WorkspaceSettings.ProjectInfo.ClientFileName.Replace('/', '+')));
 
 			ILogger PerforceLogger = ServiceProvider.GetRequiredService<ILogger<PerforceMonitor>>();
 			PerforceMonitor = new PerforceMonitor(PerforceClientSettings, WorkspaceSettings.ProjectInfo, WorkspaceSettings.LatestProjectConfigFile, WorkspaceSettings.CacheFolder, WorkspaceSettings.LocalConfigFiles, ServiceProvider);
@@ -441,9 +444,9 @@ namespace UnrealGameSync
 			ILogger<JupiterMonitor> JupiterLogger = ServiceProvider.GetRequiredService<ILogger<JupiterMonitor>>();
 			JupiterMonitor = JupiterMonitor.CreateFromConfigFile(InOidcTokenManager, JupiterLogger, WorkspaceSettings.LatestProjectConfigFile, SelectedProjectIdentifier);
 
-			UpdateColumnSettings();
+			UpdateColumnSettings(true);
 
-			SyncLog.OpenFile(ProjectLogBaseName + ".sync.log");
+			SyncLog.OpenFile(ProjectLogBaseName);
 
 			Splitter.SetLogVisibility(Settings.bShowLogWindow);
 
@@ -544,32 +547,32 @@ namespace UnrealGameSync
 			}
 			else
 			{
-				StartupCallbacks.Add(StartupCallback);
+				StartupCallbacks!.Add(StartupCallback);
 			}
 		}
 
-		private void UpdateColumnSettings()
+		private void UpdateColumnSettings(bool bForceUpdate)
 		{
-			string NextColumnSettings;
+			string? NextColumnSettings;
 			TryGetProjectSetting(PerforceMonitor.LatestProjectConfigFile, "Columns", out NextColumnSettings);
 
-			if (CustomColumns == null || NextColumnSettings != LastColumnSettings)
+			if (bForceUpdate || NextColumnSettings != LastColumnSettings)
 			{
 				LastColumnSettings = NextColumnSettings;
 
-				if (CustomColumns != null)
+				foreach (ColumnHeader CustomColumn in CustomColumns)
 				{
-					foreach (ColumnHeader CustomColumn in CustomColumns)
-					{
-						BuildList.Columns.Remove(CustomColumn);
-					}
+					BuildList.Columns.Remove(CustomColumn);
 				}
 
 				Dictionary<string, ColumnHeader> NameToColumn = new Dictionary<string, ColumnHeader>();
-				foreach (ColumnHeader Column in BuildList.Columns)
+				foreach (ColumnHeader? Column in BuildList.Columns)
 				{
-					NameToColumn[Column.Text] = Column;
-					Column.Tag = null;
+					if (Column != null)
+					{
+						NameToColumn[Column.Text] = Column;
+						Column.Tag = null;
+					}
 				}
 
 				CustomColumns = new List<ColumnHeader>();
@@ -579,10 +582,10 @@ namespace UnrealGameSync
 					{
 						ConfigObject ColumnConfig = new ConfigObject(CustomColumn);
 
-						string Name = ColumnConfig.GetValue("Name", null);
+						string? Name = ColumnConfig.GetValue("Name", null);
 						if (Name != null)
 						{
-							ColumnHeader Column;
+							ColumnHeader? Column;
 							if (NameToColumn.TryGetValue(Name, out Column))
 							{
 								Column.Tag = ColumnConfig;
@@ -640,14 +643,17 @@ namespace UnrealGameSync
 					ColumnWeights[DescriptionColumn.Index] = 1.25f;
 					ColumnWeights[CISColumn.Index] = 1.5f;
 
-					foreach (ColumnHeader Column in BuildList.Columns)
+					foreach (ColumnHeader? Column in BuildList.Columns)
 					{
-						ConfigObject ColumnConfig = (ConfigObject)Column.Tag;
-						if (ColumnConfig != null)
+						if (Column != null)
 						{
-							MinColumnWidths[Column.Index] = (int)(ColumnConfig.GetValue("MinWidth", MinColumnWidths[Column.Index]) * DpiScaleX);
-							DesiredColumnWidths[Column.Index] = (int)(ColumnConfig.GetValue("DesiredWidth", DesiredColumnWidths[Column.Index]) * DpiScaleX);
-							ColumnWeights[Column.Index] = ColumnConfig.GetValue("Weight", MinColumnWidths[Column.Index]);
+							ConfigObject? ColumnConfig = (ConfigObject?)Column.Tag;
+							if (ColumnConfig != null)
+							{
+								MinColumnWidths[Column.Index] = (int)(ColumnConfig.GetValue("MinWidth", MinColumnWidths[Column.Index]) * DpiScaleX);
+								DesiredColumnWidths[Column.Index] = (int)(ColumnConfig.GetValue("DesiredWidth", DesiredColumnWidths[Column.Index]) * DpiScaleX);
+								ColumnWeights[Column.Index] = ColumnConfig.GetValue("Weight", MinColumnWidths[Column.Index]);
+							}
 						}
 					}
 
@@ -656,7 +662,7 @@ namespace UnrealGameSync
 					{
 						if (!String.IsNullOrEmpty(BuildList.Columns[Idx].Text))
 						{
-							string StringValue;
+							string? StringValue;
 							if (TryGetProjectSetting(ProjectConfigFile, String.Format("ColumnWidth_{0}", BuildList.Columns[Idx].Text), out StringValue))
 							{
 								int IntValue;
@@ -678,10 +684,10 @@ namespace UnrealGameSync
 
 		private bool UpdateServiceBadges()
 		{
-			string[] ServiceBadgeNames;
+			string[]? ServiceBadgeNames;
 			if (!TryGetProjectSetting(PerforceMonitor.LatestProjectConfigFile, "ServiceBadges", out ServiceBadgeNames))
 			{
-				ServiceBadgeNames = new string[0];
+				ServiceBadgeNames = Array.Empty<string>();
 			}
 
 			List<ServiceBadgeInfo> PrevServiceBadges = ServiceBadges;
@@ -750,59 +756,59 @@ namespace UnrealGameSync
 				{
 					StartupCallback(this, true);
 				}
-				StartupCallbacks = null;
+				StartupCallbacks = null!;
 			}
 
 			if (IssueMonitor != null)
 			{
 				IssueMonitor.OnIssuesChanged -= IssueMonitor_OnIssuesChangedAsync;
 				Owner.ReleaseIssueMonitor(IssueMonitor);
-				IssueMonitor = null;
+				IssueMonitor = null!;
 			}
 			if (StartupTimer != null)
 			{
 				StartupTimer.Dispose();
-				StartupTimer = null;
+				StartupTimer = null!;
 			}
 			if (NotificationWindow != null)
 			{
 				NotificationWindow.Dispose();
-				NotificationWindow = null;
+				NotificationWindow = null!;
 			}
 			if (PerforceMonitor != null)
 			{
 				PerforceMonitor.Dispose();
-				PerforceMonitor = null;
+				PerforceMonitor = null!;
 			}
 			if (Workspace != null)
 			{
 				Workspace.Dispose();
-				Workspace = null;
+				Workspace = null!;
 			}
 			if (EventMonitor != null)
 			{
 				EventMonitor.Dispose();
-				EventMonitor = null;
+				EventMonitor = null!;
 			}
 			if (BuildFont != null)
 			{
 				BuildFont.Dispose();
-				BuildFont = null;
+				BuildFont = null!;
 			}
 			if (SelectedBuildFont != null)
 			{
 				SelectedBuildFont.Dispose();
-				SelectedBuildFont = null;
+				SelectedBuildFont = null!;
 			}
 			if (BadgeFont != null)
 			{
 				BadgeFont.Dispose();
-				BadgeFont = null;
+				BadgeFont = null!;
 			}
 			if (JupiterMonitor != null)
 			{
 				JupiterMonitor.Dispose();
-				JupiterMonitor = null;
+				JupiterMonitor = null!;
 			}
 
 			base.Dispose(disposing);
@@ -917,7 +923,7 @@ namespace UnrealGameSync
 			StartSync(ChangeNumber, false, null);
 		}
 
-		void StartSync(int ChangeNumber, bool bSyncOnly, WorkspaceUpdateCallback Callback)
+		void StartSync(int ChangeNumber, bool bSyncOnly, WorkspaceUpdateCallback? Callback)
 		{
 			WorkspaceUpdateOptions Options = WorkspaceUpdateOptions.Sync | WorkspaceUpdateOptions.SyncArchives | WorkspaceUpdateOptions.GenerateProjectFiles;
 			if (!bSyncOnly)
@@ -943,7 +949,7 @@ namespace UnrealGameSync
 			StartWorkspaceUpdate(ChangeNumber, Options, null);
 		}
 
-		void StartWorkspaceUpdate(int ChangeNumber, WorkspaceUpdateOptions Options, WorkspaceUpdateCallback Callback)
+		void StartWorkspaceUpdate(int ChangeNumber, WorkspaceUpdateOptions Options, WorkspaceUpdateCallback? Callback)
 		{
 			if ((Options & (WorkspaceUpdateOptions.Sync | WorkspaceUpdateOptions.Build)) != 0 && GetProcessesRunningInWorkspace().Length > 0)
 			{
@@ -983,7 +989,7 @@ namespace UnrealGameSync
 				List<IArchiveInfo> SelectedArchives = GetSelectedArchives(Archives);
 				foreach (IArchiveInfo Archive in SelectedArchives)
 				{
-					string ArchivePath = GetArchiveKeyForChangeNumber(Archive, ChangeNumber);
+					string? ArchivePath = GetArchiveKeyForChangeNumber(Archive, ChangeNumber);
 					if (ArchivePath == null)
 					{
 						MessageBox.Show(String.Format("There are no compiled {0} binaries for this change. To sync it, you must disable syncing of precompiled editor binaries.", Archive.Name));
@@ -995,7 +1001,7 @@ namespace UnrealGameSync
 						Context.Options &= ~(WorkspaceUpdateOptions.Build | WorkspaceUpdateOptions.GenerateProjectFiles | WorkspaceUpdateOptions.OpenSolutionAfterSync);
 					}
 
-					string[] ZippedBinariesSyncFilter;
+					string[]? ZippedBinariesSyncFilter;
 					if (TryGetProjectSetting(PerforceMonitor.LatestProjectConfigFile, "ZippedBinariesSyncFilter", out ZippedBinariesSyncFilter) && ZippedBinariesSyncFilter.Length > 0)
 					{
 						Context.SyncFilter = Enumerable.Concat(Context.SyncFilter, ZippedBinariesSyncFilter).ToArray();
@@ -1007,7 +1013,7 @@ namespace UnrealGameSync
 			StartWorkspaceUpdate(Context, Callback);
 		}
 
-		void StartWorkspaceUpdate(WorkspaceUpdateContext Context, WorkspaceUpdateCallback Callback)
+		void StartWorkspaceUpdate(WorkspaceUpdateContext Context, WorkspaceUpdateCallback? Callback)
 		{
 			if (Settings.bAutoResolveConflicts)
 			{
@@ -1041,7 +1047,7 @@ namespace UnrealGameSync
 			{
 				if (!Context.Options.HasFlag(WorkspaceUpdateOptions.ContentOnly) && (Context.CustomBuildSteps == null || Context.CustomBuildSteps.Count == 0))
 				{
-					foreach (BuildConfig Config in Enum.GetValues(typeof(BuildConfig)))
+					foreach (BuildConfig Config in Enum.GetValues(typeof(BuildConfig)).OfType<BuildConfig>())
 					{
 						List<FileReference> EditorReceiptPaths = GetEditorReceiptPaths(Config);
 						foreach (FileReference EditorReceiptPath in EditorReceiptPaths)
@@ -1123,7 +1129,7 @@ namespace UnrealGameSync
 
 				HashSet<string> UncontrolledFiles = new HashSet<string>();
 				
-				if (SelectedProject.LocalPath.EndsWith(".uprojectdirs", StringComparison.InvariantCultureIgnoreCase))
+				if (SelectedProject.LocalPath != null && SelectedProject.LocalPath.EndsWith(".uprojectdirs", StringComparison.InvariantCultureIgnoreCase))
 				{
 					List<string> ProjectRoots = GetProjectRoots(SelectedProject.LocalPath);
 
@@ -1134,7 +1140,7 @@ namespace UnrealGameSync
 				}
 				else
 				{
-					ParseUncontrolledChangelistsPersistencyFile(Path.GetDirectoryName(SelectedProject.LocalPath), UncontrolledFiles);
+					ParseUncontrolledChangelistsPersistencyFile(Path.GetDirectoryName(SelectedProject.LocalPath)!, UncontrolledFiles);
 				}
 
 				ClobberWindow Window = new ClobberWindow(Context.ClobberFiles, UncontrolledFiles);
@@ -1274,14 +1280,14 @@ namespace UnrealGameSync
 			Dictionary<int, ChangesRecord> ChangeNumberToSummary = Changes.ToDictionary(x => x.Number, x => x);
 			for (int Idx = BuildList.Items.Count - 1; Idx >= 0; Idx--)
 			{
-				ChangesRecord Change = BuildList.Items[Idx].Tag as ChangesRecord;
+				ChangesRecord? Change = BuildList.Items[Idx].Tag as ChangesRecord;
 				if (Change != null)
 				{
-					ChangesRecord Summary;
+					ChangesRecord? Summary;
 					if (ChangeNumberToSummary.TryGetValue(Change.Number, out Summary))
 					{
 						// Update
-						BuildList.Items[Idx].SubItems[DescriptionColumn.Index].Text = Change.Description.Replace('\n', ' ');
+						BuildList.Items[Idx].SubItems[DescriptionColumn.Index].Text = Change.Description!.Replace('\n', ' ');
 						ChangeNumberToSummary.Remove(Change.Number);
 					}
 					else
@@ -1299,7 +1305,7 @@ namespace UnrealGameSync
 			}
 
 			// Figure out which group the expand item should be in
-			ListViewGroup NewExpandItemGroup = null;
+			ListViewGroup? NewExpandItemGroup = null;
 			for (int Idx = BuildList.Items.Count - 1; Idx >= 0; Idx--)
 			{
 				ListViewItem Item = BuildList.Items[Idx];
@@ -1358,7 +1364,7 @@ namespace UnrealGameSync
 					return NextGroup;
 				}
 
-				ChangesRecord LastChange = null;
+				ChangesRecord? LastChange = null;
 				for (int Idx = NextGroup.Items.Count - 1; Idx >= 0 && LastChange == null; Idx--)
 				{
 					LastChange = NextGroup.Items[Idx].Tag as ChangesRecord;
@@ -1396,7 +1402,7 @@ namespace UnrealGameSync
 			Columns[ChangeColumn.Index] = String.Format("{0}", Change.Number);
 
 			Columns[TimeColumn.Index] = DisplayTime.ToShortTimeString();
-			string UserName = FormatUserName(Change.User);
+			string UserName = FormatUserName(Change.User!);
 
 			// If annotate robomerge is on, add a prefix to the user name
 			if (Settings.bAnnotateRobmergeChanges && IsRobomergeChange(Change))
@@ -1405,7 +1411,7 @@ namespace UnrealGameSync
 			}
 
 			Columns[AuthorColumn.Index] = UserName;
-			Columns[DescriptionColumn.Index] = Change.Description.Replace('\n', ' ');
+			Columns[DescriptionColumn.Index] = Change.Description!.Replace('\n', ' ');
 
 			for (int ColumnIdx = 1; ColumnIdx < BuildList.Columns.Count; ColumnIdx++)
 			{
@@ -1416,7 +1422,7 @@ namespace UnrealGameSync
 			int GroupInsertIdx = 0;
 			for (; GroupInsertIdx < Group.Items.Count; GroupInsertIdx++)
 			{
-				ChangesRecord OtherChange = Group.Items[GroupInsertIdx].Tag as ChangesRecord;
+				ChangesRecord? OtherChange = Group.Items[GroupInsertIdx].Tag as ChangesRecord;
 				if (OtherChange == null || Change.Number >= OtherChange.Number)
 				{
 					break;
@@ -1437,14 +1443,14 @@ namespace UnrealGameSync
 		{
 			// If hiding robomerge, filter out based on workspace name. Note - don't look at the description because we
 			// *do* want to see conflicts that were manually merged by someone
-			return Change.Client.IndexOf("ROBOMERGE", StringComparison.OrdinalIgnoreCase) >= 0;
+			return Change.Client!.IndexOf("ROBOMERGE", StringComparison.OrdinalIgnoreCase) >= 0;
 		}
 
 		bool ShouldShowChange(ChangesRecord Change, string[] ExcludeChanges)
 		{
 			if (ProjectSettings.FilterType != FilterType.None)
 			{
-				PerforceChangeDetails Details;
+				PerforceChangeDetails? Details;
 				if (!PerforceMonitor.TryGetChangeDetails(Change.Number, out Details))
 				{
 					return false;
@@ -1462,7 +1468,7 @@ namespace UnrealGameSync
 			// if filtering by user, only show changes where the author contains the filter text
 			if (!string.IsNullOrEmpty(this.AuthorFilterText)
 				&& this.AuthorFilterText != AuthorFilterPlaceholderText
-				&& Change.User.IndexOf(this.AuthorFilterText, StringComparison.OrdinalIgnoreCase) < 0)
+				&& Change.User!.IndexOf(this.AuthorFilterText, StringComparison.OrdinalIgnoreCase) < 0)
 			{
 				return false;
 			}
@@ -1503,7 +1509,7 @@ namespace UnrealGameSync
 					}
 				}
 
-				if (String.Compare(Change.User, "buildmachine", true) == 0 && Change.Description.IndexOf("lightmaps", StringComparison.InvariantCultureIgnoreCase) == -1)
+				if (String.Compare(Change.User, "buildmachine", true) == 0 && Change.Description!.IndexOf("lightmaps", StringComparison.InvariantCultureIgnoreCase) == -1)
 				{
 					return false;
 				}
@@ -1526,7 +1532,7 @@ namespace UnrealGameSync
 
 		string GetIssuesApiUrl()
 		{
-			string IssuesApiUrl;
+			string? IssuesApiUrl;
 			if (!TryGetProjectSetting(PerforceMonitor.LatestProjectConfigFile, "IssuesApiUrl", out IssuesApiUrl))
 			{
 				IssuesApiUrl = ApiUrl;
@@ -1547,7 +1553,7 @@ namespace UnrealGameSync
 			}
 
 			// Update the column settings first, since this may affect the badges we hide
-			UpdateColumnSettings();
+			UpdateColumnSettings(false);
 
 			// Clear the badge size cache
 			BadgeLabelToSize.Clear();
@@ -1556,7 +1562,7 @@ namespace UnrealGameSync
 			Dictionary<string, string> BadgeNameToGroup = new Dictionary<string, string>();
 
 			// Read the group mappings from the config file
-			string GroupDefinitions;
+			string? GroupDefinitions;
 			if (TryGetProjectSetting(PerforceMonitor.LatestProjectConfigFile, "BadgeGroups", out GroupDefinitions))
 			{
 				string[] GroupDefinitionsArray = GroupDefinitions.Split('\n');
@@ -1571,9 +1577,9 @@ namespace UnrealGameSync
 			}
 
 			// Add a dummy group for any other badges we have
-			foreach (ListViewItem Item in BuildList.Items)
+			foreach (ListViewItem? Item in BuildList.Items)
 			{
-				if (Item.Tag != null)
+				if (Item != null && Item.Tag != null)
 				{
 					EventSummary Summary = EventMonitor.GetSummaryForChange(((ChangesRecord)Item.Tag).Number);
 					if (Summary != null)
@@ -1720,7 +1726,7 @@ namespace UnrealGameSync
 					Guid UniqueId;
 					if (Guid.TryParse(Object.GetValue("UniqueId", ""), out UniqueId))
 					{
-						WorkspaceSyncCategory Category;
+						WorkspaceSyncCategory? Category;
 						if (!UniqueIdToCategory.TryGetValue(UniqueId, out Category))
 						{
 							Category = new WorkspaceSyncCategory(UniqueId);
@@ -1765,7 +1771,7 @@ namespace UnrealGameSync
 				return ProjectRoots;
 			}
 
-			string UProjectDirsRoot = Path.GetDirectoryName(InUProjectDirsPath);
+			string UProjectDirsRoot = Path.GetDirectoryName(InUProjectDirsPath)!;
 
 			foreach (string Line in File.ReadLines(InUProjectDirsPath))
 			{
@@ -1910,7 +1916,7 @@ namespace UnrealGameSync
 					if (LastChangeByCurrentUserOfType > 0 && LastBuild.ChangeNumber >= LastChangeByCurrentUserOfType)
 					{
 						// And check that there wasn't a successful build after we submitted (if there was, we're in the clear)
-						BadgeData LastSuccessfulBuild;
+						BadgeData? LastSuccessfulBuild;
 						if (!TypeToLastSucceededBuild.TryGetValue(LastBuild.BuildType, out LastSuccessfulBuild) || LastSuccessfulBuild.ChangeNumber < LastChangeByCurrentUserOfType)
 						{
 							// Add it to the list of notifications
@@ -1975,9 +1981,9 @@ namespace UnrealGameSync
 
 		class ExpandRowLayout
 		{
-			public string MainText;
+			public string? MainText;
 			public Rectangle MainRect;
-			public string LinkText;
+			public string? LinkText;
 			public Rectangle LinkRect;
 		}
 
@@ -2072,20 +2078,20 @@ namespace UnrealGameSync
 			}
 		}
 
-		private string GetArchiveKeyForChangeNumber(IArchiveInfo Archive, int ChangeNumber)
+		private string? GetArchiveKeyForChangeNumber(IArchiveInfo Archive, int ChangeNumber)
 		{
-			string ArchivePath;
+			string? ArchivePath;
 
-			Dictionary<int, string> ChangeNumberToArchivePath;
+			Dictionary<int, string?>? ChangeNumberToArchivePath;
 			if (!ArchiveToChangeNumberToArchiveKey.TryGetValue(Archive.Name, out ChangeNumberToArchivePath))
 			{
-				ChangeNumberToArchivePath = new Dictionary<int, string>();
+				ChangeNumberToArchivePath = new Dictionary<int, string?>();
 				ArchiveToChangeNumberToArchiveKey[Archive.Name] = ChangeNumberToArchivePath;
 			}
 
 			if (!ChangeNumberToArchivePath.TryGetValue(ChangeNumber, out ArchivePath))
 			{
-				PerforceChangeDetails Details;
+				PerforceChangeDetails? Details;
 				if (PerforceMonitor.TryGetChangeDetails(ChangeNumber, out Details))
 				{
 					// Try to get the archive for this CL
@@ -2123,7 +2129,7 @@ namespace UnrealGameSync
 
 		private ChangeLayoutInfo GetChangeLayoutInfo(ChangesRecord Change)
 		{
-			ChangeLayoutInfo LayoutInfo;
+			ChangeLayoutInfo? LayoutInfo;
 			if (!ChangeNumberToLayoutInfo.TryGetValue(Change.Number, out LayoutInfo))
 			{
 				LayoutInfo = new ChangeLayoutInfo();
@@ -2182,11 +2188,11 @@ namespace UnrealGameSync
 			Format.FormatFlags = StringFormatFlags.NoWrap;
 			Format.Trimming = StringTrimming.EllipsisCharacter;
 
-			Font CurrentFont = (Change.Number == Workspace.PendingChangeNumber || Change.Number == Workspace.CurrentChangeNumber) ? SelectedBuildFont : BuildFont;
+			Font CurrentFont = (Change.Number == Workspace.PendingChangeNumber || Change.Number == Workspace.CurrentChangeNumber) ? SelectedBuildFont! : BuildFont!;
 
 			bool bAllowSync = CanSyncChange(Change.Number);
 			int BadgeAlpha = bAllowSync ? 255 : 128;
-			Color TextColor = (bAllowSync || Change.Number == Workspace.PendingChangeNumber || Change.Number == Workspace.CurrentChangeNumber || (WorkspaceState != null && WorkspaceState.AdditionalChangeNumbers.Contains(Change.Number))) ? SystemColors.WindowText : Blend(SystemColors.Window, SystemColors.WindowText, 0.25f);
+			Color TextColor = (bAllowSync || Change.Number == Workspace.PendingChangeNumber || Change.Number == Workspace.CurrentChangeNumber || WorkspaceState.AdditionalChangeNumbers.Contains(Change.Number)) ? SystemColors.WindowText : Blend(SystemColors.Window, SystemColors.WindowText, 0.25f);
 
 			const int FadeRange = 6;
 			if (e.ItemIndex >= BuildList.Items.Count - FadeRange && NumChanges >= PerforceMonitor.CurrentMaxChanges && !IsBisectModeEnabled())
@@ -2229,7 +2235,7 @@ namespace UnrealGameSync
 						{
 							QualityIcon = new Rectangle(0, 0, 0, 0);
 						}
-						if (Entry.State == BisectState.Pass)
+						else if (Entry.State == BisectState.Pass)
 						{
 							QualityIcon = BisectPassIcon;
 						}
@@ -2369,7 +2375,7 @@ namespace UnrealGameSync
 						MaxX = SyncBadgeRectangle.Left - (int)(2 * DpiScaleX);
 					}
 
-					string SummaryText;
+					string? SummaryText;
 					if (WorkspaceState.LastSyncChangeNumber == -1 || WorkspaceState.LastSyncChangeNumber != Change.Number || !GetLastUpdateMessage(WorkspaceState.LastSyncResult, WorkspaceState.LastSyncResultMessage, out SummaryText))
 					{
 						StringBuilder SummaryTextBuilder = new StringBuilder();
@@ -2409,7 +2415,7 @@ namespace UnrealGameSync
 				{
 					ChangeLayoutInfo Layout = GetChangeLayoutInfo(Change);
 
-					List<BadgeInfo> Badges;
+					List<BadgeInfo>? Badges;
 					if (Layout.CustomBadges.TryGetValue(Column.Text, out Badges) && Badges.Count > 0)
 					{
 						e.Graphics.IntersectClip(e.Bounds);
@@ -2500,9 +2506,9 @@ namespace UnrealGameSync
 
 		private List<BadgeInfo> CreateDescriptionBadges(ChangesRecord Change)
 		{
-			string Description = Change.Description;
+			string Description = Change.Description ?? String.Empty;
 
-			PerforceChangeDetails Details;
+			PerforceChangeDetails? Details;
 			if (PerforceMonitor.TryGetChangeDetails(Change.Number, out Details))
 			{
 				Description = Details.Description;
@@ -2519,40 +2525,43 @@ namespace UnrealGameSync
 					foreach (string BadgeDefinition in BadgeDefinitions.Distinct())
 					{
 						ConfigObject BadgeDefinitionObject = new ConfigObject(BadgeDefinition);
-						string Pattern = BadgeDefinitionObject.GetValue("Pattern", null);
-						string Name = BadgeDefinitionObject.GetValue("Name", null);
-						string Group = BadgeDefinitionObject.GetValue("Group", null);
+						string? Pattern = BadgeDefinitionObject.GetValue("Pattern", null);
+						string? Name = BadgeDefinitionObject.GetValue("Name", null);
+						string? Group = BadgeDefinitionObject.GetValue("Group", null);
 						string Color = BadgeDefinitionObject.GetValue("Color", "#909090");
 						string HoverColor = BadgeDefinitionObject.GetValue("HoverColor", "#b0b0b0");
-						string Url = BadgeDefinitionObject.GetValue("Url", null);
-						string Arguments = BadgeDefinitionObject.GetValue("Arguments", null);
+						string? Url = BadgeDefinitionObject.GetValue("Url", null);
+						string? Arguments = BadgeDefinitionObject.GetValue("Arguments", null);
 						if (!String.IsNullOrEmpty(Name) && !String.IsNullOrEmpty(Pattern))
 						{
-							foreach (Match MatchResult in Regex.Matches(Description, Pattern, RegexOptions.Multiline))
+							foreach (Match? MatchResult in Regex.Matches(Description, Pattern, RegexOptions.Multiline))
 							{
-								Color BadgeColor = System.Drawing.ColorTranslator.FromHtml(Color);
-								Color HoverBadgeColor = System.Drawing.ColorTranslator.FromHtml(HoverColor);
+								if (MatchResult != null)
+								{
+									Color BadgeColor = System.Drawing.ColorTranslator.FromHtml(Color);
+									Color HoverBadgeColor = System.Drawing.ColorTranslator.FromHtml(HoverColor);
 
-								string UniqueId = String.IsNullOrEmpty(Url) ? null : String.Format("Description:{0}:{1}", Change.Number, Badges.Count);
+									string? UniqueId = String.IsNullOrEmpty(Url) ? null : String.Format("Description:{0}:{1}", Change.Number, Badges.Count);
 
-								string ExpandedUrl = ReplaceRegexMatches(Url, MatchResult);
-								string ExpandedArguments = ReplaceRegexMatches(Arguments, MatchResult);
+									string? ExpandedUrl = ReplaceRegexMatches(Url, MatchResult);
+									string? ExpandedArguments = ReplaceRegexMatches(Arguments, MatchResult);
 
-								Action ClickHandler;
-								if (String.IsNullOrEmpty(ExpandedUrl))
-								{
-									ClickHandler = null;
-								}
-								else if (String.IsNullOrEmpty(ExpandedArguments))
-								{
-									ClickHandler = () => SafeProcessStart(ExpandedUrl);
-								}
-								else
-								{
-									ClickHandler = () => SafeProcessStart(ExpandedUrl, ExpandedArguments);
-								}
+									Action? ClickHandler;
+									if (String.IsNullOrEmpty(ExpandedUrl))
+									{
+										ClickHandler = null;
+									}
+									else if (String.IsNullOrEmpty(ExpandedArguments))
+									{
+										ClickHandler = () => SafeProcessStart(ExpandedUrl);
+									}
+									else
+									{
+										ClickHandler = () => SafeProcessStart(ExpandedUrl, ExpandedArguments);
+									}
 
-								Badges.Add(new BadgeInfo(ReplaceRegexMatches(Name, MatchResult), Group, UniqueId, BadgeColor, HoverBadgeColor, ClickHandler));
+									Badges.Add(new BadgeInfo(ReplaceRegexMatches(Name, MatchResult), Group, UniqueId, BadgeColor, HoverBadgeColor, ClickHandler));
+								}
 							}
 						}
 					}
@@ -2632,7 +2641,8 @@ namespace UnrealGameSync
 			UpdateStatusPanel();
 		}
 
-		private string ReplaceRegexMatches(string Text, Match MatchResult)
+		[return: NotNullIfNotNull("Text")]
+		private string? ReplaceRegexMatches(string? Text, Match MatchResult)
 		{
 			if (Text != null)
 			{
@@ -2649,7 +2659,7 @@ namespace UnrealGameSync
 		{
 			List<BadgeInfo> Badges = new List<BadgeInfo>();
 
-			PerforceChangeDetails Details;
+			PerforceChangeDetails? Details;
 			if (PerforceMonitor.TryGetChangeDetails(ChangeNumber, out Details))
 			{
 				if (Details.bContainsCode)
@@ -2670,7 +2680,7 @@ namespace UnrealGameSync
 			return Badges;
 		}
 
-		private bool TryGetProjectSetting(ConfigFile ProjectConfigFile, string Name, out string Value)
+		private bool TryGetProjectSetting(ConfigFile ProjectConfigFile, string Name, [NotNullWhen(true)] out string? Value)
 		{
 			string Path = SelectedProjectIdentifier;
 			for(; ;)
@@ -2678,7 +2688,7 @@ namespace UnrealGameSync
 				ConfigSection ProjectSection = ProjectConfigFile.FindSection(Path);
 				if (ProjectSection != null)
 				{
-					string NewValue = ProjectSection.GetValue(Name, null);
+					string? NewValue = ProjectSection.GetValue(Name, null);
 					if (NewValue != null)
 					{
 						Value = NewValue;
@@ -2698,7 +2708,7 @@ namespace UnrealGameSync
 			ConfigSection DefaultSection = ProjectConfigFile.FindSection("Default");
 			if (DefaultSection != null)
 			{
-				string NewValue = DefaultSection.GetValue(Name, null);
+				string? NewValue = DefaultSection.GetValue(Name, null);
 				if (NewValue != null)
 				{
 					Value = NewValue;
@@ -2710,9 +2720,9 @@ namespace UnrealGameSync
 			return false;
 		}
 
-		private bool TryGetProjectSetting(ConfigFile ProjectConfigFile, string Name, out string[] Values)
+		private bool TryGetProjectSetting(ConfigFile ProjectConfigFile, string Name, [NotNullWhen(true)] out string[]? Values)
 		{
-			string ValueList;
+			string? ValueList;
 			if (TryGetProjectSetting(ProjectConfigFile, Name, out ValueList))
 			{
 				Values = ValueList.Split('\n').Select(x => x.Trim()).Where(x => x.Length > 0).ToArray();
@@ -2725,9 +2735,9 @@ namespace UnrealGameSync
 			}
 		}
 
-		private bool TryGetProjectSetting(ConfigFile ProjectConfigFile, string Name, string LegacyName, out string Value)
+		private bool TryGetProjectSetting(ConfigFile ProjectConfigFile, string Name, string LegacyName, [NotNullWhen(true)] out string? Value)
 		{
-			string NewValue;
+			string? NewValue;
 			if (TryGetProjectSetting(ProjectConfigFile, Name, out NewValue))
 			{
 				Value = NewValue;
@@ -2761,7 +2771,7 @@ namespace UnrealGameSync
 				// Add all the badges, sorted by group
 				foreach (KeyValuePair<string, string> BadgeNameAndGroup in BadgeNameAndGroupPairs)
 				{
-					BadgeData BadgeData;
+					BadgeData? BadgeData;
 					BadgeNameToBuildData.TryGetValue(BadgeNameAndGroup.Key, out BadgeData);
 
 					BadgeInfo BadgeInfo = CreateBadge(ChangeNumber, BadgeNameAndGroup.Key, BadgeNameAndGroup.Value, BadgeData);
@@ -2779,7 +2789,7 @@ namespace UnrealGameSync
 			return Badges;
 		}
 
-		private BadgeInfo CreateBadge(int ChangeNumber, string BadgeName, string BadgeGroup, BadgeData BadgeData)
+		private BadgeInfo CreateBadge(int ChangeNumber, string BadgeName, string BadgeGroup, BadgeData? BadgeData)
 		{
 			string BadgeLabel = BadgeName;
 			Color BadgeColor = Color.FromArgb(0, Color.White);
@@ -2792,7 +2802,7 @@ namespace UnrealGameSync
 
 			Color HoverBadgeColor = Color.FromArgb(BadgeColor.A, Math.Min(BadgeColor.R + 32, 255), Math.Min(BadgeColor.G + 32, 255), Math.Min(BadgeColor.B + 32, 255));
 
-			Action ClickHandler;
+			Action? ClickHandler;
 			if (BadgeData == null || String.IsNullOrEmpty(BadgeData.Url))
 			{
 				ClickHandler = null;
@@ -2903,7 +2913,7 @@ namespace UnrealGameSync
 			if (!BadgeLabelToSize.TryGetValue(BadgeText, out BadgeSize))
 			{
 				Size LabelSize = TextRenderer.MeasureText(BadgeText, BadgeFont);
-				int BadgeHeight = BadgeFont.Height + 1;
+				int BadgeHeight = BadgeFont!.Height + 1;
 
 				BadgeSize = new Size(LabelSize.Width + BadgeHeight - 4, BadgeHeight);
 				BadgeLabelToSize[BadgeText] = BadgeSize;
@@ -2911,7 +2921,7 @@ namespace UnrealGameSync
 			return BadgeSize;
 		}
 
-		private static bool GetLastUpdateMessage(WorkspaceUpdateResult Result, string ResultMessage, out string Message)
+		private static bool GetLastUpdateMessage(WorkspaceUpdateResult Result, string? ResultMessage, [NotNullWhen(true)] out string? Message)
 		{
 			if (Result != WorkspaceUpdateResult.Success && ResultMessage != null)
 			{
@@ -2921,7 +2931,7 @@ namespace UnrealGameSync
 			return GetGenericLastUpdateMessage(Result, out Message);
 		}
 
-		private static bool GetGenericLastUpdateMessage(WorkspaceUpdateResult Result, out string Message)
+		private static bool GetGenericLastUpdateMessage(WorkspaceUpdateResult Result, [NotNullWhen(true)] out string? Message)
 		{
 			switch (Result)
 			{
@@ -2972,7 +2982,7 @@ namespace UnrealGameSync
 
 		private static void AppendItemList(StringBuilder Builder, string Separator, string Format, IEnumerable<string> Items)
 		{
-			string ItemList = FormatItemList(Format, Items);
+			string? ItemList = FormatItemList(Format, Items);
 			if (ItemList != null)
 			{
 				if (Builder.Length > 0)
@@ -2983,7 +2993,7 @@ namespace UnrealGameSync
 			}
 		}
 
-		private static string FormatItemList(string Format, IEnumerable<string> Items)
+		private static string? FormatItemList(string Format, IEnumerable<string> Items)
 		{
 			string[] ItemsArray = Items.Distinct().ToArray();
 			if (ItemsArray.Length == 0)
@@ -3102,7 +3112,7 @@ namespace UnrealGameSync
 
 		private class ReceiptJsonObject
 		{
-			public string Launch { get; set; }
+			public string? Launch { get; set; }
 		}
 
 		private FileReference GetEditorExePath(BuildConfig Config)
@@ -3119,7 +3129,7 @@ namespace UnrealGameSync
 						string Text = FileReference.ReadAllText(ReceiptFileName);
 						ReceiptJsonObject Receipt = JsonSerializer.Deserialize<ReceiptJsonObject>(Text, Program.DefaultJsonSerializerOptions);
 
-						string LaunchFileName = Receipt.Launch;
+						string? LaunchFileName = Receipt.Launch;
 						if (LaunchFileName != null)
 						{
 							LaunchFileName = LaunchFileName.Replace("$(EngineDir)", DirectoryReference.Combine(BranchDirectoryName, "Engine").FullName);
@@ -3251,7 +3261,7 @@ namespace UnrealGameSync
 			return PossiblePaths;
 		}
 
-		private void TimerCallback(object Sender, EventArgs Args)
+		private void TimerCallback(object? Sender, EventArgs Args)
 		{
 			Tuple<string, float> Progress = Workspace.CurrentProgress;
 			if (Progress != null && Progress.Item2 > 0.0f)
@@ -3269,15 +3279,15 @@ namespace UnrealGameSync
 			BuildList.Refresh();
 		}
 
-		private string GetDefaultIssueFilter()
+		private string? GetDefaultIssueFilter()
 		{
-			string BuildHealthProject;
+			string? BuildHealthProject;
 			if (TryGetProjectSetting(PerforceMonitor.LatestProjectConfigFile, "BuildHealthProject", out BuildHealthProject))
 			{
 				return BuildHealthProject;
 			}
 
-			string DefaultIssueFilter;
+			string? DefaultIssueFilter;
 			if (TryGetProjectSetting(PerforceMonitor.LatestProjectConfigFile, "DefaultIssueFilter", out DefaultIssueFilter))
 			{
 				return DefaultIssueFilter;
@@ -3290,13 +3300,13 @@ namespace UnrealGameSync
 		{
 			Dictionary<string, Func<IssueData, bool>> CustomFilters = new Dictionary<string, Func<IssueData, bool>>();
 
-			string BuildHealthProject;
+			string? BuildHealthProject;
 			if (TryGetProjectSetting(PerforceMonitor.LatestProjectConfigFile, "BuildHealthProject", out BuildHealthProject))
 			{
 				CustomFilters[BuildHealthProject] = x => !String.IsNullOrEmpty(x.Project) && x.Project.Equals(BuildHealthProject, StringComparison.OrdinalIgnoreCase);
 			}
 
-			string[] Filters;
+			string[]? Filters;
 			if (TryGetProjectSetting(PerforceMonitor.LatestProjectConfigFile, "IssueFilters", out Filters))
 			{
 				foreach (string Filter in Filters)
@@ -3305,12 +3315,12 @@ namespace UnrealGameSync
 					{
 						ConfigObject Config = new ConfigObject(Filter);
 
-						string Name = Config.GetValue("Name", null);
-						string Pattern = Config.GetValue("Pattern", null);
+						string? Name = Config.GetValue("Name", null);
+						string? Pattern = Config.GetValue("Pattern", null);
 
 						if(Name != null && Pattern != null)
 						{
-							string DepotName;
+							string? DepotName;
 							if (PerforceUtils.TryGetDepotName(StreamName, out DepotName))
 							{
 								Pattern = Regex.Replace(Pattern, @"\$\(Depot\)", DepotName, RegexOptions.IgnoreCase);
@@ -3344,16 +3354,18 @@ namespace UnrealGameSync
 
 		void RunTool(ToolDefinition Tool, ToolLink Link)
 		{
-			DirectoryReference ToolDir = Owner.ToolUpdateMonitor.GetToolPath(Tool.Name);
-
-			Dictionary<string, string> Variables = GetWorkspaceVariables(Workspace.CurrentChangeNumber);
-			Variables["ToolDir"] = ToolDir.FullName;
+			DirectoryReference? ToolDir = Owner.ToolUpdateMonitor.GetToolPath(Tool.Name);
+			if (ToolDir != null)
+			{
+				Dictionary<string, string> Variables = GetWorkspaceVariables(Workspace.CurrentChangeNumber);
+				Variables["ToolDir"] = ToolDir.FullName;
 
-			string FileName = Utility.ExpandVariables(Link.FileName, Variables);
-			string Arguments = Utility.ExpandVariables(Link.Arguments, Variables);
-			string WorkingDir = Utility.ExpandVariables(Link.WorkingDir ?? ToolDir.FullName, Variables);
+				string FileName = Utility.ExpandVariables(Link.FileName, Variables);
+				string Arguments = Utility.ExpandVariables(Link.Arguments ?? String.Empty, Variables);
+				string WorkingDir = Utility.ExpandVariables(Link.WorkingDir ?? ToolDir.FullName, Variables);
 
-			SafeProcessStart(FileName, Arguments, WorkingDir);
+				SafeProcessStart(FileName, Arguments, WorkingDir);
+			}
 		}
 
 		private void UpdateStatusPanel_CrossThread()
@@ -3479,7 +3491,7 @@ namespace UnrealGameSync
 					ProgramsLine.AddText("  |  ");
 				}
 
-				string[] SdkInfoEntries;
+				string[]? SdkInfoEntries;
 				if (TryGetProjectSetting(PerforceMonitor.LatestProjectConfigFile, "SdkInfo", out SdkInfoEntries))
 				{
 					ProgramsLine.AddLink("SDK Info", FontStyle.Regular, () => { ShowRequiredSdkInfo(); });
@@ -3547,7 +3559,7 @@ namespace UnrealGameSync
 				// Get the summary of the last sync
 				if (WorkspaceState.LastSyncChangeNumber > 0)
 				{
-					string SummaryText;
+					string? SummaryText;
 					if (WorkspaceState.LastSyncChangeNumber == Workspace.CurrentChangeNumber && WorkspaceState.LastSyncResult == WorkspaceUpdateResult.Success && WorkspaceState.LastSyncTime.HasValue)
 					{
 						Lines.Add(new StatusLine() { LineHeight = 0.5f });
@@ -3587,26 +3599,26 @@ namespace UnrealGameSync
 				}
 			}
 
-			StatusLine Caption = null;
+			StatusLine? Caption = null;
 			if (StreamName != null && !Workspace.IsBusy())
 			{
 				Caption = new StatusLine();
 				Caption.AddLink(StreamName + "\u25BE", FontStyle.Bold, (P, R) => { SelectOtherStream(R); });
 			}
 
-			StatusLine Alert = null;
+			StatusLine? Alert = null;
 			Color? TintColor = null;
 
 			ConfigFile ProjectConfigFile = PerforceMonitor.LatestProjectConfigFile;
 			if (ProjectConfigFile != null)
 			{
-				string Message;
+				string? Message;
 				if (TryGetProjectSetting(ProjectConfigFile, "Message", out Message))
 				{
 					Alert = CreateStatusLineFromMarkdown(Message);
 				}
 
-				string StatusPanelColor;
+				string? StatusPanelColor;
 				if (TryGetProjectSetting(ProjectConfigFile, "StatusPanelColor", out StatusPanelColor))
 				{
 					TintColor = System.Drawing.ColorTranslator.FromHtml(StatusPanelColor);
@@ -3638,8 +3650,9 @@ namespace UnrealGameSync
 				BuildHealthContextMenu.Items.RemoveAt(MinSeparatorIdx + 1);
 			}
 
-			Func<IssueData, bool> Predicate;
-			if(!GetCustomIssueFilters().TryGetValue(GetDefaultIssueFilter(), out Predicate))
+			Func<IssueData, bool>? Predicate = null;
+			string? DefaultIssueFilter = GetDefaultIssueFilter();
+			if (DefaultIssueFilter == null || !GetCustomIssueFilters().TryGetValue(DefaultIssueFilter, out Predicate))
 			{
 				Predicate = x => true;
 			}
@@ -3737,14 +3750,14 @@ namespace UnrealGameSync
 
 		private void BuildHealthContextMenu_Browse_Click(object sender, EventArgs e)
 		{
-			string DefaultFilter = GetDefaultIssueFilter();
+			string? DefaultFilter = GetDefaultIssueFilter();
 			Dictionary<string, Func<IssueData, bool>> CustomFilters = GetCustomIssueFilters();
 			IssueBrowserWindow.Show(ParentForm, IssueMonitor, PerforceSettings, GetServerTimeOffset(), ServiceProvider, StreamName, CustomFilters, DefaultFilter);
 		}
 
 		private void BuildHealthContextMenu_Settings_Click(object sender, EventArgs e)
 		{
-			string BuildHealthProject;
+			string? BuildHealthProject;
 			TryGetProjectSetting(PerforceMonitor.LatestProjectConfigFile, "BuildHealthProject", out BuildHealthProject);
 
 			IssueSettingsWindow IssueSettings = new IssueSettingsWindow(Settings, BuildHealthProject ?? "");
@@ -3756,7 +3769,7 @@ namespace UnrealGameSync
 
 		private void ShowRequiredSdkInfo()
 		{
-			string[] SdkInfoEntries;
+			string[]? SdkInfoEntries;
 			if (TryGetProjectSetting(PerforceMonitor.LatestProjectConfigFile, "SdkInfo", out SdkInfoEntries))
 			{
 				Dictionary<string, string> Variables = GetWorkspaceVariables(-1);
@@ -3866,7 +3879,7 @@ namespace UnrealGameSync
 			bool bShownContextMenu = false;
 			if (StreamName != null)
 			{
-				IReadOnlyList<string> OtherStreamNames = Workspace.ProjectStreamFilter;
+				IReadOnlyList<string>? OtherStreamNames = Workspace.ProjectStreamFilter;
 				if (OtherStreamNames != null)
 				{
 					StreamContextMenu.Items.Clear();
@@ -3909,7 +3922,7 @@ namespace UnrealGameSync
 
 		private void SelectOtherStreamDialog()
 		{
-			string NewStreamName;
+			string? NewStreamName;
 			if (SelectStreamWindow.ShowModal(this, PerforceSettings, StreamName, ServiceProvider, out NewStreamName))
 			{
 				SelectStream(NewStreamName);
@@ -3932,14 +3945,14 @@ namespace UnrealGameSync
 						{
 							if (Idx == 1 || !await Perforce.OpenedAsync(OpenedOptions.None, FileSpecList.Any, CancellationToken).AnyAsync())
 							{
-								await Perforce.SwitchClientToStreamAsync(Perforce.Settings.ClientName, NewStreamName, SwitchClientOptions.IgnoreOpenFiles, CancellationToken);
+								await Perforce.SwitchClientToStreamAsync(Perforce.Settings.ClientName!, NewStreamName, SwitchClientOptions.IgnoreOpenFiles, CancellationToken);
 								return true;
 							}
 							return false;
 						};
 
-						ModalTask<bool> SwitchTask = PerforceModalTask.Execute<bool>(this, "Switching streams", "Please wait...", PerforceSettings, SwitchFunc, Logger);
-						if (!SwitchTask.Succeeded)
+						ModalTask<bool>? SwitchTask = PerforceModalTask.Execute<bool>(this, "Switching streams", "Please wait...", PerforceSettings, SwitchFunc, Logger);
+						if (SwitchTask == null || !SwitchTask.Succeeded)
 						{
 							break;
 						}
@@ -3957,10 +3970,10 @@ namespace UnrealGameSync
 
 		private void ViewLastSyncStatus()
 		{
-			string SummaryText;
+			string? SummaryText;
 			if (GetLastUpdateMessage(WorkspaceState.LastSyncResult, WorkspaceState.LastSyncResultMessage, out SummaryText))
 			{
-				string CaptionText;
+				string? CaptionText;
 				if (!GetGenericLastUpdateMessage(WorkspaceState.LastSyncResult, out CaptionText))
 				{
 					CaptionText = "Sync error";
@@ -3994,14 +4007,17 @@ namespace UnrealGameSync
 
 			PendingSelectedChangeNumber = -1;
 
-			foreach (ListViewItem Item in BuildList.Items)
+			foreach (ListViewItem? Item in BuildList.Items)
 			{
-				ChangesRecord Summary = (ChangesRecord)Item.Tag;
-				if (Summary != null && Summary.Number <= ChangeNumber)
+				if (Item != null)
 				{
-					Item.Selected = true;
-					Item.EnsureVisible();
-					return;
+					ChangesRecord Summary = (ChangesRecord)Item.Tag;
+					if (Summary != null && Summary.Number <= ChangeNumber)
+					{
+						Item.Selected = true;
+						Item.EnsureVisible();
+						return;
+					}
 				}
 			}
 
@@ -4114,7 +4130,7 @@ namespace UnrealGameSync
 								Point BuildListLocation = GetBadgeListLocation(LayoutInfo.BuildBadges, HitTest.SubItem.Bounds, HorizontalAlign.Center, VerticalAlignment.Middle);
 								BuildListLocation.X = Math.Max(BuildListLocation.X, HitTest.SubItem.Bounds.Left);
 
-								BadgeInfo BadgeInfo = HitTestBadge(Args.Location, LayoutInfo.BuildBadges, BuildListLocation);
+								BadgeInfo? BadgeInfo = HitTestBadge(Args.Location, LayoutInfo.BuildBadges, BuildListLocation);
 								if (BadgeInfo != null && BadgeInfo.ClickHandler != null)
 								{
 									BadgeInfo.ClickHandler();
@@ -4128,12 +4144,12 @@ namespace UnrealGameSync
 							{
 								ChangeLayoutInfo LayoutInfo = GetChangeLayoutInfo((ChangesRecord)HitTest.Item.Tag);
 
-								List<BadgeInfo> Badges;
+								List<BadgeInfo>? Badges;
 								if (LayoutInfo.CustomBadges.TryGetValue(CustomColumn.Text, out Badges) && Badges.Count > 0)
 								{
 									Point ListLocation = GetBadgeListLocation(Badges, HitTest.SubItem.Bounds, HorizontalAlign.Center, VerticalAlignment.Middle);
 
-									BadgeInfo BadgeInfo = HitTestBadge(Args.Location, Badges, ListLocation);
+									BadgeInfo? BadgeInfo = HitTestBadge(Args.Location, Badges, ListLocation);
 									if (BadgeInfo != null && BadgeInfo.ClickHandler != null)
 									{
 										BadgeInfo.ClickHandler();
@@ -4219,7 +4235,7 @@ namespace UnrealGameSync
 							CustomToolEnd--;
 						}
 
-						ConfigFile ProjectConfigFile = PerforceMonitor.LatestProjectConfigFile;
+						ConfigFile? ProjectConfigFile = PerforceMonitor.LatestProjectConfigFile;
 						if (ProjectConfigFile != null)
 						{
 							Dictionary<string, string> Variables = GetWorkspaceVariables(ContextMenuChange.Number);
@@ -4229,9 +4245,9 @@ namespace UnrealGameSync
 							{
 								ConfigObject Object = new ConfigObject(ChangeContextMenuEntry);
 
-								string Label = Object.GetValue("Label");
-								string Execute = Object.GetValue("Execute");
-								string Arguments = Object.GetValue("Arguments");
+								string? Label = Object.GetValue("Label");
+								string? Execute = Object.GetValue("Execute");
+								string? Arguments = Object.GetValue("Arguments");
 
 								if (Label != null && Execute != null)
 								{
@@ -4249,8 +4265,8 @@ namespace UnrealGameSync
 
 						BuildListContextMenu_CustomTool_End.Visible = (CustomToolEnd > CustomToolStart);
 
-						string SwarmURL;
-						BuildListContextMenu_ViewInSwarm.Visible = TryGetProjectSetting(ProjectConfigFile, "SwarmURL", out SwarmURL);
+						string? SwarmURL;
+						BuildListContextMenu_ViewInSwarm.Visible = ProjectConfigFile != null && TryGetProjectSetting(ProjectConfigFile, "SwarmURL", out SwarmURL);
 							
 						BuildListContextMenu.Show(BuildList, Args.Location);
 					}
@@ -4270,7 +4286,7 @@ namespace UnrealGameSync
 				BuildList.Invalidate();
 			}
 
-			string NewHoverBadgeUniqueId = null;
+			string? NewHoverBadgeUniqueId = null;
 			if (HitTest.Item != null && HitTest.Item.Tag is ChangesRecord)
 			{
 				int ColumnIndex = HitTest.Item.SubItems.IndexOf(HitTest.SubItem);
@@ -4291,7 +4307,7 @@ namespace UnrealGameSync
 						Point BuildListLocation = GetBadgeListLocation(LayoutInfo.BuildBadges, HitTest.SubItem.Bounds, HorizontalAlign.Center, VerticalAlignment.Middle);
 						BuildListLocation.X = Math.Max(BuildListLocation.X, HitTest.SubItem.Bounds.Left);
 
-						BadgeInfo Badge = HitTestBadge(e.Location, LayoutInfo.BuildBadges, BuildListLocation);
+						BadgeInfo? Badge = HitTestBadge(e.Location, LayoutInfo.BuildBadges, BuildListLocation);
 						NewHoverBadgeUniqueId = (Badge != null) ? Badge.UniqueId : null;
 
 						if (HoverBadgeUniqueId != NewHoverBadgeUniqueId)
@@ -4313,7 +4329,7 @@ namespace UnrealGameSync
 
 					ChangeLayoutInfo LayoutInfo = GetChangeLayoutInfo((ChangesRecord)HitTest.Item.Tag);
 
-					List<BadgeInfo> Badges;
+					List<BadgeInfo>? Badges;
 					if (LayoutInfo.CustomBadges.TryGetValue(Column.Text, out Badges) && Badges.Count > 0)
 					{
 						Point ListLocation = GetBadgeListLocation(Badges, HitTest.SubItem.Bounds, HorizontalAlign.Center, VerticalAlignment.Middle);
@@ -4339,7 +4355,7 @@ namespace UnrealGameSync
 			}
 		}
 
-		private BadgeInfo HitTestBadge(Point Location, List<BadgeInfo> BadgeList, Point ListLocation)
+		private BadgeInfo? HitTestBadge(Point Location, List<BadgeInfo> BadgeList, Point ListLocation)
 		{
 			foreach (BadgeInfo Badge in BadgeList)
 			{
@@ -4482,14 +4498,14 @@ namespace UnrealGameSync
 
 		private IReadOnlyList<IArchiveInfo> GetArchives()
 		{
-			IReadOnlyList<IArchiveInfo> AvailableArchives = JupiterMonitor?.AvailableArchives;
+			IReadOnlyList<IArchiveInfo>? AvailableArchives = JupiterMonitor?.AvailableArchives;
 			if (AvailableArchives != null && AvailableArchives.Count != 0)
 				return AvailableArchives;
 
 			// if jupiter had no archives we fallback to the perforce monitor
 			if (PerforceMonitor != null)
 			{
-				return PerforceMonitor?.AvailableArchives;
+				return PerforceMonitor.AvailableArchives;
 			}
 
 			return new List<IArchiveInfo>();
@@ -4541,7 +4557,7 @@ namespace UnrealGameSync
 			}
 		}
 
-		private void BuildList_FontChanged(object sender, EventArgs e)
+		private void BuildList_FontChanged(object? sender, EventArgs e)
 		{
 			if (BuildFont != null)
 			{
@@ -4578,7 +4594,7 @@ namespace UnrealGameSync
 			}
 		}
 
-		public void SyncChange(int ChangeNumber, bool bSyncOnly, WorkspaceUpdateCallback Callback)
+		public void SyncChange(int ChangeNumber, bool bSyncOnly, WorkspaceUpdateCallback? Callback)
 		{
 			if(Workspace != null)
 			{
@@ -4593,7 +4609,7 @@ namespace UnrealGameSync
 			SyncLatestChange(null);
 		}
 
-		public void SyncLatestChange(WorkspaceUpdateCallback Callback)
+		public void SyncLatestChange(WorkspaceUpdateCallback? Callback)
 		{
 			if (Workspace != null)
 			{
@@ -4615,32 +4631,50 @@ namespace UnrealGameSync
 
 		private void BuildListContextMenu_MarkGood_Click(object sender, EventArgs e)
 		{
-			EventMonitor.PostEvent(ContextMenuChange.Number, EventType.Good);
+			if (ContextMenuChange != null)
+			{
+				EventMonitor.PostEvent(ContextMenuChange.Number, EventType.Good);
+			}
 		}
 
 		private void BuildListContextMenu_MarkBad_Click(object sender, EventArgs e)
 		{
-			EventMonitor.PostEvent(ContextMenuChange.Number, EventType.Bad);
+			if (ContextMenuChange != null)
+			{
+				EventMonitor.PostEvent(ContextMenuChange.Number, EventType.Bad);
+			}
 		}
 
 		private void BuildListContextMenu_WithdrawReview_Click(object sender, EventArgs e)
 		{
-			EventMonitor.PostEvent(ContextMenuChange.Number, EventType.Unknown);
+			if (ContextMenuChange != null)
+			{
+				EventMonitor.PostEvent(ContextMenuChange.Number, EventType.Unknown);
+			}
 		}
 
 		private void BuildListContextMenu_Sync_Click(object sender, EventArgs e)
 		{
-			StartSync(ContextMenuChange.Number);
+			if (ContextMenuChange != null)
+			{
+				StartSync(ContextMenuChange.Number);
+			}
 		}
 
 		private void BuildListContextMenu_SyncContentOnly_Click(object sender, EventArgs e)
 		{
-			StartWorkspaceUpdate(ContextMenuChange.Number, WorkspaceUpdateOptions.Sync | WorkspaceUpdateOptions.ContentOnly);
+			if (ContextMenuChange != null)
+			{
+				StartWorkspaceUpdate(ContextMenuChange.Number, WorkspaceUpdateOptions.Sync | WorkspaceUpdateOptions.ContentOnly);
+			}
 		}
 
 		private void BuildListContextMenu_SyncOnlyThisChange_Click(object sender, EventArgs e)
 		{
-			StartWorkspaceUpdate(ContextMenuChange.Number, WorkspaceUpdateOptions.SyncSingleChange);
+			if (ContextMenuChange != null)
+			{
+				StartWorkspaceUpdate(ContextMenuChange.Number, WorkspaceUpdateOptions.SyncSingleChange);
+			}
 		}
 
 		private void BuildListContextMenu_Build_Click(object sender, EventArgs e)
@@ -4665,44 +4699,59 @@ namespace UnrealGameSync
 
 		private void BuildListContextMenu_MoreInfo_Click(object sender, EventArgs e)
 		{
-			Utility.SpawnP4VC(String.Format("{0} change {1}", PerforceSettings.GetArgumentsForExternalProgram(true), ContextMenuChange.Number));
+			if (ContextMenuChange != null)
+			{
+				Utility.SpawnP4VC(String.Format("{0} change {1}", PerforceSettings.GetArgumentsForExternalProgram(true), ContextMenuChange.Number));
+			}
 		}
 
 		private void BuildListContextMenu_ViewInSwarm_Click(object sender, EventArgs e)
 		{
-			string SwarmURL;
-			if (TryGetProjectSetting(PerforceMonitor.LatestProjectConfigFile, "SwarmURL", out SwarmURL))
-			{
-				Utility.OpenUrl(String.Format("{0}/changes/{1}", SwarmURL, ContextMenuChange.Number));
-			}
-			else
+			if (ContextMenuChange != null)
 			{
-				MessageBox.Show("Swarm URL is not configured.");
+				string? SwarmURL;
+				if (TryGetProjectSetting(PerforceMonitor.LatestProjectConfigFile, "SwarmURL", out SwarmURL))
+				{
+					Utility.OpenUrl(String.Format("{0}/changes/{1}", SwarmURL, ContextMenuChange.Number));
+				}
+				else
+				{
+					MessageBox.Show("Swarm URL is not configured.");
+				}
 			}
 		}
 
 		private void BuildListContextMenu_CopyChangelistNumber_Click(object sender, EventArgs e)
 		{
-			System.Windows.Forms.Clipboard.SetText(ContextMenuChange.Number.ToString());
+			if (ContextMenuChange != null)
+			{
+				System.Windows.Forms.Clipboard.SetText(ContextMenuChange.Number.ToString());
+			}
 		}
 
 		private void BuildListContextMenu_AddStar_Click(object sender, EventArgs e)
 		{
-			if (MessageBox.Show("Starred builds are meant to convey a stable, verified build to the rest of the team. Do not star a build unless it has been fully tested.\n\nAre you sure you want to star this build?", "Confirm star", MessageBoxButtons.YesNo) == DialogResult.Yes)
+			if (ContextMenuChange != null)
 			{
-				EventMonitor.PostEvent(ContextMenuChange.Number, EventType.Starred);
+				if (MessageBox.Show("Starred builds are meant to convey a stable, verified build to the rest of the team. Do not star a build unless it has been fully tested.\n\nAre you sure you want to star this build?", "Confirm star", MessageBoxButtons.YesNo) == DialogResult.Yes)
+				{
+					EventMonitor.PostEvent(ContextMenuChange.Number, EventType.Starred);
+				}
 			}
 		}
 
 		private void BuildListContextMenu_RemoveStar_Click(object sender, EventArgs e)
 		{
-			EventSummary Summary = EventMonitor.GetSummaryForChange(ContextMenuChange.Number);
-			if (Summary != null && Summary.LastStarReview != null && Summary.LastStarReview.Type == EventType.Starred)
+			if (ContextMenuChange != null)
 			{
-				string Message = String.Format("This change was starred by {0}. Are you sure you want to remove it?", FormatUserName(Summary.LastStarReview.UserName));
-				if (MessageBox.Show(Message, "Confirm removing star", MessageBoxButtons.YesNo) == DialogResult.Yes)
+				EventSummary Summary = EventMonitor.GetSummaryForChange(ContextMenuChange.Number);
+				if (Summary != null && Summary.LastStarReview != null && Summary.LastStarReview.Type == EventType.Starred)
 				{
-					EventMonitor.PostEvent(ContextMenuChange.Number, EventType.Unstarred);
+					string Message = String.Format("This change was starred by {0}. Are you sure you want to remove it?", FormatUserName(Summary.LastStarReview.UserName));
+					if (MessageBox.Show(Message, "Confirm removing star", MessageBoxButtons.YesNo) == DialogResult.Yes)
+					{
+						EventMonitor.PostEvent(ContextMenuChange.Number, EventType.Unstarred);
+					}
 				}
 			}
 		}
@@ -4714,33 +4763,39 @@ namespace UnrealGameSync
 
 		private void BuildListContextMenu_StartInvestigating_Click(object sender, EventArgs e)
 		{
-			string Message = String.Format("All changes from {0} onwards will be marked as bad while you are investigating an issue.\n\nAre you sure you want to continue?", ContextMenuChange.Number);
-			if (MessageBox.Show(Message, "Confirm investigating", MessageBoxButtons.YesNo) == DialogResult.Yes)
+			if (ContextMenuChange != null)
 			{
-				EventMonitor.StartInvestigating(ContextMenuChange.Number);
+				string Message = String.Format("All changes from {0} onwards will be marked as bad while you are investigating an issue.\n\nAre you sure you want to continue?", ContextMenuChange.Number);
+				if (MessageBox.Show(Message, "Confirm investigating", MessageBoxButtons.YesNo) == DialogResult.Yes)
+				{
+					EventMonitor.StartInvestigating(ContextMenuChange.Number);
+				}
 			}
 		}
 
 		private void BuildListContextMenu_FinishInvestigating_Click(object sender, EventArgs e)
 		{
-			int StartChangeNumber = EventMonitor.GetInvestigationStartChangeNumber(ContextMenuChange.Number);
-			if (StartChangeNumber != -1)
+			if (ContextMenuChange != null)
 			{
-				if (ContextMenuChange.Number > StartChangeNumber)
+				int StartChangeNumber = EventMonitor.GetInvestigationStartChangeNumber(ContextMenuChange.Number);
+				if (StartChangeNumber != -1)
 				{
-					string Message = String.Format("Mark all changes between {0} and {1} as bad?", StartChangeNumber, ContextMenuChange.Number);
-					if (MessageBox.Show(Message, "Finish investigating", MessageBoxButtons.YesNo) == DialogResult.Yes)
+					if (ContextMenuChange.Number > StartChangeNumber)
 					{
-						foreach (ChangesRecord Change in PerforceMonitor.GetChanges())
+						string Message = String.Format("Mark all changes between {0} and {1} as bad?", StartChangeNumber, ContextMenuChange.Number);
+						if (MessageBox.Show(Message, "Finish investigating", MessageBoxButtons.YesNo) == DialogResult.Yes)
 						{
-							if (Change.Number >= StartChangeNumber && Change.Number < ContextMenuChange.Number)
+							foreach (ChangesRecord Change in PerforceMonitor.GetChanges())
 							{
-								EventMonitor.PostEvent(Change.Number, EventType.Bad);
+								if (Change.Number >= StartChangeNumber && Change.Number < ContextMenuChange.Number)
+								{
+									EventMonitor.PostEvent(Change.Number, EventType.Bad);
+								}
 							}
 						}
 					}
+					EventMonitor.FinishInvestigating(ContextMenuChange.Number);
 				}
-				EventMonitor.FinishInvestigating(ContextMenuChange.Number);
 			}
 		}
 
@@ -5061,7 +5116,7 @@ namespace UnrealGameSync
 			{
 				int NextChangeNumber = SortedChangeNumbers[NextIndex];
 
-				PerforceChangeDetails Details;
+				PerforceChangeDetails? Details;
 				if (!PerforceMonitor.TryGetChangeDetails(NextChangeNumber, out Details) || Details.bContainsCode)
 				{
 					break;
@@ -5121,13 +5176,13 @@ namespace UnrealGameSync
 				return;
 			}
 
-			string ExtraSafeToDeleteFolders;
+			string? ExtraSafeToDeleteFolders;
 			if (!TryGetProjectSetting(Workspace.ProjectConfigFile, "SafeToDeleteFolders", out ExtraSafeToDeleteFolders))
 			{
 				ExtraSafeToDeleteFolders = "";
 			}
 
-			string ExtraSafeToDeleteExtensions;
+			string? ExtraSafeToDeleteExtensions;
 			if (!TryGetProjectSetting(Workspace.ProjectConfigFile, "SafeToDeleteExtensions", out ExtraSafeToDeleteExtensions))
 			{
 				ExtraSafeToDeleteExtensions = "";
@@ -5190,7 +5245,7 @@ namespace UnrealGameSync
 								}
 								else
 								{
-									ContextMenuStrip Menu;
+									ContextMenuStrip? Menu;
 									if (!NameToMenu.TryGetValue(MenuName, out Menu))
 									{
 										Menu = new ContextMenuStrip();
@@ -5248,10 +5303,14 @@ namespace UnrealGameSync
 					Dictionary<string, string> Variables = GetWorkspaceVariables(Workspace.CurrentChangeNumber);
 					if (Step.ToolId != Guid.Empty)
 					{
-						string ToolName = Owner.ToolUpdateMonitor.GetToolName(Step.ToolId);
+						string? ToolName = Owner.ToolUpdateMonitor.GetToolName(Step.ToolId);
 						if (ToolName != null)
 						{
-							Variables["ToolDir"] = Owner.ToolUpdateMonitor.GetToolPath(ToolName).FullName;
+							DirectoryReference? ToolPath = Owner.ToolUpdateMonitor.GetToolPath(ToolName);
+							if (ToolPath != null)
+							{
+								Variables["ToolDir"] = ToolPath.FullName;
+							}
 						}
 					}
 
@@ -5265,7 +5324,7 @@ namespace UnrealGameSync
 		{
 			BuildConfig EditorBuildConfig = GetEditorBuildConfig();
 
-			string SdkInstallerDir;
+			string? SdkInstallerDir;
 			TryGetProjectSetting(PerforceMonitor.LatestProjectConfigFile, "SdkInstallerDir", out SdkInstallerDir);
 
 			Dictionary<string, string> Variables = new Dictionary<string, string>();
@@ -5320,7 +5379,7 @@ namespace UnrealGameSync
 				{
 					if (Step.IsValid())
 					{
-						ConfigObject DefaultObject;
+						ConfigObject? DefaultObject;
 						ProjectBuildStepObjects.TryGetValue(Step.UniqueId, out DefaultObject);
 
 						ConfigObject UserConfigObject = Step.ToConfigObject(DefaultObject);
@@ -5350,7 +5409,7 @@ namespace UnrealGameSync
 				// Add or apply Object to the list of steps in Steps. Do not modify Object; make a copy first.
 				ConfigObject NewObject = new ConfigObject(Object);
 
-				ConfigObject DefaultObject;
+				ConfigObject? DefaultObject;
 				if (Steps.TryGetValue(UniqueId, out DefaultObject))
 				{
 					NewObject.SetDefaults(DefaultObject);
@@ -5372,7 +5431,7 @@ namespace UnrealGameSync
 
 		private string GetDefaultEditorTargetName()
 		{
-			string EditorTarget;
+			string? EditorTarget;
 			if (!TryGetProjectSetting(PerforceMonitor.LatestProjectConfigFile, "EditorTarget", out EditorTarget))
 			{
 				EditorTarget = "UE4Editor";
@@ -5803,10 +5862,13 @@ namespace UnrealGameSync
 			if (BuildList.SelectedItems.Count >= 2)
 			{
 				Dictionary<int, BisectState> ChangeNumberToBisectState = new Dictionary<int, BisectState>();
-				foreach (ListViewItem SelectedItem in BuildList.SelectedItems)
+				foreach (ListViewItem? SelectedItem in BuildList.SelectedItems)
 				{
-					ChangesRecord Change = (ChangesRecord)SelectedItem.Tag;
-					ChangeNumberToBisectState[Change.Number] = BisectState.Include;
+					if (SelectedItem != null)
+					{
+						ChangesRecord Change = (ChangesRecord)SelectedItem.Tag;
+						ChangeNumberToBisectState[Change.Number] = BisectState.Include;
+					}
 				}
 
 				ChangeNumberToBisectState[ChangeNumberToBisectState.Keys.Min()] = BisectState.Pass;
@@ -5831,9 +5893,9 @@ namespace UnrealGameSync
 
 		private void SetBisectStateForSelection(BisectState State)
 		{
-			foreach (ListViewItem SelectedItem in BuildList.SelectedItems)
+			foreach (ListViewItem? SelectedItem in BuildList.SelectedItems)
 			{
-				ChangesRecord Change = (ChangesRecord)SelectedItem.Tag;
+				ChangesRecord? Change = SelectedItem?.Tag as ChangesRecord;
 				if (Change != null)
 				{
 					WorkspaceState.SetBisectState(Change.Number, State);
@@ -5925,7 +5987,7 @@ namespace UnrealGameSync
 			List<KeyValuePair<string, string>> DisplayBadgeNameAndGroupPairs = new List<KeyValuePair<string, string>>(BadgeNameAndGroupPairs);
 			DisplayBadgeNameAndGroupPairs.AddRange(BadgeNames.Select(x => new KeyValuePair<string, string>(x, "User")));
 
-			string LastGroup = null;
+			string? LastGroup = null;
 			foreach (KeyValuePair<string, string> BadgeNameAndGroupPair in DisplayBadgeNameAndGroupPairs)
 			{
 				if (LastGroup != BadgeNameAndGroupPair.Value)
@@ -6137,7 +6199,7 @@ namespace UnrealGameSync
 					return null;
 				}
 
-				string OverrideValue = Section.GetValue("EditorMainWindowBackgroundOverride", null);
+				string? OverrideValue = Section.GetValue("EditorMainWindowBackgroundOverride", null);
 				if (OverrideValue == null)
 				{
 					return null;
@@ -6145,7 +6207,7 @@ namespace UnrealGameSync
 
 				ConfigObject OverrideObject = new ConfigObject(OverrideValue);
 
-				string TintColorValue = OverrideObject.GetValue("TintColor");
+				string? TintColorValue = OverrideObject.GetValue("TintColor");
 				if (TintColorValue == null)
 				{
 					return null;
@@ -6157,7 +6219,13 @@ namespace UnrealGameSync
 					return null;
 				}
 
-				ConfigObject SpecifiedColorObject = new ConfigObject(TintColorObject.GetValue("SpecifiedColor"));
+				string? SpecifiedColorValue = TintColorObject.GetValue("SpecifiedColor");
+				if(SpecifiedColorValue == null)
+				{
+					return null;
+				}
+
+				ConfigObject SpecifiedColorObject = new ConfigObject(SpecifiedColorValue);
 
 				float R, G, B;
 				if (!float.TryParse(SpecifiedColorObject.GetValue("R", ""), out R) || !float.TryParse(SpecifiedColorObject.GetValue("G", ""), out G) || !float.TryParse(SpecifiedColorObject.GetValue("B", ""), out B))
diff --git a/Engine/Source/Programs/UnrealGameSync/UnrealGameSync/DetectProjectSettingsTask.cs b/Engine/Source/Programs/UnrealGameSync/UnrealGameSync/DetectProjectSettingsTask.cs
--- a/Engine/Source/Programs/UnrealGameSync/UnrealGameSync/DetectProjectSettingsTask.cs
+++ b/Engine/Source/Programs/UnrealGameSync/UnrealGameSync/DetectProjectSettingsTask.cs
@@ -108,8 +108,6 @@ namespace UnrealGameSync
 				}
 				else
 				{
-					PerforceSettings = DefaultConnection.Settings;
-
 					// Get the perforce server settings
 					InfoRecord PerforceInfo = await DefaultConnection.GetInfoAsync(InfoOptions.ShortOutput, CancellationToken);
 
@@ -151,7 +149,8 @@ namespace UnrealGameSync
 					}
 
 					// Take the client we've chosen
-					PerforceClient = await PerforceConnection.CreateAsync(CandidateClients[0], Logger);
+					PerforceSettings = CandidateClients[0];
+					PerforceClient = await PerforceConnection.CreateAsync(PerforceSettings, Logger);
 
 					// Get the client path for the project file
 					List<WhereRecord> Records = await PerforceClient.WhereAsync(NewSelectedFileName.FullName, CancellationToken).Where(x => !x.Unmap).ToListAsync(CancellationToken);
diff --git a/Engine/Source/Programs/UnrealGameSync/UnrealGameSync/Forms/IssueBrowserWindow.cs b/Engine/Source/Programs/UnrealGameSync/UnrealGameSync/Forms/IssueBrowserWindow.cs
--- a/Engine/Source/Programs/UnrealGameSync/UnrealGameSync/Forms/IssueBrowserWindow.cs
+++ b/Engine/Source/Programs/UnrealGameSync/UnrealGameSync/Forms/IssueBrowserWindow.cs
@@ -38,7 +38,7 @@ namespace UnrealGameSync
 		Task? BackgroundTask;
 		bool bDisposed;
 
-		public IssueBrowserWindow(IssueMonitor IssueMonitor, IPerforceSettings PerforceSettings, TimeSpan? ServerTimeOffset, IServiceProvider ServiceProvider, string CurrentStream, Dictionary<string, Func<IssueData, bool>> CustomFilters, string FilterName)
+		public IssueBrowserWindow(IssueMonitor IssueMonitor, IPerforceSettings PerforceSettings, TimeSpan? ServerTimeOffset, IServiceProvider ServiceProvider, string CurrentStream, Dictionary<string, Func<IssueData, bool>> CustomFilters, string? FilterName)
 		{
 			this.IssueMonitor = IssueMonitor;
 			this.PerforceSettings = PerforceSettings;
@@ -262,7 +262,7 @@ namespace UnrealGameSync
 
 		static List<IssueBrowserWindow> ExistingWindows = new List<IssueBrowserWindow>();
 
-		public static void Show(Form Owner, IssueMonitor IssueMonitor, IPerforceSettings PerforceSettings, TimeSpan? ServerTimeOffset, IServiceProvider ServiceProvider, string CurrentStream, Dictionary<string, Func<IssueData, bool>> CustomFilters, string DefaultFilter)
+		public static void Show(Form Owner, IssueMonitor IssueMonitor, IPerforceSettings PerforceSettings, TimeSpan? ServerTimeOffset, IServiceProvider ServiceProvider, string CurrentStream, Dictionary<string, Func<IssueData, bool>> CustomFilters, string? DefaultFilter)
 		{
 			IssueBrowserWindow Window = ExistingWindows.FirstOrDefault(x => x.IssueMonitor == IssueMonitor);
 			if(Window == null)
diff --git a/Engine/Source/Programs/UnrealGameSync/UnrealGameSync/Forms/SelectStreamWindow.cs b/Engine/Source/Programs/UnrealGameSync/UnrealGameSync/Forms/SelectStreamWindow.cs
--- a/Engine/Source/Programs/UnrealGameSync/UnrealGameSync/Forms/SelectStreamWindow.cs
+++ b/Engine/Source/Programs/UnrealGameSync/UnrealGameSync/Forms/SelectStreamWindow.cs
@@ -322,7 +322,7 @@ namespace UnrealGameSync
 			return bExpand || (SelectedStream == Stream.Record.Stream);
 		}
 
-		public static bool ShowModal(IWin32Window Owner, IPerforceSettings Perforce, string StreamName, IServiceProvider ServiceProvider, out string? NewStreamName)
+		public static bool ShowModal(IWin32Window Owner, IPerforceSettings Perforce, string StreamName, IServiceProvider ServiceProvider, [NotNullWhen(true)] out string? NewStreamName)
 		{
 			ILogger Logger = ServiceProvider.GetRequiredService<ILogger<SelectStreamWindow>>();
 
@@ -334,7 +334,7 @@ namespace UnrealGameSync
 			}
 
 			SelectStreamWindow SelectStream = new SelectStreamWindow(StreamsTask.Result, StreamName);
-			if(SelectStream.ShowDialog(Owner) == DialogResult.OK)
+			if(SelectStream.ShowDialog(Owner) == DialogResult.OK && SelectStream.SelectedStream != null)
 			{
 				NewStreamName = SelectStream.SelectedStream;
 				return true;
diff --git a/Engine/Source/Programs/UnrealGameSync/UnrealGameSync/Program.cs b/Engine/Source/Programs/UnrealGameSync/UnrealGameSync/Program.cs
--- a/Engine/Source/Programs/UnrealGameSync/UnrealGameSync/Program.cs
+++ b/Engine/Source/Programs/UnrealGameSync/UnrealGameSync/Program.cs
@@ -43,7 +43,7 @@ namespace UnrealGameSync
 		public static string SyncVersion = null;
 
 		[STAThread]
-		static async Task Main(string[] Args)
+		static void Main(string[] Args)
 		{
 			bool bFirstInstance;
 			using (Mutex InstanceMutex = new Mutex(true, "UnrealGameSyncRunning", out bFirstInstance))
@@ -64,7 +64,7 @@ namespace UnrealGameSync
 
 					if (bFirstInstance)
 					{
-						await InnerMainAsync(InstanceMutex, ActivateEvent, Args);
+						InnerMainAsync(InstanceMutex, ActivateEvent, Args).GetAwaiter().GetResult();
 					}
 					else
 					{
diff --git a/Engine/Source/Programs/UnrealGameSync/UnrealGameSync/UserSettings.cs b/Engine/Source/Programs/UnrealGameSync/UnrealGameSync/UserSettings.cs
--- a/Engine/Source/Programs/UnrealGameSync/UnrealGameSync/UserSettings.cs
+++ b/Engine/Source/Programs/UnrealGameSync/UnrealGameSync/UserSettings.cs
@@ -230,7 +230,7 @@ namespace UnrealGameSync
 		public FileReference? File { get; set; }
 
 		// Workspace specific SyncFilters
-		public string[]? SyncView { get; set; }
+		public string[] SyncView { get; set; } = Array.Empty<string>();
 		public List<SyncCategory> SyncCategories { get; set; } = new List<SyncCategory>();
 		public bool? bSyncAllProjects { get; set; }
 		public bool? bIncludeAllProjectsInSolution { get; set; }
diff --git a/Engine/Source/Programs/UnrealGameSync/UnrealGameSync/Workspace.cs b/Engine/Source/Programs/UnrealGameSync/UnrealGameSync/Workspace.cs
--- a/Engine/Source/Programs/UnrealGameSync/UnrealGameSync/Workspace.cs
+++ b/Engine/Source/Programs/UnrealGameSync/UnrealGameSync/Workspace.cs
@@ -85,13 +85,13 @@ namespace UnrealGameSync
 		public DateTime StartTime = DateTime.UtcNow;
 		public int ChangeNumber;
 		public WorkspaceUpdateOptions Options;
-		public string[] SyncFilter;
-		public Dictionary<string, Tuple<IArchiveInfo, string>> ArchiveTypeToArchive = new Dictionary<string, Tuple<IArchiveInfo, string>>();
+		public string[]? SyncFilter;
+		public Dictionary<string, Tuple<IArchiveInfo, string>?> ArchiveTypeToArchive = new Dictionary<string, Tuple<IArchiveInfo, string>?>();
 		public Dictionary<string, bool> DeleteFiles = new Dictionary<string,bool>();
 		public Dictionary<string, bool> ClobberFiles = new Dictionary<string,bool>();
 		public Dictionary<Guid,ConfigObject> DefaultBuildSteps;
 		public List<ConfigObject> UserBuildStepObjects;
-		public HashSet<Guid> CustomBuildSteps;
+		public HashSet<Guid>? CustomBuildSteps;
 		public Dictionary<string, string> Variables;
 		public PerforceSyncOptions? PerforceSyncOptions;
 		public List<HaveRecord>? HaveFiles; // Cached when sync filter has changed
@@ -100,7 +100,7 @@ namespace UnrealGameSync
 		public ConfigFile? ProjectConfigFile;
 		public IReadOnlyList<string>? ProjectStreamFilter;
 
-		public WorkspaceUpdateContext(int InChangeNumber, WorkspaceUpdateOptions InOptions, string[] InSyncFilter, Dictionary<Guid, ConfigObject> InDefaultBuildSteps, List<ConfigObject> InUserBuildSteps, HashSet<Guid> InCustomBuildSteps, Dictionary<string, string> InVariables)
+		public WorkspaceUpdateContext(int InChangeNumber, WorkspaceUpdateOptions InOptions, string[]? InSyncFilter, Dictionary<Guid, ConfigObject> InDefaultBuildSteps, List<ConfigObject> InUserBuildSteps, HashSet<Guid>? InCustomBuildSteps, Dictionary<string, string> InVariables)
 		{
 			ChangeNumber = InChangeNumber;
 			Options = InOptions;
@@ -366,7 +366,7 @@ namespace UnrealGameSync
 
 		IAsyncDisposer AsyncDisposer;
 
-		public Workspace(IPerforceSettings InPerfoceSettings, ProjectInfo InProject, UserWorkspaceState InState, ConfigFile ProjectConfigFile, IReadOnlyList<string> ProjectStreamFilter, ILogger Logger, IServiceProvider ServiceProvider)
+		public Workspace(IPerforceSettings InPerfoceSettings, ProjectInfo InProject, UserWorkspaceState InState, ConfigFile ProjectConfigFile, IReadOnlyList<string>? ProjectStreamFilter, ILogger Logger, IServiceProvider ServiceProvider)
 		{
 			PerforceSettings = InPerfoceSettings;
 			Project = InProject;
@@ -393,7 +393,7 @@ namespace UnrealGameSync
 			get; private set;
 		}
 
-		public IReadOnlyList<string> ProjectStreamFilter
+		public IReadOnlyList<string>? ProjectStreamFilter
 		{
 			get; private set;
 		}
@@ -1088,7 +1088,7 @@ namespace UnrealGameSync
 					DirectoryReference.CreateDirectory(ManifestDirectoryName);
 
 					// Sync and extract (or just remove) the given archives
-					foreach (KeyValuePair<string, Tuple<IArchiveInfo, string>> ArchiveTypeAndArchive in Context.ArchiveTypeToArchive)
+					foreach (KeyValuePair<string, Tuple<IArchiveInfo, string>?> ArchiveTypeAndArchive in Context.ArchiveTypeToArchive)
 					{
 						string ArchiveType = ArchiveTypeAndArchive.Key;
 
@@ -1384,7 +1384,7 @@ namespace UnrealGameSync
 			return SyncPaths.Select(x => Project.ClientRootPath + x).ToList();
 		}
 
-		public static List<string> GetRelativeSyncPaths(ProjectInfo Project, bool bSyncAllProjects, string[] SyncFilter)
+		public static List<string> GetRelativeSyncPaths(ProjectInfo Project, bool bSyncAllProjects, string[]? SyncFilter)
 		{
 			List<string> SyncPaths = new List<string>();
 
