UGS: Convert remaining classes to use nullable references.

diff --git a/Engine/Source/Programs/UnrealGameSync/UnrealGameSync/ActivationListener.cs b/Engine/Source/Programs/UnrealGameSync/UnrealGameSync/ActivationListener.cs
--- a/Engine/Source/Programs/UnrealGameSync/UnrealGameSync/ActivationListener.cs
+++ b/Engine/Source/Programs/UnrealGameSync/UnrealGameSync/ActivationListener.cs
@@ -7,8 +7,6 @@ using System.Text;
 using System.Threading;
 using System.Threading.Tasks;
 
-#nullable enable
-
 namespace UnrealGameSync
 {
 	class ActivationListener : IDisposable
diff --git a/Engine/Source/Programs/UnrealGameSync/UnrealGameSync/ArchiveManifest.cs b/Engine/Source/Programs/UnrealGameSync/UnrealGameSync/ArchiveManifest.cs
--- a/Engine/Source/Programs/UnrealGameSync/UnrealGameSync/ArchiveManifest.cs
+++ b/Engine/Source/Programs/UnrealGameSync/UnrealGameSync/ArchiveManifest.cs
@@ -11,8 +11,6 @@ using System.Linq;
 using System.Text;
 using System.Threading.Tasks;
 
-#nullable enable
-
 namespace UnrealGameSync
 {
 	[DebuggerDisplay("{FileName}")]
diff --git a/Engine/Source/Programs/UnrealGameSync/UnrealGameSync/Automation/P4Automation.cs b/Engine/Source/Programs/UnrealGameSync/UnrealGameSync/Automation/P4Automation.cs
--- a/Engine/Source/Programs/UnrealGameSync/UnrealGameSync/Automation/P4Automation.cs
+++ b/Engine/Source/Programs/UnrealGameSync/UnrealGameSync/Automation/P4Automation.cs
@@ -8,8 +8,6 @@ using System.IO;
 using System.Threading;
 using System.Threading.Tasks;
 
-#nullable enable
-
 namespace UnrealGameSync
 {
 
diff --git a/Engine/Source/Programs/UnrealGameSync/UnrealGameSync/Automation/VisualStudioAutomation.cs b/Engine/Source/Programs/UnrealGameSync/UnrealGameSync/Automation/VisualStudioAutomation.cs
--- a/Engine/Source/Programs/UnrealGameSync/UnrealGameSync/Automation/VisualStudioAutomation.cs
+++ b/Engine/Source/Programs/UnrealGameSync/UnrealGameSync/Automation/VisualStudioAutomation.cs
@@ -11,8 +11,6 @@ using Microsoft.VisualStudio.Setup.Configuration;
 using EnvDTE;
 using System.Diagnostics;
 
-#nullable enable
-
 namespace UnrealGameSync
 {
 
diff --git a/Engine/Source/Programs/UnrealGameSync/UnrealGameSync/AutomationServer.cs b/Engine/Source/Programs/UnrealGameSync/UnrealGameSync/AutomationServer.cs
--- a/Engine/Source/Programs/UnrealGameSync/UnrealGameSync/AutomationServer.cs
+++ b/Engine/Source/Programs/UnrealGameSync/UnrealGameSync/AutomationServer.cs
@@ -11,8 +11,6 @@ using System.Text;
 using System.Threading;
 using System.Windows.Forms;
 
-#nullable enable
-
 namespace UnrealGameSync
 {
 	enum AutomationRequestType
@@ -146,9 +144,9 @@ namespace UnrealGameSync
 
 		bool bDisposing;
 		ILogger Logger;
-		string CommandLineUri;
+		string? CommandLineUri;
 
-		public AutomationServer(Action<AutomationRequest> PostRequest, string Uri, ILogger<AutomationServer> Logger)
+		public AutomationServer(Action<AutomationRequest> PostRequest, string? Uri, ILogger<AutomationServer> Logger)
 		{
 			ShutdownEvent = new ManualResetEvent(false);
 			this.PostRequest = PostRequest;
diff --git a/Engine/Source/Programs/UnrealGameSync/UnrealGameSync/BuildStep.cs b/Engine/Source/Programs/UnrealGameSync/UnrealGameSync/BuildStep.cs
--- a/Engine/Source/Programs/UnrealGameSync/UnrealGameSync/BuildStep.cs
+++ b/Engine/Source/Programs/UnrealGameSync/UnrealGameSync/BuildStep.cs
@@ -23,24 +23,24 @@ namespace UnrealGameSync
 
 		public Guid UniqueId;
 		public int OrderIndex;
-		public string Description;
-		public string StatusText;
+		public string? Description;
+		public string? StatusText;
 		public int EstimatedDuration;
 		public BuildStepType Type;
-		public string Target;
-		public string Platform;
-		public string Configuration;
-		public string FileName;
-		public string WorkingDir;
-		public string Arguments;
+		public string? Target;
+		public string? Platform;
+		public string? Configuration;
+		public string? FileName;
+		public string? WorkingDir;
+		public string? Arguments;
 		public bool bUseLogWindow;
 		public bool bNormalSync;
 		public bool bScheduledSync;
-		public string StatusPanelLink;
-		public Guid[] Requires;
+		public string? StatusPanelLink;
+		public Guid[] Requires = Array.Empty<Guid>();
 		public Guid ToolId;
 
-		public BuildStep(Guid InUniqueId, int InOrderIndex, string InDescription, string InStatusText, int InEstimatedDuration, string InFileName, string InArguments, string InWorkingDir, bool bInUseLogWindow)
+		public BuildStep(Guid InUniqueId, int InOrderIndex, string? InDescription, string? InStatusText, int InEstimatedDuration, string? InFileName, string? InArguments, string? InWorkingDir, bool bInUseLogWindow)
 		{
 			UniqueId = InUniqueId;
 			OrderIndex = InOrderIndex;
@@ -54,7 +54,7 @@ namespace UnrealGameSync
 			bUseLogWindow = bInUseLogWindow;
 		}
 
-		public BuildStep(Guid InUniqueId, int InOrderIndex, string InDescription, string InStatusText, int InEstimatedDuration, string InTarget, string InPlatform, string InConfiguration, string InArguments, bool bInSyncDefault)
+		public BuildStep(Guid InUniqueId, int InOrderIndex, string? InDescription, string? InStatusText, int InEstimatedDuration, string? InTarget, string? InPlatform, string? InConfiguration, string? InArguments, bool bInSyncDefault)
 		{
 			UniqueId = InUniqueId;
 			OrderIndex = InOrderIndex;
@@ -163,7 +163,7 @@ namespace UnrealGameSync
 				Guid UniqueId;
 				if(Guid.TryParse(ModifyObject.GetValue(UniqueIdKey, ""), out UniqueId))
 				{
-					ConfigObject DefaultObject;
+					ConfigObject? DefaultObject;
 					if(BuildStepObjects.TryGetValue(UniqueId, out DefaultObject))
 					{
 						ModifyObject.SetDefaults(DefaultObject);
@@ -213,7 +213,7 @@ namespace UnrealGameSync
 			return Result;
 		}
 
-		public ConfigObject ToConfigObject(ConfigObject DefaultObject)
+		public ConfigObject? ToConfigObject(ConfigObject? DefaultObject)
 		{
 			ConfigObject Result = new ConfigObject();
 			Result[UniqueIdKey] = UniqueId.ToString();
diff --git a/Engine/Source/Programs/UnrealGameSync/UnrealGameSync/ChildProcess.cs b/Engine/Source/Programs/UnrealGameSync/UnrealGameSync/ChildProcess.cs
deleted file mode 100644
diff --git a/Engine/Source/Programs/UnrealGameSync/UnrealGameSync/ConfigFile.cs b/Engine/Source/Programs/UnrealGameSync/UnrealGameSync/ConfigFile.cs
--- a/Engine/Source/Programs/UnrealGameSync/UnrealGameSync/ConfigFile.cs
+++ b/Engine/Source/Programs/UnrealGameSync/UnrealGameSync/ConfigFile.cs
@@ -11,8 +11,6 @@ using System.Linq;
 using System.Text;
 using System.Threading.Tasks;
 
-#nullable enable
-
 namespace UnrealGameSync
 {
 	public class ConfigObject
@@ -255,7 +253,7 @@ namespace UnrealGameSync
 			}
 		}
 
-		public void AddOverrides(ConfigObject Object, ConfigObject DefaultObject)
+		public void AddOverrides(ConfigObject Object, ConfigObject? DefaultObject)
 		{
 			foreach(KeyValuePair<string, string> Pair in Object.Pairs)
 			{
diff --git a/Engine/Source/Programs/UnrealGameSync/UnrealGameSync/Controls/AlertImageControl.cs b/Engine/Source/Programs/UnrealGameSync/UnrealGameSync/Controls/AlertImageControl.cs
--- a/Engine/Source/Programs/UnrealGameSync/UnrealGameSync/Controls/AlertImageControl.cs
+++ b/Engine/Source/Programs/UnrealGameSync/UnrealGameSync/Controls/AlertImageControl.cs
@@ -16,9 +16,9 @@ namespace UnrealGameSync
 {
 	partial class AlertImageControl : UserControl
 	{
-		Image ImageValue;
+		Image? ImageValue;
 
-		public Image Image
+		public Image? Image
 		{
 			get { return ImageValue; }
 			set { ImageValue = value; Invalidate(); }
diff --git a/Engine/Source/Programs/UnrealGameSync/UnrealGameSync/Controls/BuildListControl.cs b/Engine/Source/Programs/UnrealGameSync/UnrealGameSync/Controls/BuildListControl.cs
--- a/Engine/Source/Programs/UnrealGameSync/UnrealGameSync/Controls/BuildListControl.cs
+++ b/Engine/Source/Programs/UnrealGameSync/UnrealGameSync/Controls/BuildListControl.cs
@@ -53,7 +53,7 @@ namespace UnrealGameSync
 
 		public delegate void OnScrollDelegate();
 
-		public event OnScrollDelegate OnScroll;
+		public event OnScrollDelegate? OnScroll;
 
 		public BuildListControl()
 		{
diff --git a/Engine/Source/Programs/UnrealGameSync/UnrealGameSync/Controls/CustomListViewControl.cs b/Engine/Source/Programs/UnrealGameSync/UnrealGameSync/Controls/CustomListViewControl.cs
--- a/Engine/Source/Programs/UnrealGameSync/UnrealGameSync/Controls/CustomListViewControl.cs
+++ b/Engine/Source/Programs/UnrealGameSync/UnrealGameSync/Controls/CustomListViewControl.cs
@@ -11,8 +11,6 @@ using System.Threading.Tasks;
 using System.Windows.Forms;
 using System.Windows.Forms.VisualStyles;
 
-#nullable enable
-
 namespace UnrealGameSync
 {
 	abstract class CustomListViewWidget
@@ -107,7 +105,7 @@ namespace UnrealGameSync
 
 		public override void OnMouseDown(Point Location)
 		{
-			StatusElement Element;
+			StatusElement? Element;
 			if(Line.HitTest(Location, out Element))
 			{
 				MouseDownElement = Element;
@@ -137,7 +135,7 @@ namespace UnrealGameSync
 
 		public override void OnMouseMove(Point Location)
 		{
-			StatusElement NewMouseOverElement;
+			StatusElement? NewMouseOverElement;
 			Line.HitTest(Location, out NewMouseOverElement);
 
 			if(MouseOverElement != NewMouseOverElement)
diff --git a/Engine/Source/Programs/UnrealGameSync/UnrealGameSync/Controls/LogControl.cs b/Engine/Source/Programs/UnrealGameSync/UnrealGameSync/Controls/LogControl.cs
--- a/Engine/Source/Programs/UnrealGameSync/UnrealGameSync/Controls/LogControl.cs
+++ b/Engine/Source/Programs/UnrealGameSync/UnrealGameSync/Controls/LogControl.cs
@@ -118,16 +118,16 @@ namespace UnrealGameSync
 		int ScrollColumnsPerPage;
 
 		bool bIsSelecting;
-		TextSelection Selection;
+		TextSelection? Selection;
 
 		Size FontSize;
 
-		Timer SelectionScrollTimer;
+		Timer? SelectionScrollTimer;
 		int AutoScrollRate = 0;
 
-		Timer UpdateTimer;
+		Timer? UpdateTimer;
 		ConcurrentQueue<string> QueuedLines = new ConcurrentQueue<string>();
-		FileStream LogFileStream;
+		FileStream? LogFileStream;
 
 		public LogControl()
 		{
@@ -258,7 +258,7 @@ namespace UnrealGameSync
 				List<string> NewLines = new List<string>();
 				for(;;)
 				{
-					string NextLine;
+					string? NextLine;
 					if(!QueuedLines.TryDequeue(out NextLine))
 					{
 						break;
@@ -388,12 +388,12 @@ namespace UnrealGameSync
 			}
 		}
 
-		protected void ContextMenu_CopySelection(object sender, EventArgs e)
+		protected void ContextMenu_CopySelection(object? sender, EventArgs e)
 		{
 			CopySelection();
 		}
 
-		protected void ContextMenu_SelectAll(object sender, EventArgs e)
+		protected void ContextMenu_SelectAll(object? sender, EventArgs e)
 		{
 			SelectAll();
 		}
@@ -506,7 +506,7 @@ namespace UnrealGameSync
 				TextLocation Location = PointToTextLocation(e.Location);
 				Selection = new TextSelection(Location, Location);
 
-				SelectionScrollTimer.Start();
+				SelectionScrollTimer!.Start();
 				AutoScrollRate = 0;
 
 				Capture = true;
@@ -525,7 +525,7 @@ namespace UnrealGameSync
 		{
 			base.OnMouseMove(e);
 
-			if(bIsSelecting)
+			if(bIsSelecting && Selection != null)
 			{
 				TextLocation NewSelectionEnd = PointToTextLocation(e.Location);
 				if(NewSelectionEnd.LineIdx != Selection.End.LineIdx || NewSelectionEnd.ColumnIdx != Selection.End.ColumnIdx)
@@ -623,7 +623,7 @@ namespace UnrealGameSync
 				{
 					Selection = null;
 				}
-				SelectionScrollTimer.Stop();
+				SelectionScrollTimer!.Stop();
 				bIsSelecting = false;
 			}
 		}
@@ -797,7 +797,7 @@ namespace UnrealGameSync
 			ScrollLine = VerticalScroll.nPos;
 		}
 
-		protected void SelectionScrollTimer_TimerElapsed(object Sender, EventArgs Args)
+		protected void SelectionScrollTimer_TimerElapsed(object? Sender, EventArgs Args)
 		{
 			if(AutoScrollRate != 0 && Selection != null)
 			{
diff --git a/Engine/Source/Programs/UnrealGameSync/UnrealGameSync/Controls/LogSplitContainer.cs b/Engine/Source/Programs/UnrealGameSync/UnrealGameSync/Controls/LogSplitContainer.cs
--- a/Engine/Source/Programs/UnrealGameSync/UnrealGameSync/Controls/LogSplitContainer.cs
+++ b/Engine/Source/Programs/UnrealGameSync/UnrealGameSync/Controls/LogSplitContainer.cs
@@ -1,4 +1,4 @@
-﻿// Copyright Epic Games, Inc. All Rights Reserved.
+// Copyright Epic Games, Inc. All Rights Reserved.
 
 using System;
 using System.Collections.Generic;
@@ -28,7 +28,7 @@ namespace UnrealGameSync
 
 		int LogHeight = 200;
 
-		public event Action<bool> OnVisibilityChanged;
+		public event Action<bool>? OnVisibilityChanged;
 
 		[Browsable(true)]
 		public string Caption
diff --git a/Engine/Source/Programs/UnrealGameSync/UnrealGameSync/Controls/StatusPanel.cs b/Engine/Source/Programs/UnrealGameSync/UnrealGameSync/Controls/StatusPanel.cs
--- a/Engine/Source/Programs/UnrealGameSync/UnrealGameSync/Controls/StatusPanel.cs
+++ b/Engine/Source/Programs/UnrealGameSync/UnrealGameSync/Controls/StatusPanel.cs
@@ -3,6 +3,7 @@
 using System;
 using System.Collections.Generic;
 using System.ComponentModel;
+using System.Diagnostics.CodeAnalysis;
 using System.Drawing;
 using System.Drawing.Drawing2D;
 using System.Linq;
@@ -54,7 +55,7 @@ namespace UnrealGameSync
 
 		public Font FindOrAddFont(FontStyle Style)
 		{
-			Font Result;
+			Font? Result;
 			if(!FontCache.TryGetValue(Style, out Result))
 			{
 				Result = new Font(FontCache[FontStyle.Regular], Style);
@@ -198,7 +199,7 @@ namespace UnrealGameSync
 		string Name;
 		Color BackgroundColor;
 		Color HoverBackgroundColor;
-		Action<Point, Rectangle> ClickAction;
+		Action<Point, Rectangle>? ClickAction;
 
 		public bool MergeLeft
 		{
@@ -210,11 +211,13 @@ namespace UnrealGameSync
 			get; set;
 		}
 
-		public BadgeStatusElement(string InName, Color InBackgroundColor, Action<Point, Rectangle> InClickAction)
+		public BadgeStatusElement(string InName, Color InBackgroundColor, Action<Point, Rectangle>? InClickAction)
 		{
 			Name = InName;
 			BackgroundColor = InBackgroundColor;
-			if(ClickAction == null)
+			ClickAction = InClickAction;
+
+			if (ClickAction == null)
 			{
 				HoverBackgroundColor = BackgroundColor;
 			}
@@ -222,7 +225,6 @@ namespace UnrealGameSync
 			{
 				HoverBackgroundColor = Color.FromArgb(Math.Min(BackgroundColor.R + 32, 255), Math.Min(BackgroundColor.G + 32, 255), Math.Min(BackgroundColor.B + 32, 255));
 			}
-			ClickAction = InClickAction;
 			if(ClickAction != null)
 			{
 				Cursor = NativeCursors.Hand;
@@ -374,13 +376,13 @@ namespace UnrealGameSync
 			bModified = true;
 		}
 
-		public void AddBadge(string InText, Color InBackgroundColor, Action<Point, Rectangle> InClickAction)
+		public void AddBadge(string InText, Color InBackgroundColor, Action<Point, Rectangle>? InClickAction)
 		{
 			Elements.Add(new BadgeStatusElement(InText, InBackgroundColor, InClickAction));
 			if(Elements.Count >= 2)
 			{
-				BadgeStatusElement PrevBadge = Elements[Elements.Count - 2] as BadgeStatusElement;
-				BadgeStatusElement NextBadge = Elements[Elements.Count - 1] as BadgeStatusElement;
+				BadgeStatusElement? PrevBadge = Elements[Elements.Count - 2] as BadgeStatusElement;
+				BadgeStatusElement? NextBadge = Elements[Elements.Count - 1] as BadgeStatusElement;
 				if(PrevBadge != null && NextBadge != null)
 				{
 					PrevBadge.MergeRight = true;
@@ -396,7 +398,7 @@ namespace UnrealGameSync
 			bModified = true;
 		}
 
-		public bool HitTest(Point Location, out StatusElement OutElement)
+		public bool HitTest(Point Location, [NotNullWhen(true)] out StatusElement? OutElement)
 		{
 			OutElement = null;
 			if(Bounds.Contains(Location))
@@ -440,20 +442,20 @@ namespace UnrealGameSync
 	{
 		const float LineSpacing = 1.35f;
 
-		Image ProjectLogo;
+		Image? ProjectLogo;
 		bool bDisposeProjectLogo;
 		Rectangle ProjectLogoBounds;
-		StatusElementResources Resources;
+		StatusElementResources? Resources;
 		List<StatusLine> Lines = new List<StatusLine>();
-		StatusLine Caption;
-		Pen AlertDividerPen;
+		StatusLine? Caption;
+		Pen? AlertDividerPen;
 		int AlertDividerY;
-		StatusLine Alert;
+		StatusLine? Alert;
 		Color? TintColor;
 		Point? MouseOverLocation;
-		StatusElement MouseOverElement;
+		StatusElement? MouseOverElement;
 		Point? MouseDownLocation;
-		StatusElement MouseDownElement;
+		StatusElement? MouseDownElement;
 		int ContentWidth = 400;
 		int SuspendDisplayCount;
 
@@ -539,7 +541,7 @@ namespace UnrealGameSync
 			Caption = null;
 		}
 
-		public void Set(IEnumerable<StatusLine> NewLines, StatusLine NewCaption, StatusLine NewAlert, Color? NewTintColor)
+		public void Set(IEnumerable<StatusLine> NewLines, StatusLine? NewCaption, StatusLine? NewAlert, Color? NewTintColor)
 		{
 			if(Resources == null)
 			{
@@ -584,7 +586,7 @@ namespace UnrealGameSync
 			}
 		}
 
-		protected bool HitTest(Point Location, out StatusElement OutElement)
+		protected bool HitTest(Point Location, [NotNullWhen(true)] out StatusElement? OutElement)
 		{
 			OutElement = null;
 			foreach(StatusLine Line in Lines)
@@ -634,8 +636,8 @@ namespace UnrealGameSync
 			if(Alert != null)
 			{
 				AlertDividerY = Height - (int)(Font.Height * 2);
-				Alert.Layout(Graphics, Point.Empty, Resources);
-				Alert.Layout(Graphics, new Point((Width - Alert.Bounds.Width) / 2, (Height + AlertDividerY) / 2), Resources);
+				Alert.Layout(Graphics, Point.Empty, Resources!);
+				Alert.Layout(Graphics, new Point((Width - Alert.Bounds.Width) / 2, (Height + AlertDividerY) / 2), Resources!);
 				BodyHeight = AlertDividerY;
 			}
 
@@ -656,9 +658,9 @@ namespace UnrealGameSync
 			LogoY -= Font.Height / 2;
 			if(Caption != null)
 			{
-				Caption.Layout(Graphics, Point.Empty, Resources);
+				Caption.Layout(Graphics, Point.Empty, Resources!);
 				int CaptionWidth = Caption.Bounds.Width;
-				Caption.Layout(Graphics, new Point(Math.Min(LogoX + (LogoWidth / 2) - (CaptionWidth / 2), DividerX - CaptionWidth), LogoY + LogoHeight), Resources);
+				Caption.Layout(Graphics, new Point(Math.Min(LogoX + (LogoWidth / 2) - (CaptionWidth / 2), DividerX - CaptionWidth), LogoY + LogoHeight), Resources!);
 			}
 
 			// Set the logo rectangle
@@ -672,7 +674,7 @@ namespace UnrealGameSync
 			foreach(StatusLine Line in Lines)
 			{
 				LineY += (int)(Font.Height * LineSpacing * Line.LineHeight * 0.5f);
-				Line.Layout(Graphics, new Point(DividerX + 5, (int)LineY), Resources);
+				Line.Layout(Graphics, new Point(DividerX + 5, (int)LineY), Resources!);
 				LineY += (int)(Font.Height * LineSpacing * Line.LineHeight * 0.5f);
 			}
 		}
@@ -757,16 +759,16 @@ namespace UnrealGameSync
 
 				foreach(StatusLine Line in Lines)
 				{
-					Line.Draw(e.Graphics, Resources);
+					Line.Draw(e.Graphics, Resources!);
 				}
 				if(Caption != null)
 				{
-					Caption.Draw(e.Graphics, Resources);
+					Caption.Draw(e.Graphics, Resources!);
 				}
 				if(Alert != null)
 				{
 					e.Graphics.DrawLine(AlertDividerPen, 0, AlertDividerY, Width, AlertDividerY);
-					Alert.Draw(e.Graphics, Resources);
+					Alert.Draw(e.Graphics, Resources!);
 				}
 			}
 		}
@@ -775,7 +777,7 @@ namespace UnrealGameSync
 		{
 			MouseOverLocation = NewMouseOverLocation;
 
-			StatusElement NewMouseOverElement = null;
+			StatusElement? NewMouseOverElement = null;
 			if(MouseOverLocation.HasValue)
 			{
 				HitTest(MouseOverLocation.Value, out NewMouseOverElement);
@@ -805,7 +807,7 @@ namespace UnrealGameSync
 		{
 			MouseDownLocation = NewMouseDownLocation;
 
-			StatusElement NewMouseDownElement = null;
+			StatusElement? NewMouseDownElement = null;
 			if(MouseDownLocation.HasValue)
 			{
 				HitTest(MouseDownLocation.Value, out NewMouseDownElement);
diff --git a/Engine/Source/Programs/UnrealGameSync/UnrealGameSync/Controls/TabControl.cs b/Engine/Source/Programs/UnrealGameSync/UnrealGameSync/Controls/TabControl.cs
--- a/Engine/Source/Programs/UnrealGameSync/UnrealGameSync/Controls/TabControl.cs
+++ b/Engine/Source/Programs/UnrealGameSync/UnrealGameSync/Controls/TabControl.cs
@@ -24,8 +24,14 @@ namespace UnrealGameSync
 			public int Width;
 			public int CloseButtonWidth;
 			public Size TextSize;
-			public Tuple<Color, float> Highlight;
+			public Tuple<Color, float>? Highlight;
 			public Color? TintColor;
+
+			public TabData(string Name, object Data)
+			{
+				this.Name = Name;
+				this.Data = Data;
+			}
 		}
 
 		class TabDragData
@@ -34,39 +40,47 @@ namespace UnrealGameSync
 			public int InitialIdx;
 			public int MouseX;
 			public int RelativeMouseX;
+
+			public TabDragData(TabData Tab, int InitialIdx, int MouseX, int RelativeMouseX)
+			{
+				this.Tab = Tab;
+				this.InitialIdx = InitialIdx;
+				this.MouseX = MouseX;
+				this.RelativeMouseX = RelativeMouseX;
+			}
 		}
 
 		List<TabData> Tabs = new List<TabData>();
 		int SelectedTabIdx;
 		int HoverTabIdx;
 		int HighlightTabIdx;
-		TabDragData DragState;
+		TabDragData? DragState;
 
-		public delegate void OnTabChangedDelegate(object NewTabData);
-		public event OnTabChangedDelegate OnTabChanged;
+		public delegate void OnTabChangedDelegate(object? NewTabData);
+		public event OnTabChangedDelegate? OnTabChanged;
 
 		public delegate void OnNewTabClickDelegate(Point Location, MouseButtons Buttons);
-		public event OnNewTabClickDelegate OnNewTabClick;
+		public event OnNewTabClickDelegate? OnNewTabClick;
 
-		public delegate void OnTabClickedDelegate(object TabData, Point Location, MouseButtons Buttons);
-		public event OnTabClickedDelegate OnTabClicked;
+		public delegate void OnTabClickedDelegate(object? TabData, Point Location, MouseButtons Buttons);
+		public event OnTabClickedDelegate? OnTabClicked;
 
 		public delegate bool OnTabClosingDelegate(object TabData);
-		public event OnTabClosingDelegate OnTabClosing;
+		public event OnTabClosingDelegate? OnTabClosing;
 
 		public delegate void OnTabClosedDelegate(object TabData);
-		public event OnTabClosedDelegate OnTabClosed;
+		public event OnTabClosedDelegate? OnTabClosed;
 
 		public delegate void OnTabReorderDelegate();
-		public event OnTabReorderDelegate OnTabReorder;
+		public event OnTabReorderDelegate? OnTabReorder;
 
 		public delegate void OnButtonClickDelegate(int ButtonIdx, Point Location, MouseButtons Buttons);
-		public event OnButtonClickDelegate OnButtonClick;
+		public event OnButtonClickDelegate? OnButtonClick;
 
 		public TabControl()
 		{
 			DoubleBuffered = true;
-			Tabs.Add(new TabData{ Name = "+" });
+			Tabs.Add(new TabData("+", "+"));
 			SelectedTabIdx = -1;
 			HoverTabIdx = -1;
 			HighlightTabIdx = -1;
@@ -97,9 +111,9 @@ namespace UnrealGameSync
 			return Result;
 		}
 
-		public void SetHighlight(int TabIdx, Tuple<Color, float> Highlight)
+		public void SetHighlight(int TabIdx, Tuple<Color, float>? Highlight)
 		{
-			Tuple<Color, float> CurrentHighlight = Tabs[TabIdx].Highlight;
+			Tuple<Color, float>? CurrentHighlight = Tabs[TabIdx].Highlight;
 			if(Highlight == null || CurrentHighlight == null)
 			{
 				if(Highlight != CurrentHighlight)
@@ -139,7 +153,7 @@ namespace UnrealGameSync
 				SelectedTabIdx++;
 			}
 
-			Tabs.Insert(Idx, new TabData{ Name = Name, Data = Data, TintColor = TintColor });
+			Tabs.Insert(Idx, new TabData(Name, Data){ TintColor = TintColor });
 			LayoutTabs();
 			Invalidate();
 			return Idx;
@@ -342,18 +356,18 @@ namespace UnrealGameSync
 						}
 						else
 						{
-							DragState = new TabDragData{ Tab = Tabs[SelectedTabIdx], InitialIdx = SelectedTabIdx, MouseX = e.Location.X, RelativeMouseX = e.Location.X - Tabs[SelectedTabIdx].MinX };
+							DragState = new TabDragData(Tabs[SelectedTabIdx], SelectedTabIdx, e.Location.X, e.Location.X - Tabs[SelectedTabIdx].MinX);
 						}
 					}
 					else
 					{
 						if(HoverTabIdx > Tabs.Count - 1)
 						{
-							OnButtonClick(HoverTabIdx - Tabs.Count, e.Location, e.Button);
+							OnButtonClick?.Invoke(HoverTabIdx - Tabs.Count, e.Location, e.Button);
 						}
 						else if(HoverTabIdx == Tabs.Count - 1)
 						{
-							OnNewTabClick(e.Location, e.Button);
+							OnNewTabClick?.Invoke(e.Location, e.Button);
 						}
 						else
 						{
@@ -372,7 +386,7 @@ namespace UnrealGameSync
 
 			if(OnTabClicked != null)
 			{
-				object TabData = (HoverTabIdx == -1)? null : Tabs[HoverTabIdx].Data;
+				object? TabData = (HoverTabIdx == -1)? null : Tabs[HoverTabIdx].Data;
 				OnTabClicked(TabData, e.Location, e.Button);
 			}
 		}
@@ -558,7 +572,7 @@ namespace UnrealGameSync
 			e.Graphics.DrawLine(SystemPens.ControlLightLight, 0, ClientSize.Height - 1, ClientSize.Width, ClientSize.Height - 1);
 		}
 
-		void DrawBackground(Graphics Graphics, TabData Tab, Brush BackgroundBrush, Brush StripeBrush, Tuple<Color, float> Highlight)
+		void DrawBackground(Graphics Graphics, TabData Tab, Brush BackgroundBrush, Brush StripeBrush, Tuple<Color, float>? Highlight)
 		{
 			float DpiScaleY = Graphics.DpiY / 96.0f;
 
diff --git a/Engine/Source/Programs/UnrealGameSync/UnrealGameSync/Controls/TextBoxWithCueBanner.cs b/Engine/Source/Programs/UnrealGameSync/UnrealGameSync/Controls/TextBoxWithCueBanner.cs
--- a/Engine/Source/Programs/UnrealGameSync/UnrealGameSync/Controls/TextBoxWithCueBanner.cs
+++ b/Engine/Source/Programs/UnrealGameSync/UnrealGameSync/Controls/TextBoxWithCueBanner.cs
@@ -17,7 +17,7 @@ namespace UnrealGameSync
 		[DllImport("user32.dll", CharSet = CharSet.Unicode)]
 		private static extern IntPtr SendMessage(IntPtr hWnd, int Msg, IntPtr Wparam, string Lparam);
 
-		private string CueBannerValue;
+		private string CueBannerValue = String.Empty;
 
 		public string CueBanner
 		{
diff --git a/Engine/Source/Programs/UnrealGameSync/UnrealGameSync/Controls/WorkspaceControl.cs b/Engine/Source/Programs/UnrealGameSync/UnrealGameSync/Controls/WorkspaceControl.cs
--- a/Engine/Source/Programs/UnrealGameSync/UnrealGameSync/Controls/WorkspaceControl.cs
+++ b/Engine/Source/Programs/UnrealGameSync/UnrealGameSync/Controls/WorkspaceControl.cs
@@ -22,8 +22,6 @@ using Microsoft.Extensions.Logging;
 using Microsoft.Extensions.DependencyInjection;
 using System.Diagnostics.CodeAnalysis;
 
-#nullable enable
-
 namespace UnrealGameSync
 {
 	public class UncontrolledChangelist
@@ -61,7 +59,7 @@ namespace UnrealGameSync
 		void UpdateAlertWindows();
 		void UpdateTintColors();
 
-		IssueMonitor CreateIssueMonitor(string ApiUrl, string UserName);
+		IssueMonitor CreateIssueMonitor(string? ApiUrl, string UserName);
 		void ReleaseIssueMonitor(IssueMonitor IssueMonitor);
 	}
 
@@ -212,7 +210,7 @@ namespace UnrealGameSync
 		const string EditorArchiveType = "Editor";
 
 		IWorkspaceControlOwner Owner;
-		string ApiUrl;
+		string? ApiUrl;
 		DirectoryReference DataFolder;
 		IServiceProvider ServiceProvider;
 		ILogger Logger;
@@ -330,7 +328,7 @@ namespace UnrealGameSync
 		// Placeholder text that is in the control and cleared when the user starts editing.
 		static string AuthorFilterPlaceholderText = "<username>";
 
-		public WorkspaceControl(IWorkspaceControlOwner InOwner, string InApiUrl, string InOriginalExecutableFileName, bool bInUnstable, WorkspaceSettings WorkspaceSettings, IServiceProvider InServiceProvider, UserSettings InSettings, OIDCTokenManager InOidcTokenManager)
+		public WorkspaceControl(IWorkspaceControlOwner InOwner, string? InApiUrl, string InOriginalExecutableFileName, bool bInUnstable, WorkspaceSettings WorkspaceSettings, IServiceProvider InServiceProvider, UserSettings InSettings, OIDCTokenManager? InOidcTokenManager)
 		{
 			InitializeComponent();
 
@@ -463,7 +461,7 @@ namespace UnrealGameSync
 			StartupTimer = new System.Threading.Timer(x => MainThreadSynchronizationContext.Post((o) => { if (!IsDisposed) { StartupTimerElapsed(false); } }, null), null, TimeSpan.FromSeconds(20.0), TimeSpan.FromMilliseconds(-1.0));
 			StartupCallbacks = new List<WorkspaceStartupCallback>();
 
-			string IssuesApiUrl = GetIssuesApiUrl();
+			string? IssuesApiUrl = GetIssuesApiUrl();
 			IssueMonitor = InOwner.CreateIssueMonitor(IssuesApiUrl, WorkspaceSettings.PerforceSettings.UserName);
 			IssueMonitor.OnIssuesChanged += IssueMonitor_OnIssuesChangedAsync;
 
@@ -695,7 +693,7 @@ namespace UnrealGameSync
 			ServiceBadges = new List<ServiceBadgeInfo>();
 			foreach (string ServiceBadgeName in ServiceBadgeNames)
 			{
-				BadgeData LatestBuild;
+				BadgeData? LatestBuild;
 				if (EventMonitor.TryGetLatestBadge(ServiceBadgeName, out LatestBuild))
 				{
 					ServiceBadges.Add(new ServiceBadgeInfo(ServiceBadgeName, LatestBuild.Url, LatestBuild.Result));
@@ -1483,7 +1481,7 @@ namespace UnrealGameSync
 				else if (Settings.ShowRobomerge == UserSettings.RobomergeShowChangesOption.Badged)
 				{
 					// if this change has any badges, we'll show it unless it's later filtered out
-					EventSummary Summary = EventMonitor.GetSummaryForChange(Change.Number);
+					EventSummary? Summary = EventMonitor.GetSummaryForChange(Change.Number);
 					if (Summary == null || !Summary.Badges.Any())
 					{
 						return false;
@@ -1493,7 +1491,7 @@ namespace UnrealGameSync
 
 			if (ProjectSettings.FilterBadges.Count > 0)
 			{
-				EventSummary Summary = EventMonitor.GetSummaryForChange(Change.Number);
+				EventSummary? Summary = EventMonitor.GetSummaryForChange(Change.Number);
 				if (Summary == null || !Summary.Badges.Any(x => ProjectSettings.FilterBadges.Contains(x.BadgeName)))
 				{
 					return false;
@@ -1530,7 +1528,7 @@ namespace UnrealGameSync
 			}
 		}
 
-		string GetIssuesApiUrl()
+		string? GetIssuesApiUrl()
 		{
 			string? IssuesApiUrl;
 			if (!TryGetProjectSetting(PerforceMonitor.LatestProjectConfigFile, "IssuesApiUrl", out IssuesApiUrl))
@@ -1543,7 +1541,7 @@ namespace UnrealGameSync
 		void UpdateBuildMetadata()
 		{
 			// Refresh the issue monitor if it's changed
-			string IssuesApiUrl = GetIssuesApiUrl();
+			string? IssuesApiUrl = GetIssuesApiUrl();
 			if (IssuesApiUrl != IssueMonitor.ApiUrl)
 			{
 				Logger.LogInformation("Changing issues API url from {OldApiUrl} to {ApiUrl}", IssueMonitor.ApiUrl, IssuesApiUrl);
@@ -1581,7 +1579,7 @@ namespace UnrealGameSync
 			{
 				if (Item != null && Item.Tag != null)
 				{
-					EventSummary Summary = EventMonitor.GetSummaryForChange(((ChangesRecord)Item.Tag).Number);
+					EventSummary? Summary = EventMonitor.GetSummaryForChange(((ChangesRecord)Item.Tag).Number);
 					if (Summary != null)
 					{
 						foreach (BadgeData Badge in Summary.Badges)
@@ -1695,7 +1693,7 @@ namespace UnrealGameSync
 				return true;
 			}
 
-			EventSummary Review = EventMonitor.GetSummaryForChange(ChangeNumber);
+			EventSummary? Review = EventMonitor.GetSummaryForChange(ChangeNumber);
 			if (Review != null)
 			{
 				if (Review.LastStarReview != null && Review.LastStarReview.Type == EventType.Starred)
@@ -1877,7 +1875,7 @@ namespace UnrealGameSync
 			Dictionary<string, BadgeData> TypeToLastSucceededBuild = new Dictionary<string, BadgeData>();
 			for (int Idx = SortedChangeNumbers.Count - 1; Idx >= 0; Idx--)
 			{
-				EventSummary Summary = EventMonitor.GetSummaryForChange(SortedChangeNumbers[Idx]);
+				EventSummary? Summary = EventMonitor.GetSummaryForChange(SortedChangeNumbers[Idx]);
 				if (Summary != null)
 				{
 					foreach (BadgeData Badge in Summary.Badges)
@@ -2137,7 +2135,7 @@ namespace UnrealGameSync
 				LayoutInfo.DescriptionBadges = CreateDescriptionBadges(Change);
 				LayoutInfo.TypeBadges = CreateTypeBadges(Change.Number);
 
-				EventSummary Summary = EventMonitor.GetSummaryForChange(Change.Number);
+				EventSummary? Summary = EventMonitor.GetSummaryForChange(Change.Number);
 				LayoutInfo.BuildBadges = CreateBuildBadges(Change.Number, Summary);
 				LayoutInfo.CustomBadges = CreateCustomBadges(Change.Number, Summary);
 
@@ -2204,7 +2202,7 @@ namespace UnrealGameSync
 
 			if (e.ColumnIndex == IconColumn.Index)
 			{
-				EventSummary Summary = EventMonitor.GetSummaryForChange(Change.Number);
+				EventSummary? Summary = EventMonitor.GetSummaryForChange(Change.Number);
 
 				float MinX = 4 * DpiScaleX;
 				if ((Summary != null && EventMonitor.WasSyncedByCurrentUser(Summary.ChangeNumber)) || (Workspace != null && Workspace.CurrentChangeNumber == Change.Number))
@@ -2360,7 +2358,7 @@ namespace UnrealGameSync
 
 					if (Change.Number == Workspace.CurrentChangeNumber)
 					{
-						EventData Review = EventMonitor.GetReviewByCurrentUser(Change.Number);
+						EventData? Review = EventMonitor.GetReviewByCurrentUser(Change.Number);
 
 						MaxX -= FrownIcon.Width * DpiScaleX;
 						e.Graphics.DrawImage(Properties.Resources.Icons, MaxX, IconY, (Review == null || !EventMonitor.IsPositiveReview(Review.Type)) ? FrownIcon : DisabledFrownIcon, GraphicsUnit.Pixel);
@@ -2380,7 +2378,7 @@ namespace UnrealGameSync
 					{
 						StringBuilder SummaryTextBuilder = new StringBuilder();
 
-						EventSummary Summary = EventMonitor.GetSummaryForChange(Change.Number);
+						EventSummary? Summary = EventMonitor.GetSummaryForChange(Change.Number);
 
 						AppendItemList(SummaryTextBuilder, " ", "Under investigation by {0}.", EventMonitor.GetInvestigatingUsers(Change.Number).Select(x => FormatUserName(x)));
 
@@ -2755,7 +2753,7 @@ namespace UnrealGameSync
 			return false;
 		}
 
-		private List<BadgeInfo> CreateBuildBadges(int ChangeNumber, EventSummary Summary)
+		private List<BadgeInfo> CreateBuildBadges(int ChangeNumber, EventSummary? Summary)
 		{
 			List<BadgeInfo> Badges = new List<BadgeInfo>();
 
@@ -2816,7 +2814,7 @@ namespace UnrealGameSync
 			return new BadgeInfo(BadgeLabel, BadgeGroup, UniqueId, BadgeColor, HoverBadgeColor, ClickHandler);
 		}
 
-		private Dictionary<string, List<BadgeInfo>> CreateCustomBadges(int ChangeNumber, EventSummary Summary)
+		private Dictionary<string, List<BadgeInfo>> CreateCustomBadges(int ChangeNumber, EventSummary? Summary)
 		{
 			Dictionary<string, List<BadgeInfo>> ColumnNameToBadges = new Dictionary<string, List<BadgeInfo>>();
 			if (Summary != null && Summary.Badges.Count > 0)
@@ -3773,7 +3771,7 @@ namespace UnrealGameSync
 			if (TryGetProjectSetting(PerforceMonitor.LatestProjectConfigFile, "SdkInfo", out SdkInfoEntries))
 			{
 				Dictionary<string, string> Variables = GetWorkspaceVariables(-1);
-				SdkInfoWindow Window = new SdkInfoWindow(SdkInfoEntries, Variables, BadgeFont);
+				SdkInfoWindow Window = new SdkInfoWindow(SdkInfoEntries, Variables, BadgeFont!);
 				Window.ShowDialog();
 			}
 		}
@@ -4184,7 +4182,7 @@ namespace UnrealGameSync
 						BuildListContextMenu_StartInvestigating.Visible = !EventMonitor.IsUnderInvestigationByCurrentUser(ContextMenuChange.Number);
 						BuildListContextMenu_FinishInvestigating.Visible = EventMonitor.IsUnderInvestigation(ContextMenuChange.Number);
 
-						string CommentText;
+						string? CommentText;
 						bool bHasExistingComment = EventMonitor.GetCommentByCurrentUser(ContextMenuChange.Number, out CommentText);
 						BuildListContextMenu_LeaveComment.Visible = !bHasExistingComment;
 						BuildListContextMenu_EditComment.Visible = bHasExistingComment;
@@ -4215,7 +4213,7 @@ namespace UnrealGameSync
 						BuildListContextMenu_MarkBad.Visible = !bIsBisectMode;
 						BuildListContextMenu_WithdrawReview.Visible = !bIsBisectMode;
 
-						EventSummary Summary = EventMonitor.GetSummaryForChange(ContextMenuChange.Number);
+						EventSummary? Summary = EventMonitor.GetSummaryForChange(ContextMenuChange.Number);
 						bool bStarred = (Summary != null && Summary.LastStarReview != null && Summary.LastStarReview.Type == EventType.Starred);
 						BuildListContextMenu_AddStar.Visible = !bStarred;
 						BuildListContextMenu_RemoveStar.Visible = bStarred;
@@ -4744,7 +4742,7 @@ namespace UnrealGameSync
 		{
 			if (ContextMenuChange != null)
 			{
-				EventSummary Summary = EventMonitor.GetSummaryForChange(ContextMenuChange.Number);
+				EventSummary? Summary = EventMonitor.GetSummaryForChange(ContextMenuChange.Number);
 				if (Summary != null && Summary.LastStarReview != null && Summary.LastStarReview.Type == EventType.Starred)
 				{
 					string Message = String.Format("This change was starred by {0}. Are you sure you want to remove it?", FormatUserName(Summary.LastStarReview.UserName));
@@ -4856,7 +4854,7 @@ namespace UnrealGameSync
 						}
 					}
 
-					EventSummary Summary = EventMonitor.GetSummaryForChange(Change.Number);
+					EventSummary? Summary = EventMonitor.GetSummaryForChange(Change.Number);
 					if (Summary != null)
 					{
 						StringBuilder SummaryText = new StringBuilder();
@@ -5081,7 +5079,7 @@ namespace UnrealGameSync
 					}
 					else if (ChangeType == LatestChangeType.Good)
 					{
-						EventSummary Summary = EventMonitor.GetSummaryForChange(Change.Number);
+						EventSummary? Summary = EventMonitor.GetSummaryForChange(Change.Number);
 						if (Summary != null && Summary.Verdict == ReviewVerdict.Good && CanSyncChange(Change.Number))
 						{
 							ChangeNumber = FindNewestGoodContentChange(Change.Number);
@@ -5090,7 +5088,7 @@ namespace UnrealGameSync
 					}
 					else if (ChangeType == LatestChangeType.Starred)
 					{
-						EventSummary Summary = EventMonitor.GetSummaryForChange(Change.Number);
+						EventSummary? Summary = EventMonitor.GetSummaryForChange(Change.Number);
 						if (((Summary != null && Summary.LastStarReview != null && Summary.LastStarReview.Type == EventType.Starred) || PromotedChangeNumbers.Contains(Change.Number)) && CanSyncChange(Change.Number))
 						{
 							ChangeNumber = FindNewestGoodContentChange(Change.Number);
@@ -5122,7 +5120,7 @@ namespace UnrealGameSync
 					break;
 				}
 
-				EventSummary Summary = EventMonitor.GetSummaryForChange(NextChangeNumber);
+				EventSummary? Summary = EventMonitor.GetSummaryForChange(NextChangeNumber);
 				if (Summary != null && Summary.Verdict == ReviewVerdict.Bad)
 				{
 					break;
@@ -5140,7 +5138,7 @@ namespace UnrealGameSync
 			{
 				LeaveCommentWindow LeaveComment = new LeaveCommentWindow();
 
-				string CommentText;
+				string? CommentText;
 				if (EventMonitor.GetCommentByCurrentUser(ContextMenuChange.Number, out CommentText))
 				{
 					LeaveComment.CommentTextBox.Text = CommentText;
@@ -5382,7 +5380,7 @@ namespace UnrealGameSync
 						ConfigObject? DefaultObject;
 						ProjectBuildStepObjects.TryGetValue(Step.UniqueId, out DefaultObject);
 
-						ConfigObject UserConfigObject = Step.ToConfigObject(DefaultObject);
+						ConfigObject? UserConfigObject = Step.ToConfigObject(DefaultObject);
 						if (UserConfigObject != null && UserConfigObject.Pairs.Any(x => x.Key != "UniqueId"))
 						{
 							ModifiedBuildSteps.Add(UserConfigObject);
diff --git a/Engine/Source/Programs/UnrealGameSync/UnrealGameSync/DeleteFilesTask.cs b/Engine/Source/Programs/UnrealGameSync/UnrealGameSync/DeleteFilesTask.cs
--- a/Engine/Source/Programs/UnrealGameSync/UnrealGameSync/DeleteFilesTask.cs
+++ b/Engine/Source/Programs/UnrealGameSync/UnrealGameSync/DeleteFilesTask.cs
@@ -10,8 +10,6 @@ using System.Text;
 using System.Threading;
 using System.Threading.Tasks;
 
-#nullable enable
-
 namespace UnrealGameSync
 {
 	static class DeleteFilesTask
diff --git a/Engine/Source/Programs/UnrealGameSync/UnrealGameSync/DeploymentSettings.cs b/Engine/Source/Programs/UnrealGameSync/UnrealGameSync/DeploymentSettings.cs
--- a/Engine/Source/Programs/UnrealGameSync/UnrealGameSync/DeploymentSettings.cs
+++ b/Engine/Source/Programs/UnrealGameSync/UnrealGameSync/DeploymentSettings.cs
@@ -3,6 +3,7 @@
 using Microsoft.Extensions.Logging;
 using System;
 using System.Collections.Generic;
+using System.Diagnostics.CodeAnalysis;
 using System.IO;
 using System.Linq;
 using System.Text;
@@ -30,7 +31,7 @@ namespace UnrealGameSync
 		/// SQL connection string used to connect to the database for telemetry and review data. The 'Program' class is a partial class, to allow an
 		/// opportunistically included C# source file in NotForLicensees/ProgramSettings.cs to override this value in a static constructor.
 		/// </summary>
-		public static readonly string ApiUrl = null;
+		public static readonly string? ApiUrl = null;
 
 		/// <summary>
 		/// Servers to connect to for issue details by default
@@ -40,23 +41,18 @@ namespace UnrealGameSync
 		/// <summary>
 		/// The issue api to use for URL handler events
 		/// </summary>
-		public static readonly string UrlHandleIssueApi = null;
+		public static readonly string? UrlHandleIssueApi = null;
 
 		/// <summary>
 		/// Specifies the depot path to sync down the stable version of UGS from, without a trailing slash (eg. //depot/UnrealGameSync/bin). This is a site-specific setting. 
 		/// The UnrealGameSync executable should be located at Release/UnrealGameSync.exe under this path, with any dependent DLLs.
 		/// </summary>
-		public static readonly string DefaultDepotPath = null;
+		public static readonly string? DefaultDepotPath = null;
 
 		/// <summary>
 		/// Depot path to sync additional tools from
 		/// </summary>
-		public static readonly string ToolsDepotPath = null;
-
-		/// <summary>
-		/// URL to send info about NET Core versions to.
-		/// </summary>
-		public static readonly string NetCoreTelemetryUrl = null;
+		public static readonly string? ToolsDepotPath = null;
 
 #if WITH_TELEMETRY
 		/// <summary>
@@ -73,12 +69,12 @@ namespace UnrealGameSync
 		/// <param name="Logger">The logger</param>
 		/// <param name="Error">Receives an error on failure</param>
 		/// <returns></returns>
-		public delegate bool DetectProjectSettingsEvent(WorkspaceSettings Workspace, ILogger Logger, out string Error);
+		public delegate bool DetectProjectSettingsEvent(WorkspaceSettings Workspace, ILogger Logger, [NotNullWhen(false)] out string? Error);
 
 		/// <summary>
 		/// Called to validate the project settings
 		/// </summary>
-		public static DetectProjectSettingsEvent OnDetectProjectSettings = null;
+		public static DetectProjectSettingsEvent? OnDetectProjectSettings = null;
 #endif
 	}
 }
diff --git a/Engine/Source/Programs/UnrealGameSync/UnrealGameSync/DetectProjectSettingsTask.cs b/Engine/Source/Programs/UnrealGameSync/UnrealGameSync/DetectProjectSettingsTask.cs
--- a/Engine/Source/Programs/UnrealGameSync/UnrealGameSync/DetectProjectSettingsTask.cs
+++ b/Engine/Source/Programs/UnrealGameSync/UnrealGameSync/DetectProjectSettingsTask.cs
@@ -12,8 +12,6 @@ using System.Text.RegularExpressions;
 using System.Threading;
 using System.Threading.Tasks;
 
-#nullable enable
-
 namespace UnrealGameSync
 {
 	public class WorkspaceSettings
@@ -311,7 +309,7 @@ namespace UnrealGameSync
 				// Run any event hooks
 				if (DeploymentSettings.OnDetectProjectSettings != null)
 				{
-					string Message;
+					string? Message;
 					if (!DeploymentSettings.OnDetectProjectSettings(WorkspaceSettings, Logger, out Message))
 					{
 						throw new UserErrorException(Message);
diff --git a/Engine/Source/Programs/UnrealGameSync/UnrealGameSync/EventMonitor.cs b/Engine/Source/Programs/UnrealGameSync/UnrealGameSync/EventMonitor.cs
--- a/Engine/Source/Programs/UnrealGameSync/UnrealGameSync/EventMonitor.cs
+++ b/Engine/Source/Programs/UnrealGameSync/UnrealGameSync/EventMonitor.cs
@@ -13,6 +13,7 @@ using EpicGames.Core;
 using System.Threading.Tasks;
 using Microsoft.Extensions.Logging;
 using Microsoft.Extensions.DependencyInjection;
+using System.Diagnostics.CodeAnalysis;
 
 namespace UnrealGameSync
 {
@@ -48,18 +49,18 @@ namespace UnrealGameSync
 	{
 		public long Id { get; set; }
 		public int Change { get; set; }
-		public string UserName { get; set; }
+		public string UserName { get; set; } = String.Empty;
 		public EventType Type { get; set; }
-		public string Project { get; set; }
+		public string Project { get; set; } = String.Empty;
 	}
 
 	class CommentData
 	{
 		public long Id { get; set; }
 		public int ChangeNumber { get; set; }
-		public string UserName { get; set; }
-		public string Text { get; set; }
-		public string Project { get; set; }
+		public string UserName { get; set; } = String.Empty;
+		public string Text { get; set; } = String.Empty;
+		public string Project { get; set; } = String.Empty;
 	}
 
 	enum BadgeResult
@@ -75,10 +76,10 @@ namespace UnrealGameSync
 	{
 		public long Id { get; set; }
 		public int ChangeNumber { get; set; }
-		public string BuildType { get; set; }
+		public string BuildType { get; set; } = String.Empty;
 		public BadgeResult Result { get; set; }
-		public string Url { get; set; }
-		public string Project { get; set; }
+		public string Url { get; set; } = String.Empty;
+		public string Project { get; set; } = String.Empty;
 
 		public bool IsSuccess
 		{
@@ -152,46 +153,46 @@ namespace UnrealGameSync
 
 	class GetUserDataResponseV2
 	{
-		public string User { get; set; }
+		public string User { get; set; } = String.Empty;
 		public long? SyncTime { get; set; }
 		public UgsUserVote? Vote { get; set; }
-		public string Comment { get; set; }
+		public string Comment { get; set; } = String.Empty;
 		public bool? Investigating { get; set; }
 		public bool? Starred { get; set; }
 	}
 
 	class GetBadgeDataResponseV2
 	{
-		public string Name { get; set; }
-		public string Url { get; set; }
+		public string Name { get; set; } = String.Empty;
+		public string Url { get; set; } = String.Empty;
 		public BadgeResult State { get; set; }
 	}
 
 	class GetMetadataResponseV2
 	{
 		public int Change { get; set; }
-		public string Project { get; set; }
-		public List<GetUserDataResponseV2> Users { get; set; }
-		public List<GetBadgeDataResponseV2> Badges { get; set; }
+		public string Project { get; set; } = String.Empty;
+		public List<GetUserDataResponseV2> Users { get; set; } = new List<GetUserDataResponseV2>();
+		public List<GetBadgeDataResponseV2> Badges { get; set; } = new List<GetBadgeDataResponseV2>();
 	}
 
 	class GetMetadataListResponseV2
 	{
 		public long SequenceNumber { get; set; }
-		public List<GetMetadataResponseV2> Items { get; set; }
+		public List<GetMetadataResponseV2> Items { get; set; } = new List<GetMetadataResponseV2>();
 	}
 
 	class UpdateMetadataRequestV2
 	{
-		public string Stream { get; set; }
+		public string? Stream { get; set; }
 		public int Change { get; set; }
-		public string Project { get; set; }
-		public string UserName { get; set; }
+		public string? Project { get; set; }
+		public string? UserName { get; set; }
 		public bool? Synced { get; set; }
-		public string Vote { get; set; }
+		public string? Vote { get; set; }
 		public bool? Investigating { get; set; }
 		public bool? Starred { get; set; }
-		public string Comment { get; set; }
+		public string? Comment { get; set; }
 	}
 
 	class EventSummary
@@ -201,23 +202,23 @@ namespace UnrealGameSync
 		public List<EventData> SyncEvents = new List<EventData>();
 		public List<EventData> Reviews = new List<EventData>();
 		public List<string> CurrentUsers = new List<string>();
-		public EventData LastStarReview;
+		public EventData? LastStarReview;
 		public List<BadgeData> Badges = new List<BadgeData>();
 		public List<CommentData> Comments = new List<CommentData>();
 
-		public GetMetadataResponseV2 SharedMetadata;
-		public GetMetadataResponseV2 ProjectMetadata;
+		public GetMetadataResponseV2? SharedMetadata;
+		public GetMetadataResponseV2? ProjectMetadata;
 	}
 
 	class EventMonitor : IDisposable
 	{
-		string ApiUrl;
+		string? ApiUrl;
 		int ApiVersion;
 		string Project;
 		string CurrentUserName;
 		SynchronizationContext SynchronizationContext;
 		CancellationTokenSource CancellationSource;
-		Task WorkerTask;
+		Task? WorkerTask;
 		AsyncEvent RefreshEvent = new AsyncEvent();
 		ConcurrentQueue<EventData> OutgoingEvents = new ConcurrentQueue<EventData>();
 		ConcurrentQueue<EventData> IncomingEvents = new ConcurrentQueue<EventData>();
@@ -232,7 +233,7 @@ namespace UnrealGameSync
 		LatestData LatestIds;
 		HashSet<int> FilterChangeNumbers = new HashSet<int>();
 		List<EventData> InvestigationEvents = new List<EventData>();
-		List<EventData> ActiveInvestigations = new List<EventData>();
+		List<EventData>? ActiveInvestigations;
 
 		// MetadataV2
 		string MetadataStream;
@@ -242,20 +243,21 @@ namespace UnrealGameSync
 		int NewMinChange;
 		long MetadataSequenceNumber;
 
-		public event Action OnUpdatesReady;
+		public Action? OnUpdatesReady;
 
-		public EventMonitor(string InApiUrl, string InProject, string InCurrentUserName, IServiceProvider ServiceProvider)
+		public EventMonitor(string? InApiUrl, string InProject, string InCurrentUserName, IServiceProvider ServiceProvider)
 		{
 			ApiUrl = InApiUrl;
 			Project = InProject;
 			CurrentUserName = InCurrentUserName;
 			Logger = ServiceProvider.GetRequiredService<ILogger<EventMonitor>>();
 			AsyncDisposer = ServiceProvider.GetRequiredService<IAsyncDisposer>();
-			SynchronizationContext = SynchronizationContext.Current;
+			SynchronizationContext = SynchronizationContext.Current!;
 			CancellationSource = new CancellationTokenSource();
 
 			LatestIds = new LatestData { LastBuildId = 0, LastCommentId = 0, LastEventId = 0 };
 
+			MetadataProject = String.Empty;
 			MetadataStream = Project.ToLowerInvariant().TrimEnd('/');
 			if (MetadataStream.StartsWith("//", StringComparison.Ordinal))
 			{
@@ -359,7 +361,7 @@ namespace UnrealGameSync
 
 		protected EventSummary FindOrAddSummary(int ChangeNumber)
 		{
-			EventSummary Summary;
+			EventSummary? Summary;
 			if(!ChangeNumberToSummary.TryGetValue(ChangeNumber, out Summary))
 			{
 				Summary = new EventSummary();
@@ -373,29 +375,29 @@ namespace UnrealGameSync
 		{
 			get;
 			private set;
-		}
+		} = String.Empty;
 
 		public void ApplyUpdates()
 		{
-			GetMetadataResponseV2 Metadata;
+			GetMetadataResponseV2? Metadata;
 			while (IncomingMetadata.TryDequeue(out Metadata))
 			{
 				ConvertMetadataToEvents(Metadata);
 			}
 
-			EventData Event;
+			EventData? Event;
 			while(IncomingEvents.TryDequeue(out Event))
 			{
 				ApplyEventUpdate(Event);
 			}
 
-			BadgeData Badge;
+			BadgeData? Badge;
 			while(IncomingBadges.TryDequeue(out Badge))
 			{
 				ApplyBadgeUpdate(Badge);
 			}
 
-			CommentData Comment;
+			CommentData? Comment;
 			while(IncomingComments.TryDequeue(out Comment))
 			{
 				ApplyCommentUpdate(Comment);
@@ -436,7 +438,7 @@ namespace UnrealGameSync
 			else if(IsReview(Event.Type))
 			{
 				// Try to find an existing review by this user. If we already have a newer review, ignore this one. Otherwise remove it.
-				EventData ExistingReview = Summary.Reviews.Find(x => String.Compare(x.UserName, Event.UserName, true) == 0);
+				EventData? ExistingReview = Summary.Reviews.Find(x => String.Compare(x.UserName, Event.UserName, true) == 0);
 				if(ExistingReview != null)
 				{
 					if(ExistingReview.Id <= Event.Id)
@@ -463,7 +465,7 @@ namespace UnrealGameSync
 		{
 			EventSummary Summary = FindOrAddSummary(Badge.ChangeNumber);
 
-			BadgeData ExistingBadge = Summary.Badges.Find(x => x.ChangeNumber == Badge.ChangeNumber && x.BuildType == Badge.BuildType);
+			BadgeData? ExistingBadge = Summary.Badges.Find(x => x.ChangeNumber == Badge.ChangeNumber && x.BuildType == Badge.BuildType);
 			if(ExistingBadge != null)
 			{
 				if(ExistingBadge.Id <= Badge.Id)
@@ -479,7 +481,7 @@ namespace UnrealGameSync
 			Summary.Badges.Add(Badge);
 			Summary.Verdict = GetVerdict(Summary.Reviews, Summary.Badges);
 
-			BadgeData LatestBadge;
+			BadgeData? LatestBadge;
 			if(!BadgeNameToLatestData.TryGetValue(Badge.BadgeName, out LatestBadge) || Badge.ChangeNumber > LatestBadge.ChangeNumber || (Badge.ChangeNumber == LatestBadge.ChangeNumber && Badge.Id > LatestBadge.Id))
 			{
 				BadgeNameToLatestData[Badge.BadgeName] = Badge;
@@ -546,7 +548,7 @@ namespace UnrealGameSync
 			EventSummary NewSummary = new EventSummary();
 			NewSummary.ChangeNumber = Metadata.Change;
 
-			EventSummary Summary;
+			EventSummary? Summary;
 			if (ChangeNumberToSummary.TryGetValue(Metadata.Change, out Summary))
 			{
 				foreach (string CurrentUser in Summary.CurrentUsers)
@@ -684,7 +686,7 @@ namespace UnrealGameSync
 			if(Event.Type == EventType.Syncing && FilterChangeNumbers.Contains(Event.Change) && !String.IsNullOrEmpty(Event.UserName))
 			{
 				// Update the active users list for this change
-				EventData LastSync;
+				EventData? LastSync;
 				if(UserNameToLastSyncEvent.TryGetValue(Event.UserName, out LastSync))
 				{
 					if(Event.Id > LastSync.Id)
@@ -704,8 +706,8 @@ namespace UnrealGameSync
 
 		async Task PollForUpdatesAsync(CancellationToken CancellationToken)
 		{
-			EventData Event = null;
-			CommentData Comment = null;
+			EventData? Event = null;
+			CommentData? Comment = null;
 			bool bUpdateThrottledRequests = true;
 			double RequestThrottle = 90; // seconds to wait for throttled request;
 			Stopwatch Timer = Stopwatch.StartNew();
@@ -818,7 +820,7 @@ namespace UnrealGameSync
 			}
 		}
 
-		async Task SendMetadataUpdateUpdateV2Async(int Change, string Project, string UserName, EventType? Event, string Comment, CancellationToken CancellationToken)
+		async Task SendMetadataUpdateUpdateV2Async(int Change, string Project, string UserName, EventType? Event, string? Comment, CancellationToken CancellationToken)
 		{
 			UpdateMetadataRequestV2 Update = new UpdateMetadataRequestV2();
 			Update.Stream = MetadataStream;
@@ -1016,16 +1018,16 @@ namespace UnrealGameSync
 			}
 		}
 
-		public bool GetCommentByCurrentUser(int ChangeNumber, out string CommentText)
+		public bool GetCommentByCurrentUser(int ChangeNumber, [NotNullWhen(true)] out string? CommentText)
 		{
-			EventSummary Summary = GetSummaryForChange(ChangeNumber);
+			EventSummary? Summary = GetSummaryForChange(ChangeNumber);
 			if(Summary == null)
 			{
 				CommentText = null;
 				return false;
 			}
 
-			CommentData Comment = Summary.Comments.Find(x => String.Compare(x.UserName, CurrentUserName, true) == 0);
+			CommentData? Comment = Summary.Comments.Find(x => String.Compare(x.UserName, CurrentUserName, true) == 0);
 			if(Comment == null || String.IsNullOrWhiteSpace(Comment.Text))
 			{
 				CommentText = null;
@@ -1036,15 +1038,15 @@ namespace UnrealGameSync
 			return true;
 		}
 
-		public EventData GetReviewByCurrentUser(int ChangeNumber)
+		public EventData? GetReviewByCurrentUser(int ChangeNumber)
 		{
-			EventSummary Summary = GetSummaryForChange(ChangeNumber);
+			EventSummary? Summary = GetSummaryForChange(ChangeNumber);
 			if(Summary == null)
 			{
 				return null;
 			}
 
-			EventData Event = Summary.Reviews.FirstOrDefault(x => String.Compare(x.UserName, CurrentUserName, true) == 0);
+			EventData? Event = Summary.Reviews.FirstOrDefault(x => String.Compare(x.UserName, CurrentUserName, true) == 0);
 			if(Event == null || Event.Type == EventType.Unknown)
 			{
 				return null;
@@ -1053,14 +1055,14 @@ namespace UnrealGameSync
 			return Event;
 		}
 
-		public EventSummary GetSummaryForChange(int ChangeNumber)
+		public EventSummary? GetSummaryForChange(int ChangeNumber)
 		{
-			EventSummary Summary;
+			EventSummary? Summary;
 			ChangeNumberToSummary.TryGetValue(ChangeNumber, out Summary);
 			return Summary;
 		}
 
-		public bool TryGetLatestBadge(string BuildType, out BadgeData BadgeData)
+		public bool TryGetLatestBadge(string BuildType, [NotNullWhen(true)] out BadgeData? BadgeData)
 		{
 			return BadgeNameToLatestData.TryGetValue(BuildType, out BadgeData);
 		}
@@ -1082,7 +1084,7 @@ namespace UnrealGameSync
 
 		public bool WasSyncedByCurrentUser(int ChangeNumber)
 		{
-			EventSummary Summary = GetSummaryForChange(ChangeNumber);
+			EventSummary? Summary = GetSummaryForChange(ChangeNumber);
 			return (Summary != null && Summary.SyncEvents.Any(x => x.Type == EventType.Syncing && String.Compare(x.UserName, CurrentUserName, true) == 0));
 		}
 
@@ -1160,13 +1162,16 @@ namespace UnrealGameSync
 			UpdateActiveInvestigations();
 
 			int StartChangeNumber = -1;
-			foreach(EventData ActiveInvestigation in ActiveInvestigations)
+			if (ActiveInvestigations != null)
 			{
-				if(String.Compare(ActiveInvestigation.UserName, CurrentUserName, true) == 0)
+				foreach (EventData ActiveInvestigation in ActiveInvestigations)
 				{
-					if(ActiveInvestigation.Change <= LastChangeNumber && (StartChangeNumber == -1 || ActiveInvestigation.Change < StartChangeNumber))
+					if (String.Compare(ActiveInvestigation.UserName, CurrentUserName, true) == 0)
 					{
-						StartChangeNumber = ActiveInvestigation.Change;
+						if (ActiveInvestigation.Change <= LastChangeNumber && (StartChangeNumber == -1 || ActiveInvestigation.Change < StartChangeNumber))
+						{
+							StartChangeNumber = ActiveInvestigation.Change;
+						}
 					}
 				}
 			}
diff --git a/Engine/Source/Programs/UnrealGameSync/UnrealGameSync/FileFilter.cs b/Engine/Source/Programs/UnrealGameSync/UnrealGameSync/FileFilter.cs
--- a/Engine/Source/Programs/UnrealGameSync/UnrealGameSync/FileFilter.cs
+++ b/Engine/Source/Programs/UnrealGameSync/UnrealGameSync/FileFilter.cs
@@ -281,7 +281,7 @@ namespace UnrealGameSync
 			LastNode.Type = Type;
 
 			// Update the maximums along that path
-			for (FileFilterNode UpdateNode = LastNode; UpdateNode != null; UpdateNode = UpdateNode.Parent)
+			for (FileFilterNode? UpdateNode = LastNode; UpdateNode != null; UpdateNode = UpdateNode.Parent)
 			{
 				if (Type == FileFilterType.Include)
 				{
@@ -437,7 +437,7 @@ namespace UnrealGameSync
 		/// <summary>
 		/// Node which this is parented to. Null for the root node.
 		/// </summary>
-		public readonly FileFilterNode Parent;
+		public readonly FileFilterNode? Parent;
 
 		/// <summary>
 		/// Pattern to match for this node. May contain * or ? wildcards.
@@ -472,7 +472,7 @@ namespace UnrealGameSync
 		/// <summary>
 		/// Default constructor.
 		/// </summary>
-		public FileFilterNode(FileFilterNode InParent, string InPattern)
+		public FileFilterNode(FileFilterNode? InParent, string InPattern)
 		{
 			Parent = InParent;
 			Pattern = InPattern;
diff --git a/Engine/Source/Programs/UnrealGameSync/UnrealGameSync/FindFoldersToCleanTask.cs b/Engine/Source/Programs/UnrealGameSync/UnrealGameSync/FindFoldersToCleanTask.cs
--- a/Engine/Source/Programs/UnrealGameSync/UnrealGameSync/FindFoldersToCleanTask.cs
+++ b/Engine/Source/Programs/UnrealGameSync/UnrealGameSync/FindFoldersToCleanTask.cs
@@ -11,8 +11,6 @@ using System.Text;
 using System.Threading;
 using System.Threading.Tasks;
 
-#nullable enable
-
 namespace UnrealGameSync
 {
 	class FolderToClean
diff --git a/Engine/Source/Programs/UnrealGameSync/UnrealGameSync/Forms/ApplicationSettingsWindow.cs b/Engine/Source/Programs/UnrealGameSync/UnrealGameSync/Forms/ApplicationSettingsWindow.cs
--- a/Engine/Source/Programs/UnrealGameSync/UnrealGameSync/Forms/ApplicationSettingsWindow.cs
+++ b/Engine/Source/Programs/UnrealGameSync/UnrealGameSync/Forms/ApplicationSettingsWindow.cs
@@ -17,8 +17,6 @@ using System.Threading.Tasks;
 using System.Windows.Forms;
 using UnrealGameSync;
 
-#nullable enable
-
 namespace UnrealGameSync
 {
 	partial class ApplicationSettingsWindow : Form
@@ -99,7 +97,7 @@ namespace UnrealGameSync
 
 			this.DepotPathTextBox.Text = InitialDepotPath;
 			this.DepotPathTextBox.Select(DepotPathTextBox.TextLength, 0);
-			this.DepotPathTextBox.CueBanner = DeploymentSettings.DefaultDepotPath;
+			this.DepotPathTextBox.CueBanner = DeploymentSettings.DefaultDepotPath ?? String.Empty;
 
 			this.UseUnstableBuildCheckBox.Checked = bUnstable;
 
@@ -192,12 +190,14 @@ namespace UnrealGameSync
 				{
 					PerforceSettings Settings = Utility.OverridePerforceSettings(DefaultPerforceSettings, ServerAndPort, UserName);
 
-					string TestDepotPath = DepotPath ?? DeploymentSettings.DefaultDepotPath;
-
-					ModalTask? Task = PerforceModalTask.Execute(this, "Checking connection", "Checking connection, please wait...", Settings, (p, c) => PerforceTestConnectionTask.RunAsync(p, TestDepotPath, c), Logger);
-					if (Task == null || !Task.Succeeded)
+					string? TestDepotPath = DepotPath ?? DeploymentSettings.DefaultDepotPath;
+					if (TestDepotPath != null)
 					{
-						return;
+						ModalTask? Task = PerforceModalTask.Execute(this, "Checking connection", "Checking connection, please wait...", Settings, (p, c) => PerforceTestConnectionTask.RunAsync(p, TestDepotPath, c), Logger);
+						if (Task == null || !Task.Succeeded)
+						{
+							return;
+						}
 					}
 				}
 
diff --git a/Engine/Source/Programs/UnrealGameSync/UnrealGameSync/Forms/AutomatedBuildWindow.cs b/Engine/Source/Programs/UnrealGameSync/UnrealGameSync/Forms/AutomatedBuildWindow.cs
--- a/Engine/Source/Programs/UnrealGameSync/UnrealGameSync/Forms/AutomatedBuildWindow.cs
+++ b/Engine/Source/Programs/UnrealGameSync/UnrealGameSync/Forms/AutomatedBuildWindow.cs
@@ -6,6 +6,7 @@ using System;
 using System.Collections.Generic;
 using System.ComponentModel;
 using System.Data;
+using System.Diagnostics.CodeAnalysis;
 using System.Drawing;
 using System.IO;
 using System.Linq;
@@ -24,18 +25,26 @@ namespace UnrealGameSync.Forms
 			public string ProjectPath;
 			public bool bSync;
 			public string ExecCommand;
+
+			public BuildInfo(AutomatedSyncWindow.WorkspaceInfo SelectedWorkspaceInfo, string ProjectPath, bool bSync, string ExecCommand)
+			{
+				this.SelectedWorkspaceInfo = SelectedWorkspaceInfo;
+				this.ProjectPath = ProjectPath;
+				this.bSync = bSync;
+				this.ExecCommand = ExecCommand;
+			}
 		}
 
 		string StreamName;
 		IServiceProvider ServiceProvider;
 
-		string ServerAndPortOverride;
-		string UserNameOverride;
+		string? ServerAndPortOverride;
+		string? UserNameOverride;
 		IPerforceSettings DefaultPerforceSettings;
 
-		BuildInfo Result;
+		BuildInfo? Result;
 
-		private AutomatedBuildWindow(string StreamName, int Changelist, string Command, IPerforceSettings DefaultPerforceSettings, string DefaultWorkspaceName, string DefaultProjectPath, IServiceProvider ServiceProvider)
+		private AutomatedBuildWindow(string StreamName, int Changelist, string Command, IPerforceSettings DefaultPerforceSettings, string? DefaultWorkspaceName, string? DefaultProjectPath, IServiceProvider ServiceProvider)
 		{
 			this.StreamName = StreamName;
 			this.DefaultPerforceSettings = DefaultPerforceSettings;
@@ -73,11 +82,11 @@ namespace UnrealGameSync.Forms
 			get => Utility.OverridePerforceSettings(DefaultPerforceSettings, ServerAndPortOverride, UserNameOverride);
 		}
 
-		public static bool ShowModal(IWin32Window Owner, IPerforceSettings DefaultPerforceSettings, string StreamName, string ProjectPath, int Changelist, string Command, UserSettings Settings, IServiceProvider LoggerFactory, out BuildInfo BuildInfo)
+		public static bool ShowModal(IWin32Window Owner, IPerforceSettings DefaultPerforceSettings, string StreamName, string ProjectPath, int Changelist, string Command, UserSettings Settings, IServiceProvider LoggerFactory, [NotNullWhen(true)] out BuildInfo? BuildInfo)
 		{
-			string DefaultWorkspaceName = AutomatedSyncWindow.FindDefaultWorkspace(Owner, DefaultPerforceSettings, StreamName, LoggerFactory);
+			string? DefaultWorkspaceName = AutomatedSyncWindow.FindDefaultWorkspace(Owner, DefaultPerforceSettings, StreamName, LoggerFactory);
 
-			string DefaultProjectPath = null;
+			string? DefaultProjectPath = null;
 			if(!String.IsNullOrEmpty(ProjectPath))
 			{
 				DefaultProjectPath = ProjectPath;
@@ -98,7 +107,7 @@ namespace UnrealGameSync.Forms
 			AutomatedBuildWindow Window = new AutomatedBuildWindow(StreamName, Changelist, Command, DefaultPerforceSettings, DefaultWorkspaceName, DefaultProjectPath, LoggerFactory);
 			if (Window.ShowDialog() == DialogResult.OK)
 			{
-				BuildInfo = Window.Result;
+				BuildInfo = Window.Result!;
 				return true;
 			}
 			else
@@ -123,7 +132,7 @@ namespace UnrealGameSync.Forms
 
 		private void WorkspaceNameNewBtn_Click(object sender, EventArgs e)
 		{
-			string WorkspaceName;
+			string? WorkspaceName;
 			if (NewWorkspaceWindow.ShowModal(this, Perforce, StreamName, WorkspaceNameTextBox.Text, ServiceProvider, out WorkspaceName))
 			{
 				WorkspaceNameTextBox.Text = WorkspaceName;
@@ -133,7 +142,7 @@ namespace UnrealGameSync.Forms
 
 		private void WorkspaceNameBrowseBtn_Click(object sender, EventArgs e)
 		{
-			string WorkspaceName = WorkspaceNameTextBox.Text;
+			string? WorkspaceName = WorkspaceNameTextBox.Text;
 			if (SelectWorkspaceWindow.ShowModal(this, Perforce, WorkspaceName, ServiceProvider, out WorkspaceName))
 			{
 				WorkspaceNameTextBox.Text = WorkspaceName;
@@ -148,15 +157,10 @@ namespace UnrealGameSync.Forms
 
 		private void OkBtn_Click(object sender, EventArgs e)
 		{
-			AutomatedSyncWindow.WorkspaceInfo SelectedWorkspaceInfo;
+			AutomatedSyncWindow.WorkspaceInfo? SelectedWorkspaceInfo;
 			if (AutomatedSyncWindow.ValidateWorkspace(this, Perforce, WorkspaceNameTextBox.Text, StreamName, ServiceProvider, out SelectedWorkspaceInfo))
 			{
-				Result = new BuildInfo();
-				Result.SelectedWorkspaceInfo = SelectedWorkspaceInfo;
-				Result.ProjectPath = WorkspacePathTextBox.Text;
-				Result.bSync = SyncToChangeCheckBox.Checked;
-				Result.ExecCommand = ExecCommandTextBox.Text;
-
+				Result = new BuildInfo(SelectedWorkspaceInfo, WorkspacePathTextBox.Text, SyncToChangeCheckBox.Checked, ExecCommandTextBox.Text);
 				DialogResult = DialogResult.OK;
 				Close();
 			}
@@ -169,16 +173,16 @@ namespace UnrealGameSync.Forms
 
 		private void UpdateWorkspacePathBrowseButton()
 		{
-			string WorkspaceName;
+			string? WorkspaceName;
 			WorkspacePathBrowseBtn.Enabled = TryGetWorkspaceName(out WorkspaceName);
 		}
 
 		private void WorkspacePathBrowseBtn_Click(object sender, EventArgs e)
 		{
-			string WorkspaceName;
+			string? WorkspaceName;
 			if (TryGetWorkspaceName(out WorkspaceName))
 			{
-				string WorkspacePath = WorkspacePathTextBox.Text.Trim();
+				string? WorkspacePath = WorkspacePathTextBox.Text.Trim();
 				if (SelectProjectFromWorkspaceWindow.ShowModal(this, Perforce, WorkspaceName, WorkspacePath, ServiceProvider, out WorkspacePath))
 				{
 					WorkspacePathTextBox.Text = WorkspacePath;
@@ -187,7 +191,7 @@ namespace UnrealGameSync.Forms
 			}
 		}
 
-		private bool TryGetWorkspaceName(out string WorkspaceName)
+		private bool TryGetWorkspaceName([NotNullWhen(true)] out string? WorkspaceName)
 		{
 			string Text = WorkspaceNameTextBox.Text.Trim();
 			if (Text.Length == 0)
diff --git a/Engine/Source/Programs/UnrealGameSync/UnrealGameSync/Forms/AutomatedSyncWindow.cs b/Engine/Source/Programs/UnrealGameSync/UnrealGameSync/Forms/AutomatedSyncWindow.cs
--- a/Engine/Source/Programs/UnrealGameSync/UnrealGameSync/Forms/AutomatedSyncWindow.cs
+++ b/Engine/Source/Programs/UnrealGameSync/UnrealGameSync/Forms/AutomatedSyncWindow.cs
@@ -7,6 +7,7 @@ using System;
 using System.Collections.Generic;
 using System.ComponentModel;
 using System.Data;
+using System.Diagnostics.CodeAnalysis;
 using System.Drawing;
 using System.IO;
 using System.Linq;
@@ -136,14 +137,14 @@ namespace UnrealGameSync
 			return (WorkspaceTask != null && WorkspaceTask.Succeeded) ? WorkspaceTask.Result : null;
 		}
 
-		public static bool ShowModal(IWin32Window Owner, IPerforceSettings DefaultPerforceSettings, string StreamName, string ProjectPath, out WorkspaceInfo? WorkspaceInfo, IServiceProvider ServiceProvider)
+		public static bool ShowModal(IWin32Window Owner, IPerforceSettings DefaultPerforceSettings, string StreamName, string ProjectPath, [NotNullWhen(true)] out WorkspaceInfo? WorkspaceInfo, IServiceProvider ServiceProvider)
 		{
 			string? WorkspaceName = FindDefaultWorkspace(Owner, DefaultPerforceSettings, StreamName, ServiceProvider);
 
 			AutomatedSyncWindow Window = new AutomatedSyncWindow(StreamName, ProjectPath, WorkspaceName, DefaultPerforceSettings, ServiceProvider);
 			if(Window.ShowDialog() == DialogResult.OK)
 			{
-				WorkspaceInfo = Window.SelectedWorkspaceInfo;
+				WorkspaceInfo = Window.SelectedWorkspaceInfo!;
 				return true;
 			}
 			else
@@ -192,7 +193,7 @@ namespace UnrealGameSync
 			OkBtn.Enabled = (WorkspaceNameTextBox.Text.Length > 0);
 		}
 
-		public static bool ValidateWorkspace(IWin32Window Owner, IPerforceSettings Perforce, string WorkspaceName, string StreamName, IServiceProvider ServiceProvider, out WorkspaceInfo? SelectedWorkspaceInfo)
+		public static bool ValidateWorkspace(IWin32Window Owner, IPerforceSettings Perforce, string WorkspaceName, string StreamName, IServiceProvider ServiceProvider, [NotNullWhen(true)] out WorkspaceInfo? SelectedWorkspaceInfo)
 		{
 			ValidateWorkspaceTask ValidateWorkspace = new ValidateWorkspaceTask(WorkspaceName, StreamName);
 
diff --git a/Engine/Source/Programs/UnrealGameSync/UnrealGameSync/Forms/BuildStepWindow.cs b/Engine/Source/Programs/UnrealGameSync/UnrealGameSync/Forms/BuildStepWindow.cs
--- a/Engine/Source/Programs/UnrealGameSync/UnrealGameSync/Forms/BuildStepWindow.cs
+++ b/Engine/Source/Programs/UnrealGameSync/UnrealGameSync/Forms/BuildStepWindow.cs
@@ -20,7 +20,7 @@ namespace UnrealGameSync
 		List<string> TargetNames;
 		DirectoryReference BaseDirectory;
 		IReadOnlyDictionary<string, string> Variables;
-		VariablesWindow VariablesWindow;
+		VariablesWindow? VariablesWindow;
 
 		public BuildStepWindow(BuildStep InTask, List<string> InTargetNames, DirectoryReference InBaseDirectory, IReadOnlyDictionary<string, string> InVariables)
 		{
@@ -242,7 +242,7 @@ namespace UnrealGameSync
 			{
 				for(;;)
 				{
-					IContainerControl NextContainer = Container.ActiveControl as IContainerControl;
+					IContainerControl? NextContainer = Container.ActiveControl as IContainerControl;
 					if(NextContainer == null)
 					{
 						break;
@@ -250,7 +250,7 @@ namespace UnrealGameSync
 					Container = NextContainer;
 				}
 
-				TextBox FocusTextBox = Container.ActiveControl as TextBox;
+				TextBox? FocusTextBox = Container.ActiveControl as TextBox;
 				if(FocusTextBox != null)
 				{
 					FocusTextBox.SelectedText = Name;
diff --git a/Engine/Source/Programs/UnrealGameSync/UnrealGameSync/Forms/CleanWorkspaceWindow.cs b/Engine/Source/Programs/UnrealGameSync/UnrealGameSync/Forms/CleanWorkspaceWindow.cs
--- a/Engine/Source/Programs/UnrealGameSync/UnrealGameSync/Forms/CleanWorkspaceWindow.cs
+++ b/Engine/Source/Programs/UnrealGameSync/UnrealGameSync/Forms/CleanWorkspaceWindow.cs
@@ -18,8 +18,6 @@ using System.Threading.Tasks;
 using System.Windows.Forms;
 using System.Windows.Forms.VisualStyles;
 
-#nullable enable
-
 namespace UnrealGameSync
 {
 	partial class CleanWorkspaceWindow : Form
diff --git a/Engine/Source/Programs/UnrealGameSync/UnrealGameSync/Forms/ClobberWindow.cs b/Engine/Source/Programs/UnrealGameSync/UnrealGameSync/Forms/ClobberWindow.cs
--- a/Engine/Source/Programs/UnrealGameSync/UnrealGameSync/Forms/ClobberWindow.cs
+++ b/Engine/Source/Programs/UnrealGameSync/UnrealGameSync/Forms/ClobberWindow.cs
@@ -11,8 +11,6 @@ using System.Text;
 using System.Threading.Tasks;
 using System.Windows.Forms;
 
-#nullable enable
-
 namespace UnrealGameSync
 {
 	public partial class ClobberWindow : Form
diff --git a/Engine/Source/Programs/UnrealGameSync/UnrealGameSync/Forms/ConnectWindow.cs b/Engine/Source/Programs/UnrealGameSync/UnrealGameSync/Forms/ConnectWindow.cs
--- a/Engine/Source/Programs/UnrealGameSync/UnrealGameSync/Forms/ConnectWindow.cs
+++ b/Engine/Source/Programs/UnrealGameSync/UnrealGameSync/Forms/ConnectWindow.cs
@@ -18,11 +18,11 @@ namespace UnrealGameSync
 	partial class ConnectWindow : Form
 	{
 		IPerforceSettings DefaultPerforceSettings;
-		string ServerAndPortOverride;
-		string UserNameOverride;
+		string? ServerAndPortOverride;
+		string? UserNameOverride;
 		IServiceProvider ServiceProvider;
 
-		private ConnectWindow(IPerforceSettings DefaultPerforceSettings, string ServerAndPortOverride, string UserNameOverride, IServiceProvider ServiceProvider)
+		private ConnectWindow(IPerforceSettings DefaultPerforceSettings, string? ServerAndPortOverride, string? UserNameOverride, IServiceProvider ServiceProvider)
 		{
 			InitializeComponent();
 
@@ -47,7 +47,7 @@ namespace UnrealGameSync
 			UpdateEnabledControls();
 		}
 
-		public static bool ShowModal(IWin32Window Owner, IPerforceSettings DefaultSettings, ref string ServerAndPortOverride, ref string UserNameOverride, IServiceProvider ServiceProvider)
+		public static bool ShowModal(IWin32Window Owner, IPerforceSettings DefaultSettings, ref string? ServerAndPortOverride, ref string? UserNameOverride, IServiceProvider ServiceProvider)
 		{
 			ConnectWindow Connect = new ConnectWindow(DefaultSettings, ServerAndPortOverride, UserNameOverride, ServiceProvider);
 			if(Connect.ShowDialog(Owner) == DialogResult.OK)
@@ -109,7 +109,7 @@ namespace UnrealGameSync
 
 		private void BrowseUserBtn_Click(object sender, EventArgs e)
 		{
-			string NewUserName;
+			string? NewUserName;
 			if(SelectUserWindow.ShowModal(this, new PerforceSettings(DefaultPerforceSettings) { UserName = UserNameTextBox.Text, ServerAndPort = ServerAndPortTextBox.Text }, ServiceProvider, out NewUserName))
 			{
 				UserNameTextBox.Text = NewUserName;
diff --git a/Engine/Source/Programs/UnrealGameSync/UnrealGameSync/Forms/DeleteWindow.cs b/Engine/Source/Programs/UnrealGameSync/UnrealGameSync/Forms/DeleteWindow.cs
--- a/Engine/Source/Programs/UnrealGameSync/UnrealGameSync/Forms/DeleteWindow.cs
+++ b/Engine/Source/Programs/UnrealGameSync/UnrealGameSync/Forms/DeleteWindow.cs
@@ -34,25 +34,25 @@ namespace UnrealGameSync
 
 		private void UncheckAll_Click(object sender, EventArgs e)
 		{
-			foreach(ListViewItem Item in FileList.Items)
+			foreach(ListViewItem? Item in FileList.Items)
 			{
-				Item.Checked = false;
+				Item!.Checked = false;
 			}
 		}
 
 		private void CheckAll_Click(object sender, EventArgs e)
 		{
-			foreach(ListViewItem Item in FileList.Items)
+			foreach(ListViewItem? Item in FileList.Items)
 			{
-				Item.Checked = true;
+				Item!.Checked = true;
 			}
 		}
 
 		private void ContinueButton_Click(object sender, EventArgs e)
 		{
-			foreach(ListViewItem Item in FileList.Items)
+			foreach(ListViewItem? Item in FileList.Items)
 			{
-				FilesToDelete[(string)Item.Tag] = Item.Checked;
+				FilesToDelete[(string)Item!.Tag] = Item.Checked;
 			}
 		}
 	}
diff --git a/Engine/Source/Programs/UnrealGameSync/UnrealGameSync/Forms/IssueBrowserWindow.cs b/Engine/Source/Programs/UnrealGameSync/UnrealGameSync/Forms/IssueBrowserWindow.cs
--- a/Engine/Source/Programs/UnrealGameSync/UnrealGameSync/Forms/IssueBrowserWindow.cs
+++ b/Engine/Source/Programs/UnrealGameSync/UnrealGameSync/Forms/IssueBrowserWindow.cs
@@ -17,8 +17,6 @@ using System.Threading;
 using System.Threading.Tasks;
 using System.Windows.Forms;
 
-#nullable enable
-
 namespace UnrealGameSync
 {
 	partial class IssueBrowserWindow : Form
diff --git a/Engine/Source/Programs/UnrealGameSync/UnrealGameSync/Forms/IssueDetailsWindow.cs b/Engine/Source/Programs/UnrealGameSync/UnrealGameSync/Forms/IssueDetailsWindow.cs
--- a/Engine/Source/Programs/UnrealGameSync/UnrealGameSync/Forms/IssueDetailsWindow.cs
+++ b/Engine/Source/Programs/UnrealGameSync/UnrealGameSync/Forms/IssueDetailsWindow.cs
@@ -19,8 +19,6 @@ using System.Threading;
 using System.Threading.Tasks;
 using System.Windows.Forms;
 
-#nullable enable
-
 namespace UnrealGameSync
 {
 	partial class IssueDetailsWindow : Form
@@ -1046,7 +1044,7 @@ namespace UnrealGameSync
 			{
 				List<IssueDiagnosticData> Diagnostics = new List<IssueDiagnosticData>();
 
-				UpdateIssueDetailsTask Task = new UpdateIssueDetailsTask(IssueMonitor.ApiUrl, Issue.Id, Diagnostics);
+				UpdateIssueDetailsTask Task = new UpdateIssueDetailsTask(IssueMonitor.ApiUrl!, Issue.Id, Diagnostics);
 				ModalTask? TaskResult = ModalTask.Execute(Owner, "Fetching data", "Fetching data, please wait...", Task.RunAsync);
 				if (TaskResult == null || !TaskResult.Succeeded)
 				{
diff --git a/Engine/Source/Programs/UnrealGameSync/UnrealGameSync/Forms/MainWindow.cs b/Engine/Source/Programs/UnrealGameSync/UnrealGameSync/Forms/MainWindow.cs
--- a/Engine/Source/Programs/UnrealGameSync/UnrealGameSync/Forms/MainWindow.cs
+++ b/Engine/Source/Programs/UnrealGameSync/UnrealGameSync/Forms/MainWindow.cs
@@ -19,6 +19,7 @@ using EpicGames.Perforce;
 using System.Threading.Tasks;
 using Microsoft.Extensions.Logging;
 using Microsoft.Extensions.DependencyInjection;
+using System.Diagnostics.CodeAnalysis;
 
 namespace UnrealGameSync
 {
@@ -58,6 +59,11 @@ namespace UnrealGameSync
 		{
 			public IssueMonitor IssueMonitor;
 			public int RefCount;
+
+			public WorkspaceIssueMonitor(IssueMonitor IssueMonitor)
+			{
+				this.IssueMonitor = IssueMonitor;
+			}
 		}
 
 		[Flags]
@@ -80,7 +86,7 @@ namespace UnrealGameSync
 		List<IssueMonitor> DefaultIssueMonitors = new List<IssueMonitor>();
 		List<WorkspaceIssueMonitor> WorkspaceIssueMonitors = new List<WorkspaceIssueMonitor>();
 
-		string ApiUrl;
+		string? ApiUrl;
 		DirectoryReference DataFolder;
 		DirectoryReference CacheFolder;
 		IPerforceSettings DefaultPerforceSettings;
@@ -94,15 +100,15 @@ namespace UnrealGameSync
 
 		bool bRestoreStateOnLoad;
 
-		OIDCTokenManager OIDCTokenManager;
+		OIDCTokenManager? OIDCTokenManager;
 
-		System.Threading.Timer ScheduleTimer;
-		System.Threading.Timer ScheduleSettledTimer;
+		System.Threading.Timer? ScheduleTimer;
+		System.Threading.Timer? ScheduleSettledTimer;
 
 		string OriginalExecutableFileName;
 		bool bUnstable;
 
-		IMainWindowTabPanel CurrentTabPanel;
+		IMainWindowTabPanel? CurrentTabPanel;
 
 		AutomationServer AutomationServer;
 
@@ -113,9 +119,7 @@ namespace UnrealGameSync
 
 		public ToolUpdateMonitor ToolUpdateMonitor { get; private set; }
 
-		NetCoreWindow NetCoreWindow;
-
-		public MainWindow(UpdateMonitor InUpdateMonitor, string InApiUrl, DirectoryReference InDataFolder, DirectoryReference InCacheFolder, bool bInRestoreStateOnLoad, string InOriginalExecutableFileName, bool bInUnstable, List<(UserSelectedProjectSettings, ModalTask<WorkspaceSettings>)> StartupTasks, IPerforceSettings InDefaultPerforceSettings, IServiceProvider InServiceProvider, UserSettings InSettings, string InUri, OIDCTokenManager InOidcTokenManager)
+		public MainWindow(UpdateMonitor InUpdateMonitor, string? InApiUrl, DirectoryReference InDataFolder, DirectoryReference InCacheFolder, bool bInRestoreStateOnLoad, string InOriginalExecutableFileName, bool bInUnstable, List<(UserSelectedProjectSettings, ModalTask<WorkspaceSettings>)> StartupTasks, IPerforceSettings InDefaultPerforceSettings, IServiceProvider InServiceProvider, UserSettings InSettings, string? InUri, OIDCTokenManager? InOidcTokenManager)
 		{
 			ServiceProvider = InServiceProvider;
 			Logger = ServiceProvider.GetRequiredService<ILogger<MainWindow>>();
@@ -125,7 +129,7 @@ namespace UnrealGameSync
 			InitializeComponent();
 
 			UpdateMonitor = InUpdateMonitor;
-			MainThreadSynchronizationContext = SynchronizationContext.Current;
+			MainThreadSynchronizationContext = SynchronizationContext.Current!;
 			ApiUrl = InApiUrl;
 			DataFolder = InDataFolder;
 			CacheFolder = InCacheFolder;
@@ -220,7 +224,7 @@ namespace UnrealGameSync
 				}
 				else if (Request.Input.Type == AutomationRequestType.SyncProject)
 				{
-					AutomationRequestOutput Output = StartAutomatedSync(Request, true);
+					AutomationRequestOutput? Output = StartAutomatedSync(Request, true);
 					if (Output != null)
 					{
 						Request.SetOutput(Output);
@@ -233,7 +237,7 @@ namespace UnrealGameSync
 				}
 				else if (Request.Input.Type == AutomationRequestType.OpenProject)
 				{
-					AutomationRequestOutput Output = StartAutomatedSync(Request, false);
+					AutomationRequestOutput? Output = StartAutomatedSync(Request, false);
 					if (Output != null)
 					{
 						Request.SetOutput(Output);
@@ -269,13 +273,13 @@ namespace UnrealGameSync
 			string Command = Reader.ReadString();
 			string ProjectPath = Reader.ReadString();
 
-			AutomatedBuildWindow.BuildInfo BuildInfo;
+			AutomatedBuildWindow.BuildInfo? BuildInfo;
 			if (!AutomatedBuildWindow.ShowModal(this, DefaultPerforceSettings, StreamName, ProjectPath, Changelist, Command.ToString(), Settings, ServiceProvider, out BuildInfo))
 			{
 				return new AutomationRequestOutput(AutomationRequestResult.Canceled);
 			}
 
-			WorkspaceControl Workspace;
+			WorkspaceControl? Workspace;
 			if (!OpenWorkspaceForAutomation(BuildInfo.SelectedWorkspaceInfo, StreamName, BuildInfo.ProjectPath, out Workspace))
 			{
 				return new AutomationRequestOutput(AutomationRequestResult.Error);
@@ -304,12 +308,12 @@ namespace UnrealGameSync
 		{
 			if (Result == WorkspaceUpdateResult.Success && Command != null)
 			{
-				string CmdExe = Environment.GetEnvironmentVariable("COMSPEC");
+				string CmdExe = Environment.GetEnvironmentVariable("COMSPEC") ?? "C:\\Windows\\System32\\cmd.exe";
 				Workspace.ExecCommand("Run build command", "Running build command", CmdExe, String.Format("/c {0}", Command), Workspace.BranchDirectoryName.FullName, true);
 			}
 		}
 
-		AutomationRequestOutput StartAutomatedSync(AutomationRequest Request, bool bForceSync)
+		AutomationRequestOutput? StartAutomatedSync(AutomationRequest Request, bool bForceSync)
 		{
 			ShowAndActivate();
 
@@ -317,13 +321,13 @@ namespace UnrealGameSync
 			string StreamName = Reader.ReadString();
 			string ProjectPath = Reader.ReadString();
 
-			AutomatedSyncWindow.WorkspaceInfo WorkspaceInfo;
+			AutomatedSyncWindow.WorkspaceInfo? WorkspaceInfo;
 			if(!AutomatedSyncWindow.ShowModal(this, DefaultPerforceSettings, StreamName, ProjectPath, out WorkspaceInfo, ServiceProvider))
 			{
 				return new AutomationRequestOutput(AutomationRequestResult.Canceled);
 			}
 
-			WorkspaceControl Workspace;
+			WorkspaceControl? Workspace;
 			if(!OpenWorkspaceForAutomation(WorkspaceInfo, StreamName, ProjectPath, out Workspace))
 			{
 				return new AutomationRequestOutput(AutomationRequestResult.Error);
@@ -338,14 +342,14 @@ namespace UnrealGameSync
 			return null;
 		}
 
-		private bool OpenWorkspaceForAutomation(AutomatedSyncWindow.WorkspaceInfo WorkspaceInfo, string StreamName, string ProjectPath, out WorkspaceControl OutWorkspace)
+		private bool OpenWorkspaceForAutomation(AutomatedSyncWindow.WorkspaceInfo WorkspaceInfo, string StreamName, string ProjectPath, [NotNullWhen(true)] out WorkspaceControl? OutWorkspace)
 		{
 			if (WorkspaceInfo.bRequiresStreamSwitch)
 			{
 				// Close any tab containing this window
 				for (int ExistingTabIdx = 0; ExistingTabIdx < TabControl.GetTabCount(); ExistingTabIdx++)
 				{
-					WorkspaceControl ExistingWorkspace = TabControl.GetTabData(ExistingTabIdx) as WorkspaceControl;
+					WorkspaceControl? ExistingWorkspace = TabControl.GetTabData(ExistingTabIdx) as WorkspaceControl;
 					if (ExistingWorkspace != null && ExistingWorkspace.ClientName.Equals(WorkspaceInfo.WorkspaceName))
 					{
 						TabControl.RemoveTab(ExistingTabIdx);
@@ -356,16 +360,16 @@ namespace UnrealGameSync
 				// Switch the stream
 				Func<IPerforceConnection, CancellationToken, Task> SwitchTask = async (Connection, CancellationToken) =>
 				{
-					await Connection.SwitchClientToStreamAsync(Connection.Settings.ClientName, StreamName, SwitchClientOptions.None, CancellationToken);
+					await Connection.SwitchClientToStreamAsync(Connection.Settings.ClientName!, StreamName, SwitchClientOptions.None, CancellationToken);
 				};
 
 				PerforceSettings Settings = new PerforceSettings(WorkspaceInfo.ServerAndPort, WorkspaceInfo.UserName);
 				Settings.ClientName = WorkspaceInfo.WorkspaceName;
 
-				ModalTask Result = PerforceModalTask.Execute(Owner, "Please wait", "Switching streams, please wait...", Settings, SwitchTask, Logger, ModalTaskFlags.Quiet);
+				ModalTask? Result = PerforceModalTask.Execute(Owner, "Please wait", "Switching streams, please wait...", Settings, SwitchTask, Logger, ModalTaskFlags.Quiet);
 				if (Result == null || !Result.Succeeded)
 				{
-					Logger.LogError("Unable to switch stream ({Message})", Result.Error);
+					Logger.LogError("Unable to switch stream ({Message})", Result?.Error ?? "Operation cancelled");
 					OutWorkspace = null;
 					return false;
 				}
@@ -381,7 +385,7 @@ namespace UnrealGameSync
 				return false;
 			}
 
-			WorkspaceControl Workspace = TabControl.GetTabData(TabIdx) as WorkspaceControl;
+			WorkspaceControl? Workspace = TabControl.GetTabData(TabIdx) as WorkspaceControl;
 			if (Workspace == null)
 			{
 				Logger.LogError("Workspace was unable to open");
@@ -430,10 +434,10 @@ namespace UnrealGameSync
 
 			for(int ExistingTabIdx = 0; ExistingTabIdx < TabControl.GetTabCount(); ExistingTabIdx++)
 			{
-				WorkspaceControl ExistingWorkspace = TabControl.GetTabData(ExistingTabIdx) as WorkspaceControl;
+				WorkspaceControl? ExistingWorkspace = TabControl.GetTabData(ExistingTabIdx) as WorkspaceControl;
 				if(ExistingWorkspace != null && String.Compare(ExistingWorkspace.StreamName, StreamName, StringComparison.OrdinalIgnoreCase) == 0 && ExistingWorkspace.SelectedProject != null)
 				{
-					string ClientPath = ExistingWorkspace.SelectedProject.ClientPath;
+					string? ClientPath = ExistingWorkspace.SelectedProject.ClientPath;
 					if(ClientPath != null && ClientPath.StartsWith("//"))
 					{
 						int SlashIdx = ClientPath.IndexOf('/', 2);
@@ -459,7 +463,7 @@ namespace UnrealGameSync
 
 			for (int ExistingTabIdx = 0; ExistingTabIdx < TabControl.GetTabCount(); ExistingTabIdx++)
 			{
-				WorkspaceControl ExistingWorkspace = TabControl.GetTabData(ExistingTabIdx) as WorkspaceControl;
+				WorkspaceControl? ExistingWorkspace = TabControl.GetTabData(ExistingTabIdx) as WorkspaceControl;
 				if (ExistingWorkspace != null)
 				{
 					IssueMonitor IssueMonitor = ExistingWorkspace.GetIssueMonitor();
@@ -469,7 +473,7 @@ namespace UnrealGameSync
 						try
 						{
 							Func<CancellationToken, Task<IssueData>> Func = x => RESTApi.GetAsync<IssueData>($"{IssueMonitor.ApiUrl}/api/issues/{IssueId}", x);
-							ModalTask<IssueData> IssueTask = ModalTask.Execute(this, "Finding Issue", "Querying issue data, please wait...", Func);
+							ModalTask<IssueData>? IssueTask = ModalTask.Execute(this, "Finding Issue", "Querying issue data, please wait...", Func);
 							if (IssueTask == null)
 							{
 								Logger.LogInformation("Operation cancelled");
@@ -512,7 +516,7 @@ namespace UnrealGameSync
 			}
 		}
 
-		void TabControl_OnTabClicked(object TabData, Point Location, MouseButtons Buttons)
+		void TabControl_OnTabClicked(object? TabData, Point Location, MouseButtons Buttons)
 		{
 			if(Buttons == System.Windows.Forms.MouseButtons.Right)
 			{
@@ -605,13 +609,13 @@ namespace UnrealGameSync
 			if(AutomationServer != null)
 			{
 				AutomationServer.Dispose();
-				AutomationServer = null;
+				AutomationServer = null!;
 			}
 
 			if (ToolUpdateMonitor != null)
 			{
 				ToolUpdateMonitor.Dispose();
-				ToolUpdateMonitor = null;
+				ToolUpdateMonitor = null!;
 			}
 
 			base.Dispose(disposing);
@@ -643,7 +647,7 @@ namespace UnrealGameSync
 				{
 					for (int Idx = 0; Idx < TabControl.GetTabCount(); Idx++)
 					{
-						IMainWindowTabPanel TabPanel = (IMainWindowTabPanel)TabControl.GetTabData(Idx);
+						IMainWindowTabPanel TabPanel = (IMainWindowTabPanel)TabControl.GetTabData(Idx)!;
 						if (!TabPanel.CanClose())
 						{
 							EventArgs.Cancel = true;
@@ -790,7 +794,7 @@ namespace UnrealGameSync
 
 			for (int TabIdx = 0; TabIdx < TabControl.GetTabCount(); TabIdx++)
 			{
-				IMainWindowTabPanel TabPanel = (IMainWindowTabPanel)TabControl.GetTabData(TabIdx);
+				IMainWindowTabPanel TabPanel = (IMainWindowTabPanel)TabControl.GetTabData(TabIdx)!;
 				if(TabPanel.IsBusy())
 				{
 					return false;
@@ -954,11 +958,11 @@ namespace UnrealGameSync
 			Logger.LogInformation("Schedule: Starting Sync");
 			for(int Idx = 0; Idx < TabControl.GetTabCount(); Idx++)
 			{
-				WorkspaceControl Workspace = TabControl.GetTabData(Idx) as WorkspaceControl;
+				WorkspaceControl? Workspace = TabControl.GetTabData(Idx) as WorkspaceControl;
 				if(Workspace != null)
 				{
 					Logger.LogInformation("Schedule: Considering {File}", Workspace.SelectedFileName);
-					if(Settings.ScheduleAnyOpenProject || Settings.ScheduleProjects.Any(x => x.LocalPath.Equals(Workspace.SelectedProject.LocalPath, StringComparison.OrdinalIgnoreCase)))
+					if(Settings.ScheduleAnyOpenProject || Settings.ScheduleProjects.Any(x => x.LocalPath != null && x.LocalPath.Equals(Workspace.SelectedProject.LocalPath, StringComparison.OrdinalIgnoreCase)))
 					{
 						Logger.LogInformation("Schedule: Starting Sync");
 						Workspace.ScheduleTimerElapsed();
@@ -967,7 +971,7 @@ namespace UnrealGameSync
 			}
 		}
 
-		void TabControl_OnTabChanged(object NewTabData)
+		void TabControl_OnTabChanged(object? NewTabData)
 		{
 			if(IsHandleCreated)
 			{
@@ -1031,7 +1035,7 @@ namespace UnrealGameSync
 
 		public void OpenNewProject()
 		{
-			WorkspaceSettings WorkspaceSettings = OpenProjectWindow.ShowModal(this, null, Settings, DataFolder, CacheFolder, DefaultPerforceSettings, ServiceProvider);
+			WorkspaceSettings? WorkspaceSettings = OpenProjectWindow.ShowModal(this, null, Settings, DataFolder, CacheFolder, DefaultPerforceSettings, ServiceProvider);
 			if(WorkspaceSettings != null)
 			{
 				int NewTabIdx = TryOpenProject(WorkspaceSettings, -1, OpenProjectOptions.None);
@@ -1046,7 +1050,7 @@ namespace UnrealGameSync
 
 		void UpdateRecentProjectsList(UserSelectedProjectSettings DetectedProjectSettings)
 		{
-			Settings.RecentProjects.RemoveAll(x => x.LocalPath.Equals(DetectedProjectSettings.LocalPath, StringComparison.OrdinalIgnoreCase));
+			Settings.RecentProjects.RemoveAll(x => x.LocalPath != null && x.LocalPath.Equals(DetectedProjectSettings.LocalPath, StringComparison.OrdinalIgnoreCase));
 			Settings.RecentProjects.Insert(0, DetectedProjectSettings);
 
 			const int MaxRecentProjects = 10;
@@ -1093,7 +1097,7 @@ namespace UnrealGameSync
 
 		public void EditSelectedProject(int TabIdx, UserSelectedProjectSettings SelectedProject)
 		{
-			WorkspaceSettings WorkspaceSettings = OpenProjectWindow.ShowModal(this, SelectedProject, Settings, DataFolder, CacheFolder, DefaultPerforceSettings, ServiceProvider);
+			WorkspaceSettings? WorkspaceSettings = OpenProjectWindow.ShowModal(this, SelectedProject, Settings, DataFolder, CacheFolder, DefaultPerforceSettings, ServiceProvider);
 			if(WorkspaceSettings != null)
 			{
 				int NewTabIdx = TryOpenProject(WorkspaceSettings, TabIdx, OpenProjectOptions.None);
@@ -1116,10 +1120,10 @@ namespace UnrealGameSync
 				TaskFlags |= ModalTaskFlags.Quiet;
 			}
 
-			ModalTask<WorkspaceSettings> SettingsTask = PerforceModalTask.Execute(this, "Opening Project", "Opening project, please wait...", DefaultPerforceSettings, (p, c) => WorkspaceSettings.CreateAsync(p, Project, ProjectLogger, c), ProjectLogger, TaskFlags);
-			if(SettingsTask == null || !SettingsTask.Succeeded)
+			ModalTask<WorkspaceSettings>? SettingsTask = PerforceModalTask.Execute(this, "Opening Project", "Opening project, please wait...", DefaultPerforceSettings, (p, c) => WorkspaceSettings.CreateAsync(p, Project, ProjectLogger, c), ProjectLogger, TaskFlags);
+			if (SettingsTask == null || SettingsTask.Failed)
 			{
-				CreateErrorPanel(ReplaceTabIdx, Project, SettingsTask.Error);
+				if(SettingsTask != null) CreateErrorPanel(ReplaceTabIdx, Project, SettingsTask.Error);
 				return -1;
 			}
 
@@ -1135,7 +1139,7 @@ namespace UnrealGameSync
 			{
 				if(ReplaceTabIdx != TabIdx)
 				{
-					WorkspaceControl Workspace = TabControl.GetTabData(TabIdx) as WorkspaceControl;
+					WorkspaceControl? Workspace = TabControl.GetTabData(TabIdx) as WorkspaceControl;
 					if(Workspace != null)
 					{
 						if(Workspace.SelectedFileName == WorkspaceSettings.ProjectInfo.LocalFileName)
@@ -1170,7 +1174,7 @@ namespace UnrealGameSync
 			// Remove the current tab. We need to ensure the workspace has been shut down before creating a new one with the same log files, etc...
 			if(ReplaceTabIdx != -1)
 			{
-				WorkspaceControl OldWorkspace = TabControl.GetTabData(ReplaceTabIdx) as WorkspaceControl;
+				WorkspaceControl? OldWorkspace = TabControl.GetTabData(ReplaceTabIdx) as WorkspaceControl;
 				if(OldWorkspace != null)
 				{
 					OldWorkspace.Hide();
@@ -1236,7 +1240,7 @@ namespace UnrealGameSync
 			Settings.OpenProjects.Clear();
 			for(int TabIdx = 0; TabIdx < TabControl.GetTabCount(); TabIdx++)
 			{
-				IMainWindowTabPanel TabPanel = (IMainWindowTabPanel)TabControl.GetTabData(TabIdx);
+				IMainWindowTabPanel TabPanel = (IMainWindowTabPanel)TabControl.GetTabData(TabIdx)!;
 				Settings.OpenProjects.Add(TabPanel.SelectedProject);
 			}
 			Settings.Save();
@@ -1290,7 +1294,7 @@ namespace UnrealGameSync
 
 				for(int Idx = 0; Idx < TabControl.GetTabCount(); Idx++)
 				{
-					WorkspaceControl Workspace = TabControl.GetTabData(Idx) as WorkspaceControl;
+					WorkspaceControl? Workspace = TabControl.GetTabData(Idx) as WorkspaceControl;
 					if(Workspace != null)
 					{
 						TabControl.SetTabName(Idx, GetTabName(Workspace));
@@ -1354,7 +1358,7 @@ namespace UnrealGameSync
 
 			for(int Idx = 0; Idx < TabControl.GetTabCount(); Idx++)
 			{
-				IMainWindowTabPanel TabPanel = (IMainWindowTabPanel)TabControl.GetTabData(Idx);
+				IMainWindowTabPanel TabPanel = (IMainWindowTabPanel)TabControl.GetTabData(Idx)!;
 
 				Tuple<TaskbarState, float> DesiredTaskbarState = TabPanel.DesiredTaskbarState;
 				if(DesiredTaskbarState.Item1 == TaskbarState.Error)
@@ -1397,7 +1401,7 @@ namespace UnrealGameSync
 		{
 			for (int Idx = 0; Idx < TabControl.GetTabCount(); Idx++)
 			{
-				IMainWindowTabPanel TabPanel = (IMainWindowTabPanel)TabControl.GetTabData(Idx);
+				IMainWindowTabPanel TabPanel = (IMainWindowTabPanel)TabControl.GetTabData(Idx)!;
 				TabControl.SetTint(Idx, TabPanel.TintColor);
 			}
 		}
@@ -1412,7 +1416,7 @@ namespace UnrealGameSync
 
 			for (int Idx = 0; Idx < TabControl.GetTabCount(); Idx++)
 			{
-				IMainWindowTabPanel TabPanel = (IMainWindowTabPanel)TabControl.GetTabData(Idx);
+				IMainWindowTabPanel TabPanel = (IMainWindowTabPanel)TabControl.GetTabData(Idx)!;
 				TabPanel.UpdateSettings();
 			}
 		}
@@ -1588,7 +1592,7 @@ namespace UnrealGameSync
 		{
 			for(int Idx = 0; Idx < TabControl.GetTabCount(); Idx++)
 			{
-				WorkspaceControl Workspace = TabControl.GetTabData(Idx) as WorkspaceControl;
+				WorkspaceControl? Workspace = TabControl.GetTabData(Idx) as WorkspaceControl;
 				if(Workspace != null && Workspace.GetIssueMonitor() == Alert.IssueMonitor)
 				{
 					Workspace.ShowIssueDetails(Alert.Issue);
@@ -1640,7 +1644,7 @@ namespace UnrealGameSync
 			SetAlertWindowPositions();
 		}
 
-		public IssueMonitor CreateIssueMonitor(string ApiUrl, string UserName)
+		public IssueMonitor CreateIssueMonitor(string? ApiUrl, string UserName)
 		{
 			WorkspaceIssueMonitor WorkspaceIssueMonitor = WorkspaceIssueMonitors.FirstOrDefault(x => String.Compare(x.IssueMonitor.ApiUrl, ApiUrl, StringComparison.OrdinalIgnoreCase) == 0 && String.Compare(x.IssueMonitor.UserName, UserName, StringComparison.OrdinalIgnoreCase) == 0);
 			if (WorkspaceIssueMonitor == null)
@@ -1650,10 +1654,10 @@ namespace UnrealGameSync
 
 				FileReference LogFileName = FileReference.Combine(DataFolder, String.Format("IssueMonitor-{0}-{1}.log", ServerId, UserName));
 
-				WorkspaceIssueMonitor = new WorkspaceIssueMonitor();
-				WorkspaceIssueMonitor.IssueMonitor = new IssueMonitor(ApiUrl, UserName, TimeSpan.FromSeconds(60.0), ServiceProvider);
-				WorkspaceIssueMonitor.IssueMonitor.OnIssuesChanged += IssueMonitor_OnUpdateAsync;
+				IssueMonitor IssueMonitor = new IssueMonitor(ApiUrl, UserName, TimeSpan.FromSeconds(60.0), ServiceProvider);
+				IssueMonitor.OnIssuesChanged += IssueMonitor_OnUpdateAsync;
 
+				WorkspaceIssueMonitor = new WorkspaceIssueMonitor(IssueMonitor);
 				WorkspaceIssueMonitors.Add(WorkspaceIssueMonitor);
 			}
 
@@ -1692,18 +1696,5 @@ namespace UnrealGameSync
 				}
 			}
 		}
-
-		private void NetCoreTimer_Tick(object sender, EventArgs e)
-		{
-			if (NetCoreWindow == null || !NetCoreWindow.Created)
-			{
-				if (!EventMonitor.HasNetCore3() && Settings.bShowNetCoreInfo)
-				{
-					NetCoreWindow = new NetCoreWindow(Settings);
-					NetCoreWindow.Location = new Point(Location.X + (Width - NetCoreWindow.Width) / 2, Location.Y + (Height - NetCoreWindow.Height) / 2);
-					NetCoreWindow.Show(this);
-				}
-			}
-		}
 	}
 }
diff --git a/Engine/Source/Programs/UnrealGameSync/UnrealGameSync/Forms/ModalTaskWindow.cs b/Engine/Source/Programs/UnrealGameSync/UnrealGameSync/Forms/ModalTaskWindow.cs
--- a/Engine/Source/Programs/UnrealGameSync/UnrealGameSync/Forms/ModalTaskWindow.cs
+++ b/Engine/Source/Programs/UnrealGameSync/UnrealGameSync/Forms/ModalTaskWindow.cs
@@ -38,7 +38,7 @@ namespace UnrealGameSync
 
 		private void ModalTaskWindow_Load(object sender, EventArgs e)
 		{
-			SynchronizationContext SyncContext = SynchronizationContext.Current;
+			SynchronizationContext SyncContext = SynchronizationContext.Current!;
 			Task.ContinueWith(x => SyncContext.Post(y => Close(), null));
 		}
 
diff --git a/Engine/Source/Programs/UnrealGameSync/UnrealGameSync/Forms/ModifyBuildStepsWindow.cs b/Engine/Source/Programs/UnrealGameSync/UnrealGameSync/Forms/ModifyBuildStepsWindow.cs
--- a/Engine/Source/Programs/UnrealGameSync/UnrealGameSync/Forms/ModifyBuildStepsWindow.cs
+++ b/Engine/Source/Programs/UnrealGameSync/UnrealGameSync/Forms/ModifyBuildStepsWindow.cs
@@ -20,7 +20,7 @@ namespace UnrealGameSync
 		List<BuildStep> Steps;
 		HashSet<Guid> ProjectSteps;
 		DirectoryReference BaseDirectory;
-		ListViewItem.ListViewSubItem MouseDownSubItem = null;
+		ListViewItem.ListViewSubItem? MouseDownSubItem = null;
 		IReadOnlyDictionary<string, string> Variables;
 
 		public ModifyBuildStepsWindow(List<string> InTargetNames, List<BuildStep> InSteps, HashSet<Guid> InProjectSteps, DirectoryReference InBaseDirectory, IReadOnlyDictionary<string, string> InVariables)
@@ -78,33 +78,42 @@ namespace UnrealGameSync
 
 		private void EditStepButton_Click(object sender, EventArgs e)
 		{
-			foreach(ListViewItem Item in BuildStepList.SelectedItems)
+			foreach(ListViewItem? Item in BuildStepList.SelectedItems)
 			{
-				BuildStepWindow EditStep = new BuildStepWindow((BuildStep)Item.Tag, TargetNames, BaseDirectory, Variables);
-				EditStep.ShowDialog();
-				Item.Text = ((BuildStep)Item.Tag).Description;
-				break;
+				if (Item != null)
+				{
+					BuildStepWindow EditStep = new BuildStepWindow((BuildStep)Item.Tag, TargetNames, BaseDirectory, Variables);
+					EditStep.ShowDialog();
+					Item.Text = ((BuildStep)Item.Tag).Description;
+					break;
+				}
 			}
 		}
 
 		private void RemoveStepButton_Click(object sender, EventArgs e)
 		{
-			foreach(ListViewItem Item in BuildStepList.SelectedItems)
+			foreach(ListViewItem? Item in BuildStepList.SelectedItems)
 			{
-				if(MessageBox.Show(String.Format("Remove the '{0}' step?", Item.Text), "Remove Step", MessageBoxButtons.YesNo) == DialogResult.Yes)
+				if (Item != null)
 				{
-					BuildStepList.Items.Remove(Item);
+					if (MessageBox.Show(String.Format("Remove the '{0}' step?", Item.Text), "Remove Step", MessageBoxButtons.YesNo) == DialogResult.Yes)
+					{
+						BuildStepList.Items.Remove(Item);
+					}
+					break;
 				}
-				break;
 			}
 		}
 
 		private void CloseButton_Click(object sender, EventArgs e)
 		{
 			Steps.Clear();
-			foreach(ListViewItem Item in BuildStepList.Items)
+			foreach(ListViewItem? Item in BuildStepList.Items)
 			{
-				Steps.Add((BuildStep)Item.Tag);
+				if (Item != null)
+				{
+					Steps.Add((BuildStep)Item.Tag);
+				}
 			}
 			Close();
 		}
@@ -210,15 +219,18 @@ namespace UnrealGameSync
 		private void MoveUp_Click(object sender, EventArgs e)
 		{
 			BuildStepList.BeginUpdate();
-			foreach(ListViewItem Item in BuildStepList.SelectedItems)
+			foreach(ListViewItem? Item in BuildStepList.SelectedItems)
 			{
-				int Index = Item.Index;
-				if(Index > 0)
+				if (Item != null)
 				{
-					BuildStepList.Items.RemoveAt(Index);
-					BuildStepList.Items.Insert(Index - 1, Item);
+					int Index = Item.Index;
+					if (Index > 0)
+					{
+						BuildStepList.Items.RemoveAt(Index);
+						BuildStepList.Items.Insert(Index - 1, Item);
+					}
+					break;
 				}
-				break;
 			}
 			BuildStepList.EndUpdate();
 			UpdateEnabledButtons();
@@ -227,15 +239,18 @@ namespace UnrealGameSync
 		private void MoveDown_Click(object sender, EventArgs e)
 		{
 			BuildStepList.BeginUpdate();
-			foreach(ListViewItem Item in BuildStepList.SelectedItems)
+			foreach(ListViewItem? Item in BuildStepList.SelectedItems)
 			{
-				int Index = Item.Index;
-				if(Index < BuildStepList.Items.Count - 1)
+				if (Item != null)
 				{
-					BuildStepList.Items.RemoveAt(Index);
-					BuildStepList.Items.Insert(Index + 1, Item);
+					int Index = Item.Index;
+					if (Index < BuildStepList.Items.Count - 1)
+					{
+						BuildStepList.Items.RemoveAt(Index);
+						BuildStepList.Items.Insert(Index + 1, Item);
+					}
+					break;
 				}
-				break;
 			}
 			BuildStepList.EndUpdate();
 			UpdateEnabledButtons();
@@ -244,11 +259,14 @@ namespace UnrealGameSync
 		private void ModifyBuildStepsWindow_FormClosed(object sender, FormClosedEventArgs e)
 		{
 			Steps.Clear();
-			foreach(ListViewItem Item in BuildStepList.Items)
+			foreach(ListViewItem? Item in BuildStepList.Items)
 			{
-				BuildStep Step = (BuildStep)Item.Tag;
-				Step.OrderIndex = Steps.Count;
-				Steps.Add(Step);
+				if (Item != null)
+				{
+					BuildStep Step = (BuildStep)Item.Tag;
+					Step.OrderIndex = Steps.Count;
+					Steps.Add(Step);
+				}
 			}
 		}
 	}
diff --git a/Engine/Source/Programs/UnrealGameSync/UnrealGameSync/Forms/NewWorkspaceWindow.cs b/Engine/Source/Programs/UnrealGameSync/UnrealGameSync/Forms/NewWorkspaceWindow.cs
--- a/Engine/Source/Programs/UnrealGameSync/UnrealGameSync/Forms/NewWorkspaceWindow.cs
+++ b/Engine/Source/Programs/UnrealGameSync/UnrealGameSync/Forms/NewWorkspaceWindow.cs
@@ -17,8 +17,6 @@ using System.Threading;
 using System.Threading.Tasks;
 using System.Windows.Forms;
 
-#nullable enable
-
 namespace UnrealGameSync
 {
 	partial class NewWorkspaceWindow : Form
diff --git a/Engine/Source/Programs/UnrealGameSync/UnrealGameSync/Forms/NotificationWindow.cs b/Engine/Source/Programs/UnrealGameSync/UnrealGameSync/Forms/NotificationWindow.cs
--- a/Engine/Source/Programs/UnrealGameSync/UnrealGameSync/Forms/NotificationWindow.cs
+++ b/Engine/Source/Programs/UnrealGameSync/UnrealGameSync/Forms/NotificationWindow.cs
@@ -1,4 +1,4 @@
-﻿// Copyright Epic Games, Inc. All Rights Reserved.
+// Copyright Epic Games, Inc. All Rights Reserved.
 
 using System;
 using System.Collections.Generic;
@@ -27,16 +27,16 @@ namespace UnrealGameSync
 
 		// Logo
 		Rectangle LogoBounds;
-		Bitmap LogoBitmap;
+		Bitmap? LogoBitmap;
 
 		// Caption
 		Rectangle CaptionBounds;
-		string CaptionText;
-		Font CaptionFont;
+		string CaptionText = "Caption Text";
+		Font? CaptionFont;
 
 		// Message
 		Rectangle MessageBounds;
-		string MessageText;
+		string MessageText = "Mesage Text";
 
 		// Close button
 		Rectangle CloseButtonBounds;
@@ -44,8 +44,8 @@ namespace UnrealGameSync
 		bool bMouseDownOverClose;
 
 		// Notifications
-		public Action OnMoreInformation;
-		public Action OnDismiss;
+		public Action? OnMoreInformation;
+		public Action? OnDismiss;
 
 		public NotificationWindow(Image InLogo)
 		{
diff --git a/Engine/Source/Programs/UnrealGameSync/UnrealGameSync/Forms/OpenProjectWindow.cs b/Engine/Source/Programs/UnrealGameSync/UnrealGameSync/Forms/OpenProjectWindow.cs
--- a/Engine/Source/Programs/UnrealGameSync/UnrealGameSync/Forms/OpenProjectWindow.cs
+++ b/Engine/Source/Programs/UnrealGameSync/UnrealGameSync/Forms/OpenProjectWindow.cs
@@ -30,7 +30,7 @@ namespace UnrealGameSync
 		IServiceProvider ServiceProvider;
 		UserSettings Settings;
 
-		private OpenProjectWindow(UserSelectedProjectSettings Project, UserSettings Settings, IPerforceSettings DefaultSettings, IServiceProvider ServiceProvider)
+		private OpenProjectWindow(UserSelectedProjectSettings? Project, UserSettings Settings, IPerforceSettings DefaultSettings, IServiceProvider ServiceProvider)
 		{
 			InitializeComponent();
 
@@ -90,7 +90,7 @@ namespace UnrealGameSync
 			get => Utility.OverridePerforceSettings(DefaultSettings, ServerAndPortOverride, UserNameOverride);
 		}
 
-		public static WorkspaceSettings? ShowModal(IWin32Window Owner, UserSelectedProjectSettings Project, UserSettings Settings, DirectoryReference DataFolder, DirectoryReference CacheFolder, IPerforceSettings DefaultPerforceSettings, IServiceProvider ServiceProvider)
+		public static WorkspaceSettings? ShowModal(IWin32Window Owner, UserSelectedProjectSettings? Project, UserSettings Settings, DirectoryReference DataFolder, DirectoryReference CacheFolder, IPerforceSettings DefaultPerforceSettings, IServiceProvider ServiceProvider)
 		{
 			OpenProjectWindow Window = new OpenProjectWindow(Project, Settings, DefaultPerforceSettings, ServiceProvider);
 			if(Window.ShowDialog(Owner) == DialogResult.OK)
diff --git a/Engine/Source/Programs/UnrealGameSync/UnrealGameSync/Forms/ProgramsRunningWindow.cs b/Engine/Source/Programs/UnrealGameSync/UnrealGameSync/Forms/ProgramsRunningWindow.cs
--- a/Engine/Source/Programs/UnrealGameSync/UnrealGameSync/Forms/ProgramsRunningWindow.cs
+++ b/Engine/Source/Programs/UnrealGameSync/UnrealGameSync/Forms/ProgramsRunningWindow.cs
@@ -19,8 +19,8 @@ namespace UnrealGameSync
 		object SyncObject = new object();
 		FileReference[] Programs;
 		Func<FileReference[]> EnumeratePrograms;
-		ManualResetEvent TerminateEvent;
-		Thread BackgroundThread;
+		ManualResetEvent? TerminateEvent;
+		Thread? BackgroundThread;
 
 		public ProgramsRunningWindow(Func<FileReference[]> EnumeratePrograms, FileReference[] Programs)
 		{
@@ -43,7 +43,7 @@ namespace UnrealGameSync
 		{
 			for(;;)
 			{
-				if(TerminateEvent.WaitOne(TimeSpan.FromSeconds(2.0)))
+				if(TerminateEvent!.WaitOne(TimeSpan.FromSeconds(2.0)))
 				{
 					break;
 				}
@@ -76,7 +76,7 @@ namespace UnrealGameSync
 		{
 			if(BackgroundThread != null)
 			{
-				TerminateEvent.Set();
+				TerminateEvent!.Set();
 
 				BackgroundThread.Join();
 				BackgroundThread = null;
diff --git a/Engine/Source/Programs/UnrealGameSync/UnrealGameSync/Forms/ScheduleWindow.cs b/Engine/Source/Programs/UnrealGameSync/UnrealGameSync/Forms/ScheduleWindow.cs
--- a/Engine/Source/Programs/UnrealGameSync/UnrealGameSync/Forms/ScheduleWindow.cs
+++ b/Engine/Source/Programs/UnrealGameSync/UnrealGameSync/Forms/ScheduleWindow.cs
@@ -62,7 +62,7 @@ namespace UnrealGameSync
 			OutAnyOpenProject = false;
 
 			List<UserSelectedProjectSettings> ScheduledProjects = new List<UserSelectedProjectSettings>();
-			foreach(int Index in ProjectListBox.CheckedIndices)
+			foreach(int Index in ProjectListBox.CheckedIndices.OfType<int>())
 			{
 				if(Index == 0)
 				{
diff --git a/Engine/Source/Programs/UnrealGameSync/UnrealGameSync/Forms/SdkInfoWindow.cs b/Engine/Source/Programs/UnrealGameSync/UnrealGameSync/Forms/SdkInfoWindow.cs
--- a/Engine/Source/Programs/UnrealGameSync/UnrealGameSync/Forms/SdkInfoWindow.cs
+++ b/Engine/Source/Programs/UnrealGameSync/UnrealGameSync/Forms/SdkInfoWindow.cs
@@ -13,28 +13,42 @@ using System.Text;
 using System.Threading.Tasks;
 using System.Windows.Forms;
 
+#nullable enable
+
 namespace UnrealGameSync
 {
 	public partial class SdkInfoWindow : Form
 	{
 		class SdkItem
 		{
-			public string Category;
-			public string Description;
-			public string Install;
-			public string Browse;
+			public string Category { get; }
+			public string Description { get; }
+			public string? Install;
+			public string? Browse;
+
+			public SdkItem(string Category, string Description)
+			{
+				this.Category = Category;
+				this.Description = Description;
+			}
 		}
 
 		class BadgeInfo
 		{
-			public string UniqueId;
-			public string Label;
+			public string UniqueId { get; }
+			public string Label { get; }
 			public Rectangle Rectangle;
-			public Action OnClick;
+			public Action? OnClick;
+
+			public BadgeInfo(string UniqueId, string Label)
+			{
+				this.UniqueId = UniqueId;
+				this.Label = Label;
+			}
 		}
 
 		Font BadgeFont;
-		string HoverBadgeUniqueId;
+		string? HoverBadgeUniqueId;
 
 		public SdkInfoWindow(string[] SdkInfoEntries, Dictionary<string, string> Variables, Font BadgeFont)
 		{
@@ -49,7 +63,7 @@ namespace UnrealGameSync
 
 				string UniqueId = Object.GetValue("UniqueId", Guid.NewGuid().ToString());
 
-				ConfigObject ExistingObject;
+				ConfigObject? ExistingObject;
 				if(UniqueIdToObject.TryGetValue(UniqueId, out ExistingObject))
 				{
 					ExistingObject.AddOverrides(Object, null);
@@ -63,9 +77,9 @@ namespace UnrealGameSync
 			List<SdkItem> Items = new List<SdkItem>();
 			foreach(ConfigObject Object in UniqueIdToObject.Values)
 			{
-				SdkItem Item = new SdkItem();
-				Item.Category = Object.GetValue("Category", "Other");
-				Item.Description = Object.GetValue("Description", "");
+				string Category = Object.GetValue("Category", "Other");
+				string Description = Object.GetValue("Description", "");
+				SdkItem Item = new SdkItem(Category, Description);
 
 				Item.Install = Utility.ExpandVariables(Object.GetValue("Install", ""), Variables);
 				if(Item.Install.Contains("$("))
@@ -108,7 +122,7 @@ namespace UnrealGameSync
 				}
 			}
 
-			System.Reflection.PropertyInfo DoubleBufferedProperty = typeof(Control).GetProperty("DoubleBuffered", System.Reflection.BindingFlags.NonPublic | System.Reflection.BindingFlags.Instance);
+			System.Reflection.PropertyInfo DoubleBufferedProperty = typeof(Control).GetProperty("DoubleBuffered", System.Reflection.BindingFlags.NonPublic | System.Reflection.BindingFlags.Instance)!;
 			DoubleBufferedProperty.SetValue(SdkListView, true, null);
 
 			if(SdkListView.Items.Count > 0)
@@ -153,7 +167,7 @@ namespace UnrealGameSync
 
 		private void SdkListView_MouseMove(object sender, MouseEventArgs e)
 		{
-			string NewHoverUniqueId = null;
+			string? NewHoverUniqueId = null;
 
 			ListViewHitTestInfo HitTest = SdkListView.HitTest(e.Location);
 			if(HitTest.Item != null && HitTest.SubItem == HitTest.Item.SubItems[2])
@@ -204,19 +218,19 @@ namespace UnrealGameSync
 
 			SdkItem Sdk = (SdkItem)SubItem.Tag;
 
-			Action InstallAction = null;
+			Action? InstallAction = null;
 			if(!String.IsNullOrEmpty(Sdk.Install))
 			{
 				InstallAction = () => { Install(Sdk.Install); };
 			}
-			Badges.Add(new BadgeInfo(){ UniqueId = UniqueIdPrefix + "_Install", Label = "Install", OnClick = InstallAction });
+			Badges.Add(new BadgeInfo(UniqueIdPrefix + "_Install", "Install") { OnClick = InstallAction });
 
-			Action BrowseAction = null;
+			Action? BrowseAction = null;
 			if(!String.IsNullOrEmpty(Sdk.Browse))
 			{
 				BrowseAction = () => { Browse(Sdk.Browse); };
 			}
-			Badges.Add(new BadgeInfo(){ UniqueId = UniqueIdPrefix + "_Browse", Label = "Browse", OnClick = BrowseAction });
+			Badges.Add(new BadgeInfo(UniqueIdPrefix + "_Browse", "Browse"){ OnClick = BrowseAction });
 
 			int Right = SubItem.Bounds.Right - 10;
 			for(int Idx = Badges.Count - 1; Idx >= 0; Idx--)
diff --git a/Engine/Source/Programs/UnrealGameSync/UnrealGameSync/Forms/SyncFilter.cs b/Engine/Source/Programs/UnrealGameSync/UnrealGameSync/Forms/SyncFilter.cs
--- a/Engine/Source/Programs/UnrealGameSync/UnrealGameSync/Forms/SyncFilter.cs
+++ b/Engine/Source/Programs/UnrealGameSync/UnrealGameSync/Forms/SyncFilter.cs
@@ -57,17 +57,17 @@ namespace UnrealGameSync
 			GlobalControl.IncludeAllProjectsInSolution.CheckStateChanged += GlobalControl_IncludeAllProjectsInSolution_CheckStateChanged;
 		}
 
-		private void GlobalControl_CategoriesCheckList_ItemCheck(object sender, ItemCheckEventArgs e)
+		private void GlobalControl_CategoriesCheckList_ItemCheck(object? sender, ItemCheckEventArgs e)
 		{
 			WorkspaceControl.CategoriesCheckList.SetItemCheckState(e.Index, e.NewValue);
 		}
 
-		private void GlobalControl_SyncAllProjects_CheckStateChanged(object sender, EventArgs e)
+		private void GlobalControl_SyncAllProjects_CheckStateChanged(object? sender, EventArgs e)
 		{
 			WorkspaceControl.SyncAllProjects.Checked = GlobalControl.SyncAllProjects.Checked;
 		}
 
-		private void GlobalControl_IncludeAllProjectsInSolution_CheckStateChanged(object sender, EventArgs e)
+		private void GlobalControl_IncludeAllProjectsInSolution_CheckStateChanged(object? sender, EventArgs e)
 		{
 			WorkspaceControl.IncludeAllProjectsInSolution.Checked = GlobalControl.IncludeAllProjectsInSolution.Checked;
 		}
diff --git a/Engine/Source/Programs/UnrealGameSync/UnrealGameSync/IssueMonitor.cs b/Engine/Source/Programs/UnrealGameSync/UnrealGameSync/IssueMonitor.cs
--- a/Engine/Source/Programs/UnrealGameSync/UnrealGameSync/IssueMonitor.cs
+++ b/Engine/Source/Programs/UnrealGameSync/UnrealGameSync/IssueMonitor.cs
@@ -118,7 +118,7 @@ namespace UnrealGameSync
 	{
 		public long Id { get; set; }
 		public string Owner { get; set; } = String.Empty;
-		public string NominatedBy { get; set; } = String.Empty;
+		public string? NominatedBy { get; set; }
 		public bool? Acknowledged { get; set; }
 		public int? FixChange { get; set; }
 		public bool? Resolved { get; set; }
@@ -136,7 +136,7 @@ namespace UnrealGameSync
 
 	class IssueMonitor : IDisposable
 	{
-		public readonly string ApiUrl;
+		public readonly string? ApiUrl;
 		public readonly string UserName;
 		int RefCount = 1;
 		Task? WorkerTask;
@@ -155,7 +155,7 @@ namespace UnrealGameSync
 		// Only used by MainWindow, but easier to just store here
 		public Dictionary<long, IssueAlertReason> IssueIdToAlertReason = new Dictionary<long, IssueAlertReason>();
 
-		public IssueMonitor(string ApiUrl, string UserName, TimeSpan UpdateInterval, IServiceProvider ServiceProvider)
+		public IssueMonitor(string? ApiUrl, string UserName, TimeSpan UpdateInterval, IServiceProvider ServiceProvider)
 		{
 			this.ApiUrl = ApiUrl;
 			this.UserName = UserName;
diff --git a/Engine/Source/Programs/UnrealGameSync/UnrealGameSync/Jupiter/JupiterMonitor.cs b/Engine/Source/Programs/UnrealGameSync/UnrealGameSync/Jupiter/JupiterMonitor.cs
--- a/Engine/Source/Programs/UnrealGameSync/UnrealGameSync/Jupiter/JupiterMonitor.cs
+++ b/Engine/Source/Programs/UnrealGameSync/UnrealGameSync/Jupiter/JupiterMonitor.cs
@@ -14,6 +14,8 @@ using EpicGames.Jupiter;
 using EpicGames.Perforce;
 using Microsoft.Extensions.Logging;
 
+#nullable disable
+
 namespace UnrealGameSync
 {
 	class JupiterMonitor: IArchiveInfoSource, IDisposable
diff --git a/Engine/Source/Programs/UnrealGameSync/UnrealGameSync/ModalTask.cs b/Engine/Source/Programs/UnrealGameSync/UnrealGameSync/ModalTask.cs
--- a/Engine/Source/Programs/UnrealGameSync/UnrealGameSync/ModalTask.cs
+++ b/Engine/Source/Programs/UnrealGameSync/UnrealGameSync/ModalTask.cs
@@ -8,8 +8,6 @@ using System.Threading;
 using System.Threading.Tasks;
 using System.Windows.Forms;
 
-#nullable enable
-
 namespace UnrealGameSync
 {
 	[Flags]
diff --git a/Engine/Source/Programs/UnrealGameSync/UnrealGameSync/OIDCTokenManager.cs b/Engine/Source/Programs/UnrealGameSync/UnrealGameSync/OIDCTokenManager.cs
--- a/Engine/Source/Programs/UnrealGameSync/UnrealGameSync/OIDCTokenManager.cs
+++ b/Engine/Source/Programs/UnrealGameSync/UnrealGameSync/OIDCTokenManager.cs
@@ -15,6 +15,8 @@ using IdentityModel.Client;
 using IdentityModel.OidcClient;
 using IdentityModel.OidcClient.Results;
 
+#nullable disable
+
 namespace UnrealGameSync
 {
 	public class OIDCTokenManager
diff --git a/Engine/Source/Programs/UnrealGameSync/UnrealGameSync/PerforceModalTask.cs b/Engine/Source/Programs/UnrealGameSync/UnrealGameSync/PerforceModalTask.cs
--- a/Engine/Source/Programs/UnrealGameSync/UnrealGameSync/PerforceModalTask.cs
+++ b/Engine/Source/Programs/UnrealGameSync/UnrealGameSync/PerforceModalTask.cs
@@ -11,8 +11,6 @@ using System.Threading;
 using System.Threading.Tasks;
 using System.Windows.Forms;
 
-#nullable enable
-
 namespace UnrealGameSync
 {
 	public class PerforceLoginException : Exception
@@ -82,7 +80,7 @@ namespace UnrealGameSync
 						PasswordPrompt = $"Authentication failed. Enter the password for user '{PerforceSettings.UserName}' on server '{PerforceSettings.ServerAndPort}'.";
 					}
 
-					PasswordWindow PasswordWindow = new PasswordWindow(PasswordPrompt, Password);
+					PasswordWindow PasswordWindow = new PasswordWindow(PasswordPrompt, Password ?? String.Empty);
 					if (Owner == null)
 					{
 						PasswordWindow.StartPosition = FormStartPosition.CenterScreen;
diff --git a/Engine/Source/Programs/UnrealGameSync/UnrealGameSync/PerforceMonitor.cs b/Engine/Source/Programs/UnrealGameSync/UnrealGameSync/PerforceMonitor.cs
--- a/Engine/Source/Programs/UnrealGameSync/UnrealGameSync/PerforceMonitor.cs
+++ b/Engine/Source/Programs/UnrealGameSync/UnrealGameSync/PerforceMonitor.cs
@@ -16,8 +16,6 @@ using System.Text;
 using System.Threading;
 using System.Threading.Tasks;
 
-#nullable enable
-
 namespace UnrealGameSync
 {
 	class PerforceChangeDetails
diff --git a/Engine/Source/Programs/UnrealGameSync/UnrealGameSync/Program.cs b/Engine/Source/Programs/UnrealGameSync/UnrealGameSync/Program.cs
--- a/Engine/Source/Programs/UnrealGameSync/UnrealGameSync/Program.cs
+++ b/Engine/Source/Programs/UnrealGameSync/UnrealGameSync/Program.cs
@@ -8,6 +8,7 @@ using System;
 using System.Collections.Generic;
 using System.Data.SqlClient;
 using System.Diagnostics;
+using System.Diagnostics.CodeAnalysis;
 using System.IO;
 using System.Linq;
 using System.Net;
@@ -26,7 +27,7 @@ namespace UnrealGameSync
 	{
 		public static string GetVersionString()
 		{
-			AssemblyInformationalVersionAttribute Version = Assembly.GetExecutingAssembly().GetCustomAttribute<AssemblyInformationalVersionAttribute>();
+			AssemblyInformationalVersionAttribute? Version = Assembly.GetExecutingAssembly().GetCustomAttribute<AssemblyInformationalVersionAttribute>();
 			return Version?.InformationalVersion ?? "Unknown";
 		}
 
@@ -40,7 +41,7 @@ namespace UnrealGameSync
 
 		public static JsonSerializerOptions DefaultJsonSerializerOptions { get; } = GetDefaultJsonSerializerOptions();
 
-		public static string SyncVersion = null;
+		public static string? SyncVersion = null;
 
 		[STAThread]
 		static void Main(string[] Args)
@@ -76,17 +77,17 @@ namespace UnrealGameSync
 
 		static async Task InnerMainAsync(Mutex InstanceMutex, EventWaitHandle ActivateEvent, string[] Args)
 		{
-			string ServerAndPort = null;
-			string UserName = null;
-			string BaseUpdatePath = null;
+			string? ServerAndPort = null;
+			string? UserName = null;
+			string? BaseUpdatePath = null;
 			Utility.ReadGlobalPerforceSettings(ref ServerAndPort, ref UserName, ref BaseUpdatePath);
 
 			List<string> RemainingArgs = new List<string>(Args);
 
-			string UpdateSpawn;
+			string? UpdateSpawn;
 			ParseArgument(RemainingArgs, "-updatespawn=", out UpdateSpawn);
 
-			string UpdatePath;
+			string? UpdatePath;
 			ParseArgument(RemainingArgs, "-updatepath=", out UpdatePath);
 
 			bool bRestoreState;
@@ -95,10 +96,10 @@ namespace UnrealGameSync
 			bool bUnstable;
 			ParseOption(RemainingArgs, "-unstable", out bUnstable);
 
-            string ProjectFileName;
+            string? ProjectFileName;
             ParseArgument(RemainingArgs, "-project=", out ProjectFileName);
 
-			string Uri;
+			string? Uri;
 			ParseArgument(RemainingArgs, "-uri=", out Uri);
 
 			FileReference UpdateConfigFile = FileReference.Combine(new FileReference(Assembly.GetExecutingAssembly().Location).Directory, "AutoUpdate.ini");
@@ -110,7 +111,7 @@ namespace UnrealGameSync
 				Directory.SetCurrentDirectory(Path.GetDirectoryName(UpdateSpawn));
 			}
 
-			string SyncVersionFile = Path.Combine(Path.GetDirectoryName(Assembly.GetExecutingAssembly().Location), "SyncVersion.txt");
+			string SyncVersionFile = Path.Combine(Path.GetDirectoryName(Assembly.GetExecutingAssembly().Location!)!, "SyncVersion.txt");
 			if(File.Exists(SyncVersionFile))
 			{
 				try
@@ -123,7 +124,7 @@ namespace UnrealGameSync
 				}
 			}
 
-			DirectoryReference DataFolder = DirectoryReference.Combine(DirectoryReference.GetSpecialFolder(Environment.SpecialFolder.LocalApplicationData), "UnrealGameSync");
+			DirectoryReference DataFolder = DirectoryReference.Combine(DirectoryReference.GetSpecialFolder(Environment.SpecialFolder.LocalApplicationData)!, "UnrealGameSync");
 			DirectoryReference.CreateDirectory(DataFolder);
 
 			// Enable TLS 1.1 and 1.2. TLS 1.0 is now deprecated and not allowed by default in NET Core servers.
@@ -158,7 +159,7 @@ namespace UnrealGameSync
 					ILogger TelemetryLogger = LoggerProvider.CreateLogger("Telemetry");
 					TelemetryLogger.LogInformation("Creating telemetry sink for session {SessionId}", SessionId);
 
-					ITelemetrySink PrevTelemetrySink = Telemetry.ActiveSink;
+					ITelemetrySink? PrevTelemetrySink = Telemetry.ActiveSink;
 					using (ITelemetrySink TelemetrySink = DeploymentSettings.CreateTelemetrySink(UserName, SessionId, TelemetryLogger))
 					{
 						Telemetry.ActiveSink = TelemetrySink;
@@ -189,14 +190,14 @@ namespace UnrealGameSync
 
 		private static void CurrentDomain_UnhandledException(object Sender, UnhandledExceptionEventArgs Args)
 		{
-			Exception Ex = Args.ExceptionObject as Exception;
+			Exception? Ex = Args.ExceptionObject as Exception;
 			if(Ex != null)
 			{
 				Telemetry.SendEvent("Crash", new {Exception = Ex.ToString()});
 			}
 		}
 
-		static void MergeUpdateSettings(FileReference UpdateConfigFile, ref string UpdatePath, ref string UpdateSpawn)
+		static void MergeUpdateSettings(FileReference UpdateConfigFile, ref string? UpdatePath, ref string? UpdateSpawn)
 		{
 			try
 			{
@@ -247,7 +248,7 @@ namespace UnrealGameSync
 			return false;
 		}
 
-		static bool ParseArgument(List<string> RemainingArgs, string Prefix, out string Value)
+		static bool ParseArgument(List<string> RemainingArgs, string Prefix, [NotNullWhen(true)] out string? Value)
 		{
 			for(int Idx = 0; Idx < RemainingArgs.Count; Idx++)
 			{
diff --git a/Engine/Source/Programs/UnrealGameSync/UnrealGameSync/ProgramApplicationContext.cs b/Engine/Source/Programs/UnrealGameSync/UnrealGameSync/ProgramApplicationContext.cs
--- a/Engine/Source/Programs/UnrealGameSync/UnrealGameSync/ProgramApplicationContext.cs
+++ b/Engine/Source/Programs/UnrealGameSync/UnrealGameSync/ProgramApplicationContext.cs
@@ -27,14 +27,14 @@ namespace UnrealGameSync
 
 		IPerforceSettings DefaultPerforceSettings;
 		UpdateMonitor UpdateMonitor;
-		string ApiUrl;
+		string? ApiUrl;
 		DirectoryReference DataFolder;
 		DirectoryReference CacheFolder;
 		bool bRestoreState;
-		string UpdateSpawn;
+		string? UpdateSpawn;
 		bool bUnstable;
 		bool bIsClosing;
-		string Uri;
+		string? Uri;
 
 		IServiceProvider ServiceProvider;
 		ILogger Logger;
@@ -58,7 +58,7 @@ namespace UnrealGameSync
 
 		OIDCTokenManager? OIDCTokenManager;
 
-		public ProgramApplicationContext(IPerforceSettings DefaultPerforceSettings, UpdateMonitor UpdateMonitor, string ApiUrl, DirectoryReference DataFolder, EventWaitHandle ActivateEvent, bool bRestoreState, string UpdateSpawn, string ProjectFileName, bool bUnstable, IServiceProvider ServiceProvider, string Uri)
+		public ProgramApplicationContext(IPerforceSettings DefaultPerforceSettings, UpdateMonitor UpdateMonitor, string? ApiUrl, DirectoryReference DataFolder, EventWaitHandle ActivateEvent, bool bRestoreState, string? UpdateSpawn, string? ProjectFileName, bool bUnstable, IServiceProvider ServiceProvider, string? Uri)
 		{
 			this.DefaultPerforceSettings = DefaultPerforceSettings;
 			this.UpdateMonitor = UpdateMonitor;
@@ -241,7 +241,7 @@ namespace UnrealGameSync
 			}
 
 			// Get the application path
-			string OriginalExe = UpdateSpawn;
+			string? OriginalExe = UpdateSpawn;
 			if (OriginalExe == null)
 			{
 				OriginalExe = Path.ChangeExtension(Assembly.GetExecutingAssembly().Location, ".exe");
diff --git a/Engine/Source/Programs/UnrealGameSync/UnrealGameSync/RESTApi.cs b/Engine/Source/Programs/UnrealGameSync/UnrealGameSync/RESTApi.cs
--- a/Engine/Source/Programs/UnrealGameSync/UnrealGameSync/RESTApi.cs
+++ b/Engine/Source/Programs/UnrealGameSync/UnrealGameSync/RESTApi.cs
@@ -10,8 +10,6 @@ using System.Text.Json;
 using System.Threading;
 using System.Threading.Tasks;
 
-#nullable enable
-
 namespace UnrealGameSync
 {
 	class RestException : Exception
diff --git a/Engine/Source/Programs/UnrealGameSync/UnrealGameSync/Taskbar.cs b/Engine/Source/Programs/UnrealGameSync/UnrealGameSync/Taskbar.cs
--- a/Engine/Source/Programs/UnrealGameSync/UnrealGameSync/Taskbar.cs
+++ b/Engine/Source/Programs/UnrealGameSync/UnrealGameSync/Taskbar.cs
@@ -1,4 +1,4 @@
-﻿// Copyright Epic Games, Inc. All Rights Reserved.
+// Copyright Epic Games, Inc. All Rights Reserved.
 
 using System;
 using System.Collections.Generic;
@@ -38,7 +38,7 @@ namespace UnrealGameSync
 		{
 		}
 
-		static ITaskbarList3 Interface;
+		static ITaskbarList3? Interface;
 
 		static Taskbar()
 		{
diff --git a/Engine/Source/Programs/UnrealGameSync/UnrealGameSync/Telemetry.cs b/Engine/Source/Programs/UnrealGameSync/UnrealGameSync/Telemetry.cs
--- a/Engine/Source/Programs/UnrealGameSync/UnrealGameSync/Telemetry.cs
+++ b/Engine/Source/Programs/UnrealGameSync/UnrealGameSync/Telemetry.cs
@@ -213,7 +213,7 @@ namespace UnrealGameSync
 		/// <summary>
 		/// The current telemetry provider
 		/// </summary>
-		public static ITelemetrySink ActiveSink
+		public static ITelemetrySink? ActiveSink
 		{
 			get; set;
 		}
@@ -225,7 +225,7 @@ namespace UnrealGameSync
 		/// <param name="Attributes">Arbitrary object to include in the payload</param>
 		public static void SendEvent(string EventName, object Attributes)
 		{
-			ActiveSink.SendEvent(EventName, Attributes);
+			ActiveSink?.SendEvent(EventName, Attributes);
 		}
 	}
 }
diff --git a/Engine/Source/Programs/UnrealGameSync/UnrealGameSync/TelemetryStopwatch.cs b/Engine/Source/Programs/UnrealGameSync/UnrealGameSync/TelemetryStopwatch.cs
--- a/Engine/Source/Programs/UnrealGameSync/UnrealGameSync/TelemetryStopwatch.cs
+++ b/Engine/Source/Programs/UnrealGameSync/UnrealGameSync/TelemetryStopwatch.cs
@@ -10,14 +10,14 @@ namespace UnrealGameSync
 	class TelemetryStopwatch : IDisposable
 	{
 		readonly string EventName;
-		readonly Dictionary<string, object> EventData;
+		readonly Dictionary<string, object?> EventData;
 		readonly Stopwatch Timer;
 
 		public TelemetryStopwatch(string EventName, string Project)
 		{
 			this.EventName = EventName;
 
-			EventData = new Dictionary<string, object>();
+			EventData = new Dictionary<string, object?>();
 			EventData["Project"] = Project;
 
 			Timer = Stopwatch.StartNew();
diff --git a/Engine/Source/Programs/UnrealGameSync/UnrealGameSync/UnrealGameSync.csproj b/Engine/Source/Programs/UnrealGameSync/UnrealGameSync/UnrealGameSync.csproj
--- a/Engine/Source/Programs/UnrealGameSync/UnrealGameSync/UnrealGameSync.csproj
+++ b/Engine/Source/Programs/UnrealGameSync/UnrealGameSync/UnrealGameSync.csproj
@@ -12,6 +12,7 @@
     <StartupObject>UnrealGameSync.Program</StartupObject>
     <ApplicationIcon>Resources\Icon.ico</ApplicationIcon>
     <ApplicationManifest>Properties\app.manifest</ApplicationManifest>
+    <Nullable>enable</Nullable>
   </PropertyGroup>
   <PropertyGroup Condition="'$(Configuration)|$(Platform)'=='Release|AnyCPU'">
     <DefineConstants>TRACE;DEBUG;WITH_TELEMETRY</DefineConstants>
diff --git a/Engine/Source/Programs/UnrealGameSync/UnrealGameSync/UpdateMonitor.cs b/Engine/Source/Programs/UnrealGameSync/UnrealGameSync/UpdateMonitor.cs
--- a/Engine/Source/Programs/UnrealGameSync/UnrealGameSync/UpdateMonitor.cs
+++ b/Engine/Source/Programs/UnrealGameSync/UnrealGameSync/UpdateMonitor.cs
@@ -12,8 +12,6 @@ using System.Text;
 using System.Threading;
 using System.Threading.Tasks;
 
-#nullable enable
-
 namespace UnrealGameSync
 {
 	enum UpdateType
@@ -37,7 +35,7 @@ namespace UnrealGameSync
 			private set;
 		}
 
-		public UpdateMonitor(IPerforceSettings PerforceSettings, string WatchPath, IServiceProvider ServiceProvider)
+		public UpdateMonitor(IPerforceSettings PerforceSettings, string? WatchPath, IServiceProvider ServiceProvider)
 		{
 			this.Logger = ServiceProvider.GetRequiredService<ILogger<UpdateMonitor>>();
 			this.AsyncDisposer = ServiceProvider.GetRequiredService<IAsyncDisposer>();
diff --git a/Engine/Source/Programs/UnrealGameSync/UnrealGameSync/UriHandlers/RdpHandler.cs b/Engine/Source/Programs/UnrealGameSync/UnrealGameSync/UriHandlers/RdpHandler.cs
--- a/Engine/Source/Programs/UnrealGameSync/UnrealGameSync/UriHandlers/RdpHandler.cs
+++ b/Engine/Source/Programs/UnrealGameSync/UnrealGameSync/UriHandlers/RdpHandler.cs
@@ -16,6 +16,7 @@ using System.Runtime.InteropServices;
 using System.Runtime.InteropServices.ComTypes;
 using System.Text;
 using System.Security;
+using System.Diagnostics.CodeAnalysis;
 
 namespace UnrealGameSync
 {
@@ -39,16 +40,16 @@ namespace UnrealGameSync
 		{
 			public uint Flags;
 			public uint Type;
-			public string TargetName;
-			public string Comment;
+			public string? TargetName;
+			public string? Comment;
 			public FILETIME LastWritten;
 			public int CredentialBlobSize;
 			public IntPtr CredentialBlob;
 			public uint Persist;
 			public uint AttributeCount;
 			public IntPtr Attributes;
-			public string TargetAlias;
-			public string UserName;
+			public string? TargetAlias;
+			public string? UserName;
 		}
 
 		[DllImport("Advapi32.dll", EntryPoint = "CredReadW", CharSet = CharSet.Unicode, SetLastError = true)]
@@ -60,7 +61,7 @@ namespace UnrealGameSync
 		[DllImport("advapi32.dll", SetLastError = true)]
 		static extern bool CredFree(IntPtr Buffer);
 
-		public static bool ReadCredential(string Name, uint Type, out string UserName, out string Password)
+		public static bool ReadCredential(string Name, uint Type, [NotNullWhen(true)] out string? UserName, [NotNullWhen(true)] out string? Password)
 		{
 			IntPtr Buffer = IntPtr.Zero;
 			try
@@ -72,16 +73,16 @@ namespace UnrealGameSync
 					return false;
 				}
 
-				CREDENTIAL Credential = Marshal.PtrToStructure<CREDENTIAL>(Buffer);
-				UserName = Credential.UserName;
-				if (Credential.CredentialBlob == IntPtr.Zero)
+				CREDENTIAL? Credential = Marshal.PtrToStructure<CREDENTIAL>(Buffer);
+				if (Credential == null || Credential.UserName == null || Credential.CredentialBlob == IntPtr.Zero)
 				{
+					UserName = null;
 					Password = null;
+					return false;
 				}
-				else
-				{
-					Password = Marshal.PtrToStringUni(Credential.CredentialBlob, Credential.CredentialBlobSize / sizeof(char));
-				}
+
+				UserName = Credential.UserName;
+				Password = Marshal.PtrToStringUni(Credential.CredentialBlob, Credential.CredentialBlobSize / sizeof(char));
 				return true;
 			}
 			finally
@@ -121,7 +122,7 @@ namespace UnrealGameSync
 			// Copy the credentials from a generic Windows credential to 
 			if (!ReadCredential(Host, CRED_TYPE_DOMAIN_PASSWORD, out _, out _))
 			{
-				if (ReadCredential("UnrealGameSync:RDP", CRED_TYPE_GENERIC, out string UserName, out string Password))
+				if (ReadCredential("UnrealGameSync:RDP", CRED_TYPE_GENERIC, out string? UserName, out string? Password))
 				{
 					WriteCredential(Host, CRED_TYPE_DOMAIN_PASSWORD, UserName, Password);
 				}
diff --git a/Engine/Source/Programs/UnrealGameSync/UnrealGameSync/UriHandlers/UriHandler.cs b/Engine/Source/Programs/UnrealGameSync/UnrealGameSync/UriHandlers/UriHandler.cs
--- a/Engine/Source/Programs/UnrealGameSync/UnrealGameSync/UriHandlers/UriHandler.cs
+++ b/Engine/Source/Programs/UnrealGameSync/UnrealGameSync/UriHandlers/UriHandler.cs
@@ -22,8 +22,8 @@ namespace UnrealGameSync
 	class UriResult
 	{
 		public bool Success = false;
-		public string Error;
-		public AutomationRequest Request = null;
+		public string? Error;
+		public AutomationRequest? Request = null;
 	}
 
 	/// <summary>
@@ -38,7 +38,7 @@ namespace UnrealGameSync
 				Uri Uri = new Uri(UriIn);
 
 				// Check if this is a registered uri request
-				if (!Handlers.TryGetValue(Uri.Host, out MethodInfo Info))
+				if (!Handlers.TryGetValue(Uri.Host, out MethodInfo? Info))
 				{
 					return new UriResult() { Error = string.Format("Unknown Uri {0}", Uri.Host) };
 				}
@@ -58,7 +58,7 @@ namespace UnrealGameSync
 							return new UriResult() { Error = string.Format("Uri {0} is missing required parameter {1}", Uri.Host, Param.Name) };
 						}
 
-						Parameters.Add(Param.DefaultValue);
+						Parameters.Add(Param.DefaultValue!);
 						continue;
 					}
 
@@ -102,7 +102,7 @@ namespace UnrealGameSync
 
 				}
 
-				return Info.Invoke(null, Parameters.ToArray()) as UriResult;					
+				return (UriResult)Info.Invoke(null, Parameters.ToArray())!;					
 			}
 			catch (Exception Ex)
 			{
@@ -117,7 +117,7 @@ namespace UnrealGameSync
 		/// <summary>
 		/// Handle URI passed in via command lines
 		/// </summary>
-		public static bool ProcessCommandLine(string[] Args, bool FirstInstance, EventWaitHandle ActivateEvent = null)
+		public static bool ProcessCommandLine(string[] Args, bool FirstInstance, EventWaitHandle? ActivateEvent = null)
 		{
 			if (Args.Any(x => x.Equals(InstallHandlerArg, StringComparison.OrdinalIgnoreCase)))
 			{
@@ -171,14 +171,14 @@ namespace UnrealGameSync
 					return true;
 				}
 
-				MethodInfo Handler;
+				MethodInfo? Handler;
 				if (!Handlers.TryGetValue(Uri.Host, out Handler))
 				{
 					MessageBox.Show(String.Format("Unknown action from URI request ('{0}')", Uri.Host));
 					return true;
 				}
 
-				UriHandlerAttribute Attribute = Handler.GetCustomAttribute<UriHandlerAttribute>();
+				UriHandlerAttribute Attribute = Handler.GetCustomAttribute<UriHandlerAttribute>()!;
 
 				// handle case where we terminate after invoking handler
 				if (Attribute.Terminate)
@@ -283,10 +283,10 @@ namespace UnrealGameSync
 		{
 			public RegistryKey RootKey;
 			public string KeyName;
-			public string ValueName;
+			public string? ValueName;
 			public string Value;
 
-			public RegistrySetting(RegistryKey RootKey, string KeyName, string ValueName, string Value)
+			public RegistrySetting(RegistryKey RootKey, string KeyName, string? ValueName, string Value)
 			{
 				this.RootKey = RootKey;
 				this.KeyName = KeyName;
diff --git a/Engine/Source/Programs/UnrealGameSync/UnrealGameSync/UriHandlers/VisualStudioHandler.cs b/Engine/Source/Programs/UnrealGameSync/UnrealGameSync/UriHandlers/VisualStudioHandler.cs
--- a/Engine/Source/Programs/UnrealGameSync/UnrealGameSync/UriHandlers/VisualStudioHandler.cs
+++ b/Engine/Source/Programs/UnrealGameSync/UnrealGameSync/UriHandlers/VisualStudioHandler.cs
@@ -19,7 +19,7 @@ namespace UnrealGameSync
 		{
 			string TempFileName = P4Automation.PrintToTempFile(null, DepotPath, NullLogger.Instance).GetAwaiter().GetResult();
 
-			string ErrorMessage;
+			string? ErrorMessage;
 			if (!VisualStudioAutomation.OpenFile(TempFileName, out ErrorMessage, Line))
 			{
 				return new UriResult() { Error = ErrorMessage ?? "Unknown Visual Studio Error" };
diff --git a/Engine/Source/Programs/UnrealGameSync/UnrealGameSync/UserSettings.cs b/Engine/Source/Programs/UnrealGameSync/UnrealGameSync/UserSettings.cs
--- a/Engine/Source/Programs/UnrealGameSync/UnrealGameSync/UserSettings.cs
+++ b/Engine/Source/Programs/UnrealGameSync/UnrealGameSync/UserSettings.cs
@@ -15,8 +15,6 @@ using System.Text.RegularExpressions;
 using System.Threading.Tasks;
 using System.Windows.Forms;
 
-#nullable enable
-
 namespace UnrealGameSync
 {
 	public enum BuildConfig
diff --git a/Engine/Source/Programs/UnrealGameSync/UnrealGameSync/Utility.cs b/Engine/Source/Programs/UnrealGameSync/UnrealGameSync/Utility.cs
--- a/Engine/Source/Programs/UnrealGameSync/UnrealGameSync/Utility.cs
+++ b/Engine/Source/Programs/UnrealGameSync/UnrealGameSync/Utility.cs
@@ -19,8 +19,6 @@ using System.Threading;
 using System.Threading.Tasks;
 using System.Windows.Forms;
 
-#nullable enable
-
 namespace UnrealGameSync
 {
 	sealed class UserErrorException : Exception
diff --git a/Engine/Source/Programs/UnrealGameSync/UnrealGameSync/Workspace.cs b/Engine/Source/Programs/UnrealGameSync/UnrealGameSync/Workspace.cs
--- a/Engine/Source/Programs/UnrealGameSync/UnrealGameSync/Workspace.cs
+++ b/Engine/Source/Programs/UnrealGameSync/UnrealGameSync/Workspace.cs
@@ -19,8 +19,6 @@ using System.Text.Json;
 using System.Threading;
 using System.Threading.Tasks;
 
-#nullable enable
-
 namespace UnrealGameSync
 {
 	[Flags]
@@ -1226,7 +1224,7 @@ namespace UnrealGameSync
 					{
 						MaxProgressFraction += (float)Step.EstimatedDuration / (float)Math.Max(BuildSteps.Sum(x => x.EstimatedDuration), 1);
 
-						Progress.Set(Step.StatusText);
+						Progress.Set(Step.StatusText ?? "Executing build step");
 						Progress.Push(MaxProgressFraction);
 
 						Logger.LogInformation("{Status}", Step.StatusText);
@@ -1277,7 +1275,7 @@ namespace UnrealGameSync
 										StepStopwatch.AddData(new { Project = Path.GetFileNameWithoutExtension(Step.FileName) });
 
 										FileReference LocalRunUAT = FileReference.Combine(Project.LocalRootPath, "Engine", "Build", "BatchFiles", "RunUAT.bat");
-										string Arguments = String.Format("/C \"\"{0}\" -profile=\"{1}\"\"", LocalRunUAT, FileReference.Combine(Project.LocalRootPath, Step.FileName));
+										string Arguments = String.Format("/C \"\"{0}\" -profile=\"{1}\"\"", LocalRunUAT, FileReference.Combine(Project.LocalRootPath, Step.FileName ?? "unknown"));
 										Logger.LogInformation("uat> Running {FileName} {Argument}", LocalRunUAT, Arguments);
 
 										int ResultFromUAT = await Utility.ExecuteProcessAsync(CmdExe, null, Arguments, Line => ProcessOutput(Line, "uat> ", Progress, Logger), CancellationToken);
@@ -1295,7 +1293,7 @@ namespace UnrealGameSync
 									{
 										StepStopwatch.AddData(new { FileName = Path.GetFileNameWithoutExtension(Step.FileName) });
 
-										FileReference ToolFileName = FileReference.Combine(Project.LocalRootPath, Utility.ExpandVariables(Step.FileName, Context.Variables));
+										FileReference ToolFileName = FileReference.Combine(Project.LocalRootPath, Utility.ExpandVariables(Step.FileName ?? "unknown", Context.Variables));
 										string ToolWorkingDir = String.IsNullOrWhiteSpace(Step.WorkingDir) ? ToolFileName.Directory.FullName : Utility.ExpandVariables(Step.WorkingDir, Context.Variables);
 										string ToolArguments = Utility.ExpandVariables(Step.Arguments ?? "", Context.Variables);
 										Logger.LogInformation("tool> Running {0} {1}", ToolFileName, ToolArguments);
