UGS: Fix build errors for UnrealGameSyncLauncher and port to use new P4 library.

diff --git a/Engine/Source/Programs/Shared/EpicGames.Core/Log.cs b/Engine/Source/Programs/Shared/EpicGames.Core/Log.cs
--- a/Engine/Source/Programs/Shared/EpicGames.Core/Log.cs
+++ b/Engine/Source/Programs/Shared/EpicGames.Core/Log.cs
@@ -736,6 +736,52 @@ namespace EpicGames.Core
 		}
 	}
 
+	/// <summary>
+	/// Logger which captures the output for rendering later
+	/// </summary>
+	public class CaptureLogger : ILogger
+	{
+		class NullScope : IDisposable
+		{
+			public void Dispose() { }
+		}
+
+		/// <summary>
+		/// List of captured events
+		/// </summary>
+		public List<LogEvent> Events { get; } = new List<LogEvent>();
+
+		/// <summary>
+		/// Renders the captured events as a single string
+		/// </summary>
+		/// <returns>Rendered log text</returns>
+		public string Render() => Render("\n");
+
+		/// <summary>
+		/// Renders the captured events as a single string
+		/// </summary>
+		/// <returns>Rendered log text</returns>
+		public string Render(string NewLine) => String.Join(NewLine, RenderLines());
+
+		/// <summary>
+		/// Renders all the captured events
+		/// </summary>
+		/// <returns>List of rendered log lines</returns>
+		public List<string> RenderLines() => Events.ConvertAll(x => x.ToString());
+
+		/// <inheritdoc/>
+		public IDisposable BeginScope<TState>(TState State) => new NullScope();
+
+		/// <inheritdoc/>
+		public bool IsEnabled(LogLevel LogLevel) => true;
+
+		/// <inheritdoc/>
+		public void Log<TState>(LogLevel LogLevel, EventId EventId, TState State, Exception? Exception, Func<TState, Exception?, string> Formatter)
+		{
+			Events.Add(LogEvent.FromState(LogLevel, EventId, State, Exception, Formatter));
+		}
+	}
+
 	/// <summary>
 	/// Default log output device
 	/// </summary>
diff --git a/Engine/Source/Programs/Shared/EpicGames.Perforce/PerforceSettings.cs b/Engine/Source/Programs/Shared/EpicGames.Perforce/PerforceSettings.cs
--- a/Engine/Source/Programs/Shared/EpicGames.Perforce/PerforceSettings.cs
+++ b/Engine/Source/Programs/Shared/EpicGames.Perforce/PerforceSettings.cs
@@ -147,6 +147,27 @@ namespace EpicGames.Perforce
 	/// </summary>
 	public static class PerforceSettingExtensions
 	{
+		/// <summary>
+		/// Update common fields in a IPerforceSettings object
+		/// </summary>
+		public static IPerforceSettings MergeWith(this IPerforceSettings Settings, string? NewServerAndPort = null, string? NewUserName = null, string? NewClientName = null)
+		{
+			PerforceSettings NewSettings = new PerforceSettings(Settings);
+			if (!String.IsNullOrEmpty(NewServerAndPort))
+			{
+				NewSettings.ServerAndPort = NewServerAndPort;
+			}
+			if (!String.IsNullOrEmpty(NewUserName))
+			{
+				NewSettings.UserName = NewUserName;
+			}
+			if (!String.IsNullOrEmpty(NewClientName))
+			{
+				NewSettings.ClientName = NewClientName;
+			}
+			return NewSettings;
+		}
+
 		/// <summary>
 		/// Gets the command line arguments to launch an external program, such as P4V or P4VC
 		/// </summary>
diff --git a/Engine/Source/Programs/UnrealGameSync/UnrealGameSync/ModalTask.cs b/Engine/Source/Programs/UnrealGameSync/UnrealGameSync/ModalTask.cs
--- a/Engine/Source/Programs/UnrealGameSync/UnrealGameSync/ModalTask.cs
+++ b/Engine/Source/Programs/UnrealGameSync/UnrealGameSync/ModalTask.cs
@@ -69,6 +69,10 @@ namespace UnrealGameSync
 				Task<T> BackgroundTask = Task.Run(() => TaskFunc(CancellationSource.Token));
 
 				ModalTaskWindow Window = new ModalTaskWindow(Title, Message, (Owner == null) ? FormStartPosition.CenterScreen : FormStartPosition.CenterParent, BackgroundTask, CancellationSource);
+				if(Owner == null)
+				{
+					Window.ShowInTaskbar = true;
+				}
 				Window.ShowDialog(Owner);
 
 				if (BackgroundTask.IsCanceled || (BackgroundTask.Exception != null && BackgroundTask.Exception.InnerException is OperationCanceledException))
diff --git a/Engine/Source/Programs/UnrealGameSync/UnrealGameSync/PerforceModalTask.cs b/Engine/Source/Programs/UnrealGameSync/UnrealGameSync/PerforceModalTask.cs
--- a/Engine/Source/Programs/UnrealGameSync/UnrealGameSync/PerforceModalTask.cs
+++ b/Engine/Source/Programs/UnrealGameSync/UnrealGameSync/PerforceModalTask.cs
@@ -67,31 +67,36 @@ namespace UnrealGameSync
 			{
 				Func<CancellationToken, Task<T>> RunAsync = CancellationToken => LoginAndExecuteAsync(Password, ConnectAsync, ExecuteAsync, CancellationToken);
 
-				ModalTask<T>? Result = ModalTask.Execute(Owner, Title, Message, RunAsync, Flags);
-				if (Result != null && Result.Failed && Result.Exception is PerforceLoginException && (Flags & ModalTaskFlags.Quiet) == 0)
+				ModalTask<T>? Result = ModalTask.Execute(Owner, Title, Message, RunAsync, ModalTaskFlags.Quiet);
+				if (Result != null && Result.Failed && (Flags & ModalTaskFlags.Quiet) == 0)
 				{
-					string PasswordPrompt;
-					if (String.IsNullOrEmpty(Password))
+					if (Result.Exception is PerforceLoginException)
 					{
-						PasswordPrompt = $"Enter the password for user '{PerforceSettings.UserName}' on server '{PerforceSettings.ServerAndPort}'.";
+						string PasswordPrompt;
+						if (String.IsNullOrEmpty(Password))
+						{
+							PasswordPrompt = $"Enter the password for user '{PerforceSettings.UserName}' on server '{PerforceSettings.ServerAndPort}'.";
+						}
+						else
+						{
+							PasswordPrompt = $"Authentication failed. Enter the password for user '{PerforceSettings.UserName}' on server '{PerforceSettings.ServerAndPort}'.";
+						}
+
+						PasswordWindow PasswordWindow = new PasswordWindow(PasswordPrompt, Password ?? String.Empty);
+						if (Owner == null)
+						{
+							PasswordWindow.ShowInTaskbar = true;
+							PasswordWindow.StartPosition = FormStartPosition.CenterScreen;
+						}
+						if (PasswordWindow.ShowDialog(Owner) != DialogResult.OK)
+						{
+							return null;
+						}
+
+						Password = PasswordWindow.Password;
+						continue;
 					}
-					else
-					{
-						PasswordPrompt = $"Authentication failed. Enter the password for user '{PerforceSettings.UserName}' on server '{PerforceSettings.ServerAndPort}'.";
-					}
-
-					PasswordWindow PasswordWindow = new PasswordWindow(PasswordPrompt, Password ?? String.Empty);
-					if (Owner == null)
-					{
-						PasswordWindow.StartPosition = FormStartPosition.CenterScreen;
-					}
-					if (PasswordWindow.ShowDialog(Owner) != DialogResult.OK)
-					{
-						return null;
-					}
-
-					Password = PasswordWindow.Password;
-					continue;
+					MessageBox.Show(Owner, Title, Result.Error, MessageBoxButtons.OK);
 				}
 
 				return Result;
diff --git a/Engine/Source/Programs/UnrealGameSync/UnrealGameSyncLauncher/Forms/SettingsWindow.cs b/Engine/Source/Programs/UnrealGameSync/UnrealGameSyncLauncher/Forms/SettingsWindow.cs
--- a/Engine/Source/Programs/UnrealGameSync/UnrealGameSyncLauncher/Forms/SettingsWindow.cs
+++ b/Engine/Source/Programs/UnrealGameSync/UnrealGameSyncLauncher/Forms/SettingsWindow.cs
@@ -1,5 +1,8 @@
 // Copyright Epic Games, Inc. All Rights Reserved.
 
+using EpicGames.Core;
+using EpicGames.Perforce;
+using Microsoft.Extensions.Logging;
 using System;
 using System.Collections.Generic;
 using System.ComponentModel;
@@ -21,14 +24,14 @@ namespace UnrealGameSyncLauncher
 		[DllImport("user32.dll")]
 		private static extern IntPtr SendMessage(IntPtr hWnd, int Msg, int wParam, [MarshalAs(UnmanagedType.LPWStr)] string lParam);
 
-		public delegate bool SyncAndRunDelegate(PerforceConnection Perforce, string DepotPath, bool bUnstable, TextWriter LogWriter);
+		public delegate Task SyncAndRunDelegate(IPerforceConnection Perforce, string? DepotPath, bool bUnstable, ILogger LogWriter, CancellationToken CancellationToken);
 
 		const int EM_SETCUEBANNER = 0x1501;
 
-		string LogText;
+		string? LogText;
 		SyncAndRunDelegate SyncAndRun;
 
-		public SettingsWindow(string Prompt, string LogText, string ServerAndPort, string UserName, string DepotPath, bool bUnstable, SyncAndRunDelegate SyncAndRun)
+		public SettingsWindow(string? Prompt, string? LogText, string? ServerAndPort, string? UserName, string? DepotPath, bool bUnstable, SyncAndRunDelegate SyncAndRun)
 		{
 			InitializeComponent();
 
@@ -38,9 +41,9 @@ namespace UnrealGameSyncLauncher
 			}
 
 			this.LogText = LogText;
-			this.ServerTextBox.Text = ServerAndPort;
-			this.UserNameTextBox.Text = UserName;
-			this.DepotPathTextBox.Text = DepotPath;
+			this.ServerTextBox.Text = ServerAndPort ?? String.Empty;
+			this.UserNameTextBox.Text = UserName ?? String.Empty;
+			this.DepotPathTextBox.Text = DepotPath ?? String.Empty;
 			this.UseUnstableBuildCheckBox.Checked = bUnstable;
 			this.SyncAndRun = SyncAndRun;
 
@@ -57,26 +60,26 @@ namespace UnrealGameSyncLauncher
 
 		private void ViewLogBtn_Click(object sender, EventArgs e)
 		{
-			LogWindow Log = new LogWindow(LogText);
+			LogWindow Log = new LogWindow(LogText ?? String.Empty);
 			Log.ShowDialog(this);
 		}
 
 		private void ConnectBtn_Click(object sender, EventArgs e)
 		{
 			// Update the settings
-			string ServerAndPort = ServerTextBox.Text.Trim();
+			string? ServerAndPort = ServerTextBox.Text.Trim();
 			if(ServerAndPort.Length == 0)
 			{
 				ServerAndPort = null;
 			}
 
-			string UserName = UserNameTextBox.Text.Trim();
+			string? UserName = UserNameTextBox.Text.Trim();
 			if(UserName.Length == 0)
 			{
 				UserName = null;
 			}
 
-			string DepotPath = DepotPathTextBox.Text.Trim();
+			string? DepotPath = DepotPathTextBox.Text.Trim();
 			if(DepotPath.Length == 0)
 			{
 				DepotPath = null;
@@ -84,26 +87,34 @@ namespace UnrealGameSyncLauncher
 
 			Utility.SaveGlobalPerforceSettings(ServerAndPort, UserName, DepotPath);
 
-			// Create the task for connecting to this server
-			StringWriter Log = new StringWriter();
-			SyncAndRunPerforceTask SyncApplication = new SyncAndRunPerforceTask((Perforce, LogWriter) => SyncAndRun(Perforce, DepotPath, UseUnstableBuildCheckBox.Checked, LogWriter));
-
-			// Attempt to sync through a modal dialog
-			string ErrorMessage;
-			ModalTaskResult Result = PerforceModalTask.Execute(this, new PerforceConnection(UserName, null, ServerAndPort), SyncApplication, "Updating", "Checking for updates, please wait...", Log, out ErrorMessage);
-			if(Result == ModalTaskResult.Succeeded)
+			PerforceSettings PerforceSettings = new PerforceSettings(PerforceSettings.Default);
+			if (!String.IsNullOrEmpty(ServerAndPort))
+			{
+				PerforceSettings.ServerAndPort = ServerAndPort;
+			}
+			if (!String.IsNullOrEmpty(UserName))
 			{
-				Utility.SaveGlobalPerforceSettings(ServerAndPort, UserName, DepotPath);
-				DialogResult = DialogResult.OK;
-				Close();
+				PerforceSettings.UserName = UserName;
 			}
+			PerforceSettings.PreferNativeClient = true;
+
+			// Create the P4 connection
+			CaptureLogger Logger = new CaptureLogger();
 
-			if(ErrorMessage != null)
+			// Create the task for connecting to this server
+			ModalTask? Task = PerforceModalTask.Execute(this, "Updating", "Checking for updates, please wait...", PerforceSettings, (p, c) => SyncAndRun(p, DepotPath, UseUnstableBuildCheckBox.Checked, Logger, c), Logger);
+			if (Task != null)
 			{
-				PromptLabel.Text = ErrorMessage;
+				if(Task.Succeeded)
+				{
+					Utility.SaveGlobalPerforceSettings(ServerAndPort, UserName, DepotPath);
+					DialogResult = DialogResult.OK;
+					Close();
+				}
+				PromptLabel.Text = Task.Error;
 			}
 
-			LogText = Log.ToString();
+			LogText = Logger.Render(Environment.NewLine);
 			ViewLogBtn.Visible = true;
 		}
 
diff --git a/Engine/Source/Programs/UnrealGameSync/UnrealGameSyncLauncher/Program.cs b/Engine/Source/Programs/UnrealGameSync/UnrealGameSyncLauncher/Program.cs
--- a/Engine/Source/Programs/UnrealGameSync/UnrealGameSyncLauncher/Program.cs
+++ b/Engine/Source/Programs/UnrealGameSync/UnrealGameSyncLauncher/Program.cs
@@ -1,9 +1,13 @@
 // Copyright Epic Games, Inc. All Rights Reserved.
 
+using EpicGames.Core;
+using EpicGames.Perforce;
+using Microsoft.Extensions.Logging;
 using Microsoft.Win32;
 using System;
 using System.Collections.Generic;
 using System.Diagnostics;
+using System.Diagnostics.CodeAnalysis;
 using System.IO;
 using System.Linq;
 using System.Reflection;
@@ -35,43 +39,21 @@ namespace UnrealGameSyncLauncher
 					return 0;
 				}
 
-				// Try to find Perforce in the path
-				string PerforceFileName = null;
-				foreach(string PathDirectoryName in (Environment.GetEnvironmentVariable("PATH") ?? "").Split(new char[]{ Path.PathSeparator }, StringSplitOptions.RemoveEmptyEntries))
-				{
-					try
-					{
-						string PossibleFileName = Path.Combine(PathDirectoryName, "p4.exe");
-						if(File.Exists(PossibleFileName))
-						{
-							PerforceFileName = PossibleFileName;
-							break;
-						}
-					}
-					catch { }
-				}
-
-				// If it doesn't exist, don't continue
-				if(PerforceFileName == null)
-				{
-					MessageBox.Show("UnrealGameSync requires the Perforce command-line tools. Please download and install from http://www.perforce.com/.");
-					return 1;
-				}
-
 				// Figure out if we should sync the unstable build by default
 				bool bUnstable = Args.Contains("-unstable", StringComparer.InvariantCultureIgnoreCase);
 
 				// Read the settings
-				string ServerAndPort = null;
-				string UserName = null;
-				string DepotPath = DeploymentSettings.DefaultDepotPath;
+				string? ServerAndPort = null;
+				string? UserName = null;
+				string? DepotPath = DeploymentSettings.DefaultDepotPath;
 				Utility.ReadGlobalPerforceSettings(ref ServerAndPort, ref UserName, ref DepotPath);
 
 				// If the shift key is held down, immediately show the settings window
-				if((Control.ModifierKeys & Keys.Shift) != 0)
+				SettingsWindow.SyncAndRunDelegate SyncAndRunWrapper = (Perforce, DepotParam, bUnstableParam, LogWriter, CancellationToken) => SyncAndRun(Perforce, DepotParam, bUnstableParam, Args, InstanceMutex, LogWriter, CancellationToken);
+				if ((Control.ModifierKeys & Keys.Shift) != 0)
 				{
 					// Show the settings window immediately
-					SettingsWindow UpdateError = new SettingsWindow(null, null, ServerAndPort, UserName, DepotPath, bUnstable, (Perforce, DepotParam, bUnstableParam, LogWriter) => SyncAndRun(Perforce, DepotParam, bUnstableParam, Args, InstanceMutex, LogWriter));
+					SettingsWindow UpdateError = new SettingsWindow(null, null, ServerAndPort, UserName, DepotPath, bUnstable, SyncAndRunWrapper);
 					if(UpdateError.ShowDialog() == DialogResult.OK)
 					{
 						return 0;
@@ -80,16 +62,21 @@ namespace UnrealGameSyncLauncher
 				else
 				{
 					// Try to do a sync with the current settings first
-					StringWriter Log = new StringWriter();
-					SyncAndRunPerforceTask SyncApplication = new SyncAndRunPerforceTask((Perforce, LogWriter) => SyncAndRun(Perforce, DepotPath, bUnstable, Args, InstanceMutex, LogWriter));
+					CaptureLogger Logger = new CaptureLogger();
+
+					IPerforceSettings Settings = PerforceSettings.Default.MergeWith(NewServerAndPort: ServerAndPort, NewUserName: UserName);
 
-					string ErrorMessage;
-					if(PerforceModalTask.Execute(null, new PerforceConnection(UserName, null, ServerAndPort), SyncApplication, "Updating", "Checking for updates, please wait...", Log, out ErrorMessage) == ModalTaskResult.Succeeded)
+					ModalTask? Task = PerforceModalTask.Execute(null, "Updating", "Checking for updates, please wait...", Settings, (p, c) => SyncAndRun(p, DepotPath, bUnstable, Args, InstanceMutex, Logger, c), Logger);
+					if (Task == null)
+					{
+						Logger.LogInformation("Canceled by user");
+					}
+					else if (Task.Succeeded)
 					{
 						return 0;
 					}
 
-					SettingsWindow UpdateError = new SettingsWindow("Unable to update UnrealGameSync from Perforce. Verify that your connection settings are correct.", Log.ToString(), ServerAndPort, UserName, DepotPath, bUnstable, (Perforce, DepotParam, bUnstableParam, LogWriter) => SyncAndRun(Perforce, DepotParam, bUnstableParam, Args, InstanceMutex, LogWriter));
+					SettingsWindow UpdateError = new SettingsWindow("Unable to update UnrealGameSync from Perforce. Verify that your connection settings are correct.", Logger.Render(Environment.NewLine), ServerAndPort, UserName, DepotPath, bUnstable, SyncAndRunWrapper);
 					if(UpdateError.ShowDialog() == DialogResult.OK)
 					{
 						return 0;
@@ -99,106 +86,92 @@ namespace UnrealGameSyncLauncher
 			return 1;
 		}
 
-		public static bool SyncAndRun(PerforceConnection Perforce, string BaseDepotPath, bool bUnstable, string[] Args, Mutex InstanceMutex, TextWriter LogWriter)
+		public static async Task SyncAndRun(IPerforceConnection Perforce, string? BaseDepotPath, bool bUnstable, string[] Args, Mutex InstanceMutex, ILogger Logger, CancellationToken CancellationToken)
 		{
 			try
 			{
-				string SyncPath = BaseDepotPath.TrimEnd('/') + (bUnstable? "/UnstableRelease/..." : "/Release/...");
-				LogWriter.WriteLine("Syncing from {0}", SyncPath);
+				if (String.IsNullOrEmpty(BaseDepotPath))
+				{
+					throw new UserErrorException($"Invalid setting for sync path");
+				}
+
+				string SyncPath = BaseDepotPath.TrimEnd('/') + (bUnstable ? "/UnstableRelease/..." : "/Release/...");
+				Logger.LogInformation("Syncing from {SyncPath}", SyncPath);
 
 				// Create the target folder
 				string ApplicationFolder = Path.Combine(Environment.GetFolderPath(Environment.SpecialFolder.LocalApplicationData), "UnrealGameSync", "Latest");
-				if(!SafeCreateDirectory(ApplicationFolder))
+				if (!SafeCreateDirectory(ApplicationFolder))
 				{
-					LogWriter.WriteLine("Couldn't create directory: {0}", ApplicationFolder);
-					return false;
+					throw new UserErrorException($"Couldn't create directory: {ApplicationFolder}");
 				}
 
 				// Find the most recent changelist
-				List<PerforceChangeSummary> Changes;
-				if(!Perforce.FindChanges(SyncPath, 1, out Changes, LogWriter) || Changes.Count < 1)
-				{
-					LogWriter.WriteLine("Couldn't find last changelist");
-					return false;
-				}
-
-				// Take the first changelist number
+				List<ChangesRecord> Changes = await Perforce.GetChangesAsync(ChangesOptions.None, 1, ChangeStatus.Submitted, SyncPath, CancellationToken);
 				int RequiredChangeNumber = Changes[0].Number;
 
 				// Read the current version
 				string SyncVersionFile = Path.Combine(ApplicationFolder, "SyncVersion.txt");
-				string RequiredSyncText = String.Format("{0}\n{1}@{2}", Perforce.ServerAndPort ?? "", SyncPath, RequiredChangeNumber);
+				string RequiredSyncText = String.Format("{0}\n{1}@{2}", Perforce.Settings.ServerAndPort ?? "", SyncPath, RequiredChangeNumber);
 
 				// Check the application exists
 				string ApplicationExe = Path.Combine(ApplicationFolder, "UnrealGameSync.exe");
 
 				// Check if the version has changed
-				string SyncText;
-				if(!File.Exists(SyncVersionFile) || !File.Exists(ApplicationExe) || !TryReadAllText(SyncVersionFile, out SyncText) || SyncText != RequiredSyncText)
+				string? SyncText;
+				if (!File.Exists(SyncVersionFile) || !File.Exists(ApplicationExe) || !TryReadAllText(SyncVersionFile, out SyncText) || SyncText != RequiredSyncText)
 				{
 					// Try to delete the directory contents. Retry for a while, in case we've been spawned by an application in this folder to do an update.
-					for(int NumRetries = 0; !SafeDeleteDirectoryContents(ApplicationFolder); NumRetries++)
+					for (int NumRetries = 0; !SafeDeleteDirectoryContents(ApplicationFolder); NumRetries++)
 					{
-						if(NumRetries > 20)
+						if (NumRetries > 20)
 						{
-							LogWriter.WriteLine("Couldn't delete contents of {0} (retried {1} times).", ApplicationFolder, NumRetries);
-							return false;
+							throw new UserErrorException($"Couldn't delete contents of {ApplicationFolder} (retried {NumRetries} times).");
 						}
 						Thread.Sleep(500);
 					}
-				
+
 					// Find all the files in the sync path at this changelist
-					List<PerforceFileRecord> FileRecords;
-					if(!Perforce.Stat(String.Format("{0}@{1}", SyncPath, RequiredChangeNumber), out FileRecords, LogWriter))
+					List<FStatRecord> FileRecords = await Perforce.FStatAsync(FStatOptions.None, $"{SyncPath}@{RequiredChangeNumber}", CancellationToken).ToListAsync(CancellationToken);
+					if (FileRecords.Count == 0)
 					{
-						LogWriter.WriteLine("Couldn't find matching files.");
-						return false;
+						throw new UserErrorException($"Couldn't find any matching files for {SyncPath}@{RequiredChangeNumber}");
 					}
 
 					// Sync all the files in this list to the same directory structure under the application folder
 					string DepotPathPrefix = SyncPath.Substring(0, SyncPath.LastIndexOf('/') + 1);
-					foreach(PerforceFileRecord FileRecord in FileRecords)
+					foreach (FStatRecord FileRecord in FileRecords)
 					{
-						string LocalPath = Path.Combine(ApplicationFolder, FileRecord.DepotPath.Substring(DepotPathPrefix.Length).Replace('/', Path.DirectorySeparatorChar));
-						if(!SafeCreateDirectory(Path.GetDirectoryName(LocalPath)))
+						if (FileRecord.DepotFile == null)
 						{
-							LogWriter.WriteLine("Couldn't create folder {0}", Path.GetDirectoryName(LocalPath));
-							return false;
+							throw new UserErrorException("Missing depot path for returned file");
 						}
-						if(!Perforce.PrintToFile(FileRecord.DepotPath, LocalPath, LogWriter))
+
+						string LocalPath = Path.Combine(ApplicationFolder, FileRecord.DepotFile.Substring(DepotPathPrefix.Length).Replace('/', Path.DirectorySeparatorChar));
+						if (!SafeCreateDirectory(Path.GetDirectoryName(LocalPath)!))
 						{
-							LogWriter.WriteLine("Couldn't sync {0} to {1}", FileRecord.DepotPath, LocalPath);
-							return false;
+							throw new UserErrorException($"Couldn't create folder {Path.GetDirectoryName(LocalPath)}");
 						}
+
+						await Perforce.PrintAsync(LocalPath, FileRecord.DepotFile, CancellationToken);
 					}
 
 					// Check the application exists
-					if(!File.Exists(ApplicationExe))
+					if (!File.Exists(ApplicationExe))
 					{
-						LogWriter.WriteLine("Application was not synced from Perforce. Check that UnrealGameSync exists at {0}/UnrealGameSync.exe, and you have access to it.", SyncPath);
-						return false;
+						throw new UserErrorException($"Application was not synced from Perforce. Check that UnrealGameSync exists at {SyncPath}/UnrealGameSync.exe, and you have access to it.");
 					}
 
 					// Update the version
-					if(!TryWriteAllText(SyncVersionFile, RequiredSyncText))
+					if (!TryWriteAllText(SyncVersionFile, RequiredSyncText))
 					{
-						LogWriter.WriteLine("Couldn't write sync text to {0}", SyncVersionFile);
-						return false;
+						throw new UserErrorException("Couldn't write sync text to {SyncVersionFile}");
 					}
 				}
-				LogWriter.WriteLine();
+				Logger.LogInformation("");
 
 				// Build the command line for the synced application, including the sync path to monitor for updates
-				StringBuilder NewCommandLine = new StringBuilder(String.Format("-updatepath=\"{0}@>{1}\" -updatespawn=\"{2}\"{3}", SyncPath, RequiredChangeNumber, Assembly.GetEntryAssembly().Location, bUnstable? " -unstable" : ""));
-				if(!String.IsNullOrEmpty(Perforce.ServerAndPort))
-				{
-					NewCommandLine.AppendFormat(" -p4port={0}", QuoteArgument(Perforce.ServerAndPort));
-				}
-				if(!String.IsNullOrEmpty(Perforce.UserName))
-				{
-					NewCommandLine.AppendFormat(" -p4user={0}", QuoteArgument(Perforce.UserName));
-				}
-				foreach(string Arg in Args)
+				StringBuilder NewCommandLine = new StringBuilder(String.Format("-updatepath=\"{0}@>{1}\" -updatespawn=\"{2}\"{3}", SyncPath, RequiredChangeNumber, Assembly.GetEntryAssembly()!.Location, bUnstable ? " -unstable" : ""));
+				foreach (string Arg in Args)
 				{
 					NewCommandLine.AppendFormat(" {0}", QuoteArgument(Arg));
 				}
@@ -207,26 +180,32 @@ namespace UnrealGameSyncLauncher
 				InstanceMutex.Close();
 
 				// Spawn the application
-				LogWriter.WriteLine("Spawning {0} with command line: {1}", ApplicationExe, NewCommandLine.ToString());
-				using(Process ChildProcess = new Process())
+				Logger.LogInformation("Spawning {App} with command line: {CmdLine}", ApplicationExe, NewCommandLine.ToString());
+				using (Process ChildProcess = new Process())
 				{
 					ChildProcess.StartInfo.FileName = ApplicationExe;
 					ChildProcess.StartInfo.Arguments = NewCommandLine.ToString();
 					ChildProcess.StartInfo.UseShellExecute = false;
 					ChildProcess.StartInfo.CreateNoWindow = false;
-					if(!ChildProcess.Start())
+					if (!ChildProcess.Start())
 					{
-						LogWriter.WriteLine("Failed to start process");
-						return false;
+						throw new UserErrorException("Failed to start process");
 					}
 				}
-
-				return true;
 			}
-			catch(Exception Ex)
+			catch (UserErrorException Ex)
 			{
-				LogWriter.WriteLine(Ex.ToString());
-				return false;
+				Logger.LogError("{Message}", Ex.Message);
+				throw;
+			}
+			catch (Exception Ex)
+			{
+				Logger.LogError(Ex, "Error while syncing application.");
+				foreach (string Line in Ex.ToString().Split('\n'))
+				{
+					Logger.LogError("{Line}", Line);
+				}
+				throw;
 			}
 		}
 
@@ -242,7 +221,7 @@ namespace UnrealGameSyncLauncher
 			}
 		}
 
-		static bool TryReadAllText(string FileName, out string Text)
+		static bool TryReadAllText(string FileName, [NotNullWhen(true)] out string? Text)
 		{
 			try
 			{
diff --git a/Engine/Source/Programs/UnrealGameSync/UnrealGameSyncLauncher/SyncAndRunPerforceTask.cs b/Engine/Source/Programs/UnrealGameSync/UnrealGameSyncLauncher/SyncAndRunPerforceTask.cs
deleted file mode 100644
diff --git a/Engine/Source/Programs/UnrealGameSync/UnrealGameSyncLauncher/UnrealGameSyncLauncher.csproj b/Engine/Source/Programs/UnrealGameSync/UnrealGameSyncLauncher/UnrealGameSyncLauncher.csproj
--- a/Engine/Source/Programs/UnrealGameSync/UnrealGameSyncLauncher/UnrealGameSyncLauncher.csproj
+++ b/Engine/Source/Programs/UnrealGameSync/UnrealGameSyncLauncher/UnrealGameSyncLauncher.csproj
@@ -13,11 +13,9 @@
     <ApplicationIcon>Resources\Icon.ico</ApplicationIcon>
     <ApplicationManifest>app.manifest</ApplicationManifest>
     <RuntimeIdentifiers>win10-x64</RuntimeIdentifiers>
+    <Nullable>enable</Nullable>
   </PropertyGroup>
   <ItemGroup>
-    <Compile Include="..\UnrealGameSync\ChildProcess.cs">
-      <Link>UnrealGameSync\ChildProcess.cs</Link>
-    </Compile>
     <Compile Include="..\UnrealGameSync\Forms\ModalTaskWindow.cs">
       <Link>UnrealGameSync\ModalTaskWindow.cs</Link>
       <SubType>Form</SubType>
@@ -40,9 +38,6 @@
     <Compile Include="..\UnrealGameSync\ModalTask.cs">
       <Link>UnrealGameSync\ModalTask.cs</Link>
     </Compile>
-    <Compile Include="..\UnrealGameSync\Perforce.cs">
-      <Link>UnrealGameSync\Perforce.cs</Link>
-    </Compile>
     <Compile Include="..\UnrealGameSync\PerforceModalTask.cs">
       <Link>UnrealGameSync\PerforceModalTask.cs</Link>
     </Compile>
@@ -50,6 +45,10 @@
       <Link>UnrealGameSync\Utility.cs</Link>
     </Compile>
   </ItemGroup>
+  <ItemGroup>
+    <ProjectReference Include="..\..\Shared\EpicGames.Core\EpicGames.Core.csproj" />
+    <ProjectReference Include="..\..\Shared\EpicGames.Perforce\EpicGames.Perforce.csproj" />
+  </ItemGroup>
   <ItemGroup>
     <None Update="Properties\Settings.settings">
       <Generator>SettingsSingleFileGenerator</Generator>
