UGS: Shuffling around some code prior to splitting update logic out into a separate library.

diff --git a/Engine/Source/Programs/UnrealGameSync/UnrealGameSync/Controls/WorkspaceControl.cs b/Engine/Source/Programs/UnrealGameSync/UnrealGameSync/Controls/WorkspaceControl.cs
--- a/Engine/Source/Programs/UnrealGameSync/UnrealGameSync/Controls/WorkspaceControl.cs
+++ b/Engine/Source/Programs/UnrealGameSync/UnrealGameSync/Controls/WorkspaceControl.cs
@@ -37,13 +37,6 @@ namespace UnrealGameSync
 		public List<UncontrolledChangelist> Changelists { get; set; } = new List<UncontrolledChangelist>();
 	}
 
-	public enum LatestChangeType
-	{
-		Any,
-		Good,
-		Starred,
-	}
-
 	interface IWorkspaceControlOwner
 	{
 		ToolUpdateMonitor ToolUpdateMonitor { get; }
@@ -5187,7 +5180,7 @@ namespace UnrealGameSync
 			}
 
 			string[] CombinedSyncFilter = UserSettings.GetCombinedSyncFilter(GetSyncCategories(), Settings.SyncView, Settings.SyncCategories, WorkspaceSettings.SyncView, WorkspaceSettings.SyncCategoriesDict);
-			List<string> SyncPaths = Workspace.GetSyncPaths(Workspace.Project, WorkspaceSettings.bSyncAllProjects ?? Settings.bSyncAllProjects, CombinedSyncFilter);
+			List<string> SyncPaths = WorkspaceUpdate.GetSyncPaths(Workspace.Project, WorkspaceSettings.bSyncAllProjects ?? Settings.bSyncAllProjects, CombinedSyncFilter);
 
 			CleanWorkspaceWindow.DoClean(ParentForm, PerforceSettings, BranchDirectoryName, Workspace.Project.ClientRootPath, SyncPaths, ExtraSafeToDeleteFolders.Split('\n'), ExtraSafeToDeleteExtensions.Split('\n'), ServiceProvider.GetRequiredService<ILogger<CleanWorkspaceWindow>>());
 		}
diff --git a/Engine/Source/Programs/UnrealGameSync/UnrealGameSync/DetectProjectSettingsTask.cs b/Engine/Source/Programs/UnrealGameSync/UnrealGameSync/DetectProjectSettingsTask.cs
--- a/Engine/Source/Programs/UnrealGameSync/UnrealGameSync/DetectProjectSettingsTask.cs
+++ b/Engine/Source/Programs/UnrealGameSync/UnrealGameSync/DetectProjectSettingsTask.cs
@@ -297,8 +297,8 @@ namespace UnrealGameSync
 				ConfigFile LatestProjectConfigFile = await PerforceMonitor.ReadProjectConfigFileAsync(PerforceClient, BranchClientPath, NewSelectedClientFileName, GetCacheFolder(BranchDirectoryName), LocalConfigFiles, Logger, CancellationToken);
 
 				// Get the local config file and stream filter
-				ConfigFile WorkspaceProjectConfigFile = await Workspace.ReadProjectConfigFile(BranchDirectoryName, NewSelectedFileName, Logger);
-				IReadOnlyList<string>? WorkspaceProjectStreamFilter = await Workspace.ReadProjectStreamFilter(PerforceClient, WorkspaceProjectConfigFile, Logger, CancellationToken.None);
+				ConfigFile WorkspaceProjectConfigFile = await WorkspaceUpdate.ReadProjectConfigFile(BranchDirectoryName, NewSelectedFileName, Logger);
+				IReadOnlyList<string>? WorkspaceProjectStreamFilter = await WorkspaceUpdate.ReadProjectStreamFilter(PerforceClient, WorkspaceProjectConfigFile, Logger, CancellationToken);
 
 				ProjectInfo ProjectInfo = new ProjectInfo(BranchDirectoryName, NewSelectedFileName, BranchClientPath, NewSelectedClientFileName, NewSelectedProjectIdentifier, bIsEnterpriseProject);
 
diff --git a/Engine/Source/Programs/UnrealGameSync/UnrealGameSync/Forms/MainWindow.cs b/Engine/Source/Programs/UnrealGameSync/UnrealGameSync/Forms/MainWindow.cs
--- a/Engine/Source/Programs/UnrealGameSync/UnrealGameSync/Forms/MainWindow.cs
+++ b/Engine/Source/Programs/UnrealGameSync/UnrealGameSync/Forms/MainWindow.cs
@@ -660,7 +660,7 @@ namespace UnrealGameSync
 			}
 
 			Settings.bWindowVisible = Visible;
-			Settings.WindowState = WindowState;
+			Settings.WindowState = WindowState.ToString();
 			if(WindowState == FormWindowState.Normal)
 			{
 				Settings.WindowBounds = new Rectangle(Location, Size);
@@ -1439,7 +1439,12 @@ namespace UnrealGameSync
 					}
 				}
 			}
-			WindowState = Settings.WindowState;
+
+			FormWindowState NewWindowState;
+			if (Enum.TryParse(Settings.WindowState, true, out NewWindowState))
+			{
+				WindowState = NewWindowState;
+			}
 		}
 
 		bool ShowNotificationForIssue(IssueData Issue)
diff --git a/Engine/Source/Programs/UnrealGameSync/UnrealGameSync/PerforceMonitor.cs b/Engine/Source/Programs/UnrealGameSync/UnrealGameSync/PerforceMonitor.cs
--- a/Engine/Source/Programs/UnrealGameSync/UnrealGameSync/PerforceMonitor.cs
+++ b/Engine/Source/Programs/UnrealGameSync/UnrealGameSync/PerforceMonitor.cs
@@ -76,17 +76,6 @@ namespace UnrealGameSync
 		IReadOnlyList<IArchiveInfo> AvailableArchives { get; }
 	}
 
-	public interface IArchiveInfo
-	{
-		string Name { get; }
-		string Type { get; }
-		string BasePath { get; }
-		string? Target { get;  }
-		bool Exists();
-		bool TryGetArchiveKeyForChangeNumber(int ChangeNumber, [NotNullWhen(true)] out string? ArchiveKey);
-		Task<bool> DownloadArchive(IPerforceConnection Perforce, string ArchiveKey, DirectoryReference LocalRootPath, FileReference ManifestFileName, ILogger Logger, ProgressValue Progress, CancellationToken CancellationToken);
-	}
-
 	class PerforceMonitor : IDisposable, IArchiveInfoSource
 	{
 		internal class PerforceArchiveInfo : IArchiveInfo
diff --git a/Engine/Source/Programs/UnrealGameSync/UnrealGameSync/UserSettings.cs b/Engine/Source/Programs/UnrealGameSync/UnrealGameSync/UserSettings.cs
--- a/Engine/Source/Programs/UnrealGameSync/UnrealGameSync/UserSettings.cs
+++ b/Engine/Source/Programs/UnrealGameSync/UnrealGameSync/UserSettings.cs
@@ -13,7 +13,6 @@ using System.Text.Json;
 using System.Text.Json.Serialization;
 using System.Text.RegularExpressions;
 using System.Threading.Tasks;
-using System.Windows.Forms;
 
 namespace UnrealGameSync
 {
@@ -53,6 +52,13 @@ namespace UnrealGameSync
 		Content
 	}
 
+	public enum LatestChangeType
+	{
+		Any,
+		Good,
+		Starred,
+	}
+
 	public enum UserSettingsVersion
 	{
 		Initial = 0,
@@ -356,7 +362,7 @@ namespace UnrealGameSync
 
 		// Window settings
 		public bool bWindowVisible;
-		public FormWindowState WindowState;
+		public string WindowState;
 		public Rectangle? WindowBounds;
 		
 		// Schedule settings
@@ -550,10 +556,7 @@ namespace UnrealGameSync
 
 			// Window settings
 			bWindowVisible = ConfigFile.GetValue("Window.Visible", true);
-			if(!Enum.TryParse(ConfigFile.GetValue("Window.State", ""), true, out WindowState))
-			{
-				WindowState = FormWindowState.Normal;
-			}
+			WindowState = ConfigFile.GetValue("Window.State", "");
 			WindowBounds = ParseRectangleValue(ConfigFile.GetValue("Window.Bounds", ""));
 
 			// Schedule settings
@@ -962,7 +965,7 @@ namespace UnrealGameSync
 			ConfigSection WindowSection = ConfigFile.FindOrAddSection("Window");
 			WindowSection.Clear();
 			WindowSection.SetValue("Visible", bWindowVisible);
-			WindowSection.SetValue("State", WindowState.ToString());
+			WindowSection.SetValue("State", WindowState);
 			if(WindowBounds != null)
 			{
 				WindowSection.SetValue("Bounds", FormatRectangleValue(WindowBounds.Value));
diff --git a/Engine/Source/Programs/UnrealGameSync/UnrealGameSync/Workspace.cs b/Engine/Source/Programs/UnrealGameSync/UnrealGameSync/Workspace.cs
--- a/Engine/Source/Programs/UnrealGameSync/UnrealGameSync/Workspace.cs
+++ b/Engine/Source/Programs/UnrealGameSync/UnrealGameSync/Workspace.cs
@@ -8,6 +8,7 @@
 using System.Collections.Concurrent;
 using System.Collections.Generic;
 using System.Diagnostics;
+using System.Diagnostics.CodeAnalysis;
 using System.Globalization;
 using System.IO;
 using System.IO.Compression;
@@ -78,6 +79,17 @@
 		}
 	}
 
+	public interface IArchiveInfo
+	{
+		string Name { get; }
+		string Type { get; }
+		string BasePath { get; }
+		string? Target { get; }
+		bool Exists();
+		bool TryGetArchiveKeyForChangeNumber(int ChangeNumber, [NotNullWhen(true)] out string? ArchiveKey);
+		Task<bool> DownloadArchive(IPerforceConnection Perforce, string ArchiveKey, DirectoryReference LocalRootPath, FileReference ManifestFileName, ILogger Logger, ProgressValue Progress, CancellationToken CancellationToken);
+	}
+
 	public class WorkspaceUpdateContext
 	{
 		public DateTime StartTime = DateTime.UtcNow;
@@ -189,13 +201,6 @@
 
 	public sealed class Workspace : IDisposable
 	{
-		const string BuildVersionFileName = "/Engine/Build/Build.version";
-		const string VersionHeaderFileName = "/Engine/Source/Runtime/Launch/Resources/Version.h";
-		const string ObjectVersionFileName = "/Engine/Source/Runtime/Core/Private/UObject/ObjectVersion.cpp";
-
-		static readonly string LocalVersionHeaderFileName = VersionHeaderFileName.Replace('/', '\\');
-		static readonly string LocalObjectVersionFileName = ObjectVersionFileName.Replace('/', '\\');
-
 		public IPerforceSettings PerforceSettings { get; }
 		public UserWorkspaceState State { get; }
 		public ProjectInfo Project { get; }
@@ -203,13 +208,151 @@
 		ILogger Logger;
 
 		bool bSyncing => CurrentUpdate != null;
-		ProgressValue Progress => CurrentUpdate?.Progress ?? new ProgressValue();
 
 		WorkspaceUpdate? CurrentUpdate;
-		static SemaphoreSlim UpdateSemaphore = new SemaphoreSlim(1);
 
 		public event Action<WorkspaceUpdateContext, WorkspaceUpdateResult, string>? OnUpdateComplete;
 
+		IAsyncDisposer AsyncDisposer;
+
+		public Workspace(IPerforceSettings InPerfoceSettings, ProjectInfo InProject, UserWorkspaceState InState, ConfigFile ProjectConfigFile, IReadOnlyList<string>? ProjectStreamFilter, ILogger Logger, IServiceProvider ServiceProvider)
+		{
+			PerforceSettings = InPerfoceSettings;
+			Project = InProject;
+			State = InState;
+			this.SynchronizationContext = SynchronizationContext.Current!;
+			this.Logger = Logger;
+			this.AsyncDisposer = ServiceProvider.GetRequiredService<IAsyncDisposer>();
+
+			this.ProjectConfigFile = ProjectConfigFile;
+			this.ProjectStreamFilter = ProjectStreamFilter;
+		}
+
+		public void Dispose()
+		{
+			CancelUpdate();
+			if (PrevUpdateTask != null)
+			{
+				AsyncDisposer.Add(PrevUpdateTask);
+			}
+		}
+
+		public ConfigFile ProjectConfigFile
+		{
+			get; private set;
+		}
+
+		public IReadOnlyList<string>? ProjectStreamFilter
+		{
+			get; private set;
+		}
+
+		Task PrevUpdateTask = Task.CompletedTask;
+
+		public void Update(WorkspaceUpdateContext Context)
+		{
+			CancelUpdate();
+
+			Task PrevUpdateTaskCopy = PrevUpdateTask;
+			CurrentUpdate = new WorkspaceUpdate(Context, (Update, CancellationToken) => UpdateWorkspaceMini(Update, PrevUpdateTaskCopy, CancellationToken));
+		}
+
+		public void CancelUpdate()
+		{
+			// Cancel the current task. We actually terminate the operation asynchronously, but we can signal the cancellation and 
+			// send a cancelled event, then wait for the heavy lifting to finish in the new update task.
+			if (CurrentUpdate != null)
+			{
+				PrevUpdateTask = CurrentUpdate.DisposeAsync().AsTask();
+				CompleteUpdate(CurrentUpdate, WorkspaceUpdateResult.Canceled, "Cancelled");
+			}
+		}
+
+		async Task UpdateWorkspaceMini(WorkspaceUpdate Update, Task PrevUpdateTask, CancellationToken CancellationToken)
+		{
+			if (PrevUpdateTask != null)
+			{
+				await PrevUpdateTask;
+			}
+
+			WorkspaceUpdateContext Context = Update.Context;
+			Context.ProjectConfigFile = ProjectConfigFile;
+			Context.ProjectStreamFilter = ProjectStreamFilter;
+
+			string StatusMessage;
+			WorkspaceUpdateResult Result = WorkspaceUpdateResult.FailedToSync;
+
+			try
+			{
+				(Result, StatusMessage) = await Update.ExecuteAsync(PerforceSettings, Project, State, Logger, CancellationToken);
+				if (Result != WorkspaceUpdateResult.Success)
+				{
+					Logger.LogError("{Message}", StatusMessage);
+				}
+			}
+			catch (OperationCanceledException)
+			{
+				StatusMessage = "Canceled.";
+				Logger.LogError("Canceled.");
+			}
+			catch (Exception Ex)
+			{
+				StatusMessage = "Failed with exception - " + Ex.ToString();
+				Logger.LogError(Ex, "Failed with exception");
+			}
+
+			ProjectConfigFile = Context.ProjectConfigFile;
+			ProjectStreamFilter = Context.ProjectStreamFilter;
+
+			SynchronizationContext.Post(x => CompleteUpdate(Update, Result, StatusMessage), null);
+		}
+
+		void CompleteUpdate(WorkspaceUpdate Update, WorkspaceUpdateResult Result, string StatusMessage)
+		{
+			if (CurrentUpdate == Update)
+			{
+				WorkspaceUpdateContext Context = Update.Context;
+
+				State.LastSyncChangeNumber = Context.ChangeNumber;
+				State.LastSyncResult = Result;
+				State.LastSyncResultMessage = StatusMessage;
+				State.LastSyncTime = DateTime.UtcNow;
+				State.LastSyncDurationSeconds = (int)(State.LastSyncTime.Value - Context.StartTime).TotalSeconds;
+				State.Save();
+
+				OnUpdateComplete?.Invoke(Context, Result, StatusMessage);
+				CurrentUpdate = null;
+			}
+		}
+
+		public bool IsBusy()
+		{
+			return bSyncing;
+		}
+
+		public int CurrentChangeNumber => State.CurrentChangeNumber;
+
+		public int PendingChangeNumber => CurrentUpdate?.Context?.ChangeNumber ?? CurrentChangeNumber;
+
+		public string ClientName
+		{
+			get { return PerforceSettings.ClientName!; }
+		}
+
+		public Tuple<string, float> CurrentProgress => CurrentUpdate?.CurrentProgress ?? new Tuple<string, float>("", 0.0f);
+	}
+
+	class WorkspaceUpdate : IAsyncDisposable
+	{
+		const string BuildVersionFileName = "/Engine/Build/Build.version";
+		const string VersionHeaderFileName = "/Engine/Source/Runtime/Launch/Resources/Version.h";
+		const string ObjectVersionFileName = "/Engine/Source/Runtime/Core/Private/UObject/ObjectVersion.cpp";
+
+		static readonly string LocalVersionHeaderFileName = VersionHeaderFileName.Replace('/', '\\');
+		static readonly string LocalObjectVersionFileName = ObjectVersionFileName.Replace('/', '\\');
+
+		static SemaphoreSlim UpdateSemaphore = new SemaphoreSlim(1);
+
 		class RecordCounter : IDisposable
 		{
 			ProgressValue Progress;
@@ -362,147 +505,32 @@
 			}
 		}
 
-		IAsyncDisposer AsyncDisposer;
+		CancellationTokenSource CancellationSource;
+		public WorkspaceUpdateContext Context { get; }
+		public ProgressValue Progress { get; } = new ProgressValue();
+		public Task UpdateTask { get; }
 
-		public Workspace(IPerforceSettings InPerfoceSettings, ProjectInfo InProject, UserWorkspaceState InState, ConfigFile ProjectConfigFile, IReadOnlyList<string>? ProjectStreamFilter, ILogger Logger, IServiceProvider ServiceProvider)
+		public WorkspaceUpdate(WorkspaceUpdateContext Context, Func<WorkspaceUpdate, CancellationToken, Task> TaskFunc)
 		{
-			PerforceSettings = InPerfoceSettings;
-			Project = InProject;
-			State = InState;
-			this.SynchronizationContext = SynchronizationContext.Current!;
-			this.Logger = Logger;
-			this.AsyncDisposer = ServiceProvider.GetRequiredService<IAsyncDisposer>();
-
-			this.ProjectConfigFile = ProjectConfigFile;
-			this.ProjectStreamFilter = ProjectStreamFilter;
+			this.CancellationSource = new CancellationTokenSource();
+			this.Context = Context;
+			this.UpdateTask = Task.Run(() => TaskFunc(this, CancellationSource.Token));
 		}
 
-		public void Dispose()
+		public void Cancel()
 		{
-			CancelUpdate();
-			if (PrevUpdateTask != null)
-			{
-				AsyncDisposer.Add(PrevUpdateTask);
-			}
+			CancellationSource.Cancel();
 		}
 
-		public ConfigFile ProjectConfigFile
+		public ValueTask DisposeAsync()
 		{
-			get; private set;
-		}
-
-		public IReadOnlyList<string>? ProjectStreamFilter
-		{
-			get; private set;
-		}
-
-		class WorkspaceUpdate : IAsyncDisposable
-		{
-			CancellationTokenSource CancellationSource;
-			public WorkspaceUpdateContext Context { get; }
-			public ProgressValue Progress { get; } = new ProgressValue();
-			public Task UpdateTask { get; }
-
-			public WorkspaceUpdate(WorkspaceUpdateContext Context, Func<WorkspaceUpdate, CancellationToken, Task> TaskFunc)
+			Task DisposeTask = Task.CompletedTask;
+			if (UpdateTask != null)
 			{
-				this.CancellationSource = new CancellationTokenSource();
-				this.Context = Context;
-				this.UpdateTask = Task.Run(() => TaskFunc(this, CancellationSource.Token));
-			}
-
-			public void Cancel()
-			{
 				CancellationSource.Cancel();
+				DisposeTask = UpdateTask.ContinueWith(Task => CancellationSource.Dispose());
 			}
-
-			public ValueTask DisposeAsync()
-			{
-				Task DisposeTask = Task.CompletedTask;
-				if (UpdateTask != null)
-				{
-					CancellationSource.Cancel();
-					DisposeTask = UpdateTask.ContinueWith(Task => CancellationSource.Dispose());
-				}
-				return new ValueTask(DisposeTask);
-			}
+			return new ValueTask(DisposeTask);
 		}
 
-		Task PrevUpdateTask = Task.CompletedTask;
-
-		public void Update(WorkspaceUpdateContext Context)
-		{
-			CancelUpdate();
-
-			Task PrevUpdateTaskCopy = PrevUpdateTask;
-			CurrentUpdate = new WorkspaceUpdate(Context, (Update, CancellationToken) => UpdateWorkspaceMini(Update, PrevUpdateTaskCopy, CancellationToken));
-		}
-
-		public void CancelUpdate()
-		{
-			// Cancel the current task. We actually terminate the operation asynchronously, but we can signal the cancellation and 
-			// send a cancelled event, then wait for the heavy lifting to finish in the new update task.
-			if (CurrentUpdate != null)
-			{
-				PrevUpdateTask = CurrentUpdate.DisposeAsync().AsTask();
-				CompleteUpdate(CurrentUpdate, WorkspaceUpdateResult.Canceled, "Cancelled");
-			}
-		}
-
-		async Task UpdateWorkspaceMini(WorkspaceUpdate Update, Task PrevUpdateTask, CancellationToken CancellationToken)
-		{
-			if (PrevUpdateTask != null)
-			{
-				await PrevUpdateTask;
-			}
-
-			WorkspaceUpdateContext Context = Update.Context;
-			Context.ProjectConfigFile = ProjectConfigFile;
-			Context.ProjectStreamFilter = ProjectStreamFilter;
-
-			string StatusMessage;
-			WorkspaceUpdateResult Result = WorkspaceUpdateResult.FailedToSync;
-
-			try
-			{
-				(Result, StatusMessage) = await UpdateWorkspaceInternal(PerforceSettings, Project, State, Context, Progress, Logger, CancellationToken);
-				if (Result != WorkspaceUpdateResult.Success)
-				{
-					Logger.LogError("{Message}", StatusMessage);
-				}
-			}
-			catch (OperationCanceledException)
-			{
-				StatusMessage = "Canceled.";
-				Logger.LogError("Canceled.");
-			}
-			catch (Exception Ex)
-			{
-				StatusMessage = "Failed with exception - " + Ex.ToString();
-				Logger.LogError(Ex, "Failed with exception");
-			}
-
-			ProjectConfigFile = Context.ProjectConfigFile;
-			ProjectStreamFilter = Context.ProjectStreamFilter;
-
-			SynchronizationContext.Post(x => CompleteUpdate(Update, Result, StatusMessage), null);
-		}
-
-		void CompleteUpdate(WorkspaceUpdate Update, WorkspaceUpdateResult Result, string StatusMessage)
-		{
-			if (CurrentUpdate == Update)
-			{
-				WorkspaceUpdateContext Context = Update.Context;
-
-				State.LastSyncChangeNumber = Context.ChangeNumber;
-				State.LastSyncResult = Result;
-				State.LastSyncResultMessage = StatusMessage;
-				State.LastSyncTime = DateTime.UtcNow;
-				State.LastSyncDurationSeconds = (int)(State.LastSyncTime.Value - Context.StartTime).TotalSeconds;
-				State.Save();
-
-				OnUpdateComplete?.Invoke(Context, Result, StatusMessage);
-				CurrentUpdate = null;
-			}
-		}
-
 		class SyncFile
@@ -553,7 +581,7 @@
 			public void Dispose() => Release();
 		}
 
-		static async Task<(WorkspaceUpdateResult, string)> UpdateWorkspaceInternal(IPerforceSettings PerforceSettings, ProjectInfo Project, UserWorkspaceState State, WorkspaceUpdateContext Context, ProgressValue Progress, ILogger Logger, CancellationToken CancellationToken)
+		public async Task<(WorkspaceUpdateResult, string)> ExecuteAsync(IPerforceSettings PerforceSettings, ProjectInfo Project, UserWorkspaceState State, ILogger Logger, CancellationToken CancellationToken)
 		{
 			using IPerforceConnection Perforce = await PerforceConnection.CreateAsync(PerforceSettings, Logger);
 
@@ -1932,23 +1960,9 @@
 			return Line.Substring(StartIdx, LineIdx - StartIdx);
 		}
 
-		public bool IsBusy()
-		{
-			return bSyncing;
-		}
-
 		public Tuple<string, float> CurrentProgress
 		{
 			get { return Progress.Current; }
 		}
-
-		public int CurrentChangeNumber => State.CurrentChangeNumber;
-
-		public int PendingChangeNumber => CurrentUpdate?.Context?.ChangeNumber ?? CurrentChangeNumber;
-
-		public string ClientName
-		{
-			get { return PerforceSettings.ClientName!; }
-		}
 	}
 }
