UGS: Separation of Winforms/standalone functionality prior to moving into separate project.

diff --git a/Engine/Source/Programs/UnrealGameSync/UnrealGameSync/Automation/P4Automation.cs b/Engine/Source/Programs/UnrealGameSync/UnrealGameSync/Automation/P4Automation.cs
--- a/Engine/Source/Programs/UnrealGameSync/UnrealGameSync/Automation/P4Automation.cs
+++ b/Engine/Source/Programs/UnrealGameSync/UnrealGameSync/Automation/P4Automation.cs
@@ -21,7 +21,7 @@ namespace UnrealGameSync
 			string? UserName = null;
 			string? DepotPathSettings = null;
 
-			Utility.ReadGlobalPerforceSettings(ref ServerAndPort, ref UserName, ref DepotPathSettings);
+			GlobalSettings.ReadGlobalPerforceSettings(ref ServerAndPort, ref UserName, ref DepotPathSettings);
 
 			return Utility.OverridePerforceSettings(PerforceSettings.Default, ServerAndPort, UserName);
 
diff --git a/Engine/Source/Programs/UnrealGameSync/UnrealGameSync/AutomationServer.cs b/Engine/Source/Programs/UnrealGameSync/UnrealGameSync/AutomationServer.cs
--- a/Engine/Source/Programs/UnrealGameSync/UnrealGameSync/AutomationServer.cs
+++ b/Engine/Source/Programs/UnrealGameSync/UnrealGameSync/AutomationServer.cs
@@ -190,7 +190,7 @@ namespace UnrealGameSync
 		{
 			if(PortNumber <= 0)
 			{
-				Utility.DeleteRegistryKey(Registry.CurrentUser, "Software\\Epic Games\\UnrealGameSync", "AutomationPort");
+				GlobalSettings.DeleteRegistryKey(Registry.CurrentUser, "Software\\Epic Games\\UnrealGameSync", "AutomationPort");
 			}
 			else
 			{
diff --git a/Engine/Source/Programs/UnrealGameSync/UnrealGameSync/Controls/WorkspaceControl.cs b/Engine/Source/Programs/UnrealGameSync/UnrealGameSync/Controls/WorkspaceControl.cs
--- a/Engine/Source/Programs/UnrealGameSync/UnrealGameSync/Controls/WorkspaceControl.cs
+++ b/Engine/Source/Programs/UnrealGameSync/UnrealGameSync/Controls/WorkspaceControl.cs
@@ -4692,7 +4692,7 @@ namespace UnrealGameSync
 		{
 			if (ContextMenuChange != null)
 			{
-				Utility.SpawnP4VC(String.Format("{0} change {1}", PerforceSettings.GetArgumentsForExternalProgram(true), ContextMenuChange.Number));
+				Program.SpawnP4VC(String.Format("{0} change {1}", PerforceSettings.GetArgumentsForExternalProgram(true), ContextMenuChange.Number));
 			}
 		}
 
diff --git a/Engine/Source/Programs/UnrealGameSync/UnrealGameSync/Forms/ApplicationSettingsWindow.cs b/Engine/Source/Programs/UnrealGameSync/UnrealGameSync/Forms/ApplicationSettingsWindow.cs
--- a/Engine/Source/Programs/UnrealGameSync/UnrealGameSync/Forms/ApplicationSettingsWindow.cs
+++ b/Engine/Source/Programs/UnrealGameSync/UnrealGameSync/Forms/ApplicationSettingsWindow.cs
@@ -76,7 +76,7 @@ namespace UnrealGameSync
 			this.ToolUpdateMonitor = ToolUpdateMonitor;
 			this.Logger = Logger;
 
-			Utility.ReadGlobalPerforceSettings(ref InitialServerAndPort, ref InitialUserName, ref InitialDepotPath);
+			GlobalSettings.ReadGlobalPerforceSettings(ref InitialServerAndPort, ref InitialUserName, ref InitialDepotPath);
 			bInitialUnstable = bUnstable;
 
 			InitialAutomationPortNumber = AutomationServer.GetPortNumber();
@@ -207,7 +207,7 @@ namespace UnrealGameSync
 				}
 
 				bRestartUnstable = UseUnstableBuildCheckBox.Checked;
-				Utility.SaveGlobalPerforceSettings(ServerAndPort, UserName, DepotPath);
+				GlobalSettings.SaveGlobalPerforceSettings(ServerAndPort, UserName, DepotPath);
 				AutomationServer.SetPortNumber(AutomationPortNumber);
 			}
 
diff --git a/Engine/Source/Programs/UnrealGameSync/UnrealGameSync/Forms/IssueDetailsWindow.cs b/Engine/Source/Programs/UnrealGameSync/UnrealGameSync/Forms/IssueDetailsWindow.cs
--- a/Engine/Source/Programs/UnrealGameSync/UnrealGameSync/Forms/IssueDetailsWindow.cs
+++ b/Engine/Source/Programs/UnrealGameSync/UnrealGameSync/Forms/IssueDetailsWindow.cs
@@ -1282,7 +1282,7 @@ namespace UnrealGameSync
 		{
 			if (ContextMenuChange != null)
 			{
-				Utility.SpawnP4VC(String.Format("{0} change {1}", PerforceSettings.GetArgumentsForExternalProgram(true), ContextMenuChange.Number));
+				Program.SpawnP4VC(String.Format("{0} change {1}", PerforceSettings.GetArgumentsForExternalProgram(true), ContextMenuChange.Number));
 			}
 		}
 
diff --git a/Engine/Source/Programs/UnrealGameSync/UnrealGameSync/Forms/IssueFixedWindow.cs b/Engine/Source/Programs/UnrealGameSync/UnrealGameSync/Forms/IssueFixedWindow.cs
--- a/Engine/Source/Programs/UnrealGameSync/UnrealGameSync/Forms/IssueFixedWindow.cs
+++ b/Engine/Source/Programs/UnrealGameSync/UnrealGameSync/Forms/IssueFixedWindow.cs
@@ -260,7 +260,7 @@ namespace UnrealGameSync
 		private void ChangesContextMenu_MoreInfo_Click(object sender, EventArgs e)
 		{
 			DescribeRecord Record = (DescribeRecord)ChangesListContextMenu.Tag;
-			Utility.SpawnP4VC(String.Format("{0} change {1}", PerforceSettings.GetArgumentsForExternalProgram(true), Record.Number));
+			Program.SpawnP4VC(String.Format("{0} change {1}", PerforceSettings.GetArgumentsForExternalProgram(true), Record.Number));
 		}
 
 		private void ChangeNumberTextBox_Enter(object sender, EventArgs e)
diff --git a/Engine/Source/Programs/UnrealGameSync/UnrealGameSync/GlobalSettings.cs b/Engine/Source/Programs/UnrealGameSync/UnrealGameSync/GlobalSettings.cs
new file mode 100644
--- /dev/null
+++ b/Engine/Source/Programs/UnrealGameSync/UnrealGameSync/GlobalSettings.cs
@@ -0,0 +1,108 @@
+// Copyright Epic Games, Inc. All Rights Reserved.
+
+using Microsoft.Win32;
+using System;
+using System.Collections.Generic;
+using System.Linq;
+using System.Text;
+using System.Text.RegularExpressions;
+using System.Windows.Forms;
+
+namespace UnrealGameSync
+{
+	internal class GlobalSettings
+	{
+		public static void ReadGlobalPerforceSettings(ref string? ServerAndPort, ref string? UserName, ref string? DepotPath)
+		{
+			using (RegistryKey Key = Registry.CurrentUser.OpenSubKey("SOFTWARE\\Epic Games\\UnrealGameSync", false))
+			{
+				if (Key != null)
+				{
+					ServerAndPort = Key.GetValue("ServerAndPort", ServerAndPort) as string;
+					UserName = Key.GetValue("UserName", UserName) as string;
+					DepotPath = Key.GetValue("DepotPath", DepotPath) as string;
+
+					// Fix corrupted depot path string
+					if (DepotPath != null)
+					{
+						Match Match = Regex.Match(DepotPath, "^(.*)/(Release|UnstableRelease)/\\.\\.\\.@.*$");
+						if (Match.Success)
+						{
+							DepotPath = Match.Groups[1].Value;
+							SaveGlobalPerforceSettings(ServerAndPort, UserName, DepotPath);
+						}
+					}
+				}
+			}
+		}
+
+		public static void DeleteRegistryKey(RegistryKey RootKey, string KeyName, string ValueName)
+		{
+			using (RegistryKey Key = RootKey.OpenSubKey(KeyName, true))
+			{
+				if (Key != null)
+				{
+					DeleteRegistryKey(Key, ValueName);
+				}
+			}
+		}
+
+		public static void DeleteRegistryKey(RegistryKey Key, string Name)
+		{
+			string[] ValueNames = Key.GetValueNames();
+			if (ValueNames.Any(x => String.Compare(x, Name, StringComparison.OrdinalIgnoreCase) == 0))
+			{
+				try
+				{
+					Key.DeleteValue(Name);
+				}
+				catch
+				{
+				}
+			}
+		}
+
+		public static void SaveGlobalPerforceSettings(string? ServerAndPort, string? UserName, string? DepotPath)
+		{
+			try
+			{
+				using (RegistryKey Key = Registry.CurrentUser.CreateSubKey("SOFTWARE\\Epic Games\\UnrealGameSync"))
+				{
+					// Delete this legacy setting
+					DeleteRegistryKey(Key, "Server");
+
+					if (String.IsNullOrEmpty(ServerAndPort))
+					{
+						try { Key.DeleteValue("ServerAndPort"); } catch (Exception) { }
+					}
+					else
+					{
+						Key.SetValue("ServerAndPort", ServerAndPort);
+					}
+
+					if (String.IsNullOrEmpty(UserName))
+					{
+						try { Key.DeleteValue("UserName"); } catch (Exception) { }
+					}
+					else
+					{
+						Key.SetValue("UserName", UserName);
+					}
+
+					if (String.IsNullOrEmpty(DepotPath) || (DeploymentSettings.DefaultDepotPath != null && String.Equals(DepotPath, DeploymentSettings.DefaultDepotPath, StringComparison.InvariantCultureIgnoreCase)))
+					{
+						DeleteRegistryKey(Key, "DepotPath");
+					}
+					else
+					{
+						Key.SetValue("DepotPath", DepotPath);
+					}
+				}
+			}
+			catch (Exception Ex)
+			{
+				MessageBox.Show("Unable to save settings.\n\n" + Ex.ToString());
+			}
+		}
+	}
+}
diff --git a/Engine/Source/Programs/UnrealGameSync/UnrealGameSync/Program.cs b/Engine/Source/Programs/UnrealGameSync/UnrealGameSync/Program.cs
--- a/Engine/Source/Programs/UnrealGameSync/UnrealGameSync/Program.cs
+++ b/Engine/Source/Programs/UnrealGameSync/UnrealGameSync/Program.cs
@@ -80,7 +80,7 @@ namespace UnrealGameSync
 			string? ServerAndPort = null;
 			string? UserName = null;
 			string? BaseUpdatePath = null;
-			Utility.ReadGlobalPerforceSettings(ref ServerAndPort, ref UserName, ref BaseUpdatePath);
+			GlobalSettings.ReadGlobalPerforceSettings(ref ServerAndPort, ref UserName, ref BaseUpdatePath);
 
 			List<string> RemainingArgs = new List<string>(Args);
 
@@ -153,7 +153,7 @@ namespace UnrealGameSync
 						Logger.LogInformation("Missing server settings; finding defaults.");
 						ServerAndPort ??= PerforceSettings.Default.ServerAndPort;
 						UserName ??= PerforceSettings.Default.UserName;
-						Utility.SaveGlobalPerforceSettings(ServerAndPort, UserName, BaseUpdatePath);
+						GlobalSettings.SaveGlobalPerforceSettings(ServerAndPort, UserName, BaseUpdatePath);
 					}
 
 					ILogger TelemetryLogger = LoggerProvider.CreateLogger("Telemetry");
@@ -263,5 +263,54 @@ namespace UnrealGameSync
 			Value = null;
 			return false;
 		}
+
+		public static IEnumerable<string> GetPerforcePaths()
+		{
+			string? PathList = Environment.GetEnvironmentVariable("PATH");
+			if (!String.IsNullOrEmpty(PathList))
+			{
+				foreach (string PathEntry in PathList.Split(Path.PathSeparator))
+				{
+					string? PerforcePath = null;
+					try
+					{
+						string TestPerforcePath = Path.Combine(PathEntry, "p4.exe");
+						if (File.Exists(TestPerforcePath))
+						{
+							PerforcePath = TestPerforcePath;
+						}
+					}
+					catch
+					{
+					}
+
+					if (PerforcePath != null)
+					{
+						yield return PerforcePath;
+					}
+				}
+			}
+		}
+
+		public static void SpawnP4VC(string Arguments)
+		{
+			string Executable = "p4vc.exe";
+
+			foreach (string PerforcePath in GetPerforcePaths())
+			{
+				string? PerforceDir = Path.GetDirectoryName(PerforcePath);
+				if (PerforceDir != null && File.Exists(Path.Combine(PerforceDir, "p4vc.bat")) && !File.Exists(Path.Combine(PerforceDir, "p4vc.exe")))
+				{
+					Executable = Path.Combine(PerforceDir, "p4v.exe");
+					Arguments = "-p4vc " + Arguments;
+					break;
+				}
+			}
+
+			if (!Utility.SpawnHiddenProcess(Executable, Arguments))
+			{
+				MessageBox.Show("Unable to spawn p4vc. Check you have P4V installed.");
+			}
+		}
 	}
 }
diff --git a/Engine/Source/Programs/UnrealGameSync/UnrealGameSync/UriHandlers/P4Handler.cs b/Engine/Source/Programs/UnrealGameSync/UnrealGameSync/UriHandlers/P4Handler.cs
--- a/Engine/Source/Programs/UnrealGameSync/UnrealGameSync/UriHandlers/P4Handler.cs
+++ b/Engine/Source/Programs/UnrealGameSync/UnrealGameSync/UriHandlers/P4Handler.cs
@@ -26,7 +26,7 @@ namespace UnrealGameSync
 		{
 			string CommandLine = string.Format("timelapse {0}{1}", Line == -1 ? "" : string.Format(" -l {0} ", Line), DepotPath);
 
-			Utility.SpawnP4VC(CommandLine);
+			Program.SpawnP4VC(CommandLine);
 
 			return new UriResult() { Success = true };
 		}		
diff --git a/Engine/Source/Programs/UnrealGameSync/UnrealGameSync/Utility.cs b/Engine/Source/Programs/UnrealGameSync/UnrealGameSync/Utility.cs
--- a/Engine/Source/Programs/UnrealGameSync/UnrealGameSync/Utility.cs
+++ b/Engine/Source/Programs/UnrealGameSync/UnrealGameSync/Utility.cs
@@ -17,7 +17,6 @@ using System.Text.Json.Serialization;
 using System.Text.RegularExpressions;
 using System.Threading;
 using System.Threading.Tasks;
-using System.Windows.Forms;
 
 namespace UnrealGameSync
 {
@@ -115,15 +114,6 @@ namespace UnrealGameSync
 			}
 		}
 
-		public static void WriteException(this TextWriter Writer, Exception Ex, string Format, params object[] Args)
-		{
-			Writer.WriteLine(Format, Args);
-			foreach(string Line in Ex.ToString().Split('\n'))
-			{
-				Writer.WriteLine(">>>  " + Line);
-			}
-		}
-
 		public static bool SpawnProcess(string FileName, string CommandLine)
 		{
 			using(Process ChildProcess = new Process())
@@ -192,7 +182,7 @@ namespace UnrealGameSync
 		/// Any unknown variables are ignored.
 		/// </summary>
 		/// <param name="InputString">String to search for variable names</param>
-		/// <param name="Variables">Lookup of variable names to values</param>
+		/// <param name="AdditionalVariables">Lookup of variable names to values</param>
 		/// <returns>String with all variables replaced</returns>
 		public static string ExpandVariables(string InputString, Dictionary<string, string>? AdditionalVariables = null)
 		{
@@ -371,99 +361,6 @@ namespace UnrealGameSync
 
 		/******/
 
-		public static void ReadGlobalPerforceSettings(ref string? ServerAndPort, ref string? UserName, ref string? DepotPath)
-		{
-			using (RegistryKey Key = Registry.CurrentUser.OpenSubKey("SOFTWARE\\Epic Games\\UnrealGameSync", false))
-			{
-				if(Key != null)
-				{
-					ServerAndPort = Key.GetValue("ServerAndPort", ServerAndPort) as string;
-					UserName = Key.GetValue("UserName", UserName) as string;
-					DepotPath = Key.GetValue("DepotPath", DepotPath) as string;
-
-					// Fix corrupted depot path string
-					if(DepotPath != null)
-					{
-						Match Match = Regex.Match(DepotPath, "^(.*)/(Release|UnstableRelease)/\\.\\.\\.@.*$");
-						if (Match.Success)
-						{
-							DepotPath = Match.Groups[1].Value;
-							SaveGlobalPerforceSettings(ServerAndPort, UserName, DepotPath);
-						}
-					}
-				}
-			}
-		}
-
-		public static void DeleteRegistryKey(RegistryKey RootKey, string KeyName, string ValueName)
-		{
-			using (RegistryKey Key = RootKey.OpenSubKey(KeyName, true))
-			{
-				if (Key != null)
-				{
-					DeleteRegistryKey(Key, ValueName);
-				}
-			}
-		}
-
-		public static void DeleteRegistryKey(RegistryKey Key, string Name)
-		{
-			string[] ValueNames = Key.GetValueNames();
-			if (ValueNames.Any(x => String.Compare(x, Name, StringComparison.OrdinalIgnoreCase) == 0))
-			{
-				try
-				{ 
-					Key.DeleteValue(Name);
-				}
-				catch
-				{ 
-				}
-			}
-		}
-
-		public static void SaveGlobalPerforceSettings(string? ServerAndPort, string? UserName, string? DepotPath)
-		{
-			try
-			{
-				using (RegistryKey Key = Registry.CurrentUser.CreateSubKey("SOFTWARE\\Epic Games\\UnrealGameSync"))
-				{
-					// Delete this legacy setting
-					DeleteRegistryKey(Key, "Server");
-
-					if(String.IsNullOrEmpty(ServerAndPort))
-					{
-						try { Key.DeleteValue("ServerAndPort"); } catch(Exception) { }
-					}
-					else
-					{
-						Key.SetValue("ServerAndPort", ServerAndPort);
-					}
-
-					if(String.IsNullOrEmpty(UserName))
-					{
-						try { Key.DeleteValue("UserName"); } catch(Exception) { }
-					}
-					else
-					{
-						Key.SetValue("UserName", UserName);
-					}
-
-					if(String.IsNullOrEmpty(DepotPath) || (DeploymentSettings.DefaultDepotPath != null && String.Equals(DepotPath, DeploymentSettings.DefaultDepotPath, StringComparison.InvariantCultureIgnoreCase)))
-					{
-						DeleteRegistryKey(Key, "DepotPath");
-					}
-					else
-					{
-						Key.SetValue("DepotPath", DepotPath);
-					}
-				}
-			}
-			catch(Exception Ex)
-			{
-				MessageBox.Show("Unable to save settings.\n\n" + Ex.ToString());
-			}
-		}
-
 		public static async Task<string[]?> TryPrintFileUsingCacheAsync(IPerforceConnection Perforce, string DepotPath, DirectoryReference CacheFolder, string? Digest, ILogger Logger, CancellationToken CancellationToken)
 		{
 			if(Digest == null)
@@ -639,55 +536,6 @@ namespace UnrealGameSync
 			return NormalUserName.ToString();
 		}
 
-		public static IEnumerable<string> GetPerforcePaths()
-		{
-			string? PathList = Environment.GetEnvironmentVariable("PATH");
-			if (!String.IsNullOrEmpty(PathList))
-			{
-				foreach (string PathEntry in PathList.Split(Path.PathSeparator))
-				{
-					string? PerforcePath = null;
-					try
-					{
-						string TestPerforcePath = Path.Combine(PathEntry, "p4.exe");
-						if (File.Exists(TestPerforcePath))
-						{
-							PerforcePath = TestPerforcePath;
-						}
-					}
-					catch
-					{
-					}
-
-					if(PerforcePath != null)
-					{
-						yield return PerforcePath;
-					}
-				}
-			}
-		}
-
-		public static void SpawnP4VC(string Arguments)
-		{
-			string Executable = "p4vc.exe";
-
-			foreach (string PerforcePath in GetPerforcePaths())
-			{
-				string? PerforceDir = Path.GetDirectoryName(PerforcePath);
-				if (PerforceDir != null && File.Exists(Path.Combine(PerforceDir, "p4vc.bat")) && !File.Exists(Path.Combine(PerforceDir, "p4vc.exe")))
-				{
-					Executable = Path.Combine(PerforceDir, "p4v.exe");
-					Arguments = "-p4vc " + Arguments;
-					break;
-				}
-			}
-
-			if (!Utility.SpawnHiddenProcess(Executable, Arguments))
-			{
-				MessageBox.Show("Unable to spawn p4vc. Check you have P4V installed.");
-			}
-		}
-
 		public static void OpenUrl(string Url)
 		{
 			ProcessStartInfo StartInfo = new ProcessStartInfo();
diff --git a/Engine/Source/Programs/UnrealGameSync/UnrealGameSync/Workspace.cs b/Engine/Source/Programs/UnrealGameSync/UnrealGameSync/Workspace.cs
--- a/Engine/Source/Programs/UnrealGameSync/UnrealGameSync/Workspace.cs
+++ b/Engine/Source/Programs/UnrealGameSync/UnrealGameSync/Workspace.cs
@@ -22,183 +22,6 @@ using System.Threading.Tasks;
 
 namespace UnrealGameSync
 {
-	[Flags]
-	public enum WorkspaceUpdateOptions
-	{
-		Sync = 0x01,
-		SyncSingleChange = 0x02,
-		AutoResolveChanges = 0x04,
-		GenerateProjectFiles = 0x08,
-		SyncArchives = 0x10,
-		Build = 0x20,
-		Clean = 0x40,
-		ScheduledBuild = 0x80,
-		RunAfterSync = 0x100,
-		OpenSolutionAfterSync = 0x200,
-		ContentOnly = 0x400,
-		UpdateFilter = 0x800,
-		SyncAllProjects = 0x1000,
-		IncludeAllProjectsInSolution = 0x2000,
-	}
-
-	public enum WorkspaceUpdateResult
-	{
-		Canceled,
-		FailedToSync,
-		FailedToSyncLoginExpired,
-		FilesToDelete,
-		FilesToResolve,
-		FilesToClobber,
-		FailedToCompile,
-		FailedToCompileWithCleanWorkspace,
-		Success,
-	}
-
-	public class PerforceSyncOptions
-	{
-		public const int DefaultNumRetries = 0;
-		public const int DefaultNumThreads = 2;
-		public const int DefaultTcpBufferSize = 0;
-		public const int DefaultFileBufferSize = 0;
-
-		public const int DefaultMaxCommandsPerBatch = 200;
-		public const int DefaultMaxSizePerBatch = 128 * 1024 * 1024;
-
-		public int NumRetries = DefaultNumRetries;
-		public int NumThreads = DefaultNumThreads;
-		public int TcpBufferSize = DefaultTcpBufferSize;
-		public int FileBufferSize = DefaultFileBufferSize;
-
-		public int MaxCommandsPerBatch = DefaultMaxCommandsPerBatch;
-		public int MaxSizePerBatch = DefaultMaxSizePerBatch;
-
-
-		public PerforceSyncOptions Clone()
-		{
-			return (PerforceSyncOptions)MemberwiseClone();
-		}
-	}
-
-	public interface IArchiveInfo
-	{
-		string Name { get; }
-		string Type { get; }
-		string BasePath { get; }
-		string? Target { get; }
-		bool Exists();
-		bool TryGetArchiveKeyForChangeNumber(int ChangeNumber, [NotNullWhen(true)] out string? ArchiveKey);
-		Task<bool> DownloadArchive(IPerforceConnection Perforce, string ArchiveKey, DirectoryReference LocalRootPath, FileReference ManifestFileName, ILogger Logger, ProgressValue Progress, CancellationToken CancellationToken);
-	}
-
-	public class WorkspaceUpdateContext
-	{
-		public DateTime StartTime = DateTime.UtcNow;
-		public int ChangeNumber;
-		public WorkspaceUpdateOptions Options;
-		public string[]? SyncFilter;
-		public Dictionary<string, Tuple<IArchiveInfo, string>?> ArchiveTypeToArchive = new Dictionary<string, Tuple<IArchiveInfo, string>?>();
-		public Dictionary<string, bool> DeleteFiles = new Dictionary<string,bool>();
-		public Dictionary<string, bool> ClobberFiles = new Dictionary<string,bool>();
-		public Dictionary<Guid,ConfigObject> DefaultBuildSteps;
-		public List<ConfigObject> UserBuildStepObjects;
-		public HashSet<Guid>? CustomBuildSteps;
-		public Dictionary<string, string> Variables;
-		public PerforceSyncOptions? PerforceSyncOptions;
-		public List<HaveRecord>? HaveFiles; // Cached when sync filter has changed
-
-		// May be updated during sync
-		public ConfigFile? ProjectConfigFile;
-		public IReadOnlyList<string>? ProjectStreamFilter;
-
-		public WorkspaceUpdateContext(int InChangeNumber, WorkspaceUpdateOptions InOptions, string[]? InSyncFilter, Dictionary<Guid, ConfigObject> InDefaultBuildSteps, List<ConfigObject> InUserBuildSteps, HashSet<Guid>? InCustomBuildSteps, Dictionary<string, string> InVariables)
-		{
-			ChangeNumber = InChangeNumber;
-			Options = InOptions;
-			SyncFilter = InSyncFilter;
-			DefaultBuildSteps = InDefaultBuildSteps;
-			UserBuildStepObjects = InUserBuildSteps;
-			CustomBuildSteps = InCustomBuildSteps;
-			Variables = InVariables;
-		}
-	}
-
-	class WorkspaceSyncCategory
-	{
-		public Guid UniqueId;
-		public bool bEnable;
-		public string Name;
-		public string[] Paths;
-		public bool bHidden;
-		public Guid[] Requires;
-
-		public WorkspaceSyncCategory(Guid UniqueId) : this(UniqueId, "Unnamed")
-		{
-		}
-
-		public WorkspaceSyncCategory(Guid UniqueId, string Name, params string[] Paths)
-		{
-			this.UniqueId = UniqueId;
-			this.bEnable = true;
-			this.Name = Name;
-			this.Paths = Paths;
-			this.Requires = new Guid[0];
-		}
-
-		public static Dictionary<Guid, bool> GetDefault(IEnumerable<WorkspaceSyncCategory> Categories)
-		{
-			return Categories.ToDictionary(x => x.UniqueId, x => x.bEnable);
-		}
-
-		public static Dictionary<Guid, bool> GetDelta(Dictionary<Guid, bool> Source, Dictionary<Guid, bool> Target)
-		{
-			Dictionary<Guid, bool> Changes = new Dictionary<Guid, bool>();
-			foreach (KeyValuePair<Guid, bool> Pair in Target)
-			{
-				bool bValue;
-				if (!Source.TryGetValue(Pair.Key, out bValue) || bValue != Pair.Value)
-				{
-					Changes[Pair.Key] = Pair.Value;
-				}
-			}
-			return Changes;
-		}
-
-		public static void ApplyDelta(Dictionary<Guid, bool> Categories, Dictionary<Guid, bool> Delta)
-		{
-			foreach(KeyValuePair<Guid, bool> Pair in Delta)
-			{
-				Categories[Pair.Key] = Pair.Value;
-			}
-		}
-
-		public override string ToString()
-		{
-			return Name;
-		}
-	}
-
-	public class ProjectInfo
-	{
-		public DirectoryReference LocalRootPath { get; }
-		public FileReference LocalFileName { get; }
-		public string ClientRootPath { get; }
-		public string ClientFileName { get; }
-		public string ProjectIdentifier { get; }
-		public string TelemetryProjectIdentifier { get; }
-		public bool bIsEnterpriseProject { get; }
-
-		public ProjectInfo(DirectoryReference LocalRootPath, FileReference LocalFileName, string ClientRootPath, string ClientFileName, string ProjectIdentifier, bool bIsEnterpriseProject)
-		{
-			this.LocalRootPath = LocalRootPath;
-			this.LocalFileName = LocalFileName;
-			this.ClientRootPath = ClientRootPath;
-			this.ClientFileName = ClientFileName;
-			this.ProjectIdentifier = ProjectIdentifier;
-			this.TelemetryProjectIdentifier = PerforceUtils.GetClientOrDepotDirectoryName(ProjectIdentifier);
-			this.bIsEnterpriseProject = bIsEnterpriseProject;
-		}
-	}
-
 	public sealed class Workspace : IDisposable
 	{
 		public IPerforceSettings PerforceSettings { get; }
@@ -341,1628 +164,4 @@ namespace UnrealGameSync
 
 		public Tuple<string, float> CurrentProgress => CurrentUpdate?.CurrentProgress ?? new Tuple<string, float>("", 0.0f);
 	}
-
-	class WorkspaceUpdate : IAsyncDisposable
-	{
-		const string BuildVersionFileName = "/Engine/Build/Build.version";
-		const string VersionHeaderFileName = "/Engine/Source/Runtime/Launch/Resources/Version.h";
-		const string ObjectVersionFileName = "/Engine/Source/Runtime/Core/Private/UObject/ObjectVersion.cpp";
-
-		static readonly string LocalVersionHeaderFileName = VersionHeaderFileName.Replace('/', '\\');
-		static readonly string LocalObjectVersionFileName = ObjectVersionFileName.Replace('/', '\\');
-
-		static SemaphoreSlim UpdateSemaphore = new SemaphoreSlim(1);
-
-		class RecordCounter : IDisposable
-		{
-			ProgressValue Progress;
-			string Message;
-			int Count;
-			Stopwatch Timer = Stopwatch.StartNew();
-
-			public RecordCounter(ProgressValue Progress, string Message)
-			{
-				this.Progress = Progress;
-				this.Message = Message;
-
-				Progress.Set(Message);
-			}
-
-			public void Dispose()
-			{
-				UpdateMessage();
-			}
-
-			public void Increment()
-			{
-				Count++;
-				if (Timer.ElapsedMilliseconds > 250)
-				{
-					UpdateMessage();
-				}
-			}
-
-			public void UpdateMessage()
-			{
-				Progress.Set(String.Format("{0} ({1:N0})", Message, Count));
-				Timer.Restart();
-			}
-		}
-
-		class SyncBatchBuilder
-		{
-			public int MaxCommandsPerList { get; }
-			public long MaxSizePerList { get; }
-			public Queue<List<string>> Batches { get; }
-
-			List<string>? Commands;
-			List<string>? DeleteCommands;
-			long Size;
-
-			public SyncBatchBuilder(int? MaxCommandsPerList, long? MaxSizePerList)
-			{
-				this.MaxCommandsPerList = MaxCommandsPerList ?? PerforceSyncOptions.DefaultMaxCommandsPerBatch;
-				this.MaxSizePerList = MaxSizePerList ?? PerforceSyncOptions.DefaultMaxSizePerBatch;
-				this.Batches = new Queue<List<string>>();
-			}
-
-			public void Add(string NewCommand, long NewSize)
-			{
-				if (NewSize == 0)
-				{
-					if (DeleteCommands == null || DeleteCommands.Count >= MaxCommandsPerList)
-					{
-						DeleteCommands = new List<string>();
-						Batches.Enqueue(DeleteCommands);
-					}
-
-					DeleteCommands.Add(NewCommand);
-				}
-				else
-				{
-					if (Commands == null || Commands.Count >= MaxCommandsPerList || Size + NewSize >= MaxSizePerList)
-					{
-						Commands = new List<string>();
-						Batches.Enqueue(Commands);
-						Size = 0;
-					}
-
-					Commands.Add(NewCommand);
-					Size += NewSize;
-				}
-			}
-		}
-
-		class SyncTree
-		{
-			public bool bCanUseWildcard;
-			public int TotalIncludedFiles;
-			public long TotalSize;
-			public int TotalExcludedFiles;
-			public Dictionary<string, long> IncludedFiles = new Dictionary<string, long>();
-			public Dictionary<string, SyncTree> NameToSubTree = new Dictionary<string, SyncTree>(StringComparer.OrdinalIgnoreCase);
-
-			public SyncTree(bool bCanUseWildcard)
-			{
-				this.bCanUseWildcard = bCanUseWildcard;
-			}
-
-			public SyncTree FindOrAddSubTree(string Name)
-			{
-				SyncTree? Result;
-				if (!NameToSubTree.TryGetValue(Name, out Result))
-				{
-					Result = new SyncTree(bCanUseWildcard);
-					NameToSubTree.Add(Name, Result);
-				}
-				return Result;
-			}
-
-			public void IncludeFile(string Path, long Size)
-			{
-				int Idx = Path.IndexOf('/');
-				if (Idx == -1)
-				{
-					IncludedFiles.Add(Path, Size);
-				}
-				else
-				{
-					SyncTree SubTree = FindOrAddSubTree(Path.Substring(0, Idx));
-					SubTree.IncludeFile(Path.Substring(Idx + 1), Size);
-				}
-				TotalIncludedFiles++;
-				TotalSize += Size;
-			}
-
-			public void ExcludeFile(string Path)
-			{
-				int Idx = Path.IndexOf('/');
-				if (Idx != -1)
-				{
-					SyncTree SubTree = FindOrAddSubTree(Path.Substring(0, Idx));
-					SubTree.ExcludeFile(Path.Substring(Idx + 1));
-				}
-				TotalExcludedFiles++;
-			}
-
-			public void GetOptimizedSyncCommands(string Prefix, int ChangeNumber, SyncBatchBuilder Builder)
-			{
-				if (bCanUseWildcard && TotalExcludedFiles == 0 && TotalSize < Builder.MaxSizePerList)
-				{
-					Builder.Add(String.Format("{0}/...@{1}", Prefix, ChangeNumber), TotalSize);
-				}
-				else
-				{
-					foreach (KeyValuePair<string, long> File in IncludedFiles)
-					{
-						Builder.Add(String.Format("{0}/{1}@{2}", Prefix, File.Key, ChangeNumber), File.Value);
-					}
-					foreach (KeyValuePair<string, SyncTree> Pair in NameToSubTree)
-					{
-						Pair.Value.GetOptimizedSyncCommands(String.Format("{0}/{1}", Prefix, PerforceUtils.EscapePath(Pair.Key)), ChangeNumber, Builder);
-					}
-				}
-			}
-		}
-
-		CancellationTokenSource CancellationSource;
-		public WorkspaceUpdateContext Context { get; }
-		public ProgressValue Progress { get; } = new ProgressValue();
-		public Task UpdateTask { get; }
-
-		public WorkspaceUpdate(WorkspaceUpdateContext Context, Func<WorkspaceUpdate, CancellationToken, Task> TaskFunc)
-		{
-			this.CancellationSource = new CancellationTokenSource();
-			this.Context = Context;
-			this.UpdateTask = Task.Run(() => TaskFunc(this, CancellationSource.Token));
-		}
-
-		public void Cancel()
-		{
-			CancellationSource.Cancel();
-		}
-
-		public ValueTask DisposeAsync()
-		{
-			Task DisposeTask = Task.CompletedTask;
-			if (UpdateTask != null)
-			{
-				CancellationSource.Cancel();
-				DisposeTask = UpdateTask.ContinueWith(Task => CancellationSource.Dispose());
-			}
-			return new ValueTask(DisposeTask);
-		}
-
-		class SyncFile
-		{
-			public string DepotFile;
-			public string RelativePath;
-			public long Size;
-
-			public SyncFile(string DepotFile, string RelativePath, long Size)
-			{
-				this.DepotFile = DepotFile;
-				this.RelativePath = RelativePath;
-				this.Size = Size;
-			}
-		};
-
-		class SemaphoreScope : IDisposable
-		{
-			SemaphoreSlim Semaphore;
-			public bool HasLock { get; private set; }
-
-			public SemaphoreScope(SemaphoreSlim Semaphore)
-			{
-				this.Semaphore = Semaphore;
-			}
-
-			public bool TryAcquire()
-			{
-				HasLock = HasLock || Semaphore.Wait(0);
-				return HasLock;
-			}
-
-			public async Task AcquireAsync(CancellationToken CancellationToken)
-			{
-				await Semaphore.WaitAsync(CancellationToken);
-				HasLock = true;
-			}
-
-			public void Release()
-			{
-				if (HasLock)
-				{
-					Semaphore.Release();
-					HasLock = false;
-				}
-			}
-
-			public void Dispose() => Release();
-		}
-
-		public async Task<(WorkspaceUpdateResult, string)> ExecuteAsync(IPerforceSettings PerforceSettings, ProjectInfo Project, UserWorkspaceState State, ILogger Logger, CancellationToken CancellationToken)
-		{
-			using IPerforceConnection Perforce = await PerforceConnection.CreateAsync(PerforceSettings, Logger);
-
-			string CmdExe = Path.Combine(Environment.GetFolderPath(Environment.SpecialFolder.System), "cmd.exe");
-			if (!File.Exists(CmdExe))
-			{
-				return (WorkspaceUpdateResult.FailedToSync, $"Missing {CmdExe}.");
-			}
-
-			List<Tuple<string, TimeSpan>> Times = new List<Tuple<string, TimeSpan>>();
-
-			int NumFilesSynced = 0;
-			if (Context.Options.HasFlag(WorkspaceUpdateOptions.Sync) || Context.Options.HasFlag(WorkspaceUpdateOptions.SyncSingleChange))
-			{
-				using (TelemetryStopwatch SyncTelemetryStopwatch = new TelemetryStopwatch("Workspace_Sync", Project.TelemetryProjectIdentifier))
-				{
-					Logger.LogInformation("Syncing to {Change}...", Context.ChangeNumber);
-
-					// Make sure we're logged in
-					PerforceResponse<LoginRecord> LoginResponse = await Perforce.TryGetLoginStateAsync(CancellationToken);
-					if (!LoginResponse.Succeeded)
-					{
-						return (WorkspaceUpdateResult.FailedToSyncLoginExpired, "User is not logged in.");
-					}
-
-					// Figure out which paths to sync
-					List<string> RelativeSyncPaths = GetRelativeSyncPaths(Project, (Context.Options & WorkspaceUpdateOptions.SyncAllProjects) != 0, Context.SyncFilter);
-					List<string> SyncPaths = new List<string>(RelativeSyncPaths.Select(x => Project.ClientRootPath + x));
-
-					// Get the user's sync filter
-					FileFilter UserFilter = new FileFilter(FileFilterType.Include);
-					if (Context.SyncFilter != null)
-					{
-						UserFilter.AddRules(Context.SyncFilter.Select(x => x.Trim()).Where(x => x.Length > 0 && !x.StartsWith(";") && !x.StartsWith("#")));
-					}
-
-					// Check if the new sync filter matches the previous one. If not, we'll enumerate all files in the workspace and make sure there's nothing extra there.
-					string? NextSyncFilterHash = null;
-					using (SHA1Managed SHA = new SHA1Managed())
-					{
-						StringBuilder CombinedFilter = new StringBuilder();
-						foreach (string RelativeSyncPath in RelativeSyncPaths)
-						{
-							CombinedFilter.AppendFormat("{0}\n", RelativeSyncPath);
-						}
-						if (Context.SyncFilter != null)
-						{
-							CombinedFilter.Append("--FROM--\n");
-							CombinedFilter.Append(String.Join("\n", Context.SyncFilter));
-						}
-						NextSyncFilterHash = BitConverter.ToString(SHA.ComputeHash(Encoding.UTF8.GetBytes(CombinedFilter.ToString()))).Replace("-", "");
-					}
-
-					// If the hash differs, enumerate everything in the workspace to find what needs to be removed
-					if (NextSyncFilterHash != State.CurrentSyncFilterHash)
-					{
-						using (TelemetryStopwatch FilterStopwatch = new TelemetryStopwatch("Workspace_Sync_FilterChanged", Project.TelemetryProjectIdentifier))
-						{
-							Logger.LogInformation("Filter has changed ({PrevHash} -> {NextHash}); finding files in workspace that need to be removed.", (String.IsNullOrEmpty(State.CurrentSyncFilterHash)) ? "None" : State.CurrentSyncFilterHash, NextSyncFilterHash);
-
-							// Find all the files that are in this workspace
-							List<HaveRecord>? HaveFiles = Context.HaveFiles;
-							if (HaveFiles == null)
-							{
-								HaveFiles = new List<HaveRecord>();
-								using (RecordCounter HaveCounter = new RecordCounter(Progress, "Sync filter changed; checking workspace..."))
-								{
-									await foreach (PerforceResponse<HaveRecord> Record in Perforce.TryHaveAsync(FileSpecList.Any, CancellationToken))
-									{
-										if (Record.Succeeded)
-										{
-											HaveFiles.Add(Record.Data);
-											HaveCounter.Increment();
-										}
-										else
-										{
-											return (WorkspaceUpdateResult.FailedToSync, $"Unable to query files ({Record}).");
-										}
-									}
-								}
-								Context.HaveFiles = HaveFiles;
-							}
-
-							// Build a filter for the current sync paths
-							FileFilter SyncPathsFilter = new FileFilter(FileFilterType.Exclude);
-							foreach (string RelativeSyncPath in RelativeSyncPaths)
-							{
-								SyncPathsFilter.Include(RelativeSyncPath);
-							}
-
-							// Remove all the files that are not included by the filter
-							List<string> RemoveDepotPaths = new List<string>();
-							foreach (HaveRecord HaveFile in HaveFiles)
-							{
-								try
-								{
-									FileReference FullPath = new FileReference(HaveFile.Path);
-									if (MatchFilter(Project, FullPath, SyncPathsFilter) && !MatchFilter(Project, FullPath, UserFilter))
-									{
-										Logger.LogInformation("  {DepotFile}", HaveFile.DepotFile);
-										RemoveDepotPaths.Add(HaveFile.DepotFile);
-									}
-								}
-								catch (PathTooLongException)
-								{
-									// We don't actually care about this when looking for files to remove. Perforce may think that it's synced the path, and silently failed. Just ignore it.
-								}
-							}
-
-							// Check if there are any paths outside the regular sync paths
-							if (RemoveDepotPaths.Count > 0)
-							{
-								bool bDeleteListMatches = true;
-
-								Dictionary<string, bool> NewDeleteFiles = new Dictionary<string, bool>(StringComparer.OrdinalIgnoreCase);
-								foreach (string RemoveDepotPath in RemoveDepotPaths)
-								{
-									bool bDelete;
-									if (!Context.DeleteFiles.TryGetValue(RemoveDepotPath, out bDelete))
-									{
-										bDeleteListMatches = false;
-										bDelete = true;
-									}
-									NewDeleteFiles[RemoveDepotPath] = bDelete;
-								}
-								Context.DeleteFiles = NewDeleteFiles;
-
-								if (!bDeleteListMatches)
-								{
-									return (WorkspaceUpdateResult.FilesToDelete, $"Cancelled after finding {NewDeleteFiles.Count} files excluded by filter");
-								}
-
-								RemoveDepotPaths.RemoveAll(x => !Context.DeleteFiles[x]);
-							}
-
-							// Actually delete any files that we don't want
-							if (RemoveDepotPaths.Count > 0)
-							{
-								// Clear the current sync filter hash. If the sync is canceled, we'll be in an indeterminate state, and we should always clean next time round.
-								State.CurrentSyncFilterHash = "INVALID";
-								State.Save();
-
-								// Find all the depot paths that will be synced
-								HashSet<string> RemainingDepotPathsToRemove = new HashSet<string>(StringComparer.OrdinalIgnoreCase);
-								RemainingDepotPathsToRemove.UnionWith(RemoveDepotPaths);
-
-								// Build the list of revisions to sync
-								List<string> RevisionsToRemove = new List<string>();
-								RevisionsToRemove.AddRange(RemoveDepotPaths.Select(x => String.Format("{0}#0", x)));
-
-								(WorkspaceUpdateResult, string) RemoveResult = await SyncFileRevisions(Perforce, "Removing files...", Context, RevisionsToRemove, RemainingDepotPathsToRemove, Progress, Logger, CancellationToken);
-								if (RemoveResult.Item1 != WorkspaceUpdateResult.Success)
-								{
-									return RemoveResult;
-								}
-							}
-
-							// Update the sync filter hash. We've removed any files we need to at this point.
-							State.CurrentSyncFilterHash = NextSyncFilterHash;
-							State.Save();
-						}
-					}
-
-					// Create a filter for all the files we don't want
-					FileFilter Filter = new FileFilter(UserFilter);
-					Filter.Exclude(BuildVersionFileName);
-					if (Context.Options.HasFlag(WorkspaceUpdateOptions.ContentOnly))
-					{
-						Filter.Exclude("*.usf");
-						Filter.Exclude("*.ush");
-					}
-
-					// Create a tree to store the sync path
-					SyncTree SyncTree = new SyncTree(false);
-					if (!Context.Options.HasFlag(WorkspaceUpdateOptions.SyncSingleChange))
-					{
-						foreach (string RelativeSyncPath in RelativeSyncPaths)
-						{
-							const string WildcardSuffix = "/...";
-							if (RelativeSyncPath.EndsWith(WildcardSuffix, StringComparison.Ordinal))
-							{
-								SyncTree Leaf = SyncTree;
-
-								string[] Fragments = RelativeSyncPath.Split('/');
-								for (int Idx = 1; Idx < Fragments.Length - 1; Idx++)
-								{
-									Leaf = Leaf.FindOrAddSubTree(Fragments[Idx]);
-								}
-
-								Leaf.bCanUseWildcard = true;
-							}
-						}
-					}
-
-					// Find all the server changes, and anything that's opened for edit locally. We need to sync files we have open to schedule a resolve.
-					SyncBatchBuilder BatchBuilder = new SyncBatchBuilder(Context.PerforceSyncOptions?.MaxCommandsPerBatch, Context.PerforceSyncOptions?.MaxSizePerBatch);
-					List<string> SyncDepotPaths = new List<string>();
-					using (RecordCounter Counter = new RecordCounter(Progress, "Filtering files..."))
-					{
-						// Track the total new bytes that will be required on disk when syncing. Add an extra 100MB for padding.
-						long RequiredFreeSpace = 100 * 1024 * 1024;
-
-						foreach (string SyncPath in SyncPaths)
-						{
-							string SyncFilter = Context.Options.HasFlag(WorkspaceUpdateOptions.Sync) ? $"{Context.ChangeNumber}" : $"={Context.ChangeNumber}";
-
-							List<SyncFile> SyncFiles = new List<SyncFile>();
-							await foreach (PerforceResponse<SyncRecord> Response in Perforce.TrySyncAsync(SyncOptions.PreviewOnly, -1, $"{SyncPath}@{SyncFilter}", CancellationToken))
-							{
-								if (!Response.Succeeded)
-								{
-									return (WorkspaceUpdateResult.FailedToSync, $"Couldn't enumerate changes matching {SyncPath}.");
-								}
-
-								SyncRecord Record = Response.Data;
-
-								// Get the full local path
-								string RelativePath;
-								try
-								{
-									FileReference SyncFile = new FileReference(Record.Path);
-									RelativePath = PerforceUtils.GetClientRelativePath(Project.LocalRootPath, SyncFile);
-								}
-								catch (PathTooLongException)
-								{
-									Logger.LogInformation("The local path for {Path} exceeds the maximum allowed by Windows. Re-sync your workspace to a directory with a shorter name, or delete the file from the server.", Record.Path);
-									return (WorkspaceUpdateResult.FailedToSync, "File exceeds maximum path length allowed by Windows.");
-								}
-
-								// Create the sync record
-								long SyncSize = (Record.Action == SyncAction.Deleted) ? 0 : Record.FileSize;
-								SyncFiles.Add(new SyncFile(Record.DepotFile, RelativePath, SyncSize));
-								Counter.Increment();
-							}
-
-							// Also sync the currently open files
-							await foreach (PerforceResponse<OpenedRecord> Response in Perforce.TryOpenedAsync(OpenedOptions.None, FileSpecList.Any, CancellationToken))
-							{
-								if (!Response.Succeeded)
-								{
-									return (WorkspaceUpdateResult.FailedToSync, $"Couldn't enumerate changes matching {SyncPath}.");
-								}
-
-								OpenedRecord Record = Response.Data;
-								if (Record.Action != FileAction.Add || Record.Action != FileAction.Branch || Record.Action != FileAction.MoveAdd)
-								{
-									string RelativePath = PerforceUtils.GetClientRelativePath(Record.ClientFile);
-									SyncFiles.Add(new SyncFile(Record.DepotFile, RelativePath, 0));
-								}
-							}
-
-							// Enumerate all the files to be synced. NOTE: depotPath is escaped, whereas clientPath is not.
-							foreach (SyncFile SyncRecord in SyncFiles)
-							{
-								if (Filter.Matches(SyncRecord.RelativePath))
-								{
-									SyncTree.IncludeFile(PerforceUtils.EscapePath(SyncRecord.RelativePath), SyncRecord.Size);
-									SyncDepotPaths.Add(SyncRecord.DepotFile);
-									RequiredFreeSpace += SyncRecord.Size;
-
-									// If the file exists the required free space can be reduced as those bytes will be replaced.
-									FileInfo LocalFileInfo = FileReference.Combine(Project.LocalRootPath, SyncRecord.RelativePath).ToFileInfo();
-									if (LocalFileInfo.Exists)
-									{
-										RequiredFreeSpace -= LocalFileInfo.Length;
-									}
-								}
-								else
-								{
-									SyncTree.ExcludeFile(PerforceUtils.EscapePath(SyncRecord.RelativePath));
-								}
-							}
-						}
-
-						try
-						{
-							DirectoryInfo LocalRootInfo = Project.LocalRootPath.ToDirectoryInfo();
-							DriveInfo Drive = new DriveInfo(LocalRootInfo.Root.FullName);
-
-							if (Drive.AvailableFreeSpace < RequiredFreeSpace)
-							{
-								Logger.LogInformation("Syncing requires {RequiredSpace} which exceeds the {AvailableSpace} available free space on {Drive}.", StringUtils.FormatBytesString(RequiredFreeSpace), StringUtils.FormatBytesString(Drive.AvailableFreeSpace), Drive.Name);
-								return (WorkspaceUpdateResult.FailedToSync, "Not enough available free space.");
-							}
-						}
-						catch (SystemException)
-						{
-							Logger.LogInformation("Unable to check available free space for {RootPath}.", Project.LocalRootPath);
-						}
-					}
-					SyncTree.GetOptimizedSyncCommands(Project.ClientRootPath, Context.ChangeNumber, BatchBuilder);
-
-					// Clear the current sync changelist, in case we cancel
-					if (!Context.Options.HasFlag(WorkspaceUpdateOptions.SyncSingleChange))
-					{
-						State.CurrentChangeNumber = -1;
-						State.AdditionalChangeNumbers.Clear();
-						State.Save();
-					}
-
-					// Find all the depot paths that will be synced
-					HashSet<string> RemainingDepotPaths = new HashSet<string>(StringComparer.OrdinalIgnoreCase);
-					RemainingDepotPaths.UnionWith(SyncDepotPaths);
-
-					using (TelemetryStopwatch TransferStopwatch = new TelemetryStopwatch("Workspace_Sync_TransferFiles", Project.TelemetryProjectIdentifier))
-					{
-						TransferStopwatch.AddData(new { MachineName = Environment.MachineName, DomainName = Environment.UserDomainName, ServerAndPort = Perforce.Settings.ServerAndPort, UserName = Perforce.Settings.UserName, IncludedFiles = SyncTree.TotalIncludedFiles, ExcludedFiles = SyncTree.TotalExcludedFiles, Size = SyncTree.TotalSize, NumThreads = Context.PerforceSyncOptions?.NumThreads ?? PerforceSyncOptions.DefaultNumThreads });
-
-						(WorkspaceUpdateResult, string) SyncResult = await SyncFileRevisions(Perforce, "Syncing files...", Context, BatchBuilder.Batches, RemainingDepotPaths, Progress, Logger, CancellationToken);
-						if (SyncResult.Item1 != WorkspaceUpdateResult.Success)
-						{
-							TransferStopwatch.AddData(new { SyncResult = SyncResult.Item2.ToString(), CompletedFilesFiles = SyncDepotPaths.Count - RemainingDepotPaths.Count });
-							return SyncResult;
-						}
-
-						TransferStopwatch.Stop("Ok");
-						TransferStopwatch.AddData(new { TransferRate = SyncTree.TotalSize / Math.Max(TransferStopwatch.Elapsed.TotalSeconds, 0.0001f) });
-					}
-
-					int VersionChangeNumber = -1;
-					if (Context.Options.HasFlag(WorkspaceUpdateOptions.Sync) && !Context.Options.HasFlag(WorkspaceUpdateOptions.UpdateFilter))
-					{
-						// Read the new config file
-						Context.ProjectConfigFile = await ReadProjectConfigFile(Project.LocalRootPath, Project.LocalFileName, Logger);
-						Context.ProjectStreamFilter = await ReadProjectStreamFilter(Perforce, Context.ProjectConfigFile, Logger, CancellationToken);
-
-						// Get the branch name
-						string? BranchOrStreamName = await Perforce.GetCurrentStreamAsync(CancellationToken);
-						if (BranchOrStreamName != null)
-						{
-							// If it's a virtual stream, take the concrete parent stream instead
-							for (; ; )
-							{
-								StreamRecord StreamSpec = await Perforce.GetStreamAsync(BranchOrStreamName, false, CancellationToken);
-								if (StreamSpec.Type != "virtual" || StreamSpec.Parent == "none" || StreamSpec.Parent == null)
-								{
-									break;
-								}
-								BranchOrStreamName = StreamSpec.Parent;
-							}
-						}
-						else
-						{
-							// Otherwise use the depot path for GenerateProjectFiles.bat in the root of the workspace
-							List<WhereRecord> Files = await Perforce.WhereAsync(Project.ClientRootPath + "/GenerateProjectFiles.bat").Where(x => !x.Unmap).ToListAsync(CancellationToken);
-							if (Files.Count != 1)
-							{
-								return (WorkspaceUpdateResult.FailedToSync, $"Couldn't determine branch name for {Project.ClientFileName}.");
-							}
-							BranchOrStreamName = PerforceUtils.GetClientOrDepotDirectoryName(Files[0].DepotFile);
-						}
-
-						// Find the last code change before this changelist. For consistency in versioning between local builds and precompiled binaries, we need to use the last submitted code changelist as our version number.
-						string[] CodeFilter = new string[] { ".cs", ".h", ".cpp", ".usf", ".ush", ".uproject", ".uplugin" }.SelectMany(x => SyncPaths.Select(y => String.Format("{0}{1}@<={2}", y, x, Context.ChangeNumber))).ToArray();
-
-						PerforceResponseList<ChangesRecord> CodeChanges = await Perforce.TryGetChangesAsync(ChangesOptions.None, 1, ChangeStatus.Submitted, CodeFilter, CancellationToken);
-						if (!CodeChanges.Succeeded)
-						{
-							return (WorkspaceUpdateResult.FailedToSync, $"Couldn't determine last code changelist before CL {Context.ChangeNumber}.");
-						}
-						if (CodeChanges.Count == 0)
-						{
-							return (WorkspaceUpdateResult.FailedToSync, $"Could not find any code changes before CL {Context.ChangeNumber}.");
-						}
-
-						// Get the last code change
-						if (Context.ProjectConfigFile.GetValue("Options.VersionToLastCodeChange", true))
-						{
-							VersionChangeNumber = CodeChanges.Max(x => x.Data.Number);
-						}
-						else
-						{
-							VersionChangeNumber = Context.ChangeNumber;
-						}
-
-						// Update the version files
-						if (Context.ProjectConfigFile.GetValue("Options.UseFastModularVersioningV2", false))
-						{
-							bool bIsLicenseeVersion = await IsLicenseeVersion(Perforce, Project, Logger, CancellationToken);
-							if (!await UpdateVersionFile(Perforce, Project.ClientRootPath + BuildVersionFileName, Context.ChangeNumber, Text => UpdateBuildVersion(Text, Context.ChangeNumber, VersionChangeNumber, BranchOrStreamName, bIsLicenseeVersion), Logger, CancellationToken))
-							{
-								return (WorkspaceUpdateResult.FailedToSync, $"Failed to update {BuildVersionFileName}.");
-							}
-						}
-						else if (Context.ProjectConfigFile.GetValue("Options.UseFastModularVersioning", false))
-						{
-							bool bIsLicenseeVersion = await IsLicenseeVersion(Perforce, Project, Logger, CancellationToken);
-							if (!await UpdateVersionFile(Perforce, Project.ClientRootPath + BuildVersionFileName, Context.ChangeNumber, Text => UpdateBuildVersion(Text, Context.ChangeNumber, VersionChangeNumber, BranchOrStreamName, bIsLicenseeVersion), Logger, CancellationToken))
-							{
-								return (WorkspaceUpdateResult.FailedToSync, $"Failed to update {BuildVersionFileName}");
-							}
-
-							Dictionary<string, string> VersionHeaderStrings = new Dictionary<string, string>();
-							VersionHeaderStrings["#define ENGINE_IS_PROMOTED_BUILD"] = " (0)";
-							VersionHeaderStrings["#define BUILT_FROM_CHANGELIST"] = " 0";
-							VersionHeaderStrings["#define BRANCH_NAME"] = " \"" + BranchOrStreamName.Replace('/', '+') + "\"";
-							if (!await UpdateVersionFile(Perforce, Project.ClientRootPath + VersionHeaderFileName, VersionHeaderStrings, Context.ChangeNumber, Logger, CancellationToken))
-							{
-								return (WorkspaceUpdateResult.FailedToSync, $"Failed to update {VersionHeaderFileName}.");
-							}
-							if (!await UpdateVersionFile(Perforce, Project.ClientRootPath + ObjectVersionFileName, new Dictionary<string, string>(), Context.ChangeNumber, Logger, CancellationToken))
-							{
-								return (WorkspaceUpdateResult.FailedToSync, $"Failed to update {ObjectVersionFileName}.");
-							}
-						}
-						else
-						{
-							if (!await UpdateVersionFile(Perforce, Project.ClientRootPath + BuildVersionFileName, new Dictionary<string, string>(), Context.ChangeNumber, Logger, CancellationToken))
-							{
-								return (WorkspaceUpdateResult.FailedToSync, $"Failed to update {BuildVersionFileName}");
-							}
-
-							Dictionary<string, string> VersionStrings = new Dictionary<string, string>();
-							VersionStrings["#define ENGINE_VERSION"] = " " + VersionChangeNumber.ToString();
-							VersionStrings["#define ENGINE_IS_PROMOTED_BUILD"] = " (0)";
-							VersionStrings["#define BUILT_FROM_CHANGELIST"] = " " + VersionChangeNumber.ToString();
-							VersionStrings["#define BRANCH_NAME"] = " \"" + BranchOrStreamName.Replace('/', '+') + "\"";
-							if (!await UpdateVersionFile(Perforce, Project.ClientRootPath + VersionHeaderFileName, VersionStrings, Context.ChangeNumber, Logger, CancellationToken))
-							{
-								return (WorkspaceUpdateResult.FailedToSync, $"Failed to update {VersionHeaderFileName}");
-							}
-							if (!await UpdateVersionFile(Perforce, Project.ClientRootPath + ObjectVersionFileName, VersionStrings, Context.ChangeNumber, Logger, CancellationToken))
-							{
-								return (WorkspaceUpdateResult.FailedToSync, $"Failed to update {ObjectVersionFileName}");
-							}
-						}
-
-						// Remove all the receipts for build targets in this branch
-						if (Project.ClientFileName.EndsWith(".uproject", StringComparison.OrdinalIgnoreCase))
-						{
-							await Perforce.SyncQuietAsync(SyncOptions.None, -1, PerforceUtils.GetClientOrDepotDirectoryName(Project.ClientFileName) + "/Build/Receipts/...#0", CancellationToken);
-						}
-					}
-
-					// Check if there are any files which need resolving
-					List<FStatRecord> UnresolvedFiles = await FindUnresolvedFiles(Perforce, SyncPaths, CancellationToken);
-					if (UnresolvedFiles.Count > 0 && Context.Options.HasFlag(WorkspaceUpdateOptions.AutoResolveChanges))
-					{
-						foreach (FStatRecord UnresolvedFile in UnresolvedFiles)
-						{
-							if (UnresolvedFile.DepotFile != null)
-							{
-								await Perforce.ResolveAsync(-1, ResolveOptions.Automatic, UnresolvedFile.DepotFile, CancellationToken);
-							}
-						}
-						UnresolvedFiles = await FindUnresolvedFiles(Perforce, SyncPaths, CancellationToken);
-					}
-					if (UnresolvedFiles.Count > 0)
-					{
-						Logger.LogInformation("{NumFiles} files need resolving:", UnresolvedFiles.Count);
-						foreach (FStatRecord UnresolvedFile in UnresolvedFiles)
-						{
-							Logger.LogInformation("  {ClientFile}", UnresolvedFile.ClientFile);
-						}
-						return (WorkspaceUpdateResult.FilesToResolve, "Files need resolving.");
-					}
-
-					// Continue processing sync-only actions
-					if (Context.Options.HasFlag(WorkspaceUpdateOptions.Sync) && !Context.Options.HasFlag(WorkspaceUpdateOptions.UpdateFilter))
-					{
-						Context.ProjectConfigFile ??= await ReadProjectConfigFile(Project.LocalRootPath, Project.LocalFileName, Logger);
-
-						// Execute any project specific post-sync steps
-						string[]? PostSyncSteps = Context.ProjectConfigFile.GetValues("Sync.Step", null);
-						if (PostSyncSteps != null)
-						{
-							Logger.LogInformation("");
-							Logger.LogInformation("Executing post-sync steps...");
-
-							Dictionary<string, string> PostSyncVariables = new Dictionary<string, string>(Context.Variables);
-							PostSyncVariables["Change"] = Context.ChangeNumber.ToString();
-							PostSyncVariables["CodeChange"] = VersionChangeNumber.ToString();
-
-							foreach (string PostSyncStep in PostSyncSteps.Select(x => x.Trim()))
-							{
-								ConfigObject PostSyncStepObject = new ConfigObject(PostSyncStep);
-
-								string ToolFileName = Utility.ExpandVariables(PostSyncStepObject.GetValue("FileName", ""), PostSyncVariables);
-								if (ToolFileName != null)
-								{
-									string ToolArguments = Utility.ExpandVariables(PostSyncStepObject.GetValue("Arguments", ""), PostSyncVariables);
-
-									Logger.LogInformation("post-sync> Running {FileName} {Arguments}", ToolFileName, ToolArguments);
-
-									int ResultFromTool = await Utility.ExecuteProcessAsync(ToolFileName, null, ToolArguments, Line => ProcessOutput(Line, "post-sync> ", Progress, Logger), CancellationToken);
-									if (ResultFromTool != 0)
-									{
-										return (WorkspaceUpdateResult.FailedToSync, $"Post-sync step terminated with exit code {ResultFromTool}.");
-									}
-								}
-							}
-						}
-					}
-
-					// Update the current state
-					if (Context.Options.HasFlag(WorkspaceUpdateOptions.SyncSingleChange))
-					{
-						State.AdditionalChangeNumbers.Add(Context.ChangeNumber);
-					}
-					else
-					{
-						State.CurrentChangeNumber = Context.ChangeNumber;
-					}
-					State.Save();
-
-					// Update the timing info
-					Times.Add(new Tuple<string, TimeSpan>("Sync", SyncTelemetryStopwatch.Stop("Success")));
-
-					// Save the number of files synced
-					NumFilesSynced = SyncDepotPaths.Count;
-					Logger.LogInformation("");
-				}
-			}
-
-			// Extract an archive from the depot path
-			if (Context.Options.HasFlag(WorkspaceUpdateOptions.SyncArchives))
-			{
-				using (TelemetryStopwatch Stopwatch = new TelemetryStopwatch("Workspace_SyncArchives", Project.TelemetryProjectIdentifier))
-				{
-					// Create the directory for extracted archive manifests
-					DirectoryReference ManifestDirectoryName;
-					if (Project.LocalFileName.HasExtension(".uproject"))
-					{
-						ManifestDirectoryName = DirectoryReference.Combine(Project.LocalFileName.Directory, "Saved", "UnrealGameSync");
-					}
-					else
-					{
-						ManifestDirectoryName = DirectoryReference.Combine(Project.LocalFileName.Directory, "Engine", "Saved", "UnrealGameSync");
-					}
-					DirectoryReference.CreateDirectory(ManifestDirectoryName);
-
-					// Sync and extract (or just remove) the given archives
-					foreach (KeyValuePair<string, Tuple<IArchiveInfo, string>?> ArchiveTypeAndArchive in Context.ArchiveTypeToArchive)
-					{
-						string ArchiveType = ArchiveTypeAndArchive.Key;
-
-						// Remove any existing binaries
-						FileReference ManifestFileName = FileReference.Combine(ManifestDirectoryName, String.Format("{0}.zipmanifest", ArchiveType));
-						if (FileReference.Exists(ManifestFileName))
-						{
-							Logger.LogInformation("Removing {ArchiveType} binaries...", ArchiveType);
-							Progress.Set(String.Format("Removing {0} binaries...", ArchiveType), 0.0f);
-							ArchiveUtils.RemoveExtractedFiles(Project.LocalRootPath, ManifestFileName, Progress, Logger);
-							FileReference.Delete(ManifestFileName);
-							Logger.LogInformation("");
-						}
-
-						// If we have a new depot path, sync it down and extract it
-						if (ArchiveTypeAndArchive.Value != null)
-						{
-							IArchiveInfo ArchiveInfo = ArchiveTypeAndArchive.Value.Item1;
-							string ArchiveKey = ArchiveTypeAndArchive.Value.Item2;
-
-							Logger.LogInformation("Syncing {ArchiveType} binaries...", ArchiveType.ToLowerInvariant());
-							Progress.Set(String.Format("Syncing {0} binaries...", ArchiveType.ToLowerInvariant()), 0.0f);
-							if (!ArchiveInfo.DownloadArchive(Perforce, ArchiveKey, Project.LocalRootPath, ManifestFileName, Logger, Progress, CancellationToken.None).Result)
-							{
-								return (WorkspaceUpdateResult.FailedToSync, $"Couldn't read {ArchiveKey}");
-							}
-						}
-					}
-
-					// Update the state
-					State.ExpandedArchiveTypes = Context.ArchiveTypeToArchive.Where(x => x.Value != null).Select(x => x.Key).ToArray();
-					State.Save();
-
-					// Add the finish time
-					Times.Add(new Tuple<string, TimeSpan>("Archive", Stopwatch.Stop("Success")));
-				}
-			}
-
-			// Take the lock before doing anything else. Building and generating project files can only be done on one workspace at a time.
-			using SemaphoreScope Scope = new SemaphoreScope(UpdateSemaphore);
-			if (Context.Options.HasFlag(WorkspaceUpdateOptions.GenerateProjectFiles) || Context.Options.HasFlag(WorkspaceUpdateOptions.Build))
-			{
-				if (!Scope.TryAcquire())
-				{
-					Logger.LogInformation("Waiting for other workspaces to finish...");
-					await Scope.AcquireAsync(CancellationToken);
-				}
-			}
-
-			// Generate project files in the workspace
-			if (Context.Options.HasFlag(WorkspaceUpdateOptions.GenerateProjectFiles))
-			{
-				using (TelemetryStopwatch Stopwatch = new TelemetryStopwatch("Workspace_GenerateProjectFiles", Project.TelemetryProjectIdentifier))
-				{
-					Progress.Set("Generating project files...", 0.0f);
-
-					StringBuilder CommandLine = new StringBuilder();
-					CommandLine.AppendFormat("/C \"\"{0}\"", FileReference.Combine(Project.LocalRootPath, "GenerateProjectFiles.bat"));
-					if ((Context.Options & WorkspaceUpdateOptions.SyncAllProjects) == 0 && (Context.Options & WorkspaceUpdateOptions.IncludeAllProjectsInSolution) == 0)
-					{
-						if (Project.LocalFileName.HasExtension(".uproject"))
-						{
-							CommandLine.AppendFormat(" \"{0}\"", Project.LocalFileName);
-						}
-					}
-					CommandLine.Append(" -progress\"");
-
-					Logger.LogInformation("Generating project files...");
-					Logger.LogInformation("gpf> Running {FileName} {Arguments}", CmdExe, CommandLine);
-
-					int GenerateProjectFilesResult = await Utility.ExecuteProcessAsync(CmdExe, null, CommandLine.ToString(), Line => ProcessOutput(Line, "gpf> ", Progress, Logger), CancellationToken);
-					if (GenerateProjectFilesResult != 0)
-					{
-						return (WorkspaceUpdateResult.FailedToCompile, $"Failed to generate project files (exit code {GenerateProjectFilesResult}).");
-					}
-
-					Logger.LogInformation("");
-					Times.Add(new Tuple<string, TimeSpan>("Prj gen", Stopwatch.Stop("Success")));
-				}
-			}
-
-			// Build everything using MegaXGE
-			if (Context.Options.HasFlag(WorkspaceUpdateOptions.Build))
-			{
-				Context.ProjectConfigFile ??= await ReadProjectConfigFile(Project.LocalRootPath, Project.LocalFileName, Logger);
-
-				// Compile all the build steps together
-				Dictionary<Guid, ConfigObject> BuildStepObjects = Context.DefaultBuildSteps.ToDictionary(x => x.Key, x => new ConfigObject(x.Value));
-				BuildStep.MergeBuildStepObjects(BuildStepObjects, Context.ProjectConfigFile.GetValues("Build.Step", new string[0]).Select(x => new ConfigObject(x)));
-				BuildStep.MergeBuildStepObjects(BuildStepObjects, Context.UserBuildStepObjects);
-
-				// Construct build steps from them
-				List<BuildStep> BuildSteps = BuildStepObjects.Values.Select(x => new BuildStep(x)).OrderBy(x => (x.OrderIndex == -1) ? 10000 : x.OrderIndex).ToList();
-				if (Context.CustomBuildSteps != null && Context.CustomBuildSteps.Count > 0)
-				{
-					BuildSteps.RemoveAll(x => !Context.CustomBuildSteps.Contains(x.UniqueId));
-				}
-				else if (Context.Options.HasFlag(WorkspaceUpdateOptions.ScheduledBuild))
-				{
-					BuildSteps.RemoveAll(x => !x.bScheduledSync);
-				}
-				else
-				{
-					BuildSteps.RemoveAll(x => !x.bNormalSync);
-				}
-
-				// Check if the last successful build was before a change that we need to force a clean for
-				bool bForceClean = false;
-				if (State.LastBuiltChangeNumber != 0)
-				{
-					foreach (string CleanBuildChange in Context.ProjectConfigFile.GetValues("ForceClean.Changelist", new string[0]))
-					{
-						int ChangeNumber;
-						if (int.TryParse(CleanBuildChange, out ChangeNumber))
-						{
-							if ((State.LastBuiltChangeNumber >= ChangeNumber) != (State.CurrentChangeNumber >= ChangeNumber))
-							{
-								Logger.LogInformation("Forcing clean build due to changelist {Change}.", ChangeNumber);
-								Logger.LogInformation("");
-								bForceClean = true;
-								break;
-							}
-						}
-					}
-				}
-
-				// Execute them all
-				using (TelemetryStopwatch Stopwatch = new TelemetryStopwatch("Workspace_Build", Project.TelemetryProjectIdentifier))
-				{
-					Progress.Set("Starting build...", 0.0f);
-
-					// Execute all the steps
-					float MaxProgressFraction = 0.0f;
-					foreach (BuildStep Step in BuildSteps)
-					{
-						MaxProgressFraction += (float)Step.EstimatedDuration / (float)Math.Max(BuildSteps.Sum(x => x.EstimatedDuration), 1);
-
-						Progress.Set(Step.StatusText ?? "Executing build step");
-						Progress.Push(MaxProgressFraction);
-
-						Logger.LogInformation("{Status}", Step.StatusText);
-
-						if (Step.IsValid())
-						{
-							switch (Step.Type)
-							{
-								case BuildStepType.Compile:
-									using (TelemetryStopwatch StepStopwatch = new TelemetryStopwatch("Workspace_Execute_Compile", Project.TelemetryProjectIdentifier))
-									{
-										StepStopwatch.AddData(new { Target = Step.Target });
-
-										FileReference BuildBat = FileReference.Combine(Project.LocalRootPath, "Engine", "Build", "BatchFiles", "Build.bat");
-										string CommandLine = String.Format("{0} {1} {2} {3} -NoHotReloadFromIDE", Step.Target, Step.Platform, Step.Configuration, Utility.ExpandVariables(Step.Arguments ?? "", Context.Variables));
-										if (Context.Options.HasFlag(WorkspaceUpdateOptions.Clean) || bForceClean)
-										{
-											Logger.LogInformation("ubt> Running {FileName} {Arguments}", BuildBat, CommandLine + " -clean");
-											await Utility.ExecuteProcessAsync(CmdExe, null, "/C \"\"" + BuildBat + "\" " + CommandLine + " -clean\"", Line => ProcessOutput(Line, "ubt> ", Progress, Logger), CancellationToken);
-										}
-
-										Logger.LogInformation("ubt> Running {FileName} {Arguments}", BuildBat, CommandLine + " -progress");
-
-										int ResultFromBuild = await Utility.ExecuteProcessAsync(CmdExe, null, "/C \"\"" + BuildBat + "\" " + CommandLine + " -progress\"", Line => ProcessOutput(Line, "ubt> ", Progress, Logger), CancellationToken);
-										if (ResultFromBuild != 0)
-										{
-											StepStopwatch.Stop("Failed");
-
-											WorkspaceUpdateResult Result;
-											if (await HasModifiedSourceFiles(Perforce, Project, CancellationToken) || Context.UserBuildStepObjects.Count > 0)
-											{
-												Result = WorkspaceUpdateResult.FailedToCompile;
-											}
-											else
-											{
-												Result = WorkspaceUpdateResult.FailedToCompileWithCleanWorkspace;
-											}
-
-											return (Result, $"Failed to compile {Step.Target}");
-										}
-
-										StepStopwatch.Stop("Success");
-									}
-									break;
-								case BuildStepType.Cook:
-									using (TelemetryStopwatch StepStopwatch = new TelemetryStopwatch("Workspace_Execute_Cook", Project.TelemetryProjectIdentifier))
-									{
-										StepStopwatch.AddData(new { Project = Path.GetFileNameWithoutExtension(Step.FileName) });
-
-										FileReference LocalRunUAT = FileReference.Combine(Project.LocalRootPath, "Engine", "Build", "BatchFiles", "RunUAT.bat");
-										string Arguments = String.Format("/C \"\"{0}\" -profile=\"{1}\"\"", LocalRunUAT, FileReference.Combine(Project.LocalRootPath, Step.FileName ?? "unknown"));
-										Logger.LogInformation("uat> Running {FileName} {Argument}", LocalRunUAT, Arguments);
-
-										int ResultFromUAT = await Utility.ExecuteProcessAsync(CmdExe, null, Arguments, Line => ProcessOutput(Line, "uat> ", Progress, Logger), CancellationToken);
-										if (ResultFromUAT != 0)
-										{
-											StepStopwatch.Stop("Failed");
-											return (WorkspaceUpdateResult.FailedToCompile, $"Cook failed. ({ResultFromUAT})");
-										}
-
-										StepStopwatch.Stop("Success");
-									}
-									break;
-								case BuildStepType.Other:
-									using (TelemetryStopwatch StepStopwatch = new TelemetryStopwatch("Workspace_Execute_Custom", Project.TelemetryProjectIdentifier))
-									{
-										StepStopwatch.AddData(new { FileName = Path.GetFileNameWithoutExtension(Step.FileName) });
-
-										FileReference ToolFileName = FileReference.Combine(Project.LocalRootPath, Utility.ExpandVariables(Step.FileName ?? "unknown", Context.Variables));
-										string ToolWorkingDir = String.IsNullOrWhiteSpace(Step.WorkingDir) ? ToolFileName.Directory.FullName : Utility.ExpandVariables(Step.WorkingDir, Context.Variables);
-										string ToolArguments = Utility.ExpandVariables(Step.Arguments ?? "", Context.Variables);
-										Logger.LogInformation("tool> Running {0} {1}", ToolFileName, ToolArguments);
-
-										if (Step.bUseLogWindow)
-										{
-											int ResultFromTool = await Utility.ExecuteProcessAsync(ToolFileName.FullName, ToolWorkingDir, ToolArguments, Line => ProcessOutput(Line, "tool> ", Progress, Logger), CancellationToken);
-											if (ResultFromTool != 0)
-											{
-												StepStopwatch.Stop("Failed");
-												return (WorkspaceUpdateResult.FailedToCompile, $"Tool terminated with exit code {ResultFromTool}.");
-											}
-										}
-										else
-										{
-											ProcessStartInfo StartInfo = new ProcessStartInfo(ToolFileName.FullName, ToolArguments);
-											StartInfo.WorkingDirectory = ToolWorkingDir;
-											using (Process.Start(StartInfo))
-											{
-											}
-										}
-
-										StepStopwatch.Stop("Success");
-									}
-									break;
-							}
-						}
-
-						Logger.LogInformation("");
-						Progress.Pop();
-					}
-
-					Times.Add(new Tuple<string, TimeSpan>("Build", Stopwatch.Stop("Success")));
-				}
-
-				// Update the last successful build change number
-				if (Context.CustomBuildSteps == null || Context.CustomBuildSteps.Count == 0)
-				{
-					State.LastBuiltChangeNumber = State.CurrentChangeNumber;
-					State.Save();
-				}
-			}
-
-			// Write out all the timing information
-			Logger.LogInformation("Total time: {TotalTime}", FormatTime(Times.Sum(x => (long)(x.Item2.TotalMilliseconds / 1000))));
-			foreach (Tuple<string, TimeSpan> Time in Times)
-			{
-				Logger.LogInformation("   {Name,-8}: {TimeSeconds}", Time.Item1, FormatTime((long)(Time.Item2.TotalMilliseconds / 1000)));
-			}
-			if (NumFilesSynced > 0)
-			{
-				Logger.LogInformation("{NumFiles} files synced.", NumFilesSynced);
-			}
-
-			DateTime FinishTime = DateTime.Now;
-			Logger.LogInformation("");
-			Logger.LogInformation("UPDATE SUCCEEDED ({FinishDate} {FinishTime})", FinishTime.ToShortDateString(), FinishTime.ToShortTimeString());
-
-			return (WorkspaceUpdateResult.Success, "Update succeeded");
-		}
-
-		static void ProcessOutput(string Line, string Prefix, ProgressValue Progress, ILogger Logger)
-		{
-			string? ParsedLine = ProgressTextWriter.ParseLine(Line, Progress);
-			if (ParsedLine != null)
-			{
-				Logger.LogInformation("{Prefix}{Message}", Prefix, Line);
-			}
-		}
-
-		static async Task<bool> IsLicenseeVersion(IPerforceConnection Perforce, ProjectInfo Project, ILogger Logger, CancellationToken CancellationToken)
-		{
-			string[] Files =
-			{
-				Project.ClientRootPath + "/Engine/Build/NotForLicensees/EpicInternal.txt",
-				Project.ClientRootPath + "/Engine/Restricted/NotForLicensees/Build/EpicInternal.txt"
-			};
-
-			List<FStatRecord> Records = await Perforce.FStatAsync(Files, CancellationToken).ToListAsync(CancellationToken);
-			return !Records.Any(x => x.IsMapped);
-		}
-
-		public static List<string> GetSyncPaths(ProjectInfo Project, bool bSyncAllProjects, string[] SyncFilter)
-		{
-			List<string> SyncPaths = GetRelativeSyncPaths(Project, bSyncAllProjects, SyncFilter);
-			return SyncPaths.Select(x => Project.ClientRootPath + x).ToList();
-		}
-
-		public static List<string> GetRelativeSyncPaths(ProjectInfo Project, bool bSyncAllProjects, string[]? SyncFilter)
-		{
-			List<string> SyncPaths = new List<string>();
-
-			// Check the client path is formatted correctly
-			if (!Project.ClientFileName.StartsWith(Project.ClientRootPath + "/"))
-			{
-				throw new Exception(String.Format("Expected '{0}' to start with '{1}'", Project.ClientFileName, Project.ClientRootPath));
-			}
-
-			// Add the default project paths
-			int LastSlashIdx = Project.ClientFileName.LastIndexOf('/');
-			if (bSyncAllProjects || !Project.ClientFileName.EndsWith(".uproject", StringComparison.OrdinalIgnoreCase) || LastSlashIdx <= Project.ClientRootPath.Length)
-			{
-				SyncPaths.Add("/...");
-			}
-			else
-			{
-				SyncPaths.Add("/*");
-				SyncPaths.Add("/Engine/...");
-				if (Project.bIsEnterpriseProject)
-				{
-					SyncPaths.Add("/Enterprise/...");
-				}
-				SyncPaths.Add(Project.ClientFileName.Substring(Project.ClientRootPath.Length, LastSlashIdx - Project.ClientRootPath.Length) + "/...");
-			}
-
-			// Apply the sync filter to that list. We only want inclusive rules in the output list, but we can manually apply exclusions to previous entries.
-			if (SyncFilter != null)
-			{
-				foreach (string SyncPath in SyncFilter)
-				{
-					string TrimSyncPath = SyncPath.Trim();
-					if (TrimSyncPath.StartsWith("/"))
-					{
-						SyncPaths.Add(TrimSyncPath);
-					}
-					else if (TrimSyncPath.StartsWith("-/") && TrimSyncPath.EndsWith("..."))
-					{
-						SyncPaths.RemoveAll(x => x.StartsWith(TrimSyncPath.Substring(1, TrimSyncPath.Length - 4)));
-					}
-				}
-			}
-
-			// Sort the remaining paths by length, and remove any paths which are included twice
-			SyncPaths = SyncPaths.OrderBy(x => x.Length).ToList();
-			for (int Idx = 0; Idx < SyncPaths.Count; Idx++)
-			{
-				string SyncPath = SyncPaths[Idx];
-				if (SyncPath.EndsWith("..."))
-				{
-					string SyncPathPrefix = SyncPath.Substring(0, SyncPath.Length - 3);
-					for (int OtherIdx = SyncPaths.Count - 1; OtherIdx > Idx; OtherIdx--)
-					{
-						if (SyncPaths[OtherIdx].StartsWith(SyncPathPrefix))
-						{
-							SyncPaths.RemoveAt(OtherIdx);
-						}
-					}
-				}
-			}
-
-			return SyncPaths;
-		}
-
-		public static bool MatchFilter(ProjectInfo Project, FileReference FileName, FileFilter Filter)
-		{
-			bool bMatch = true;
-			if (FileName.IsUnderDirectory(Project.LocalRootPath))
-			{
-				string RelativePath = FileName.MakeRelativeTo(Project.LocalRootPath);
-				if (!Filter.Matches(RelativePath))
-				{
-					bMatch = false;
-				}
-			}
-			return bMatch;
-		}
-
-		class SyncState
-		{
-			public int TotalDepotPaths;
-			public HashSet<string> RemainingDepotPaths;
-			public Queue<List<string>> SyncCommandLists;
-			public string StatusMessage;
-			public WorkspaceUpdateResult Result = WorkspaceUpdateResult.Success;
-
-			public SyncState(HashSet<string> RemainingDepotPaths, Queue<List<string>> SyncCommandLists)
-			{
-				this.TotalDepotPaths = RemainingDepotPaths.Count;
-				this.RemainingDepotPaths = RemainingDepotPaths;
-				this.SyncCommandLists = SyncCommandLists;
-				this.StatusMessage = "Succeeded.";
-			}
-		}
-
-		static Task<(WorkspaceUpdateResult, string)> SyncFileRevisions(IPerforceConnection Perforce, string Prefix, WorkspaceUpdateContext Context, List<string> SyncCommands, HashSet<string> RemainingDepotPaths, ProgressValue Progress, ILogger Logger, CancellationToken CancellationToken)
-		{
-			Queue<List<string>> SyncCommandLists = new Queue<List<string>>();
-			SyncCommandLists.Enqueue(SyncCommands);
-			return SyncFileRevisions(Perforce, Prefix, Context, SyncCommandLists, RemainingDepotPaths, Progress, Logger, CancellationToken);
-		}
-
-		static async Task<(WorkspaceUpdateResult, string)> SyncFileRevisions(IPerforceConnection Perforce, string Prefix, WorkspaceUpdateContext Context, Queue<List<string>> SyncCommandLists, HashSet<string> RemainingDepotPaths, ProgressValue Progress, ILogger Logger, CancellationToken CancellationToken)
-		{
-			// Figure out the number of additional background threads we want to run with. We can run worker on the current thread.
-			int NumThreads = Context.PerforceSyncOptions?.NumThreads ?? PerforceSyncOptions.DefaultNumThreads;
-			int NumExtraThreads = Math.Max(Math.Min(SyncCommandLists.Count, NumThreads) - 1, 0);
-
-			List<IPerforceConnection> ChildConnections = new List<IPerforceConnection>();
-			List<Task> ChildTasks = new List<Task>(NumExtraThreads);
-			try
-			{
-				// Create the state object shared by all the worker threads
-				SyncState State = new SyncState(RemainingDepotPaths, SyncCommandLists);
-
-				// Wrapper writer around the log class to prevent multiple threads writing to it at once
-				ILogger LogWrapper = Logger;
-
-				// Initialize Sync Progress
-				UpdateSyncState(Prefix, State, Progress);
-
-				// Delegate for updating the sync state after a file has been synced
-				Action<SyncRecord, ILogger> SyncOutput = (Record, LocalLog) => { UpdateSyncState(Prefix, Record, State, Progress, LocalLog); };
-
-				// Create all the child threads
-				for (int ThreadIdx = 0; ThreadIdx < NumExtraThreads; ThreadIdx++)
-				{
-					int ThreadNumber = ThreadIdx + 2;
-
-					// Create connection
-					IPerforceConnection ChildConnection = await PerforceConnection.CreateAsync(Perforce.Settings, Perforce.Logger);
-					ChildConnections.Add(ChildConnection);
-
-					Task ChildTask = Task.Run(() => StaticSyncWorker(ThreadNumber, ChildConnection, Context, State, SyncOutput, LogWrapper, CancellationToken));
-					ChildTasks.Add(ChildTask);
-				}
-
-				// Run one worker on the current thread
-				await StaticSyncWorker(1, Perforce, Context, State, SyncOutput, LogWrapper, CancellationToken);
-
-				// Return the result that was set on the state object
-				return (State.Result, State.StatusMessage);
-			}
-			finally
-			{
-				await Task.WhenAll(ChildTasks);
-
-				foreach (IPerforceConnection ChildConnection in ChildConnections)
-				{
-					ChildConnection.Dispose();
-				}
-
-				foreach (Task ChildTask in ChildTasks)
-				{
-					await ChildTask;
-				}
-			}
-		}
-
-		static void UpdateSyncState(string Prefix, SyncState State, ProgressValue Progress)
-		{
-			lock (State)
-			{
-				string Message = String.Format("{0} ({1}/{2})", Prefix, State.TotalDepotPaths - State.RemainingDepotPaths.Count, State.TotalDepotPaths);
-				float Fraction = Math.Min((float)(State.TotalDepotPaths - State.RemainingDepotPaths.Count) / (float)State.TotalDepotPaths, 1.0f);
-				Progress.Set(Message, Fraction);
-			}
-		}
-
-		static void UpdateSyncState(string Prefix, SyncRecord Record, SyncState State, ProgressValue Progress, ILogger Logger)
-		{
-			lock (State)
-			{
-				State.RemainingDepotPaths.Remove(Record.DepotFile);
-
-				string Message = String.Format("{0} ({1}/{2})", Prefix, State.TotalDepotPaths - State.RemainingDepotPaths.Count, State.TotalDepotPaths);
-				float Fraction = Math.Min((float)(State.TotalDepotPaths - State.RemainingDepotPaths.Count) / (float)State.TotalDepotPaths, 1.0f);
-				Progress.Set(Message, Fraction);
-
-				Logger.LogInformation("p4>   {Action} {Path}", Record.Action, Record.Path);
-			}
-		}
-
-		static async Task StaticSyncWorker(int ThreadNumber, IPerforceConnection Perforce, WorkspaceUpdateContext Context, SyncState State, Action<SyncRecord, ILogger> SyncOutput, ILogger GlobalLog, CancellationToken CancellationToken)
-		{
-			PrefixedTextWriter ThreadLog = new PrefixedTextWriter(String.Format("{0}:", ThreadNumber), GlobalLog);
-			for (; ; )
-			{
-				// Remove the next batch that needs to be synced
-				List<string> SyncCommands;
-				lock (State)
-				{
-					if (State.Result == WorkspaceUpdateResult.Success && State.SyncCommandLists.Count > 0)
-					{
-						SyncCommands = State.SyncCommandLists.Dequeue();
-					}
-					else
-					{
-						break;
-					}
-				}
-
-				// Sync the files
-				(WorkspaceUpdateResult Result, string StatusMessage) = await StaticSyncFileRevisions(Perforce, Context, SyncCommands, Record => SyncOutput(Record, ThreadLog), CancellationToken);
-
-				// If it failed, try to set it on the state if nothing else has failed first
-				if (Result != WorkspaceUpdateResult.Success)
-				{
-					lock (State)
-					{
-						if (State.Result == WorkspaceUpdateResult.Success)
-						{
-							State.Result = Result;
-							State.StatusMessage = StatusMessage;
-						}
-					}
-					break;
-				}
-			}
-		}
-
-		static async Task<(WorkspaceUpdateResult, string)> StaticSyncFileRevisions(IPerforceConnection Perforce, WorkspaceUpdateContext Context, List<string> SyncCommands, Action<SyncRecord> SyncOutput, CancellationToken CancellationToken)
-		{
-			// Sync them all
-			List<PerforceResponse<SyncRecord>> Responses = await Perforce.TrySyncAsync(SyncOptions.None, -1, SyncCommands, CancellationToken).ToListAsync(CancellationToken);
-
-			List<string> TamperedFiles = new List<string>();
-			foreach (PerforceResponse<SyncRecord> Response in Responses)
-			{
-				const string NoClobberPrefix = "Can't clobber writable file ";
-				if (Response.Succeeded)
-				{
-					SyncOutput(Response.Data);
-				}
-				else if (Response.Error != null && Response.Error.Generic == PerforceGenericCode.Client && Response.Error.Data.StartsWith(NoClobberPrefix, StringComparison.OrdinalIgnoreCase))
-				{
-					TamperedFiles.Add(Response.Error.Data.Substring(NoClobberPrefix.Length).Trim());
-				}
-				else
-				{
-					return (WorkspaceUpdateResult.FailedToSync, "Aborted sync due to errors.");
-				}
-			}
-
-			// If any files need to be clobbered, defer to the main thread to figure out which ones
-			if (TamperedFiles.Count > 0)
-			{
-				int NumNewFilesToClobber = 0;
-				foreach (string TamperedFile in TamperedFiles)
-				{
-					if (!Context.ClobberFiles.ContainsKey(TamperedFile))
-					{
-						Context.ClobberFiles[TamperedFile] = true;
-						if (TamperedFile.EndsWith(LocalObjectVersionFileName, StringComparison.OrdinalIgnoreCase) || TamperedFile.EndsWith(LocalVersionHeaderFileName, StringComparison.OrdinalIgnoreCase))
-						{
-							// Hack for UseFastModularVersioningV2; we don't need to update these files any more.
-							continue;
-						}
-						NumNewFilesToClobber++;
-					}
-				}
-				if (NumNewFilesToClobber > 0)
-				{
-					return (WorkspaceUpdateResult.FilesToClobber, $"Cancelled sync after checking files to clobber ({NumNewFilesToClobber} new files).");
-				}
-				foreach (string TamperedFile in TamperedFiles)
-				{
-					if (Context.ClobberFiles[TamperedFile])
-					{
-						List<PerforceResponse<SyncRecord>> Response = await Perforce.TrySyncAsync(SyncOptions.Force, -1, TamperedFile, CancellationToken).ToListAsync();
-						if (!Response.Succeeded())
-						{
-							return (WorkspaceUpdateResult.FailedToSync, $"Couldn't sync {TamperedFile}.");
-						}
-					}
-				}
-			}
-
-			// All succeeded
-			return (WorkspaceUpdateResult.Success, "Succeeded.");
-		}
-
-		public static async Task<ConfigFile> ReadProjectConfigFile(DirectoryReference LocalRootPath, FileReference SelectedLocalFileName, ILogger Logger)
-		{
-			// Find the valid config file paths
-			DirectoryInfo EngineDir = DirectoryReference.Combine(LocalRootPath, "Engine").ToDirectoryInfo();
-			List<FileInfo> LocalConfigFiles = Utility.GetLocalConfigPaths(EngineDir, SelectedLocalFileName.ToFileInfo());
-
-			// Read them in
-			ConfigFile ProjectConfig = new ConfigFile();
-			foreach (FileInfo LocalConfigFile in LocalConfigFiles)
-			{
-				try
-				{
-					string[] Lines = await File.ReadAllLinesAsync(LocalConfigFile.FullName);
-					ProjectConfig.Parse(Lines);
-					Logger.LogInformation("Read config file from {FileName}", LocalConfigFile.FullName);
-				}
-				catch (Exception Ex)
-				{
-					Logger.LogInformation(Ex, "Failed to read config file from {FileName}", LocalConfigFile.FullName);
-				}
-			}
-			return ProjectConfig;
-		}
-
-		public static async Task<IReadOnlyList<string>?> ReadProjectStreamFilter(IPerforceConnection Perforce, ConfigFile ProjectConfigFile, ILogger Logger, CancellationToken CancellationToken)
-		{
-			string? StreamListDepotPath = ProjectConfigFile.GetValue("Options.QuickSelectStreamList", null);
-			if (StreamListDepotPath == null)
-			{
-				return null;
-			}
-
-			PerforceResponse<PrintRecord<string[]>> Response = await Perforce.TryPrintLinesAsync(StreamListDepotPath, CancellationToken);
-			if (!Response.Succeeded)
-			{
-				return null;
-			}
-
-			return Response.Data.Contents.Select(x => x.Trim()).Where(x => x.Length > 0).ToList().AsReadOnly();
-		}
-
-		static string FormatTime(long Seconds)
-		{
-			if (Seconds >= 60)
-			{
-				return String.Format("{0,3}m {1:00}s", Seconds / 60, Seconds % 60);
-			}
-			else
-			{
-				return String.Format("     {0,2}s", Seconds);
-			}
-		}
-
-		static async Task<bool> HasModifiedSourceFiles(IPerforceConnection Perforce, ProjectInfo Project, CancellationToken CancellationToken)
-		{
-			List<OpenedRecord> OpenFiles = await Perforce.OpenedAsync(OpenedOptions.None, Project.ClientRootPath + "/...", CancellationToken).ToListAsync(CancellationToken);
-			if (OpenFiles.Any(x => x.DepotFile.IndexOf("/Source/", StringComparison.OrdinalIgnoreCase) != -1))
-			{
-				return true;
-			}
-			return false;
-		}
-
-		static async Task<List<FStatRecord>> FindUnresolvedFiles(IPerforceConnection Perforce, IEnumerable<string> SyncPaths, CancellationToken CancellationToken)
-		{
-			List<FStatRecord> UnresolvedFiles = new List<FStatRecord>();
-			foreach (string SyncPath in SyncPaths)
-			{
-				List<FStatRecord> Records = await Perforce.FStatAsync(FStatOptions.OnlyUnresolved, SyncPath, CancellationToken).ToListAsync(CancellationToken);
-				UnresolvedFiles.AddRange(Records);
-			}
-			return UnresolvedFiles;
-		}
-
-		static Task<bool> UpdateVersionFile(IPerforceConnection Perforce, string ClientPath, Dictionary<string, string> VersionStrings, int ChangeNumber, ILogger Logger, CancellationToken CancellationToken)
-		{
-			return UpdateVersionFile(Perforce, ClientPath, ChangeNumber, Text => UpdateVersionStrings(Text, VersionStrings), Logger, CancellationToken);
-		}
-
-		static async Task<bool> UpdateVersionFile(IPerforceConnection Perforce, string ClientPath, int ChangeNumber, Func<string, string> Update, ILogger Logger, CancellationToken CancellationToken)
-		{
-			List<PerforceResponse<FStatRecord>> Records = await Perforce.TryFStatAsync(FStatOptions.None, ClientPath, CancellationToken).ToListAsync(CancellationToken);
-			if (!Records.Succeeded())
-			{
-				Logger.LogInformation("Failed to query records for {ClientPath}", ClientPath);
-				return false;
-			}
-			if (Records.Count > 1)
-			{
-				// Attempt to remove any existing file which is synced
-				await Perforce.SyncAsync(SyncOptions.Force, -1, $"{ClientPath}#0", CancellationToken).ToListAsync(CancellationToken);
-
-				// Try to get the mapped files again
-				Records = await Perforce.TryFStatAsync(FStatOptions.None, ClientPath, CancellationToken).ToListAsync(CancellationToken);
-				if (!Records.Succeeded())
-				{
-					Logger.LogInformation("Failed to query records for {ClientPath}", ClientPath);
-					return false;
-				}
-			}
-			if (Records.Count == 0)
-			{
-				Logger.LogInformation("Ignoring {ClientPath}; not found on server.", ClientPath);
-				return true;
-			}
-
-			FStatRecord Record = Records[0].Data;
-			string? LocalPath = Record.ClientFile; // Actually a filesystem path
-			if (LocalPath == null)
-			{
-				Logger.LogInformation("Version file is not mapped to workspace ({ClientFile})", ClientPath);
-				return false;
-			}
-			string? DepotPath = Record.DepotFile;
-			if (DepotPath == null)
-			{
-				Logger.LogInformation("Version file does not exist in depot ({ClientFile})", ClientPath);
-				return false;
-			}
-
-			PerforceResponse<PrintRecord<string[]>> Response = await Perforce.TryPrintLinesAsync($"{DepotPath}@{ChangeNumber}", CancellationToken);
-			if (!Response.Succeeded)
-			{
-				Logger.LogInformation("Couldn't get default contents of {DepotPath}", DepotPath);
-				return false;
-			}
-
-			string[]? Contents = Response.Data.Contents;
-			if (Contents == null)
-			{
-				Logger.LogInformation("No data returned for {DepotPath}", DepotPath);
-				return false;
-			}
-
-			string Text = String.Join("\n", Contents);
-			Text = Update(Text);
-			return await WriteVersionFile(Perforce, LocalPath, DepotPath, Text, Logger, CancellationToken);
-		}
-
-		static string UpdateVersionStrings(string Text, Dictionary<string, string> VersionStrings)
-		{
-			StringWriter Writer = new StringWriter();
-			foreach (string Line in Text.Split('\n'))
-			{
-				string NewLine = Line;
-				foreach (KeyValuePair<string, string> VersionString in VersionStrings)
-				{
-					if (UpdateVersionLine(ref NewLine, VersionString.Key, VersionString.Value))
-					{
-						break;
-					}
-				}
-				Writer.WriteLine(NewLine);
-			}
-			return Writer.ToString();
-		}
-
-		static string UpdateBuildVersion(string Text, int Changelist, int CodeChangelist, string BranchOrStreamName, bool bIsLicenseeVersion)
-		{
-			Dictionary<string, object> Object = JsonSerializer.Deserialize<Dictionary<string, object>>(Text, Program.DefaultJsonSerializerOptions);
-
-			int PrevCompatibleChangelist = 0;
-			if (Object.TryGetValue("CompatibleChangelist", out object? PrevCompatibleChangelistObj))
-			{
-				int.TryParse(PrevCompatibleChangelistObj?.ToString(), out PrevCompatibleChangelist);
-			}
-
-			int PrevIsLicenseeVersion = 0;
-			if (Object.TryGetValue("IsLicenseeVersion", out object? PrevIsLicenseeVersionObj))
-			{
-				int.TryParse(PrevIsLicenseeVersionObj?.ToString(), out PrevIsLicenseeVersion);
-			}
-
-			Object["Changelist"] = Changelist;
-			if (PrevCompatibleChangelist == 0 || (PrevIsLicenseeVersion != 0) != bIsLicenseeVersion)
-			{
-				// Don't overwrite the compatible changelist if we're in a hotfix release
-				Object["CompatibleChangelist"] = CodeChangelist;
-			}
-			Object["BranchName"] = BranchOrStreamName.Replace('/', '+');
-			Object["IsPromotedBuild"] = 0;
-			Object["IsLicenseeVersion"] = bIsLicenseeVersion ? 1 : 0;
-
-			return JsonSerializer.Serialize(Object, new JsonSerializerOptions
-			{
-				WriteIndented = true,
-				// do not escape +
-				Encoder = JavaScriptEncoder.UnsafeRelaxedJsonEscaping
-			});
-		}
-
-		static async Task<bool> WriteVersionFile(IPerforceConnection Perforce, string LocalPath, string DepotPath, string NewText, ILogger Logger, CancellationToken CancellationToken)
-		{
-			try
-			{
-				if (File.Exists(LocalPath) && File.ReadAllText(LocalPath) == NewText)
-				{
-					Logger.LogInformation("Ignored {FileName}; contents haven't changed", LocalPath);
-				}
-				else
-				{
-					Directory.CreateDirectory(Path.GetDirectoryName(LocalPath));
-					Utility.ForceDeleteFile(LocalPath);
-					if (DepotPath != null)
-					{
-						await Perforce.SyncAsync(DepotPath + "#0", CancellationToken).ToListAsync(CancellationToken);
-					}
-					File.WriteAllText(LocalPath, NewText);
-					Logger.LogInformation("Written {FileName}", LocalPath);
-				}
-				return true;
-			}
-			catch (Exception Ex)
-			{
-				Logger.LogError(Ex, "Failed to write to {FileName}.", LocalPath);
-				return false;
-			}
-		}
-
-		static bool UpdateVersionLine(ref string Line, string Prefix, string Suffix)
-		{
-			int LineIdx = 0;
-			int PrefixIdx = 0;
-			for (; ; )
-			{
-				string? PrefixToken = ReadToken(Prefix, ref PrefixIdx);
-				if (PrefixToken == null)
-				{
-					break;
-				}
-
-				string? LineToken = ReadToken(Line, ref LineIdx);
-				if (LineToken == null || LineToken != PrefixToken)
-				{
-					return false;
-				}
-			}
-			Line = Line.Substring(0, LineIdx) + Suffix;
-			return true;
-		}
-
-		static string? ReadToken(string Line, ref int LineIdx)
-		{
-			for (; ; LineIdx++)
-			{
-				if (LineIdx == Line.Length)
-				{
-					return null;
-				}
-				else if (!Char.IsWhiteSpace(Line[LineIdx]))
-				{
-					break;
-				}
-			}
-
-			int StartIdx = LineIdx++;
-			if (Char.IsLetterOrDigit(Line[StartIdx]) || Line[StartIdx] == '_')
-			{
-				while (LineIdx < Line.Length && (Char.IsLetterOrDigit(Line[LineIdx]) || Line[LineIdx] == '_'))
-				{
-					LineIdx++;
-				}
-			}
-
-			return Line.Substring(StartIdx, LineIdx - StartIdx);
-		}
-
-		public Tuple<string, float> CurrentProgress
-		{
-			get { return Progress.Current; }
-		}
-	}
 }
diff --git a/Engine/Source/Programs/UnrealGameSync/UnrealGameSync/WorkspaceUpdate.cs b/Engine/Source/Programs/UnrealGameSync/UnrealGameSync/WorkspaceUpdate.cs
new file mode 100644
--- /dev/null
+++ b/Engine/Source/Programs/UnrealGameSync/UnrealGameSync/WorkspaceUpdate.cs
@@ -0,0 +1,1825 @@
+// Copyright Epic Games, Inc. All Rights Reserved.
+
+using EpicGames.Core;
+using EpicGames.Perforce;
+using Microsoft.Extensions.DependencyInjection;
+using Microsoft.Extensions.Logging;
+using System;
+using System.Collections.Concurrent;
+using System.Collections.Generic;
+using System.Diagnostics;
+using System.Diagnostics.CodeAnalysis;
+using System.Globalization;
+using System.IO;
+using System.IO.Compression;
+using System.Linq;
+using System.Security.Cryptography;
+using System.Text;
+using System.Text.Encodings.Web;
+using System.Text.Json;
+using System.Threading;
+using System.Threading.Tasks;
+
+namespace UnrealGameSync
+{
+	[Flags]
+	public enum WorkspaceUpdateOptions
+	{
+		Sync = 0x01,
+		SyncSingleChange = 0x02,
+		AutoResolveChanges = 0x04,
+		GenerateProjectFiles = 0x08,
+		SyncArchives = 0x10,
+		Build = 0x20,
+		Clean = 0x40,
+		ScheduledBuild = 0x80,
+		RunAfterSync = 0x100,
+		OpenSolutionAfterSync = 0x200,
+		ContentOnly = 0x400,
+		UpdateFilter = 0x800,
+		SyncAllProjects = 0x1000,
+		IncludeAllProjectsInSolution = 0x2000,
+	}
+
+	public enum WorkspaceUpdateResult
+	{
+		Canceled,
+		FailedToSync,
+		FailedToSyncLoginExpired,
+		FilesToDelete,
+		FilesToResolve,
+		FilesToClobber,
+		FailedToCompile,
+		FailedToCompileWithCleanWorkspace,
+		Success,
+	}
+
+	public class PerforceSyncOptions
+	{
+		public const int DefaultNumRetries = 0;
+		public const int DefaultNumThreads = 2;
+		public const int DefaultTcpBufferSize = 0;
+		public const int DefaultFileBufferSize = 0;
+
+		public const int DefaultMaxCommandsPerBatch = 200;
+		public const int DefaultMaxSizePerBatch = 128 * 1024 * 1024;
+
+		public int NumRetries = DefaultNumRetries;
+		public int NumThreads = DefaultNumThreads;
+		public int TcpBufferSize = DefaultTcpBufferSize;
+		public int FileBufferSize = DefaultFileBufferSize;
+
+		public int MaxCommandsPerBatch = DefaultMaxCommandsPerBatch;
+		public int MaxSizePerBatch = DefaultMaxSizePerBatch;
+
+
+		public PerforceSyncOptions Clone()
+		{
+			return (PerforceSyncOptions)MemberwiseClone();
+		}
+	}
+
+	public interface IArchiveInfo
+	{
+		string Name { get; }
+		string Type { get; }
+		string BasePath { get; }
+		string? Target { get; }
+		bool Exists();
+		bool TryGetArchiveKeyForChangeNumber(int ChangeNumber, [NotNullWhen(true)] out string? ArchiveKey);
+		Task<bool> DownloadArchive(IPerforceConnection Perforce, string ArchiveKey, DirectoryReference LocalRootPath, FileReference ManifestFileName, ILogger Logger, ProgressValue Progress, CancellationToken CancellationToken);
+	}
+
+	public class WorkspaceUpdateContext
+	{
+		public DateTime StartTime = DateTime.UtcNow;
+		public int ChangeNumber;
+		public WorkspaceUpdateOptions Options;
+		public string[]? SyncFilter;
+		public Dictionary<string, Tuple<IArchiveInfo, string>?> ArchiveTypeToArchive = new Dictionary<string, Tuple<IArchiveInfo, string>?>();
+		public Dictionary<string, bool> DeleteFiles = new Dictionary<string,bool>();
+		public Dictionary<string, bool> ClobberFiles = new Dictionary<string,bool>();
+		public Dictionary<Guid,ConfigObject> DefaultBuildSteps;
+		public List<ConfigObject> UserBuildStepObjects;
+		public HashSet<Guid>? CustomBuildSteps;
+		public Dictionary<string, string> Variables;
+		public PerforceSyncOptions? PerforceSyncOptions;
+		public List<HaveRecord>? HaveFiles; // Cached when sync filter has changed
+
+		// May be updated during sync
+		public ConfigFile? ProjectConfigFile;
+		public IReadOnlyList<string>? ProjectStreamFilter;
+
+		public WorkspaceUpdateContext(int InChangeNumber, WorkspaceUpdateOptions InOptions, string[]? InSyncFilter, Dictionary<Guid, ConfigObject> InDefaultBuildSteps, List<ConfigObject> InUserBuildSteps, HashSet<Guid>? InCustomBuildSteps, Dictionary<string, string> InVariables)
+		{
+			ChangeNumber = InChangeNumber;
+			Options = InOptions;
+			SyncFilter = InSyncFilter;
+			DefaultBuildSteps = InDefaultBuildSteps;
+			UserBuildStepObjects = InUserBuildSteps;
+			CustomBuildSteps = InCustomBuildSteps;
+			Variables = InVariables;
+		}
+	}
+
+	class WorkspaceSyncCategory
+	{
+		public Guid UniqueId;
+		public bool bEnable;
+		public string Name;
+		public string[] Paths;
+		public bool bHidden;
+		public Guid[] Requires;
+
+		public WorkspaceSyncCategory(Guid UniqueId) : this(UniqueId, "Unnamed")
+		{
+		}
+
+		public WorkspaceSyncCategory(Guid UniqueId, string Name, params string[] Paths)
+		{
+			this.UniqueId = UniqueId;
+			this.bEnable = true;
+			this.Name = Name;
+			this.Paths = Paths;
+			this.Requires = new Guid[0];
+		}
+
+		public static Dictionary<Guid, bool> GetDefault(IEnumerable<WorkspaceSyncCategory> Categories)
+		{
+			return Categories.ToDictionary(x => x.UniqueId, x => x.bEnable);
+		}
+
+		public static Dictionary<Guid, bool> GetDelta(Dictionary<Guid, bool> Source, Dictionary<Guid, bool> Target)
+		{
+			Dictionary<Guid, bool> Changes = new Dictionary<Guid, bool>();
+			foreach (KeyValuePair<Guid, bool> Pair in Target)
+			{
+				bool bValue;
+				if (!Source.TryGetValue(Pair.Key, out bValue) || bValue != Pair.Value)
+				{
+					Changes[Pair.Key] = Pair.Value;
+				}
+			}
+			return Changes;
+		}
+
+		public static void ApplyDelta(Dictionary<Guid, bool> Categories, Dictionary<Guid, bool> Delta)
+		{
+			foreach(KeyValuePair<Guid, bool> Pair in Delta)
+			{
+				Categories[Pair.Key] = Pair.Value;
+			}
+		}
+
+		public override string ToString()
+		{
+			return Name;
+		}
+	}
+
+	public class ProjectInfo
+	{
+		public DirectoryReference LocalRootPath { get; }
+		public FileReference LocalFileName { get; }
+		public string ClientRootPath { get; }
+		public string ClientFileName { get; }
+		public string ProjectIdentifier { get; }
+		public string TelemetryProjectIdentifier { get; }
+		public bool bIsEnterpriseProject { get; }
+
+		public ProjectInfo(DirectoryReference LocalRootPath, FileReference LocalFileName, string ClientRootPath, string ClientFileName, string ProjectIdentifier, bool bIsEnterpriseProject)
+		{
+			this.LocalRootPath = LocalRootPath;
+			this.LocalFileName = LocalFileName;
+			this.ClientRootPath = ClientRootPath;
+			this.ClientFileName = ClientFileName;
+			this.ProjectIdentifier = ProjectIdentifier;
+			this.TelemetryProjectIdentifier = PerforceUtils.GetClientOrDepotDirectoryName(ProjectIdentifier);
+			this.bIsEnterpriseProject = bIsEnterpriseProject;
+		}
+	}
+
+	class WorkspaceUpdate : IAsyncDisposable
+	{
+		const string BuildVersionFileName = "/Engine/Build/Build.version";
+		const string VersionHeaderFileName = "/Engine/Source/Runtime/Launch/Resources/Version.h";
+		const string ObjectVersionFileName = "/Engine/Source/Runtime/Core/Private/UObject/ObjectVersion.cpp";
+
+		static readonly string LocalVersionHeaderFileName = VersionHeaderFileName.Replace('/', '\\');
+		static readonly string LocalObjectVersionFileName = ObjectVersionFileName.Replace('/', '\\');
+
+		static SemaphoreSlim UpdateSemaphore = new SemaphoreSlim(1);
+
+		class RecordCounter : IDisposable
+		{
+			ProgressValue Progress;
+			string Message;
+			int Count;
+			Stopwatch Timer = Stopwatch.StartNew();
+
+			public RecordCounter(ProgressValue Progress, string Message)
+			{
+				this.Progress = Progress;
+				this.Message = Message;
+
+				Progress.Set(Message);
+			}
+
+			public void Dispose()
+			{
+				UpdateMessage();
+			}
+
+			public void Increment()
+			{
+				Count++;
+				if (Timer.ElapsedMilliseconds > 250)
+				{
+					UpdateMessage();
+				}
+			}
+
+			public void UpdateMessage()
+			{
+				Progress.Set(String.Format("{0} ({1:N0})", Message, Count));
+				Timer.Restart();
+			}
+		}
+
+		class SyncBatchBuilder
+		{
+			public int MaxCommandsPerList { get; }
+			public long MaxSizePerList { get; }
+			public Queue<List<string>> Batches { get; }
+
+			List<string>? Commands;
+			List<string>? DeleteCommands;
+			long Size;
+
+			public SyncBatchBuilder(int? MaxCommandsPerList, long? MaxSizePerList)
+			{
+				this.MaxCommandsPerList = MaxCommandsPerList ?? PerforceSyncOptions.DefaultMaxCommandsPerBatch;
+				this.MaxSizePerList = MaxSizePerList ?? PerforceSyncOptions.DefaultMaxSizePerBatch;
+				this.Batches = new Queue<List<string>>();
+			}
+
+			public void Add(string NewCommand, long NewSize)
+			{
+				if (NewSize == 0)
+				{
+					if (DeleteCommands == null || DeleteCommands.Count >= MaxCommandsPerList)
+					{
+						DeleteCommands = new List<string>();
+						Batches.Enqueue(DeleteCommands);
+					}
+
+					DeleteCommands.Add(NewCommand);
+				}
+				else
+				{
+					if (Commands == null || Commands.Count >= MaxCommandsPerList || Size + NewSize >= MaxSizePerList)
+					{
+						Commands = new List<string>();
+						Batches.Enqueue(Commands);
+						Size = 0;
+					}
+
+					Commands.Add(NewCommand);
+					Size += NewSize;
+				}
+			}
+		}
+
+		class SyncTree
+		{
+			public bool bCanUseWildcard;
+			public int TotalIncludedFiles;
+			public long TotalSize;
+			public int TotalExcludedFiles;
+			public Dictionary<string, long> IncludedFiles = new Dictionary<string, long>();
+			public Dictionary<string, SyncTree> NameToSubTree = new Dictionary<string, SyncTree>(StringComparer.OrdinalIgnoreCase);
+
+			public SyncTree(bool bCanUseWildcard)
+			{
+				this.bCanUseWildcard = bCanUseWildcard;
+			}
+
+			public SyncTree FindOrAddSubTree(string Name)
+			{
+				SyncTree? Result;
+				if (!NameToSubTree.TryGetValue(Name, out Result))
+				{
+					Result = new SyncTree(bCanUseWildcard);
+					NameToSubTree.Add(Name, Result);
+				}
+				return Result;
+			}
+
+			public void IncludeFile(string Path, long Size)
+			{
+				int Idx = Path.IndexOf('/');
+				if (Idx == -1)
+				{
+					IncludedFiles.Add(Path, Size);
+				}
+				else
+				{
+					SyncTree SubTree = FindOrAddSubTree(Path.Substring(0, Idx));
+					SubTree.IncludeFile(Path.Substring(Idx + 1), Size);
+				}
+				TotalIncludedFiles++;
+				TotalSize += Size;
+			}
+
+			public void ExcludeFile(string Path)
+			{
+				int Idx = Path.IndexOf('/');
+				if (Idx != -1)
+				{
+					SyncTree SubTree = FindOrAddSubTree(Path.Substring(0, Idx));
+					SubTree.ExcludeFile(Path.Substring(Idx + 1));
+				}
+				TotalExcludedFiles++;
+			}
+
+			public void GetOptimizedSyncCommands(string Prefix, int ChangeNumber, SyncBatchBuilder Builder)
+			{
+				if (bCanUseWildcard && TotalExcludedFiles == 0 && TotalSize < Builder.MaxSizePerList)
+				{
+					Builder.Add(String.Format("{0}/...@{1}", Prefix, ChangeNumber), TotalSize);
+				}
+				else
+				{
+					foreach (KeyValuePair<string, long> File in IncludedFiles)
+					{
+						Builder.Add(String.Format("{0}/{1}@{2}", Prefix, File.Key, ChangeNumber), File.Value);
+					}
+					foreach (KeyValuePair<string, SyncTree> Pair in NameToSubTree)
+					{
+						Pair.Value.GetOptimizedSyncCommands(String.Format("{0}/{1}", Prefix, PerforceUtils.EscapePath(Pair.Key)), ChangeNumber, Builder);
+					}
+				}
+			}
+		}
+
+		CancellationTokenSource CancellationSource;
+		public WorkspaceUpdateContext Context { get; }
+		public ProgressValue Progress { get; } = new ProgressValue();
+		public Task UpdateTask { get; }
+
+		public WorkspaceUpdate(WorkspaceUpdateContext Context, Func<WorkspaceUpdate, CancellationToken, Task> TaskFunc)
+		{
+			this.CancellationSource = new CancellationTokenSource();
+			this.Context = Context;
+			this.UpdateTask = Task.Run(() => TaskFunc(this, CancellationSource.Token));
+		}
+
+		public void Cancel()
+		{
+			CancellationSource.Cancel();
+		}
+
+		public ValueTask DisposeAsync()
+		{
+			Task DisposeTask = Task.CompletedTask;
+			if (UpdateTask != null)
+			{
+				CancellationSource.Cancel();
+				DisposeTask = UpdateTask.ContinueWith(Task => CancellationSource.Dispose());
+			}
+			return new ValueTask(DisposeTask);
+		}
+
+		class SyncFile
+		{
+			public string DepotFile;
+			public string RelativePath;
+			public long Size;
+
+			public SyncFile(string DepotFile, string RelativePath, long Size)
+			{
+				this.DepotFile = DepotFile;
+				this.RelativePath = RelativePath;
+				this.Size = Size;
+			}
+		};
+
+		class SemaphoreScope : IDisposable
+		{
+			SemaphoreSlim Semaphore;
+			public bool HasLock { get; private set; }
+
+			public SemaphoreScope(SemaphoreSlim Semaphore)
+			{
+				this.Semaphore = Semaphore;
+			}
+
+			public bool TryAcquire()
+			{
+				HasLock = HasLock || Semaphore.Wait(0);
+				return HasLock;
+			}
+
+			public async Task AcquireAsync(CancellationToken CancellationToken)
+			{
+				await Semaphore.WaitAsync(CancellationToken);
+				HasLock = true;
+			}
+
+			public void Release()
+			{
+				if (HasLock)
+				{
+					Semaphore.Release();
+					HasLock = false;
+				}
+			}
+
+			public void Dispose() => Release();
+		}
+
+		public async Task<(WorkspaceUpdateResult, string)> ExecuteAsync(IPerforceSettings PerforceSettings, ProjectInfo Project, UserWorkspaceState State, ILogger Logger, CancellationToken CancellationToken)
+		{
+			using IPerforceConnection Perforce = await PerforceConnection.CreateAsync(PerforceSettings, Logger);
+
+			string CmdExe = Path.Combine(Environment.GetFolderPath(Environment.SpecialFolder.System), "cmd.exe");
+			if (!File.Exists(CmdExe))
+			{
+				return (WorkspaceUpdateResult.FailedToSync, $"Missing {CmdExe}.");
+			}
+
+			List<Tuple<string, TimeSpan>> Times = new List<Tuple<string, TimeSpan>>();
+
+			int NumFilesSynced = 0;
+			if (Context.Options.HasFlag(WorkspaceUpdateOptions.Sync) || Context.Options.HasFlag(WorkspaceUpdateOptions.SyncSingleChange))
+			{
+				using (TelemetryStopwatch SyncTelemetryStopwatch = new TelemetryStopwatch("Workspace_Sync", Project.TelemetryProjectIdentifier))
+				{
+					Logger.LogInformation("Syncing to {Change}...", Context.ChangeNumber);
+
+					// Make sure we're logged in
+					PerforceResponse<LoginRecord> LoginResponse = await Perforce.TryGetLoginStateAsync(CancellationToken);
+					if (!LoginResponse.Succeeded)
+					{
+						return (WorkspaceUpdateResult.FailedToSyncLoginExpired, "User is not logged in.");
+					}
+
+					// Figure out which paths to sync
+					List<string> RelativeSyncPaths = GetRelativeSyncPaths(Project, (Context.Options & WorkspaceUpdateOptions.SyncAllProjects) != 0, Context.SyncFilter);
+					List<string> SyncPaths = new List<string>(RelativeSyncPaths.Select(x => Project.ClientRootPath + x));
+
+					// Get the user's sync filter
+					FileFilter UserFilter = new FileFilter(FileFilterType.Include);
+					if (Context.SyncFilter != null)
+					{
+						UserFilter.AddRules(Context.SyncFilter.Select(x => x.Trim()).Where(x => x.Length > 0 && !x.StartsWith(";") && !x.StartsWith("#")));
+					}
+
+					// Check if the new sync filter matches the previous one. If not, we'll enumerate all files in the workspace and make sure there's nothing extra there.
+					string? NextSyncFilterHash = null;
+					using (SHA1Managed SHA = new SHA1Managed())
+					{
+						StringBuilder CombinedFilter = new StringBuilder();
+						foreach (string RelativeSyncPath in RelativeSyncPaths)
+						{
+							CombinedFilter.AppendFormat("{0}\n", RelativeSyncPath);
+						}
+						if (Context.SyncFilter != null)
+						{
+							CombinedFilter.Append("--FROM--\n");
+							CombinedFilter.Append(String.Join("\n", Context.SyncFilter));
+						}
+						NextSyncFilterHash = BitConverter.ToString(SHA.ComputeHash(Encoding.UTF8.GetBytes(CombinedFilter.ToString()))).Replace("-", "");
+					}
+
+					// If the hash differs, enumerate everything in the workspace to find what needs to be removed
+					if (NextSyncFilterHash != State.CurrentSyncFilterHash)
+					{
+						using (TelemetryStopwatch FilterStopwatch = new TelemetryStopwatch("Workspace_Sync_FilterChanged", Project.TelemetryProjectIdentifier))
+						{
+							Logger.LogInformation("Filter has changed ({PrevHash} -> {NextHash}); finding files in workspace that need to be removed.", (String.IsNullOrEmpty(State.CurrentSyncFilterHash)) ? "None" : State.CurrentSyncFilterHash, NextSyncFilterHash);
+
+							// Find all the files that are in this workspace
+							List<HaveRecord>? HaveFiles = Context.HaveFiles;
+							if (HaveFiles == null)
+							{
+								HaveFiles = new List<HaveRecord>();
+								using (RecordCounter HaveCounter = new RecordCounter(Progress, "Sync filter changed; checking workspace..."))
+								{
+									await foreach (PerforceResponse<HaveRecord> Record in Perforce.TryHaveAsync(FileSpecList.Any, CancellationToken))
+									{
+										if (Record.Succeeded)
+										{
+											HaveFiles.Add(Record.Data);
+											HaveCounter.Increment();
+										}
+										else
+										{
+											return (WorkspaceUpdateResult.FailedToSync, $"Unable to query files ({Record}).");
+										}
+									}
+								}
+								Context.HaveFiles = HaveFiles;
+							}
+
+							// Build a filter for the current sync paths
+							FileFilter SyncPathsFilter = new FileFilter(FileFilterType.Exclude);
+							foreach (string RelativeSyncPath in RelativeSyncPaths)
+							{
+								SyncPathsFilter.Include(RelativeSyncPath);
+							}
+
+							// Remove all the files that are not included by the filter
+							List<string> RemoveDepotPaths = new List<string>();
+							foreach (HaveRecord HaveFile in HaveFiles)
+							{
+								try
+								{
+									FileReference FullPath = new FileReference(HaveFile.Path);
+									if (MatchFilter(Project, FullPath, SyncPathsFilter) && !MatchFilter(Project, FullPath, UserFilter))
+									{
+										Logger.LogInformation("  {DepotFile}", HaveFile.DepotFile);
+										RemoveDepotPaths.Add(HaveFile.DepotFile);
+									}
+								}
+								catch (PathTooLongException)
+								{
+									// We don't actually care about this when looking for files to remove. Perforce may think that it's synced the path, and silently failed. Just ignore it.
+								}
+							}
+
+							// Check if there are any paths outside the regular sync paths
+							if (RemoveDepotPaths.Count > 0)
+							{
+								bool bDeleteListMatches = true;
+
+								Dictionary<string, bool> NewDeleteFiles = new Dictionary<string, bool>(StringComparer.OrdinalIgnoreCase);
+								foreach (string RemoveDepotPath in RemoveDepotPaths)
+								{
+									bool bDelete;
+									if (!Context.DeleteFiles.TryGetValue(RemoveDepotPath, out bDelete))
+									{
+										bDeleteListMatches = false;
+										bDelete = true;
+									}
+									NewDeleteFiles[RemoveDepotPath] = bDelete;
+								}
+								Context.DeleteFiles = NewDeleteFiles;
+
+								if (!bDeleteListMatches)
+								{
+									return (WorkspaceUpdateResult.FilesToDelete, $"Cancelled after finding {NewDeleteFiles.Count} files excluded by filter");
+								}
+
+								RemoveDepotPaths.RemoveAll(x => !Context.DeleteFiles[x]);
+							}
+
+							// Actually delete any files that we don't want
+							if (RemoveDepotPaths.Count > 0)
+							{
+								// Clear the current sync filter hash. If the sync is canceled, we'll be in an indeterminate state, and we should always clean next time round.
+								State.CurrentSyncFilterHash = "INVALID";
+								State.Save();
+
+								// Find all the depot paths that will be synced
+								HashSet<string> RemainingDepotPathsToRemove = new HashSet<string>(StringComparer.OrdinalIgnoreCase);
+								RemainingDepotPathsToRemove.UnionWith(RemoveDepotPaths);
+
+								// Build the list of revisions to sync
+								List<string> RevisionsToRemove = new List<string>();
+								RevisionsToRemove.AddRange(RemoveDepotPaths.Select(x => String.Format("{0}#0", x)));
+
+								(WorkspaceUpdateResult, string) RemoveResult = await SyncFileRevisions(Perforce, "Removing files...", Context, RevisionsToRemove, RemainingDepotPathsToRemove, Progress, Logger, CancellationToken);
+								if (RemoveResult.Item1 != WorkspaceUpdateResult.Success)
+								{
+									return RemoveResult;
+								}
+							}
+
+							// Update the sync filter hash. We've removed any files we need to at this point.
+							State.CurrentSyncFilterHash = NextSyncFilterHash;
+							State.Save();
+						}
+					}
+
+					// Create a filter for all the files we don't want
+					FileFilter Filter = new FileFilter(UserFilter);
+					Filter.Exclude(BuildVersionFileName);
+					if (Context.Options.HasFlag(WorkspaceUpdateOptions.ContentOnly))
+					{
+						Filter.Exclude("*.usf");
+						Filter.Exclude("*.ush");
+					}
+
+					// Create a tree to store the sync path
+					SyncTree SyncTree = new SyncTree(false);
+					if (!Context.Options.HasFlag(WorkspaceUpdateOptions.SyncSingleChange))
+					{
+						foreach (string RelativeSyncPath in RelativeSyncPaths)
+						{
+							const string WildcardSuffix = "/...";
+							if (RelativeSyncPath.EndsWith(WildcardSuffix, StringComparison.Ordinal))
+							{
+								SyncTree Leaf = SyncTree;
+
+								string[] Fragments = RelativeSyncPath.Split('/');
+								for (int Idx = 1; Idx < Fragments.Length - 1; Idx++)
+								{
+									Leaf = Leaf.FindOrAddSubTree(Fragments[Idx]);
+								}
+
+								Leaf.bCanUseWildcard = true;
+							}
+						}
+					}
+
+					// Find all the server changes, and anything that's opened for edit locally. We need to sync files we have open to schedule a resolve.
+					SyncBatchBuilder BatchBuilder = new SyncBatchBuilder(Context.PerforceSyncOptions?.MaxCommandsPerBatch, Context.PerforceSyncOptions?.MaxSizePerBatch);
+					List<string> SyncDepotPaths = new List<string>();
+					using (RecordCounter Counter = new RecordCounter(Progress, "Filtering files..."))
+					{
+						// Track the total new bytes that will be required on disk when syncing. Add an extra 100MB for padding.
+						long RequiredFreeSpace = 100 * 1024 * 1024;
+
+						foreach (string SyncPath in SyncPaths)
+						{
+							string SyncFilter = Context.Options.HasFlag(WorkspaceUpdateOptions.Sync) ? $"{Context.ChangeNumber}" : $"={Context.ChangeNumber}";
+
+							List<SyncFile> SyncFiles = new List<SyncFile>();
+							await foreach (PerforceResponse<SyncRecord> Response in Perforce.TrySyncAsync(SyncOptions.PreviewOnly, -1, $"{SyncPath}@{SyncFilter}", CancellationToken))
+							{
+								if (!Response.Succeeded)
+								{
+									return (WorkspaceUpdateResult.FailedToSync, $"Couldn't enumerate changes matching {SyncPath}.");
+								}
+
+								SyncRecord Record = Response.Data;
+
+								// Get the full local path
+								string RelativePath;
+								try
+								{
+									FileReference SyncFile = new FileReference(Record.Path);
+									RelativePath = PerforceUtils.GetClientRelativePath(Project.LocalRootPath, SyncFile);
+								}
+								catch (PathTooLongException)
+								{
+									Logger.LogInformation("The local path for {Path} exceeds the maximum allowed by Windows. Re-sync your workspace to a directory with a shorter name, or delete the file from the server.", Record.Path);
+									return (WorkspaceUpdateResult.FailedToSync, "File exceeds maximum path length allowed by Windows.");
+								}
+
+								// Create the sync record
+								long SyncSize = (Record.Action == SyncAction.Deleted) ? 0 : Record.FileSize;
+								SyncFiles.Add(new SyncFile(Record.DepotFile, RelativePath, SyncSize));
+								Counter.Increment();
+							}
+
+							// Also sync the currently open files
+							await foreach (PerforceResponse<OpenedRecord> Response in Perforce.TryOpenedAsync(OpenedOptions.None, FileSpecList.Any, CancellationToken))
+							{
+								if (!Response.Succeeded)
+								{
+									return (WorkspaceUpdateResult.FailedToSync, $"Couldn't enumerate changes matching {SyncPath}.");
+								}
+
+								OpenedRecord Record = Response.Data;
+								if (Record.Action != FileAction.Add || Record.Action != FileAction.Branch || Record.Action != FileAction.MoveAdd)
+								{
+									string RelativePath = PerforceUtils.GetClientRelativePath(Record.ClientFile);
+									SyncFiles.Add(new SyncFile(Record.DepotFile, RelativePath, 0));
+								}
+							}
+
+							// Enumerate all the files to be synced. NOTE: depotPath is escaped, whereas clientPath is not.
+							foreach (SyncFile SyncRecord in SyncFiles)
+							{
+								if (Filter.Matches(SyncRecord.RelativePath))
+								{
+									SyncTree.IncludeFile(PerforceUtils.EscapePath(SyncRecord.RelativePath), SyncRecord.Size);
+									SyncDepotPaths.Add(SyncRecord.DepotFile);
+									RequiredFreeSpace += SyncRecord.Size;
+
+									// If the file exists the required free space can be reduced as those bytes will be replaced.
+									FileInfo LocalFileInfo = FileReference.Combine(Project.LocalRootPath, SyncRecord.RelativePath).ToFileInfo();
+									if (LocalFileInfo.Exists)
+									{
+										RequiredFreeSpace -= LocalFileInfo.Length;
+									}
+								}
+								else
+								{
+									SyncTree.ExcludeFile(PerforceUtils.EscapePath(SyncRecord.RelativePath));
+								}
+							}
+						}
+
+						try
+						{
+							DirectoryInfo LocalRootInfo = Project.LocalRootPath.ToDirectoryInfo();
+							DriveInfo Drive = new DriveInfo(LocalRootInfo.Root.FullName);
+
+							if (Drive.AvailableFreeSpace < RequiredFreeSpace)
+							{
+								Logger.LogInformation("Syncing requires {RequiredSpace} which exceeds the {AvailableSpace} available free space on {Drive}.", StringUtils.FormatBytesString(RequiredFreeSpace), StringUtils.FormatBytesString(Drive.AvailableFreeSpace), Drive.Name);
+								return (WorkspaceUpdateResult.FailedToSync, "Not enough available free space.");
+							}
+						}
+						catch (SystemException)
+						{
+							Logger.LogInformation("Unable to check available free space for {RootPath}.", Project.LocalRootPath);
+						}
+					}
+					SyncTree.GetOptimizedSyncCommands(Project.ClientRootPath, Context.ChangeNumber, BatchBuilder);
+
+					// Clear the current sync changelist, in case we cancel
+					if (!Context.Options.HasFlag(WorkspaceUpdateOptions.SyncSingleChange))
+					{
+						State.CurrentChangeNumber = -1;
+						State.AdditionalChangeNumbers.Clear();
+						State.Save();
+					}
+
+					// Find all the depot paths that will be synced
+					HashSet<string> RemainingDepotPaths = new HashSet<string>(StringComparer.OrdinalIgnoreCase);
+					RemainingDepotPaths.UnionWith(SyncDepotPaths);
+
+					using (TelemetryStopwatch TransferStopwatch = new TelemetryStopwatch("Workspace_Sync_TransferFiles", Project.TelemetryProjectIdentifier))
+					{
+						TransferStopwatch.AddData(new { MachineName = Environment.MachineName, DomainName = Environment.UserDomainName, ServerAndPort = Perforce.Settings.ServerAndPort, UserName = Perforce.Settings.UserName, IncludedFiles = SyncTree.TotalIncludedFiles, ExcludedFiles = SyncTree.TotalExcludedFiles, Size = SyncTree.TotalSize, NumThreads = Context.PerforceSyncOptions?.NumThreads ?? PerforceSyncOptions.DefaultNumThreads });
+
+						(WorkspaceUpdateResult, string) SyncResult = await SyncFileRevisions(Perforce, "Syncing files...", Context, BatchBuilder.Batches, RemainingDepotPaths, Progress, Logger, CancellationToken);
+						if (SyncResult.Item1 != WorkspaceUpdateResult.Success)
+						{
+							TransferStopwatch.AddData(new { SyncResult = SyncResult.Item2.ToString(), CompletedFilesFiles = SyncDepotPaths.Count - RemainingDepotPaths.Count });
+							return SyncResult;
+						}
+
+						TransferStopwatch.Stop("Ok");
+						TransferStopwatch.AddData(new { TransferRate = SyncTree.TotalSize / Math.Max(TransferStopwatch.Elapsed.TotalSeconds, 0.0001f) });
+					}
+
+					int VersionChangeNumber = -1;
+					if (Context.Options.HasFlag(WorkspaceUpdateOptions.Sync) && !Context.Options.HasFlag(WorkspaceUpdateOptions.UpdateFilter))
+					{
+						// Read the new config file
+						Context.ProjectConfigFile = await ReadProjectConfigFile(Project.LocalRootPath, Project.LocalFileName, Logger);
+						Context.ProjectStreamFilter = await ReadProjectStreamFilter(Perforce, Context.ProjectConfigFile, Logger, CancellationToken);
+
+						// Get the branch name
+						string? BranchOrStreamName = await Perforce.GetCurrentStreamAsync(CancellationToken);
+						if (BranchOrStreamName != null)
+						{
+							// If it's a virtual stream, take the concrete parent stream instead
+							for (; ; )
+							{
+								StreamRecord StreamSpec = await Perforce.GetStreamAsync(BranchOrStreamName, false, CancellationToken);
+								if (StreamSpec.Type != "virtual" || StreamSpec.Parent == "none" || StreamSpec.Parent == null)
+								{
+									break;
+								}
+								BranchOrStreamName = StreamSpec.Parent;
+							}
+						}
+						else
+						{
+							// Otherwise use the depot path for GenerateProjectFiles.bat in the root of the workspace
+							List<WhereRecord> Files = await Perforce.WhereAsync(Project.ClientRootPath + "/GenerateProjectFiles.bat").Where(x => !x.Unmap).ToListAsync(CancellationToken);
+							if (Files.Count != 1)
+							{
+								return (WorkspaceUpdateResult.FailedToSync, $"Couldn't determine branch name for {Project.ClientFileName}.");
+							}
+							BranchOrStreamName = PerforceUtils.GetClientOrDepotDirectoryName(Files[0].DepotFile);
+						}
+
+						// Find the last code change before this changelist. For consistency in versioning between local builds and precompiled binaries, we need to use the last submitted code changelist as our version number.
+						string[] CodeFilter = new string[] { ".cs", ".h", ".cpp", ".usf", ".ush", ".uproject", ".uplugin" }.SelectMany(x => SyncPaths.Select(y => String.Format("{0}{1}@<={2}", y, x, Context.ChangeNumber))).ToArray();
+
+						PerforceResponseList<ChangesRecord> CodeChanges = await Perforce.TryGetChangesAsync(ChangesOptions.None, 1, ChangeStatus.Submitted, CodeFilter, CancellationToken);
+						if (!CodeChanges.Succeeded)
+						{
+							return (WorkspaceUpdateResult.FailedToSync, $"Couldn't determine last code changelist before CL {Context.ChangeNumber}.");
+						}
+						if (CodeChanges.Count == 0)
+						{
+							return (WorkspaceUpdateResult.FailedToSync, $"Could not find any code changes before CL {Context.ChangeNumber}.");
+						}
+
+						// Get the last code change
+						if (Context.ProjectConfigFile.GetValue("Options.VersionToLastCodeChange", true))
+						{
+							VersionChangeNumber = CodeChanges.Max(x => x.Data.Number);
+						}
+						else
+						{
+							VersionChangeNumber = Context.ChangeNumber;
+						}
+
+						// Update the version files
+						if (Context.ProjectConfigFile.GetValue("Options.UseFastModularVersioningV2", false))
+						{
+							bool bIsLicenseeVersion = await IsLicenseeVersion(Perforce, Project, Logger, CancellationToken);
+							if (!await UpdateVersionFile(Perforce, Project.ClientRootPath + BuildVersionFileName, Context.ChangeNumber, Text => UpdateBuildVersion(Text, Context.ChangeNumber, VersionChangeNumber, BranchOrStreamName, bIsLicenseeVersion), Logger, CancellationToken))
+							{
+								return (WorkspaceUpdateResult.FailedToSync, $"Failed to update {BuildVersionFileName}.");
+							}
+						}
+						else if (Context.ProjectConfigFile.GetValue("Options.UseFastModularVersioning", false))
+						{
+							bool bIsLicenseeVersion = await IsLicenseeVersion(Perforce, Project, Logger, CancellationToken);
+							if (!await UpdateVersionFile(Perforce, Project.ClientRootPath + BuildVersionFileName, Context.ChangeNumber, Text => UpdateBuildVersion(Text, Context.ChangeNumber, VersionChangeNumber, BranchOrStreamName, bIsLicenseeVersion), Logger, CancellationToken))
+							{
+								return (WorkspaceUpdateResult.FailedToSync, $"Failed to update {BuildVersionFileName}");
+							}
+
+							Dictionary<string, string> VersionHeaderStrings = new Dictionary<string, string>();
+							VersionHeaderStrings["#define ENGINE_IS_PROMOTED_BUILD"] = " (0)";
+							VersionHeaderStrings["#define BUILT_FROM_CHANGELIST"] = " 0";
+							VersionHeaderStrings["#define BRANCH_NAME"] = " \"" + BranchOrStreamName.Replace('/', '+') + "\"";
+							if (!await UpdateVersionFile(Perforce, Project.ClientRootPath + VersionHeaderFileName, VersionHeaderStrings, Context.ChangeNumber, Logger, CancellationToken))
+							{
+								return (WorkspaceUpdateResult.FailedToSync, $"Failed to update {VersionHeaderFileName}.");
+							}
+							if (!await UpdateVersionFile(Perforce, Project.ClientRootPath + ObjectVersionFileName, new Dictionary<string, string>(), Context.ChangeNumber, Logger, CancellationToken))
+							{
+								return (WorkspaceUpdateResult.FailedToSync, $"Failed to update {ObjectVersionFileName}.");
+							}
+						}
+						else
+						{
+							if (!await UpdateVersionFile(Perforce, Project.ClientRootPath + BuildVersionFileName, new Dictionary<string, string>(), Context.ChangeNumber, Logger, CancellationToken))
+							{
+								return (WorkspaceUpdateResult.FailedToSync, $"Failed to update {BuildVersionFileName}");
+							}
+
+							Dictionary<string, string> VersionStrings = new Dictionary<string, string>();
+							VersionStrings["#define ENGINE_VERSION"] = " " + VersionChangeNumber.ToString();
+							VersionStrings["#define ENGINE_IS_PROMOTED_BUILD"] = " (0)";
+							VersionStrings["#define BUILT_FROM_CHANGELIST"] = " " + VersionChangeNumber.ToString();
+							VersionStrings["#define BRANCH_NAME"] = " \"" + BranchOrStreamName.Replace('/', '+') + "\"";
+							if (!await UpdateVersionFile(Perforce, Project.ClientRootPath + VersionHeaderFileName, VersionStrings, Context.ChangeNumber, Logger, CancellationToken))
+							{
+								return (WorkspaceUpdateResult.FailedToSync, $"Failed to update {VersionHeaderFileName}");
+							}
+							if (!await UpdateVersionFile(Perforce, Project.ClientRootPath + ObjectVersionFileName, VersionStrings, Context.ChangeNumber, Logger, CancellationToken))
+							{
+								return (WorkspaceUpdateResult.FailedToSync, $"Failed to update {ObjectVersionFileName}");
+							}
+						}
+
+						// Remove all the receipts for build targets in this branch
+						if (Project.ClientFileName.EndsWith(".uproject", StringComparison.OrdinalIgnoreCase))
+						{
+							await Perforce.SyncQuietAsync(SyncOptions.None, -1, PerforceUtils.GetClientOrDepotDirectoryName(Project.ClientFileName) + "/Build/Receipts/...#0", CancellationToken);
+						}
+					}
+
+					// Check if there are any files which need resolving
+					List<FStatRecord> UnresolvedFiles = await FindUnresolvedFiles(Perforce, SyncPaths, CancellationToken);
+					if (UnresolvedFiles.Count > 0 && Context.Options.HasFlag(WorkspaceUpdateOptions.AutoResolveChanges))
+					{
+						foreach (FStatRecord UnresolvedFile in UnresolvedFiles)
+						{
+							if (UnresolvedFile.DepotFile != null)
+							{
+								await Perforce.ResolveAsync(-1, ResolveOptions.Automatic, UnresolvedFile.DepotFile, CancellationToken);
+							}
+						}
+						UnresolvedFiles = await FindUnresolvedFiles(Perforce, SyncPaths, CancellationToken);
+					}
+					if (UnresolvedFiles.Count > 0)
+					{
+						Logger.LogInformation("{NumFiles} files need resolving:", UnresolvedFiles.Count);
+						foreach (FStatRecord UnresolvedFile in UnresolvedFiles)
+						{
+							Logger.LogInformation("  {ClientFile}", UnresolvedFile.ClientFile);
+						}
+						return (WorkspaceUpdateResult.FilesToResolve, "Files need resolving.");
+					}
+
+					// Continue processing sync-only actions
+					if (Context.Options.HasFlag(WorkspaceUpdateOptions.Sync) && !Context.Options.HasFlag(WorkspaceUpdateOptions.UpdateFilter))
+					{
+						Context.ProjectConfigFile ??= await ReadProjectConfigFile(Project.LocalRootPath, Project.LocalFileName, Logger);
+
+						// Execute any project specific post-sync steps
+						string[]? PostSyncSteps = Context.ProjectConfigFile.GetValues("Sync.Step", null);
+						if (PostSyncSteps != null)
+						{
+							Logger.LogInformation("");
+							Logger.LogInformation("Executing post-sync steps...");
+
+							Dictionary<string, string> PostSyncVariables = new Dictionary<string, string>(Context.Variables);
+							PostSyncVariables["Change"] = Context.ChangeNumber.ToString();
+							PostSyncVariables["CodeChange"] = VersionChangeNumber.ToString();
+
+							foreach (string PostSyncStep in PostSyncSteps.Select(x => x.Trim()))
+							{
+								ConfigObject PostSyncStepObject = new ConfigObject(PostSyncStep);
+
+								string ToolFileName = Utility.ExpandVariables(PostSyncStepObject.GetValue("FileName", ""), PostSyncVariables);
+								if (ToolFileName != null)
+								{
+									string ToolArguments = Utility.ExpandVariables(PostSyncStepObject.GetValue("Arguments", ""), PostSyncVariables);
+
+									Logger.LogInformation("post-sync> Running {FileName} {Arguments}", ToolFileName, ToolArguments);
+
+									int ResultFromTool = await Utility.ExecuteProcessAsync(ToolFileName, null, ToolArguments, Line => ProcessOutput(Line, "post-sync> ", Progress, Logger), CancellationToken);
+									if (ResultFromTool != 0)
+									{
+										return (WorkspaceUpdateResult.FailedToSync, $"Post-sync step terminated with exit code {ResultFromTool}.");
+									}
+								}
+							}
+						}
+					}
+
+					// Update the current state
+					if (Context.Options.HasFlag(WorkspaceUpdateOptions.SyncSingleChange))
+					{
+						State.AdditionalChangeNumbers.Add(Context.ChangeNumber);
+					}
+					else
+					{
+						State.CurrentChangeNumber = Context.ChangeNumber;
+					}
+					State.Save();
+
+					// Update the timing info
+					Times.Add(new Tuple<string, TimeSpan>("Sync", SyncTelemetryStopwatch.Stop("Success")));
+
+					// Save the number of files synced
+					NumFilesSynced = SyncDepotPaths.Count;
+					Logger.LogInformation("");
+				}
+			}
+
+			// Extract an archive from the depot path
+			if (Context.Options.HasFlag(WorkspaceUpdateOptions.SyncArchives))
+			{
+				using (TelemetryStopwatch Stopwatch = new TelemetryStopwatch("Workspace_SyncArchives", Project.TelemetryProjectIdentifier))
+				{
+					// Create the directory for extracted archive manifests
+					DirectoryReference ManifestDirectoryName;
+					if (Project.LocalFileName.HasExtension(".uproject"))
+					{
+						ManifestDirectoryName = DirectoryReference.Combine(Project.LocalFileName.Directory, "Saved", "UnrealGameSync");
+					}
+					else
+					{
+						ManifestDirectoryName = DirectoryReference.Combine(Project.LocalFileName.Directory, "Engine", "Saved", "UnrealGameSync");
+					}
+					DirectoryReference.CreateDirectory(ManifestDirectoryName);
+
+					// Sync and extract (or just remove) the given archives
+					foreach (KeyValuePair<string, Tuple<IArchiveInfo, string>?> ArchiveTypeAndArchive in Context.ArchiveTypeToArchive)
+					{
+						string ArchiveType = ArchiveTypeAndArchive.Key;
+
+						// Remove any existing binaries
+						FileReference ManifestFileName = FileReference.Combine(ManifestDirectoryName, String.Format("{0}.zipmanifest", ArchiveType));
+						if (FileReference.Exists(ManifestFileName))
+						{
+							Logger.LogInformation("Removing {ArchiveType} binaries...", ArchiveType);
+							Progress.Set(String.Format("Removing {0} binaries...", ArchiveType), 0.0f);
+							ArchiveUtils.RemoveExtractedFiles(Project.LocalRootPath, ManifestFileName, Progress, Logger);
+							FileReference.Delete(ManifestFileName);
+							Logger.LogInformation("");
+						}
+
+						// If we have a new depot path, sync it down and extract it
+						if (ArchiveTypeAndArchive.Value != null)
+						{
+							IArchiveInfo ArchiveInfo = ArchiveTypeAndArchive.Value.Item1;
+							string ArchiveKey = ArchiveTypeAndArchive.Value.Item2;
+
+							Logger.LogInformation("Syncing {ArchiveType} binaries...", ArchiveType.ToLowerInvariant());
+							Progress.Set(String.Format("Syncing {0} binaries...", ArchiveType.ToLowerInvariant()), 0.0f);
+							if (!ArchiveInfo.DownloadArchive(Perforce, ArchiveKey, Project.LocalRootPath, ManifestFileName, Logger, Progress, CancellationToken.None).Result)
+							{
+								return (WorkspaceUpdateResult.FailedToSync, $"Couldn't read {ArchiveKey}");
+							}
+						}
+					}
+
+					// Update the state
+					State.ExpandedArchiveTypes = Context.ArchiveTypeToArchive.Where(x => x.Value != null).Select(x => x.Key).ToArray();
+					State.Save();
+
+					// Add the finish time
+					Times.Add(new Tuple<string, TimeSpan>("Archive", Stopwatch.Stop("Success")));
+				}
+			}
+
+			// Take the lock before doing anything else. Building and generating project files can only be done on one workspace at a time.
+			using SemaphoreScope Scope = new SemaphoreScope(UpdateSemaphore);
+			if (Context.Options.HasFlag(WorkspaceUpdateOptions.GenerateProjectFiles) || Context.Options.HasFlag(WorkspaceUpdateOptions.Build))
+			{
+				if (!Scope.TryAcquire())
+				{
+					Logger.LogInformation("Waiting for other workspaces to finish...");
+					await Scope.AcquireAsync(CancellationToken);
+				}
+			}
+
+			// Generate project files in the workspace
+			if (Context.Options.HasFlag(WorkspaceUpdateOptions.GenerateProjectFiles))
+			{
+				using (TelemetryStopwatch Stopwatch = new TelemetryStopwatch("Workspace_GenerateProjectFiles", Project.TelemetryProjectIdentifier))
+				{
+					Progress.Set("Generating project files...", 0.0f);
+
+					StringBuilder CommandLine = new StringBuilder();
+					CommandLine.AppendFormat("/C \"\"{0}\"", FileReference.Combine(Project.LocalRootPath, "GenerateProjectFiles.bat"));
+					if ((Context.Options & WorkspaceUpdateOptions.SyncAllProjects) == 0 && (Context.Options & WorkspaceUpdateOptions.IncludeAllProjectsInSolution) == 0)
+					{
+						if (Project.LocalFileName.HasExtension(".uproject"))
+						{
+							CommandLine.AppendFormat(" \"{0}\"", Project.LocalFileName);
+						}
+					}
+					CommandLine.Append(" -progress\"");
+
+					Logger.LogInformation("Generating project files...");
+					Logger.LogInformation("gpf> Running {FileName} {Arguments}", CmdExe, CommandLine);
+
+					int GenerateProjectFilesResult = await Utility.ExecuteProcessAsync(CmdExe, null, CommandLine.ToString(), Line => ProcessOutput(Line, "gpf> ", Progress, Logger), CancellationToken);
+					if (GenerateProjectFilesResult != 0)
+					{
+						return (WorkspaceUpdateResult.FailedToCompile, $"Failed to generate project files (exit code {GenerateProjectFilesResult}).");
+					}
+
+					Logger.LogInformation("");
+					Times.Add(new Tuple<string, TimeSpan>("Prj gen", Stopwatch.Stop("Success")));
+				}
+			}
+
+			// Build everything using MegaXGE
+			if (Context.Options.HasFlag(WorkspaceUpdateOptions.Build))
+			{
+				Context.ProjectConfigFile ??= await ReadProjectConfigFile(Project.LocalRootPath, Project.LocalFileName, Logger);
+
+				// Compile all the build steps together
+				Dictionary<Guid, ConfigObject> BuildStepObjects = Context.DefaultBuildSteps.ToDictionary(x => x.Key, x => new ConfigObject(x.Value));
+				BuildStep.MergeBuildStepObjects(BuildStepObjects, Context.ProjectConfigFile.GetValues("Build.Step", new string[0]).Select(x => new ConfigObject(x)));
+				BuildStep.MergeBuildStepObjects(BuildStepObjects, Context.UserBuildStepObjects);
+
+				// Construct build steps from them
+				List<BuildStep> BuildSteps = BuildStepObjects.Values.Select(x => new BuildStep(x)).OrderBy(x => (x.OrderIndex == -1) ? 10000 : x.OrderIndex).ToList();
+				if (Context.CustomBuildSteps != null && Context.CustomBuildSteps.Count > 0)
+				{
+					BuildSteps.RemoveAll(x => !Context.CustomBuildSteps.Contains(x.UniqueId));
+				}
+				else if (Context.Options.HasFlag(WorkspaceUpdateOptions.ScheduledBuild))
+				{
+					BuildSteps.RemoveAll(x => !x.bScheduledSync);
+				}
+				else
+				{
+					BuildSteps.RemoveAll(x => !x.bNormalSync);
+				}
+
+				// Check if the last successful build was before a change that we need to force a clean for
+				bool bForceClean = false;
+				if (State.LastBuiltChangeNumber != 0)
+				{
+					foreach (string CleanBuildChange in Context.ProjectConfigFile.GetValues("ForceClean.Changelist", new string[0]))
+					{
+						int ChangeNumber;
+						if (int.TryParse(CleanBuildChange, out ChangeNumber))
+						{
+							if ((State.LastBuiltChangeNumber >= ChangeNumber) != (State.CurrentChangeNumber >= ChangeNumber))
+							{
+								Logger.LogInformation("Forcing clean build due to changelist {Change}.", ChangeNumber);
+								Logger.LogInformation("");
+								bForceClean = true;
+								break;
+							}
+						}
+					}
+				}
+
+				// Execute them all
+				using (TelemetryStopwatch Stopwatch = new TelemetryStopwatch("Workspace_Build", Project.TelemetryProjectIdentifier))
+				{
+					Progress.Set("Starting build...", 0.0f);
+
+					// Execute all the steps
+					float MaxProgressFraction = 0.0f;
+					foreach (BuildStep Step in BuildSteps)
+					{
+						MaxProgressFraction += (float)Step.EstimatedDuration / (float)Math.Max(BuildSteps.Sum(x => x.EstimatedDuration), 1);
+
+						Progress.Set(Step.StatusText ?? "Executing build step");
+						Progress.Push(MaxProgressFraction);
+
+						Logger.LogInformation("{Status}", Step.StatusText);
+
+						if (Step.IsValid())
+						{
+							switch (Step.Type)
+							{
+								case BuildStepType.Compile:
+									using (TelemetryStopwatch StepStopwatch = new TelemetryStopwatch("Workspace_Execute_Compile", Project.TelemetryProjectIdentifier))
+									{
+										StepStopwatch.AddData(new { Target = Step.Target });
+
+										FileReference BuildBat = FileReference.Combine(Project.LocalRootPath, "Engine", "Build", "BatchFiles", "Build.bat");
+										string CommandLine = String.Format("{0} {1} {2} {3} -NoHotReloadFromIDE", Step.Target, Step.Platform, Step.Configuration, Utility.ExpandVariables(Step.Arguments ?? "", Context.Variables));
+										if (Context.Options.HasFlag(WorkspaceUpdateOptions.Clean) || bForceClean)
+										{
+											Logger.LogInformation("ubt> Running {FileName} {Arguments}", BuildBat, CommandLine + " -clean");
+											await Utility.ExecuteProcessAsync(CmdExe, null, "/C \"\"" + BuildBat + "\" " + CommandLine + " -clean\"", Line => ProcessOutput(Line, "ubt> ", Progress, Logger), CancellationToken);
+										}
+
+										Logger.LogInformation("ubt> Running {FileName} {Arguments}", BuildBat, CommandLine + " -progress");
+
+										int ResultFromBuild = await Utility.ExecuteProcessAsync(CmdExe, null, "/C \"\"" + BuildBat + "\" " + CommandLine + " -progress\"", Line => ProcessOutput(Line, "ubt> ", Progress, Logger), CancellationToken);
+										if (ResultFromBuild != 0)
+										{
+											StepStopwatch.Stop("Failed");
+
+											WorkspaceUpdateResult Result;
+											if (await HasModifiedSourceFiles(Perforce, Project, CancellationToken) || Context.UserBuildStepObjects.Count > 0)
+											{
+												Result = WorkspaceUpdateResult.FailedToCompile;
+											}
+											else
+											{
+												Result = WorkspaceUpdateResult.FailedToCompileWithCleanWorkspace;
+											}
+
+											return (Result, $"Failed to compile {Step.Target}");
+										}
+
+										StepStopwatch.Stop("Success");
+									}
+									break;
+								case BuildStepType.Cook:
+									using (TelemetryStopwatch StepStopwatch = new TelemetryStopwatch("Workspace_Execute_Cook", Project.TelemetryProjectIdentifier))
+									{
+										StepStopwatch.AddData(new { Project = Path.GetFileNameWithoutExtension(Step.FileName) });
+
+										FileReference LocalRunUAT = FileReference.Combine(Project.LocalRootPath, "Engine", "Build", "BatchFiles", "RunUAT.bat");
+										string Arguments = String.Format("/C \"\"{0}\" -profile=\"{1}\"\"", LocalRunUAT, FileReference.Combine(Project.LocalRootPath, Step.FileName ?? "unknown"));
+										Logger.LogInformation("uat> Running {FileName} {Argument}", LocalRunUAT, Arguments);
+
+										int ResultFromUAT = await Utility.ExecuteProcessAsync(CmdExe, null, Arguments, Line => ProcessOutput(Line, "uat> ", Progress, Logger), CancellationToken);
+										if (ResultFromUAT != 0)
+										{
+											StepStopwatch.Stop("Failed");
+											return (WorkspaceUpdateResult.FailedToCompile, $"Cook failed. ({ResultFromUAT})");
+										}
+
+										StepStopwatch.Stop("Success");
+									}
+									break;
+								case BuildStepType.Other:
+									using (TelemetryStopwatch StepStopwatch = new TelemetryStopwatch("Workspace_Execute_Custom", Project.TelemetryProjectIdentifier))
+									{
+										StepStopwatch.AddData(new { FileName = Path.GetFileNameWithoutExtension(Step.FileName) });
+
+										FileReference ToolFileName = FileReference.Combine(Project.LocalRootPath, Utility.ExpandVariables(Step.FileName ?? "unknown", Context.Variables));
+										string ToolWorkingDir = String.IsNullOrWhiteSpace(Step.WorkingDir) ? ToolFileName.Directory.FullName : Utility.ExpandVariables(Step.WorkingDir, Context.Variables);
+										string ToolArguments = Utility.ExpandVariables(Step.Arguments ?? "", Context.Variables);
+										Logger.LogInformation("tool> Running {0} {1}", ToolFileName, ToolArguments);
+
+										if (Step.bUseLogWindow)
+										{
+											int ResultFromTool = await Utility.ExecuteProcessAsync(ToolFileName.FullName, ToolWorkingDir, ToolArguments, Line => ProcessOutput(Line, "tool> ", Progress, Logger), CancellationToken);
+											if (ResultFromTool != 0)
+											{
+												StepStopwatch.Stop("Failed");
+												return (WorkspaceUpdateResult.FailedToCompile, $"Tool terminated with exit code {ResultFromTool}.");
+											}
+										}
+										else
+										{
+											ProcessStartInfo StartInfo = new ProcessStartInfo(ToolFileName.FullName, ToolArguments);
+											StartInfo.WorkingDirectory = ToolWorkingDir;
+											using (Process.Start(StartInfo))
+											{
+											}
+										}
+
+										StepStopwatch.Stop("Success");
+									}
+									break;
+							}
+						}
+
+						Logger.LogInformation("");
+						Progress.Pop();
+					}
+
+					Times.Add(new Tuple<string, TimeSpan>("Build", Stopwatch.Stop("Success")));
+				}
+
+				// Update the last successful build change number
+				if (Context.CustomBuildSteps == null || Context.CustomBuildSteps.Count == 0)
+				{
+					State.LastBuiltChangeNumber = State.CurrentChangeNumber;
+					State.Save();
+				}
+			}
+
+			// Write out all the timing information
+			Logger.LogInformation("Total time: {TotalTime}", FormatTime(Times.Sum(x => (long)(x.Item2.TotalMilliseconds / 1000))));
+			foreach (Tuple<string, TimeSpan> Time in Times)
+			{
+				Logger.LogInformation("   {Name,-8}: {TimeSeconds}", Time.Item1, FormatTime((long)(Time.Item2.TotalMilliseconds / 1000)));
+			}
+			if (NumFilesSynced > 0)
+			{
+				Logger.LogInformation("{NumFiles} files synced.", NumFilesSynced);
+			}
+
+			DateTime FinishTime = DateTime.Now;
+			Logger.LogInformation("");
+			Logger.LogInformation("UPDATE SUCCEEDED ({FinishDate} {FinishTime})", FinishTime.ToShortDateString(), FinishTime.ToShortTimeString());
+
+			return (WorkspaceUpdateResult.Success, "Update succeeded");
+		}
+
+		static void ProcessOutput(string Line, string Prefix, ProgressValue Progress, ILogger Logger)
+		{
+			string? ParsedLine = ProgressTextWriter.ParseLine(Line, Progress);
+			if (ParsedLine != null)
+			{
+				Logger.LogInformation("{Prefix}{Message}", Prefix, Line);
+			}
+		}
+
+		static async Task<bool> IsLicenseeVersion(IPerforceConnection Perforce, ProjectInfo Project, ILogger Logger, CancellationToken CancellationToken)
+		{
+			string[] Files =
+			{
+				Project.ClientRootPath + "/Engine/Build/NotForLicensees/EpicInternal.txt",
+				Project.ClientRootPath + "/Engine/Restricted/NotForLicensees/Build/EpicInternal.txt"
+			};
+
+			List<FStatRecord> Records = await Perforce.FStatAsync(Files, CancellationToken).ToListAsync(CancellationToken);
+			return !Records.Any(x => x.IsMapped);
+		}
+
+		public static List<string> GetSyncPaths(ProjectInfo Project, bool bSyncAllProjects, string[] SyncFilter)
+		{
+			List<string> SyncPaths = GetRelativeSyncPaths(Project, bSyncAllProjects, SyncFilter);
+			return SyncPaths.Select(x => Project.ClientRootPath + x).ToList();
+		}
+
+		public static List<string> GetRelativeSyncPaths(ProjectInfo Project, bool bSyncAllProjects, string[]? SyncFilter)
+		{
+			List<string> SyncPaths = new List<string>();
+
+			// Check the client path is formatted correctly
+			if (!Project.ClientFileName.StartsWith(Project.ClientRootPath + "/"))
+			{
+				throw new Exception(String.Format("Expected '{0}' to start with '{1}'", Project.ClientFileName, Project.ClientRootPath));
+			}
+
+			// Add the default project paths
+			int LastSlashIdx = Project.ClientFileName.LastIndexOf('/');
+			if (bSyncAllProjects || !Project.ClientFileName.EndsWith(".uproject", StringComparison.OrdinalIgnoreCase) || LastSlashIdx <= Project.ClientRootPath.Length)
+			{
+				SyncPaths.Add("/...");
+			}
+			else
+			{
+				SyncPaths.Add("/*");
+				SyncPaths.Add("/Engine/...");
+				if (Project.bIsEnterpriseProject)
+				{
+					SyncPaths.Add("/Enterprise/...");
+				}
+				SyncPaths.Add(Project.ClientFileName.Substring(Project.ClientRootPath.Length, LastSlashIdx - Project.ClientRootPath.Length) + "/...");
+			}
+
+			// Apply the sync filter to that list. We only want inclusive rules in the output list, but we can manually apply exclusions to previous entries.
+			if (SyncFilter != null)
+			{
+				foreach (string SyncPath in SyncFilter)
+				{
+					string TrimSyncPath = SyncPath.Trim();
+					if (TrimSyncPath.StartsWith("/"))
+					{
+						SyncPaths.Add(TrimSyncPath);
+					}
+					else if (TrimSyncPath.StartsWith("-/") && TrimSyncPath.EndsWith("..."))
+					{
+						SyncPaths.RemoveAll(x => x.StartsWith(TrimSyncPath.Substring(1, TrimSyncPath.Length - 4)));
+					}
+				}
+			}
+
+			// Sort the remaining paths by length, and remove any paths which are included twice
+			SyncPaths = SyncPaths.OrderBy(x => x.Length).ToList();
+			for (int Idx = 0; Idx < SyncPaths.Count; Idx++)
+			{
+				string SyncPath = SyncPaths[Idx];
+				if (SyncPath.EndsWith("..."))
+				{
+					string SyncPathPrefix = SyncPath.Substring(0, SyncPath.Length - 3);
+					for (int OtherIdx = SyncPaths.Count - 1; OtherIdx > Idx; OtherIdx--)
+					{
+						if (SyncPaths[OtherIdx].StartsWith(SyncPathPrefix))
+						{
+							SyncPaths.RemoveAt(OtherIdx);
+						}
+					}
+				}
+			}
+
+			return SyncPaths;
+		}
+
+		public static bool MatchFilter(ProjectInfo Project, FileReference FileName, FileFilter Filter)
+		{
+			bool bMatch = true;
+			if (FileName.IsUnderDirectory(Project.LocalRootPath))
+			{
+				string RelativePath = FileName.MakeRelativeTo(Project.LocalRootPath);
+				if (!Filter.Matches(RelativePath))
+				{
+					bMatch = false;
+				}
+			}
+			return bMatch;
+		}
+
+		class SyncState
+		{
+			public int TotalDepotPaths;
+			public HashSet<string> RemainingDepotPaths;
+			public Queue<List<string>> SyncCommandLists;
+			public string StatusMessage;
+			public WorkspaceUpdateResult Result = WorkspaceUpdateResult.Success;
+
+			public SyncState(HashSet<string> RemainingDepotPaths, Queue<List<string>> SyncCommandLists)
+			{
+				this.TotalDepotPaths = RemainingDepotPaths.Count;
+				this.RemainingDepotPaths = RemainingDepotPaths;
+				this.SyncCommandLists = SyncCommandLists;
+				this.StatusMessage = "Succeeded.";
+			}
+		}
+
+		static Task<(WorkspaceUpdateResult, string)> SyncFileRevisions(IPerforceConnection Perforce, string Prefix, WorkspaceUpdateContext Context, List<string> SyncCommands, HashSet<string> RemainingDepotPaths, ProgressValue Progress, ILogger Logger, CancellationToken CancellationToken)
+		{
+			Queue<List<string>> SyncCommandLists = new Queue<List<string>>();
+			SyncCommandLists.Enqueue(SyncCommands);
+			return SyncFileRevisions(Perforce, Prefix, Context, SyncCommandLists, RemainingDepotPaths, Progress, Logger, CancellationToken);
+		}
+
+		static async Task<(WorkspaceUpdateResult, string)> SyncFileRevisions(IPerforceConnection Perforce, string Prefix, WorkspaceUpdateContext Context, Queue<List<string>> SyncCommandLists, HashSet<string> RemainingDepotPaths, ProgressValue Progress, ILogger Logger, CancellationToken CancellationToken)
+		{
+			// Figure out the number of additional background threads we want to run with. We can run worker on the current thread.
+			int NumThreads = Context.PerforceSyncOptions?.NumThreads ?? PerforceSyncOptions.DefaultNumThreads;
+			int NumExtraThreads = Math.Max(Math.Min(SyncCommandLists.Count, NumThreads) - 1, 0);
+
+			List<IPerforceConnection> ChildConnections = new List<IPerforceConnection>();
+			List<Task> ChildTasks = new List<Task>(NumExtraThreads);
+			try
+			{
+				// Create the state object shared by all the worker threads
+				SyncState State = new SyncState(RemainingDepotPaths, SyncCommandLists);
+
+				// Wrapper writer around the log class to prevent multiple threads writing to it at once
+				ILogger LogWrapper = Logger;
+
+				// Initialize Sync Progress
+				UpdateSyncState(Prefix, State, Progress);
+
+				// Delegate for updating the sync state after a file has been synced
+				Action<SyncRecord, ILogger> SyncOutput = (Record, LocalLog) => { UpdateSyncState(Prefix, Record, State, Progress, LocalLog); };
+
+				// Create all the child threads
+				for (int ThreadIdx = 0; ThreadIdx < NumExtraThreads; ThreadIdx++)
+				{
+					int ThreadNumber = ThreadIdx + 2;
+
+					// Create connection
+					IPerforceConnection ChildConnection = await PerforceConnection.CreateAsync(Perforce.Settings, Perforce.Logger);
+					ChildConnections.Add(ChildConnection);
+
+					Task ChildTask = Task.Run(() => StaticSyncWorker(ThreadNumber, ChildConnection, Context, State, SyncOutput, LogWrapper, CancellationToken));
+					ChildTasks.Add(ChildTask);
+				}
+
+				// Run one worker on the current thread
+				await StaticSyncWorker(1, Perforce, Context, State, SyncOutput, LogWrapper, CancellationToken);
+
+				// Return the result that was set on the state object
+				return (State.Result, State.StatusMessage);
+			}
+			finally
+			{
+				await Task.WhenAll(ChildTasks);
+
+				foreach (IPerforceConnection ChildConnection in ChildConnections)
+				{
+					ChildConnection.Dispose();
+				}
+
+				foreach (Task ChildTask in ChildTasks)
+				{
+					await ChildTask;
+				}
+			}
+		}
+
+		static void UpdateSyncState(string Prefix, SyncState State, ProgressValue Progress)
+		{
+			lock (State)
+			{
+				string Message = String.Format("{0} ({1}/{2})", Prefix, State.TotalDepotPaths - State.RemainingDepotPaths.Count, State.TotalDepotPaths);
+				float Fraction = Math.Min((float)(State.TotalDepotPaths - State.RemainingDepotPaths.Count) / (float)State.TotalDepotPaths, 1.0f);
+				Progress.Set(Message, Fraction);
+			}
+		}
+
+		static void UpdateSyncState(string Prefix, SyncRecord Record, SyncState State, ProgressValue Progress, ILogger Logger)
+		{
+			lock (State)
+			{
+				State.RemainingDepotPaths.Remove(Record.DepotFile);
+
+				string Message = String.Format("{0} ({1}/{2})", Prefix, State.TotalDepotPaths - State.RemainingDepotPaths.Count, State.TotalDepotPaths);
+				float Fraction = Math.Min((float)(State.TotalDepotPaths - State.RemainingDepotPaths.Count) / (float)State.TotalDepotPaths, 1.0f);
+				Progress.Set(Message, Fraction);
+
+				Logger.LogInformation("p4>   {Action} {Path}", Record.Action, Record.Path);
+			}
+		}
+
+		static async Task StaticSyncWorker(int ThreadNumber, IPerforceConnection Perforce, WorkspaceUpdateContext Context, SyncState State, Action<SyncRecord, ILogger> SyncOutput, ILogger GlobalLog, CancellationToken CancellationToken)
+		{
+			PrefixedTextWriter ThreadLog = new PrefixedTextWriter(String.Format("{0}:", ThreadNumber), GlobalLog);
+			for (; ; )
+			{
+				// Remove the next batch that needs to be synced
+				List<string> SyncCommands;
+				lock (State)
+				{
+					if (State.Result == WorkspaceUpdateResult.Success && State.SyncCommandLists.Count > 0)
+					{
+						SyncCommands = State.SyncCommandLists.Dequeue();
+					}
+					else
+					{
+						break;
+					}
+				}
+
+				// Sync the files
+				(WorkspaceUpdateResult Result, string StatusMessage) = await StaticSyncFileRevisions(Perforce, Context, SyncCommands, Record => SyncOutput(Record, ThreadLog), CancellationToken);
+
+				// If it failed, try to set it on the state if nothing else has failed first
+				if (Result != WorkspaceUpdateResult.Success)
+				{
+					lock (State)
+					{
+						if (State.Result == WorkspaceUpdateResult.Success)
+						{
+							State.Result = Result;
+							State.StatusMessage = StatusMessage;
+						}
+					}
+					break;
+				}
+			}
+		}
+
+		static async Task<(WorkspaceUpdateResult, string)> StaticSyncFileRevisions(IPerforceConnection Perforce, WorkspaceUpdateContext Context, List<string> SyncCommands, Action<SyncRecord> SyncOutput, CancellationToken CancellationToken)
+		{
+			// Sync them all
+			List<PerforceResponse<SyncRecord>> Responses = await Perforce.TrySyncAsync(SyncOptions.None, -1, SyncCommands, CancellationToken).ToListAsync(CancellationToken);
+
+			List<string> TamperedFiles = new List<string>();
+			foreach (PerforceResponse<SyncRecord> Response in Responses)
+			{
+				const string NoClobberPrefix = "Can't clobber writable file ";
+				if (Response.Succeeded)
+				{
+					SyncOutput(Response.Data);
+				}
+				else if (Response.Error != null && Response.Error.Generic == PerforceGenericCode.Client && Response.Error.Data.StartsWith(NoClobberPrefix, StringComparison.OrdinalIgnoreCase))
+				{
+					TamperedFiles.Add(Response.Error.Data.Substring(NoClobberPrefix.Length).Trim());
+				}
+				else
+				{
+					return (WorkspaceUpdateResult.FailedToSync, "Aborted sync due to errors.");
+				}
+			}
+
+			// If any files need to be clobbered, defer to the main thread to figure out which ones
+			if (TamperedFiles.Count > 0)
+			{
+				int NumNewFilesToClobber = 0;
+				foreach (string TamperedFile in TamperedFiles)
+				{
+					if (!Context.ClobberFiles.ContainsKey(TamperedFile))
+					{
+						Context.ClobberFiles[TamperedFile] = true;
+						if (TamperedFile.EndsWith(LocalObjectVersionFileName, StringComparison.OrdinalIgnoreCase) || TamperedFile.EndsWith(LocalVersionHeaderFileName, StringComparison.OrdinalIgnoreCase))
+						{
+							// Hack for UseFastModularVersioningV2; we don't need to update these files any more.
+							continue;
+						}
+						NumNewFilesToClobber++;
+					}
+				}
+				if (NumNewFilesToClobber > 0)
+				{
+					return (WorkspaceUpdateResult.FilesToClobber, $"Cancelled sync after checking files to clobber ({NumNewFilesToClobber} new files).");
+				}
+				foreach (string TamperedFile in TamperedFiles)
+				{
+					if (Context.ClobberFiles[TamperedFile])
+					{
+						List<PerforceResponse<SyncRecord>> Response = await Perforce.TrySyncAsync(SyncOptions.Force, -1, TamperedFile, CancellationToken).ToListAsync();
+						if (!Response.Succeeded())
+						{
+							return (WorkspaceUpdateResult.FailedToSync, $"Couldn't sync {TamperedFile}.");
+						}
+					}
+				}
+			}
+
+			// All succeeded
+			return (WorkspaceUpdateResult.Success, "Succeeded.");
+		}
+
+		public static async Task<ConfigFile> ReadProjectConfigFile(DirectoryReference LocalRootPath, FileReference SelectedLocalFileName, ILogger Logger)
+		{
+			// Find the valid config file paths
+			DirectoryInfo EngineDir = DirectoryReference.Combine(LocalRootPath, "Engine").ToDirectoryInfo();
+			List<FileInfo> LocalConfigFiles = Utility.GetLocalConfigPaths(EngineDir, SelectedLocalFileName.ToFileInfo());
+
+			// Read them in
+			ConfigFile ProjectConfig = new ConfigFile();
+			foreach (FileInfo LocalConfigFile in LocalConfigFiles)
+			{
+				try
+				{
+					string[] Lines = await File.ReadAllLinesAsync(LocalConfigFile.FullName);
+					ProjectConfig.Parse(Lines);
+					Logger.LogInformation("Read config file from {FileName}", LocalConfigFile.FullName);
+				}
+				catch (Exception Ex)
+				{
+					Logger.LogInformation(Ex, "Failed to read config file from {FileName}", LocalConfigFile.FullName);
+				}
+			}
+			return ProjectConfig;
+		}
+
+		public static async Task<IReadOnlyList<string>?> ReadProjectStreamFilter(IPerforceConnection Perforce, ConfigFile ProjectConfigFile, ILogger Logger, CancellationToken CancellationToken)
+		{
+			string? StreamListDepotPath = ProjectConfigFile.GetValue("Options.QuickSelectStreamList", null);
+			if (StreamListDepotPath == null)
+			{
+				return null;
+			}
+
+			PerforceResponse<PrintRecord<string[]>> Response = await Perforce.TryPrintLinesAsync(StreamListDepotPath, CancellationToken);
+			if (!Response.Succeeded)
+			{
+				return null;
+			}
+
+			return Response.Data.Contents.Select(x => x.Trim()).Where(x => x.Length > 0).ToList().AsReadOnly();
+		}
+
+		static string FormatTime(long Seconds)
+		{
+			if (Seconds >= 60)
+			{
+				return String.Format("{0,3}m {1:00}s", Seconds / 60, Seconds % 60);
+			}
+			else
+			{
+				return String.Format("     {0,2}s", Seconds);
+			}
+		}
+
+		static async Task<bool> HasModifiedSourceFiles(IPerforceConnection Perforce, ProjectInfo Project, CancellationToken CancellationToken)
+		{
+			List<OpenedRecord> OpenFiles = await Perforce.OpenedAsync(OpenedOptions.None, Project.ClientRootPath + "/...", CancellationToken).ToListAsync(CancellationToken);
+			if (OpenFiles.Any(x => x.DepotFile.IndexOf("/Source/", StringComparison.OrdinalIgnoreCase) != -1))
+			{
+				return true;
+			}
+			return false;
+		}
+
+		static async Task<List<FStatRecord>> FindUnresolvedFiles(IPerforceConnection Perforce, IEnumerable<string> SyncPaths, CancellationToken CancellationToken)
+		{
+			List<FStatRecord> UnresolvedFiles = new List<FStatRecord>();
+			foreach (string SyncPath in SyncPaths)
+			{
+				List<FStatRecord> Records = await Perforce.FStatAsync(FStatOptions.OnlyUnresolved, SyncPath, CancellationToken).ToListAsync(CancellationToken);
+				UnresolvedFiles.AddRange(Records);
+			}
+			return UnresolvedFiles;
+		}
+
+		static Task<bool> UpdateVersionFile(IPerforceConnection Perforce, string ClientPath, Dictionary<string, string> VersionStrings, int ChangeNumber, ILogger Logger, CancellationToken CancellationToken)
+		{
+			return UpdateVersionFile(Perforce, ClientPath, ChangeNumber, Text => UpdateVersionStrings(Text, VersionStrings), Logger, CancellationToken);
+		}
+
+		static async Task<bool> UpdateVersionFile(IPerforceConnection Perforce, string ClientPath, int ChangeNumber, Func<string, string> Update, ILogger Logger, CancellationToken CancellationToken)
+		{
+			List<PerforceResponse<FStatRecord>> Records = await Perforce.TryFStatAsync(FStatOptions.None, ClientPath, CancellationToken).ToListAsync(CancellationToken);
+			if (!Records.Succeeded())
+			{
+				Logger.LogInformation("Failed to query records for {ClientPath}", ClientPath);
+				return false;
+			}
+			if (Records.Count > 1)
+			{
+				// Attempt to remove any existing file which is synced
+				await Perforce.SyncAsync(SyncOptions.Force, -1, $"{ClientPath}#0", CancellationToken).ToListAsync(CancellationToken);
+
+				// Try to get the mapped files again
+				Records = await Perforce.TryFStatAsync(FStatOptions.None, ClientPath, CancellationToken).ToListAsync(CancellationToken);
+				if (!Records.Succeeded())
+				{
+					Logger.LogInformation("Failed to query records for {ClientPath}", ClientPath);
+					return false;
+				}
+			}
+			if (Records.Count == 0)
+			{
+				Logger.LogInformation("Ignoring {ClientPath}; not found on server.", ClientPath);
+				return true;
+			}
+
+			FStatRecord Record = Records[0].Data;
+			string? LocalPath = Record.ClientFile; // Actually a filesystem path
+			if (LocalPath == null)
+			{
+				Logger.LogInformation("Version file is not mapped to workspace ({ClientFile})", ClientPath);
+				return false;
+			}
+			string? DepotPath = Record.DepotFile;
+			if (DepotPath == null)
+			{
+				Logger.LogInformation("Version file does not exist in depot ({ClientFile})", ClientPath);
+				return false;
+			}
+
+			PerforceResponse<PrintRecord<string[]>> Response = await Perforce.TryPrintLinesAsync($"{DepotPath}@{ChangeNumber}", CancellationToken);
+			if (!Response.Succeeded)
+			{
+				Logger.LogInformation("Couldn't get default contents of {DepotPath}", DepotPath);
+				return false;
+			}
+
+			string[]? Contents = Response.Data.Contents;
+			if (Contents == null)
+			{
+				Logger.LogInformation("No data returned for {DepotPath}", DepotPath);
+				return false;
+			}
+
+			string Text = String.Join("\n", Contents);
+			Text = Update(Text);
+			return await WriteVersionFile(Perforce, LocalPath, DepotPath, Text, Logger, CancellationToken);
+		}
+
+		static string UpdateVersionStrings(string Text, Dictionary<string, string> VersionStrings)
+		{
+			StringWriter Writer = new StringWriter();
+			foreach (string Line in Text.Split('\n'))
+			{
+				string NewLine = Line;
+				foreach (KeyValuePair<string, string> VersionString in VersionStrings)
+				{
+					if (UpdateVersionLine(ref NewLine, VersionString.Key, VersionString.Value))
+					{
+						break;
+					}
+				}
+				Writer.WriteLine(NewLine);
+			}
+			return Writer.ToString();
+		}
+
+		static string UpdateBuildVersion(string Text, int Changelist, int CodeChangelist, string BranchOrStreamName, bool bIsLicenseeVersion)
+		{
+			Dictionary<string, object> Object = JsonSerializer.Deserialize<Dictionary<string, object>>(Text, Program.DefaultJsonSerializerOptions);
+
+			int PrevCompatibleChangelist = 0;
+			if (Object.TryGetValue("CompatibleChangelist", out object? PrevCompatibleChangelistObj))
+			{
+				int.TryParse(PrevCompatibleChangelistObj?.ToString(), out PrevCompatibleChangelist);
+			}
+
+			int PrevIsLicenseeVersion = 0;
+			if (Object.TryGetValue("IsLicenseeVersion", out object? PrevIsLicenseeVersionObj))
+			{
+				int.TryParse(PrevIsLicenseeVersionObj?.ToString(), out PrevIsLicenseeVersion);
+			}
+
+			Object["Changelist"] = Changelist;
+			if (PrevCompatibleChangelist == 0 || (PrevIsLicenseeVersion != 0) != bIsLicenseeVersion)
+			{
+				// Don't overwrite the compatible changelist if we're in a hotfix release
+				Object["CompatibleChangelist"] = CodeChangelist;
+			}
+			Object["BranchName"] = BranchOrStreamName.Replace('/', '+');
+			Object["IsPromotedBuild"] = 0;
+			Object["IsLicenseeVersion"] = bIsLicenseeVersion ? 1 : 0;
+
+			return JsonSerializer.Serialize(Object, new JsonSerializerOptions
+			{
+				WriteIndented = true,
+				// do not escape +
+				Encoder = JavaScriptEncoder.UnsafeRelaxedJsonEscaping
+			});
+		}
+
+		static async Task<bool> WriteVersionFile(IPerforceConnection Perforce, string LocalPath, string DepotPath, string NewText, ILogger Logger, CancellationToken CancellationToken)
+		{
+			try
+			{
+				if (File.Exists(LocalPath) && File.ReadAllText(LocalPath) == NewText)
+				{
+					Logger.LogInformation("Ignored {FileName}; contents haven't changed", LocalPath);
+				}
+				else
+				{
+					Directory.CreateDirectory(Path.GetDirectoryName(LocalPath));
+					Utility.ForceDeleteFile(LocalPath);
+					if (DepotPath != null)
+					{
+						await Perforce.SyncAsync(DepotPath + "#0", CancellationToken).ToListAsync(CancellationToken);
+					}
+					File.WriteAllText(LocalPath, NewText);
+					Logger.LogInformation("Written {FileName}", LocalPath);
+				}
+				return true;
+			}
+			catch (Exception Ex)
+			{
+				Logger.LogError(Ex, "Failed to write to {FileName}.", LocalPath);
+				return false;
+			}
+		}
+
+		static bool UpdateVersionLine(ref string Line, string Prefix, string Suffix)
+		{
+			int LineIdx = 0;
+			int PrefixIdx = 0;
+			for (; ; )
+			{
+				string? PrefixToken = ReadToken(Prefix, ref PrefixIdx);
+				if (PrefixToken == null)
+				{
+					break;
+				}
+
+				string? LineToken = ReadToken(Line, ref LineIdx);
+				if (LineToken == null || LineToken != PrefixToken)
+				{
+					return false;
+				}
+			}
+			Line = Line.Substring(0, LineIdx) + Suffix;
+			return true;
+		}
+
+		static string? ReadToken(string Line, ref int LineIdx)
+		{
+			for (; ; LineIdx++)
+			{
+				if (LineIdx == Line.Length)
+				{
+					return null;
+				}
+				else if (!Char.IsWhiteSpace(Line[LineIdx]))
+				{
+					break;
+				}
+			}
+
+			int StartIdx = LineIdx++;
+			if (Char.IsLetterOrDigit(Line[StartIdx]) || Line[StartIdx] == '_')
+			{
+				while (LineIdx < Line.Length && (Char.IsLetterOrDigit(Line[LineIdx]) || Line[LineIdx] == '_'))
+				{
+					LineIdx++;
+				}
+			}
+
+			return Line.Substring(StartIdx, LineIdx - StartIdx);
+		}
+
+		public Tuple<string, float> CurrentProgress
+		{
+			get { return Progress.Current; }
+		}
+	}
+}
