UGS: More core functionality into a shared library.

diff --git a/Engine/Source/Programs/UnrealGameSync/UnrealGameSync/ArchiveManifest.cs b/Engine/Source/Programs/UnrealGameSync/UnrealGameSync/ArchiveManifest.cs
deleted file mode 100644
diff --git a/Engine/Source/Programs/UnrealGameSync/UnrealGameSync/BuildStep.cs b/Engine/Source/Programs/UnrealGameSync/UnrealGameSync/BuildStep.cs
deleted file mode 100644
diff --git a/Engine/Source/Programs/UnrealGameSync/UnrealGameSync/ConfigFile.cs b/Engine/Source/Programs/UnrealGameSync/UnrealGameSync/ConfigFile.cs
deleted file mode 100644
diff --git a/Engine/Source/Programs/UnrealGameSync/UnrealGameSync/Controls/WorkspaceControl.cs b/Engine/Source/Programs/UnrealGameSync/UnrealGameSync/Controls/WorkspaceControl.cs
--- a/Engine/Source/Programs/UnrealGameSync/UnrealGameSync/Controls/WorkspaceControl.cs
+++ b/Engine/Source/Programs/UnrealGameSync/UnrealGameSync/Controls/WorkspaceControl.cs
@@ -1806,7 +1806,7 @@ namespace UnrealGameSync
 				try
 				{
 					string JsonString = File.ReadAllText(UncontrolledChangelistPersistencyFilePath);
-					UncontrolledChangelistPersistency UCLPersistency = JsonSerializer.Deserialize<UncontrolledChangelistPersistency>(JsonString, Program.DefaultJsonSerializerOptions);
+					UncontrolledChangelistPersistency UCLPersistency = JsonSerializer.Deserialize<UncontrolledChangelistPersistency>(JsonString, Utility.DefaultJsonSerializerOptions);
 
 					foreach (UncontrolledChangelist UCL in UCLPersistency.Changelists)
 					{
@@ -3118,7 +3118,7 @@ namespace UnrealGameSync
 					try
 					{
 						string Text = FileReference.ReadAllText(ReceiptFileName);
-						ReceiptJsonObject Receipt = JsonSerializer.Deserialize<ReceiptJsonObject>(Text, Program.DefaultJsonSerializerOptions);
+						ReceiptJsonObject Receipt = JsonSerializer.Deserialize<ReceiptJsonObject>(Text, Utility.DefaultJsonSerializerOptions);
 
 						string? LaunchFileName = Receipt.Launch;
 						if (LaunchFileName != null)
diff --git a/Engine/Source/Programs/UnrealGameSync/UnrealGameSync/FileFilter.cs b/Engine/Source/Programs/UnrealGameSync/UnrealGameSync/FileFilter.cs
deleted file mode 100644
diff --git a/Engine/Source/Programs/UnrealGameSync/UnrealGameSync/OutputAdapters.cs b/Engine/Source/Programs/UnrealGameSync/UnrealGameSync/OutputAdapters.cs
deleted file mode 100644
diff --git a/Engine/Source/Programs/UnrealGameSync/UnrealGameSync/Program.cs b/Engine/Source/Programs/UnrealGameSync/UnrealGameSync/Program.cs
--- a/Engine/Source/Programs/UnrealGameSync/UnrealGameSync/Program.cs
+++ b/Engine/Source/Programs/UnrealGameSync/UnrealGameSync/Program.cs
@@ -31,16 +31,6 @@ namespace UnrealGameSync
 			return Version?.InformationalVersion ?? "Unknown";
 		}
 
-		static JsonSerializerOptions GetDefaultJsonSerializerOptions()
-		{
-			JsonSerializerOptions Options = new JsonSerializerOptions();
-			Options.PropertyNameCaseInsensitive = true;
-			Options.Converters.Add(new JsonStringEnumConverter());
-			return Options;
-		}
-
-		public static JsonSerializerOptions DefaultJsonSerializerOptions { get; } = GetDefaultJsonSerializerOptions();
-
 		public static string? SyncVersion = null;
 
 		[STAThread]
diff --git a/Engine/Source/Programs/UnrealGameSync/UnrealGameSync/RESTApi.cs b/Engine/Source/Programs/UnrealGameSync/UnrealGameSync/RESTApi.cs
--- a/Engine/Source/Programs/UnrealGameSync/UnrealGameSync/RESTApi.cs
+++ b/Engine/Source/Programs/UnrealGameSync/UnrealGameSync/RESTApi.cs
@@ -76,7 +76,7 @@ namespace UnrealGameSync
 
 		public static async Task<T> GetAsync<T>(string Url, CancellationToken CancellationToken)
 		{
-			return JsonSerializer.Deserialize<T>(await GetAsync(Url, CancellationToken), Program.DefaultJsonSerializerOptions);
+			return JsonSerializer.Deserialize<T>(await GetAsync(Url, CancellationToken), Utility.DefaultJsonSerializerOptions);
 		}
 
 		public static Task<string> PutAsync<T>(string Url, T Object, CancellationToken CancellationToken)
diff --git a/Engine/Source/Programs/UnrealGameSync/UnrealGameSync/Telemetry.cs b/Engine/Source/Programs/UnrealGameSync/UnrealGameSync/Telemetry.cs
deleted file mode 100644
diff --git a/Engine/Source/Programs/UnrealGameSync/UnrealGameSync/TelemetryStopwatch.cs b/Engine/Source/Programs/UnrealGameSync/UnrealGameSync/TelemetryStopwatch.cs
deleted file mode 100644
diff --git a/Engine/Source/Programs/UnrealGameSync/UnrealGameSync/UserSettings.cs b/Engine/Source/Programs/UnrealGameSync/UnrealGameSync/UserSettings.cs
deleted file mode 100644
diff --git a/Engine/Source/Programs/UnrealGameSync/UnrealGameSync/Utility.cs b/Engine/Source/Programs/UnrealGameSync/UnrealGameSync/Utility.cs
deleted file mode 100644
diff --git a/Engine/Source/Programs/UnrealGameSync/UnrealGameSync/WorkspaceUpdate.cs b/Engine/Source/Programs/UnrealGameSync/UnrealGameSync/WorkspaceUpdate.cs
deleted file mode 100644
diff --git a/Engine/Source/Programs/UnrealGameSync/UnrealGameSyncLauncher/Forms/SettingsWindow.cs b/Engine/Source/Programs/UnrealGameSync/UnrealGameSyncLauncher/Forms/SettingsWindow.cs
--- a/Engine/Source/Programs/UnrealGameSync/UnrealGameSyncLauncher/Forms/SettingsWindow.cs
+++ b/Engine/Source/Programs/UnrealGameSync/UnrealGameSyncLauncher/Forms/SettingsWindow.cs
@@ -85,7 +85,7 @@ namespace UnrealGameSyncLauncher
 				DepotPath = null;
 			}
 
-			Utility.SaveGlobalPerforceSettings(ServerAndPort, UserName, DepotPath);
+			GlobalSettings.SaveGlobalPerforceSettings(ServerAndPort, UserName, DepotPath);
 
 			PerforceSettings PerforceSettings = new PerforceSettings(PerforceSettings.Default);
 			if (!String.IsNullOrEmpty(ServerAndPort))
@@ -107,7 +107,7 @@ namespace UnrealGameSyncLauncher
 			{
 				if(Task.Succeeded)
 				{
-					Utility.SaveGlobalPerforceSettings(ServerAndPort, UserName, DepotPath);
+					GlobalSettings.SaveGlobalPerforceSettings(ServerAndPort, UserName, DepotPath);
 					DialogResult = DialogResult.OK;
 					Close();
 				}
diff --git a/Engine/Source/Programs/UnrealGameSync/UnrealGameSyncLauncher/Program.cs b/Engine/Source/Programs/UnrealGameSync/UnrealGameSyncLauncher/Program.cs
--- a/Engine/Source/Programs/UnrealGameSync/UnrealGameSyncLauncher/Program.cs
+++ b/Engine/Source/Programs/UnrealGameSync/UnrealGameSyncLauncher/Program.cs
@@ -46,7 +46,7 @@ namespace UnrealGameSyncLauncher
 				string? ServerAndPort = null;
 				string? UserName = null;
 				string? DepotPath = DeploymentSettings.DefaultDepotPath;
-				Utility.ReadGlobalPerforceSettings(ref ServerAndPort, ref UserName, ref DepotPath);
+				GlobalSettings.ReadGlobalPerforceSettings(ref ServerAndPort, ref UserName, ref DepotPath);
 
 				// If the shift key is held down, immediately show the settings window
 				SettingsWindow.SyncAndRunDelegate SyncAndRunWrapper = (Perforce, DepotParam, bUnstableParam, LogWriter, CancellationToken) => SyncAndRun(Perforce, DepotParam, bUnstableParam, Args, InstanceMutex, LogWriter, CancellationToken);
diff --git a/Engine/Source/Programs/UnrealGameSync/UnrealGameSyncLauncher/UnrealGameSyncLauncher.csproj b/Engine/Source/Programs/UnrealGameSync/UnrealGameSyncLauncher/UnrealGameSyncLauncher.csproj
--- a/Engine/Source/Programs/UnrealGameSync/UnrealGameSyncLauncher/UnrealGameSyncLauncher.csproj
+++ b/Engine/Source/Programs/UnrealGameSync/UnrealGameSyncLauncher/UnrealGameSyncLauncher.csproj
@@ -1,4 +1,4 @@
-﻿<?xml version="1.0" encoding="utf-8"?>
+<?xml version="1.0" encoding="utf-8"?>
 <Project Sdk="Microsoft.NET.Sdk.WindowsDesktop">
 
   <PropertyGroup>
@@ -41,8 +41,11 @@
     <Compile Include="..\UnrealGameSync\PerforceModalTask.cs">
       <Link>UnrealGameSync\PerforceModalTask.cs</Link>
     </Compile>
-    <Compile Include="..\UnrealGameSync\Utility.cs">
-      <Link>UnrealGameSync\Utility.cs</Link>
+    <Compile Include="..\UnrealGameSync\GlobalSettings.cs">
+      <Link>UnrealGameSync\GlobalSettings.cs</Link>
+    </Compile>
+    <Compile Include="..\UnrealGameSyncShared\Utility.cs">
+      <Link>UnrealGameSyncShared\Utility.cs</Link>
     </Compile>
   </ItemGroup>
   <ItemGroup>
diff --git a/Engine/Source/Programs/UnrealGameSync/UnrealGameSyncShared/ArchiveManifest.cs b/Engine/Source/Programs/UnrealGameSync/UnrealGameSyncShared/ArchiveManifest.cs
new file mode 100644
--- /dev/null
+++ b/Engine/Source/Programs/UnrealGameSync/UnrealGameSyncShared/ArchiveManifest.cs
@@ -0,0 +1,163 @@
+// Copyright Epic Games, Inc. All Rights Reserved.
+
+using EpicGames.Core;
+using Microsoft.Extensions.Logging;
+using System;
+using System.Collections.Generic;
+using System.Diagnostics;
+using System.IO;
+using System.IO.Compression;
+using System.Linq;
+using System.Text;
+using System.Threading.Tasks;
+
+namespace UnrealGameSync
+{
+	[DebuggerDisplay("{FileName}")]
+	public class ArchiveManifestFile
+	{
+		public string FileName;
+		public long Length;
+		public DateTime LastWriteTimeUtc;
+
+		public ArchiveManifestFile(BinaryReader Reader)
+		{
+			FileName = Reader.ReadString();
+			Length = Reader.ReadInt64();
+			LastWriteTimeUtc = new DateTime(Reader.ReadInt64());
+		}
+
+		public ArchiveManifestFile(string InFileName, long InLength, DateTime InLastWriteTimeUtc)
+		{
+			FileName = InFileName;
+			Length = InLength;
+			LastWriteTimeUtc = InLastWriteTimeUtc;
+		}
+
+		public void Write(BinaryWriter Writer)
+		{
+			Writer.Write(FileName);
+			Writer.Write(Length);
+			Writer.Write(LastWriteTimeUtc.Ticks);
+		}
+	}
+
+	public class ArchiveManifest
+	{
+		const int Signature = ((int)'U' << 24) | ((int)'A' << 16) | ((int)'M' << 8) | 1;
+
+		public List<ArchiveManifestFile> Files = new List<ArchiveManifestFile>();
+
+		public ArchiveManifest()
+		{
+		}
+
+		public ArchiveManifest(FileStream InputStream)
+		{
+			BinaryReader Reader = new BinaryReader(InputStream);
+			if(Reader.ReadInt32() != Signature)
+			{
+				throw new Exception("Archive manifest signature does not match");
+			}
+
+			int NumFiles = Reader.ReadInt32();
+			for(int Idx = 0; Idx < NumFiles; Idx++)
+			{
+				Files.Add(new ArchiveManifestFile(Reader));
+			}
+		}
+
+		public void Write(FileStream OutputStream)
+		{
+			BinaryWriter Writer = new BinaryWriter(OutputStream);
+			Writer.Write(Signature);
+			Writer.Write(Files.Count);
+			foreach(ArchiveManifestFile File in Files)
+			{
+				File.Write(Writer);
+			}
+		}
+	}
+
+	public static class ArchiveUtils
+	{
+		public static void ExtractFiles(FileReference ArchiveFileName, DirectoryReference BaseDirectoryName, FileReference? ManifestFileName, ProgressValue Progress, ILogger Logger)
+		{
+			DateTime TimeStamp = DateTime.UtcNow;
+			using (ZipArchive Zip = new ZipArchive(File.OpenRead(ArchiveFileName.FullName)))
+			{
+				if (ManifestFileName != null)
+				{
+					FileReference.Delete(ManifestFileName);
+
+					// Create the manifest
+					ArchiveManifest Manifest = new ArchiveManifest();
+					foreach (ZipArchiveEntry Entry in Zip.Entries)
+					{
+						if (!Entry.FullName.EndsWith("/") && !Entry.FullName.EndsWith("\\"))
+						{
+							Manifest.Files.Add(new ArchiveManifestFile(Entry.FullName, Entry.Length, TimeStamp));
+						}
+					}
+
+					// Write it out to a temporary file, then move it into place
+					FileReference TempManifestFileName = ManifestFileName + ".tmp";
+					using (FileStream OutputStream = FileReference.Open(TempManifestFileName, FileMode.Create, FileAccess.Write))
+					{
+						Manifest.Write(OutputStream);
+					}
+					FileReference.Move(TempManifestFileName, ManifestFileName);
+				}
+
+				// Extract all the files
+				int EntryIdx = 0;
+				foreach(ZipArchiveEntry Entry in Zip.Entries)
+				{
+					if(!Entry.FullName.EndsWith("/") && !Entry.FullName.EndsWith("\\"))
+					{
+						FileReference FileName = FileReference.Combine(BaseDirectoryName, Entry.FullName);
+						DirectoryReference.CreateDirectory(FileName.Directory);
+						Logger.LogInformation("Writing {0}", FileName);
+
+						Entry.ExtractToFile(FileName.FullName, true);
+						FileReference.SetLastWriteTimeUtc(FileName, TimeStamp);
+					}
+					Progress.Set((float)++EntryIdx / (float)Zip.Entries.Count);
+				}
+			}
+		}
+
+		public static void RemoveExtractedFiles(DirectoryReference BaseDirectoryName, FileReference ManifestFileName, ProgressValue Progress, ILogger LogWriter)
+		{
+			// Read the manifest in
+			ArchiveManifest Manifest;
+			using(FileStream InputStream = FileReference.Open(ManifestFileName, FileMode.Open, FileAccess.Read))
+			{
+				Manifest = new ArchiveManifest(InputStream);
+			}
+
+			// Remove all the files that haven't been modified match
+			for(int Idx = 0; Idx < Manifest.Files.Count; Idx++)
+			{
+				FileInfo File = FileReference.Combine(BaseDirectoryName, Manifest.Files[Idx].FileName).ToFileInfo();
+				if(File.Exists)
+				{
+					if(File.Length != Manifest.Files[Idx].Length)
+					{
+						LogWriter.LogInformation("Skipping {FileName} due to modified length", File.FullName);
+					}
+					else if(Math.Abs((File.LastWriteTimeUtc - Manifest.Files[Idx].LastWriteTimeUtc).TotalSeconds) > 2.0)
+					{
+						LogWriter.LogInformation("Skipping {FileName} due to modified timestamp", File.FullName);
+					}
+					else
+					{
+						LogWriter.LogInformation("Removing {FileName}", File.FullName);
+						File.Delete();
+					}
+				}
+				Progress.Set((float)(Idx + 1) / (float)Manifest.Files.Count);
+			}
+		}
+	}
+}
diff --git a/Engine/Source/Programs/UnrealGameSync/UnrealGameSyncShared/BuildStep.cs b/Engine/Source/Programs/UnrealGameSync/UnrealGameSyncShared/BuildStep.cs
new file mode 100644
--- /dev/null
+++ b/Engine/Source/Programs/UnrealGameSync/UnrealGameSyncShared/BuildStep.cs
@@ -0,0 +1,224 @@
+// Copyright Epic Games, Inc. All Rights Reserved.
+
+using System;
+using System.Collections.Generic;
+using System.Diagnostics;
+using System.Linq;
+using System.Text;
+using System.Threading.Tasks;
+
+namespace UnrealGameSync
+{
+	public enum BuildStepType
+	{
+		Compile,
+		Cook,
+		Other,
+	}
+
+	[DebuggerDisplay("{Description}")]
+	public class BuildStep
+	{
+		public const string UniqueIdKey = "UniqueId";
+
+		public Guid UniqueId;
+		public int OrderIndex;
+		public string? Description;
+		public string? StatusText;
+		public int EstimatedDuration;
+		public BuildStepType Type;
+		public string? Target;
+		public string? Platform;
+		public string? Configuration;
+		public string? FileName;
+		public string? WorkingDir;
+		public string? Arguments;
+		public bool bUseLogWindow;
+		public bool bNormalSync;
+		public bool bScheduledSync;
+		public string? StatusPanelLink;
+		public Guid[] Requires = Array.Empty<Guid>();
+		public Guid ToolId;
+
+		public BuildStep(Guid InUniqueId, int InOrderIndex, string? InDescription, string? InStatusText, int InEstimatedDuration, string? InFileName, string? InArguments, string? InWorkingDir, bool bInUseLogWindow)
+		{
+			UniqueId = InUniqueId;
+			OrderIndex = InOrderIndex;
+			Description = InDescription;
+			StatusText = InStatusText;
+			EstimatedDuration = InEstimatedDuration;
+			Type = BuildStepType.Other;
+			FileName = InFileName;
+			Arguments = InArguments;
+			WorkingDir = InWorkingDir;
+			bUseLogWindow = bInUseLogWindow;
+		}
+
+		public BuildStep(Guid InUniqueId, int InOrderIndex, string? InDescription, string? InStatusText, int InEstimatedDuration, string? InTarget, string? InPlatform, string? InConfiguration, string? InArguments, bool bInSyncDefault)
+		{
+			UniqueId = InUniqueId;
+			OrderIndex = InOrderIndex;
+			Description = InDescription;
+			StatusText = InStatusText;
+			EstimatedDuration = InEstimatedDuration;
+			Type = BuildStepType.Compile;
+			Target = InTarget;
+			Platform = InPlatform;
+			Configuration = InConfiguration;
+			Arguments = InArguments;
+			bUseLogWindow = true;
+			bNormalSync = bInSyncDefault;
+			bScheduledSync = bInSyncDefault;
+		}
+
+		public BuildStep(ConfigObject Object)
+		{
+			if(!Guid.TryParse(Object.GetValue(UniqueIdKey, ""), out UniqueId))
+			{
+				UniqueId = Guid.NewGuid();
+			}
+			if(!Int32.TryParse(Object.GetValue("OrderIndex", ""), out OrderIndex))
+			{
+				OrderIndex = -1;
+			}
+
+			Description = Object.GetValue("Description", "Untitled");
+			StatusText = Object.GetValue("StatusText", "Untitled");
+
+			if(!int.TryParse(Object.GetValue("EstimatedDuration", ""), out EstimatedDuration) || EstimatedDuration < 1)
+			{
+				EstimatedDuration = 1;
+			}
+
+			if(!Enum.TryParse(Object.GetValue("Type", ""), true, out Type))
+			{
+				Type = BuildStepType.Other;
+			}
+
+			Target = Object.GetValue("Target");
+			Platform = Object.GetValue("Platform");
+			Configuration = Object.GetValue("Configuration");
+			FileName = Object.GetValue("FileName");
+			WorkingDir = Object.GetValue("WorkingDir");
+			Arguments = Object.GetValue("Arguments");
+
+			if(!Boolean.TryParse(Object.GetValue("bUseLogWindow", ""), out bUseLogWindow))
+			{
+				bUseLogWindow = true;
+			}
+			if(!Boolean.TryParse(Object.GetValue("bNormalSync", ""), out bNormalSync))
+			{
+				bNormalSync = true;
+			}
+			if(!Boolean.TryParse(Object.GetValue("bScheduledSync", ""), out bScheduledSync))
+			{
+				bScheduledSync = bNormalSync;
+			}
+
+			StatusPanelLink = Object.GetValue("Link", null);
+			if (String.IsNullOrEmpty(StatusPanelLink))
+			{
+				bool bShowAsTool;
+				if (Boolean.TryParse(Object.GetValue("bShowAsTool", ""), out bShowAsTool) && bShowAsTool)
+				{
+					StatusPanelLink = $"More...|{Description}";
+				}
+			}
+
+			List<Guid> Requires = new List<Guid>();
+			foreach (string RequireString in Object.GetValue("Requires", String.Empty).Split(';', StringSplitOptions.RemoveEmptyEntries))
+			{
+				if (Guid.TryParse(RequireString, out Guid Require))
+				{
+					Requires.Add(Require);
+				}
+			}
+			this.Requires = Requires.ToArray();
+	
+			if (!Guid.TryParse(Object.GetValue("Tool", ""), out ToolId))
+			{
+				ToolId = Guid.Empty;
+			}
+		}
+
+		public bool IsValid()
+		{
+			switch(Type)
+			{
+				case BuildStepType.Compile:
+					return Target != null && Platform != null && Configuration != null;
+				case BuildStepType.Cook:
+					return FileName != null;
+				case BuildStepType.Other:
+					return FileName != null;
+				default:
+					return false;
+			}
+		}
+
+		public static void MergeBuildStepObjects(Dictionary<Guid, ConfigObject> BuildStepObjects, IEnumerable<ConfigObject> ModifyObjects)
+		{
+			foreach(ConfigObject ModifyObject in ModifyObjects)
+			{
+				Guid UniqueId;
+				if(Guid.TryParse(ModifyObject.GetValue(UniqueIdKey, ""), out UniqueId))
+				{
+					ConfigObject? DefaultObject;
+					if(BuildStepObjects.TryGetValue(UniqueId, out DefaultObject))
+					{
+						ModifyObject.SetDefaults(DefaultObject);
+					}
+					BuildStepObjects[UniqueId] = ModifyObject;
+				}
+			}
+		}
+
+		public ConfigObject ToConfigObject()
+		{
+			ConfigObject Result = new ConfigObject();
+			Result["UniqueId"] = UniqueId.ToString();
+			Result["Description"] = Description;
+			Result["StatusText"] = StatusText;
+			Result["EstimatedDuration"] = EstimatedDuration.ToString();
+			Result["Type"] = Type.ToString();
+			switch(Type)
+			{
+				case BuildStepType.Compile:
+					Result["Target"] = Target;
+					Result["Platform"] = Platform;
+					Result["Configuration"] = Configuration;
+					Result["Arguments"] = Arguments;
+					break;
+				case BuildStepType.Cook:
+					Result["FileName"] = FileName;
+					break;
+				case BuildStepType.Other:
+					Result["FileName"] = FileName;
+					Result["WorkingDir"] = WorkingDir;
+					Result["Arguments"] = Arguments;
+					Result["bUseLogWindow"] = bUseLogWindow.ToString();
+					break;
+			}
+			Result["OrderIndex"] = OrderIndex.ToString();
+			Result["bNormalSync"] = bNormalSync.ToString();
+			Result["bScheduledSync"] = bScheduledSync.ToString();
+			if (!String.IsNullOrEmpty(StatusPanelLink))
+			{
+				Result["Link"] = StatusPanelLink;
+			}
+			if (ToolId != Guid.Empty)
+			{
+				Result["Tool"] = ToolId.ToString();
+			}
+			return Result;
+		}
+
+		public ConfigObject? ToConfigObject(ConfigObject? DefaultObject)
+		{
+			ConfigObject Result = new ConfigObject();
+			Result[UniqueIdKey] = UniqueId.ToString();
+			Result.AddOverrides(ToConfigObject(), DefaultObject);
+			return (Result.Pairs.Count <= 1)? null : Result;
+		}
+	}
+}
diff --git a/Engine/Source/Programs/UnrealGameSync/UnrealGameSyncShared/ConfigFile.cs b/Engine/Source/Programs/UnrealGameSync/UnrealGameSyncShared/ConfigFile.cs
new file mode 100644
--- /dev/null
+++ b/Engine/Source/Programs/UnrealGameSync/UnrealGameSyncShared/ConfigFile.cs
@@ -0,0 +1,640 @@
+// Copyright Epic Games, Inc. All Rights Reserved.
+
+using EpicGames.Core;
+using Microsoft.Extensions.Logging;
+using System;
+using System.Collections.Generic;
+using System.Diagnostics;
+using System.Diagnostics.CodeAnalysis;
+using System.IO;
+using System.Linq;
+using System.Text;
+using System.Threading.Tasks;
+
+namespace UnrealGameSync
+{
+	public class ConfigObject
+	{
+		const string ConfigSeparatorCharacters = "(),= \t\"";
+
+		public List<KeyValuePair<string, string>> Pairs;
+
+		public ConfigObject()
+		{
+			Pairs = new List<KeyValuePair<string,string>>();
+		}
+
+		public ConfigObject(ConfigObject Other)
+		{
+			Pairs = new List<KeyValuePair<string,string>>(Other.Pairs);
+		}
+
+		public ConfigObject(string Text)
+		{
+			Pairs = new List<KeyValuePair<string,string>>();
+			ParseConfigString(Text);
+		}
+
+		public ConfigObject(ConfigObject BaseObject, string Text)
+		{
+			Pairs = new List<KeyValuePair<string,string>>(BaseObject.Pairs);
+			ParseConfigString(Text);
+		}
+
+		void ParseConfigString(string Text)
+		{
+			int Idx = 0;
+			if(ParseConfigToken(Text, ref Idx) == "(")
+			{
+				while(Idx < Text.Length)
+				{
+					// Read the next key/value pair
+					string? Key = ParseConfigToken(Text, ref Idx);
+					if (Key == null)
+					{
+						return;
+					}
+					if(ParseConfigToken(Text, ref Idx) == "=")
+					{
+						string? Value = ParseConfigValueToken(Text, ref Idx);
+						if (Value == null)
+						{
+							return;
+						}
+						SetValue(Key, Value);
+					}
+
+					// Check for the end of the list, or a comma before the next pair
+					for(;;)
+					{
+						string? Token = ParseConfigValueToken(Text, ref Idx);
+						if(Token == ",")
+						{
+							break;
+						}
+						if(Token == ")" || Token == null)
+						{
+							return;
+						}
+					}
+				}
+			}
+		}
+
+		static string? ParseConfigToken(string Text, ref int Idx)
+		{
+			// Skip whitespace
+			while(Idx < Text.Length && Char.IsWhiteSpace(Text[Idx]))
+			{
+				Idx++;
+			}
+			if(Idx == Text.Length)
+			{
+				return null;
+			}
+
+			// Read the token
+			if (Text[Idx] == '\"')
+			{
+				StringBuilder Token = new StringBuilder();
+				while (++Idx < Text.Length)
+				{
+					if (Text[Idx] == '\"')
+					{
+						Idx++;
+						break;
+					}
+					if (Text[Idx] == '\\' && Idx + 1 < Text.Length)
+					{
+						Idx++;
+					}
+					Token.Append(Text[Idx]);
+				}
+				return Token.ToString();
+			}
+			else if (ConfigSeparatorCharacters.IndexOf(Text[Idx]) != -1)
+			{
+				return Text[Idx++].ToString();
+			}
+			else
+			{
+				int StartIdx = Idx;
+				while (Idx < Text.Length && ConfigSeparatorCharacters.IndexOf(Text[Idx]) == -1)
+				{
+					Idx++;
+				}
+				return Text.Substring(StartIdx, Idx - StartIdx);
+			}
+		}
+
+		static string? ParseConfigValueToken(string Text, ref int Idx)
+		{
+			string? Token = ParseConfigToken(Text, ref Idx);
+			if (Token == "(")
+			{
+				int StartIdx = Idx - 1;
+				for (; ; )
+				{
+					string? NextToken = ParseConfigValueToken(Text, ref Idx);
+					if (NextToken == null || NextToken == ")")
+					{
+						break;
+					}
+				}
+				Token = Text.Substring(StartIdx, Idx - StartIdx);
+			}
+			return Token;
+		}
+
+		[return: NotNullIfNotNull("DefaultValue")]
+		public string? GetValue(string Key, string? DefaultValue = null)
+		{
+			for(int Idx = 0; Idx < Pairs.Count; Idx++)
+			{
+				if(Pairs[Idx].Key.Equals(Key, StringComparison.InvariantCultureIgnoreCase))
+				{
+					return Pairs[Idx].Value;
+				}
+			}
+			return DefaultValue;
+		}
+
+		public Guid GetValue(string Key, Guid DefaultValue)
+		{
+			string? StringValue = GetValue(Key);
+			if(StringValue != null)
+			{
+				Guid Value;
+				if(Guid.TryParse(StringValue, out Value))
+				{
+					return Value;
+				}
+			}
+			return DefaultValue;
+		}
+
+		public int GetValue(string Key, int DefaultValue)
+		{
+			string? StringValue = GetValue(Key);
+			if(StringValue != null)
+			{
+				int Value;
+				if(int.TryParse(StringValue, out Value))
+				{
+					return Value;
+				}
+			}
+			return DefaultValue;
+		}
+
+		public bool GetValue(string Key, bool DefaultValue)
+		{
+			string? StringValue = GetValue(Key);
+			if(StringValue != null)
+			{
+				bool Value;
+				if(bool.TryParse(StringValue, out Value))
+				{
+					return Value;
+				}
+			}
+			return DefaultValue;
+		}
+
+		public void SetValue(string Key, string? Value)
+		{
+			if (Value == null)
+			{
+				Pairs.RemoveAll(x => x.Key.Equals(Key, StringComparison.OrdinalIgnoreCase));
+			}
+			else
+			{
+				for (int Idx = 0; Idx < Pairs.Count; Idx++)
+				{
+					if (Pairs[Idx].Key.Equals(Key, StringComparison.OrdinalIgnoreCase))
+					{
+						Pairs[Idx] = new KeyValuePair<string, string>(Key, Value);
+						return;
+					}
+				}
+				Pairs.Add(new KeyValuePair<string, string>(Key, Value));
+			}
+		}
+
+		public void SetValue(string Key, Guid Value)
+		{
+			SetValue(Key, Value.ToString());
+		}
+
+		public void SetValue(string Key, int Value)
+		{
+			SetValue(Key, Value.ToString());
+		}
+
+		public void SetValue(string Key, bool Value)
+		{
+			SetValue(Key, Value.ToString());
+		}
+
+		public string? this[string Key]
+		{
+			get { return GetValue(Key); }
+			set { SetValue(Key, value); }
+		}
+
+		public void SetDefaults(ConfigObject Other)
+		{
+			foreach(KeyValuePair<string, string> Pair in Other.Pairs)
+			{
+				if(GetValue(Pair.Key) == null)
+				{
+					SetValue(Pair.Key, Pair.Value);
+				}
+			}
+		}
+
+		public void AddOverrides(ConfigObject Object, ConfigObject? DefaultObject)
+		{
+			foreach(KeyValuePair<string, string> Pair in Object.Pairs)
+			{
+				if(DefaultObject == null || DefaultObject.GetValue(Pair.Key) != Pair.Value)
+				{
+					SetValue(Pair.Key, Pair.Value);
+				}
+			}
+		}
+
+		public string ToString(ConfigObject? BaseObject)
+		{
+			StringBuilder Result = new StringBuilder();
+			Result.Append("(");
+			foreach(KeyValuePair<string, string> Pair in Pairs)
+			{
+				if(BaseObject == null || BaseObject.GetValue(Pair.Key) != Pair.Value)
+				{
+					if(Result.Length > 1)
+					{
+						Result.Append(", ");
+					}
+					Result.Append(Pair.Key);
+					Result.Append("=");
+					if(Pair.Value == null)
+					{
+						Result.Append("\"\"");
+					}
+					else
+					{
+						Result.AppendFormat("\"{0}\"", Pair.Value.Replace("\\", "\\\\").Replace("\"", "\\\""));
+					}
+				}
+			}
+			Result.Append(")");
+			return Result.ToString();
+		}
+ 
+		public override string ToString()
+		{
+			return ToString(null);
+		}
+	}
+
+	[DebuggerDisplay("{Name}")]
+	public class ConfigSection
+	{
+		public string Name;
+		public Dictionary<string, string> Pairs = new Dictionary<string,string>();
+
+		public ConfigSection(string InName)
+		{
+			Name = InName;
+		}
+
+		public void Clear()
+		{
+			Pairs.Clear();
+		}
+
+		public void SetValue(string Key, int Value)
+		{
+			Pairs[Key] = Value.ToString();
+		}
+
+		public void SetValue(string Key, bool Value)
+		{
+			Pairs[Key] = Value? "1" : "0";
+		}
+
+		public void SetValue(string Key, string Value)
+		{
+			if(Value == null)
+			{
+				RemoveValue(Key);
+			}
+			else
+			{
+				Pairs[Key] = Value;
+			}
+		}
+
+		public void SetValues(string Key, string[] Values)
+		{
+			if(Values == null || Values.Length == 0)
+			{
+				RemoveValue(Key);
+			}
+			else
+			{
+				Pairs[Key] = String.Join("\n", Values);
+			}
+		}
+
+		public void SetValues(string Key, Guid[] Values)
+		{
+			if(Values == null)
+			{
+				RemoveValue(Key);
+			}
+			else
+			{
+				Pairs[Key] = String.Join("\n", Values.Select(x => x.ToString()));
+			}
+		}
+
+		public void AppendValue(string Key, string Value)
+		{
+			string? CurrentValue;
+			if(Pairs.TryGetValue(Key, out CurrentValue))
+			{
+				Pairs[Key] = CurrentValue + "\n" + Value;
+			}
+			else
+			{
+				Pairs[Key] = Value;
+			}
+		}
+
+		public void RemoveValue(string Key)
+		{
+			Pairs.Remove(Key);
+		}
+
+		public int GetValue(string Key, int DefaultValue)
+		{
+			string? ValueString = GetValue(Key);
+			if(ValueString != null)
+			{
+				int Value;
+				if(int.TryParse(ValueString, out Value))
+				{
+					return Value;
+				}
+			}
+			return DefaultValue;
+		}
+
+		public bool GetValue(string Key, bool DefaultValue)
+		{
+			return GetValue(Key, DefaultValue? 1 : 0) != 0;
+		}
+
+		[return: NotNullIfNotNull("DefaultValue")]
+		public string? GetValue(string Key, string? DefaultValue = null)
+		{
+			string? Value;
+			if(!Pairs.TryGetValue(Key, out Value))
+			{
+				Value = DefaultValue;
+			}
+			return Value;
+		}
+
+		[return: NotNullIfNotNull("DefaultValue")]
+		public string[]? GetValues(string Key, string[]? DefaultValue = null)
+		{
+			string? Value = GetValue(Key, null);
+			if(Value == null)
+			{
+				return DefaultValue;
+			}
+			else
+			{
+				return Value.Split('\n');
+			}
+		}
+
+		[return: NotNullIfNotNull("DefaultValue")]
+		public Guid[]? GetValues(string Key, Guid[]? DefaultValue = null)
+		{
+			string[]? StringValues = GetValues(Key, (string[]?)null);
+			if(StringValues == null)
+			{
+				return DefaultValue;
+			}
+
+			List<Guid> GuidValues = new List<Guid>();
+			foreach(string StringValue in StringValues)
+			{
+				Guid GuidValue;
+				if(Guid.TryParse(StringValue, out GuidValue))
+				{
+					GuidValues.Add(GuidValue);
+				}
+			}
+
+			return GuidValues.ToArray();
+		}
+	}
+
+	public class ConfigFile
+	{
+		List<ConfigSection> Sections = new List<ConfigSection>();
+
+		public ConfigFile()
+		{
+		}
+
+		FileReference GetTempFileName(FileReference FileName)
+		{
+			return FileName + ".tmp";
+		}
+
+		public void Load(FileReference FileName)
+		{
+			Parse(FileReference.ReadAllLines(FileName));
+		}
+
+		public bool TryLoad(FileReference FileName, ILogger Logger)
+		{
+			FileInfo FileInfo = FileName.ToFileInfo();
+			if (FileInfo.Exists)
+			{
+				Logger.LogInformation("Loading config file from {File} ({Size} bytes)", FileInfo.FullName, FileInfo.Length);
+				Load(FileName);
+				return true;
+			}
+
+			FileReference TempFileName = GetTempFileName(FileName);
+
+			FileInfo TempFileInfo = TempFileName.ToFileInfo();
+			if (TempFileInfo.Exists)
+			{
+				Logger.LogInformation("Loading temporary config file from {File} ({Size} bytes)", TempFileInfo.FullName, TempFileInfo.Length);
+				Load(TempFileName);
+				return true;
+			}
+
+			Logger.LogInformation("No existing config file at {File}", FileName);
+			return false;
+		}
+
+		public void Parse(string[] Lines)
+		{
+			ConfigSection? CurrentSection = null;
+			foreach(string Line in Lines)
+			{
+				string TrimLine = Line.Trim();
+				if(!TrimLine.StartsWith(";"))
+				{
+					if(TrimLine.StartsWith("[") && TrimLine.EndsWith("]"))
+					{
+						string SectionName = TrimLine.Substring(1, TrimLine.Length - 2).Trim();
+						CurrentSection = FindOrAddSection(SectionName);
+					}
+					else if(CurrentSection != null)
+					{
+						int EqualsIdx = TrimLine.IndexOf('=');
+						if(EqualsIdx != -1)
+						{
+							string Value = Line.Substring(EqualsIdx + 1).TrimStart();
+							if(TrimLine.StartsWith("+"))
+							{
+								CurrentSection.AppendValue(TrimLine.Substring(1, EqualsIdx - 1).Trim(), Value);
+							}
+							else
+							{
+								CurrentSection.SetValue(TrimLine.Substring(0, EqualsIdx).TrimEnd(), Value);
+							}
+						}
+					}
+				}
+			}
+		}
+
+		public void Save(FileReference FileName)
+		{
+			FileReference TempFileName = GetTempFileName(FileName);
+			FileReference.Delete(TempFileName);
+
+			using (StreamWriter Writer = new StreamWriter(TempFileName.FullName))
+			{
+				for(int Idx = 0; Idx < Sections.Count; Idx++)
+				{
+					Writer.WriteLine("[{0}]", Sections[Idx].Name);
+					foreach(KeyValuePair<string, string> Pair in Sections[Idx].Pairs)
+					{
+						if(Pair.Value.Contains('\n'))
+						{
+							foreach(string Line in Pair.Value.Split('\n'))
+							{
+								Writer.WriteLine("+{0}={1}", Pair.Key, Line);
+							}
+						}
+						else
+						{
+							Writer.WriteLine("{0}={1}", Pair.Key, Pair.Value);
+						}
+					}
+					if(Idx < Sections.Count - 1)
+					{
+						Writer.WriteLine();
+					}
+				}
+			}
+
+			FileReference.Delete(FileName);
+			FileReference.Move(TempFileName, FileName);
+		}
+
+		public ConfigSection FindSection(string Name)
+		{
+			return Sections.FirstOrDefault(x => String.Compare(x.Name, Name, true) == 0);
+		}
+
+		public ConfigSection FindOrAddSection(string Name)
+		{
+			ConfigSection Section = FindSection(Name);
+			if(Section == null)
+			{
+				Section = new ConfigSection(Name);
+				Sections.Add(Section);
+			}
+			return Section;
+		}
+
+		public void SetValue(string Key, int Value)
+		{
+			int DotIdx = Key.IndexOf('.');
+			ConfigSection Section = FindOrAddSection(Key.Substring(0, DotIdx));
+			Section.SetValue(Key.Substring(DotIdx + 1), Value);
+		}
+
+		public void SetValue(string Key, bool Value)
+		{
+			int DotIdx = Key.IndexOf('.');
+			ConfigSection Section = FindOrAddSection(Key.Substring(0, DotIdx));
+			Section.SetValue(Key.Substring(DotIdx + 1), Value);
+		}
+
+		public void SetValue(string Key, string Value)
+		{
+			int DotIdx = Key.IndexOf('.');
+			ConfigSection Section = FindOrAddSection(Key.Substring(0, DotIdx));
+			Section.SetValue(Key.Substring(DotIdx + 1), Value);
+		}
+
+		public void SetValues(string Key, string[] Values)
+		{
+			int DotIdx = Key.IndexOf('.');
+			ConfigSection Section = FindOrAddSection(Key.Substring(0, DotIdx));
+			Section.SetValues(Key.Substring(DotIdx + 1), Values);
+		}
+
+		public bool GetValue(string Key, bool DefaultValue)
+		{
+			int DotIdx = Key.IndexOf('.');
+			ConfigSection Section = FindSection(Key.Substring(0, DotIdx));
+			return (Section == null)? DefaultValue : Section.GetValue(Key.Substring(DotIdx + 1), DefaultValue);
+		}
+
+		public int GetValue(string Key, int DefaultValue)
+		{
+			int DotIdx = Key.IndexOf('.');
+			ConfigSection Section = FindSection(Key.Substring(0, DotIdx));
+			return (Section == null)? DefaultValue : Section.GetValue(Key.Substring(DotIdx + 1), DefaultValue);
+		}
+
+		[return: NotNullIfNotNull("DefaultValue")]
+		public string? GetValue(string Key, string? DefaultValue)
+		{
+			int DotIdx = Key.IndexOf('.');
+			ConfigSection Section = FindSection(Key.Substring(0, DotIdx));
+			return (Section == null)? DefaultValue : Section.GetValue(Key.Substring(DotIdx + 1), DefaultValue);
+		}
+
+		[return: NotNullIfNotNull("DefaultValue")]
+		public string[]? GetValues(string Key, string[]? DefaultValue)
+		{
+			int DotIdx = Key.IndexOf('.');
+			ConfigSection Section = FindSection(Key.Substring(0, DotIdx));
+			return (Section == null)? DefaultValue : Section.GetValues(Key.Substring(DotIdx + 1), DefaultValue);
+		}
+
+		[return: NotNullIfNotNull("DefaultValue")]
+		public Guid[]? GetGuidValues(string Key, Guid[]? DefaultValue)
+		{
+			int DotIdx = Key.IndexOf('.');
+			ConfigSection Section = FindSection(Key.Substring(0, DotIdx));
+			return (Section == null)? DefaultValue : Section.GetValues(Key.Substring(DotIdx + 1), DefaultValue);
+		}
+	}
+}
diff --git a/Engine/Source/Programs/UnrealGameSync/UnrealGameSyncShared/FileFilter.cs b/Engine/Source/Programs/UnrealGameSync/UnrealGameSyncShared/FileFilter.cs
new file mode 100644
--- /dev/null
+++ b/Engine/Source/Programs/UnrealGameSync/UnrealGameSyncShared/FileFilter.cs
@@ -0,0 +1,541 @@
+// Copyright Epic Games, Inc. All Rights Reserved.
+
+using System;
+using System.Collections.Generic;
+using System.IO;
+using System.Linq;
+using System.Text;
+using System.Threading.Tasks;
+
+namespace UnrealGameSync
+{
+	/// <summary>
+	/// Indicates whether files which match a pattern should be included or excluded
+	/// </summary>
+	public enum FileFilterType
+	{
+		Include,
+		Exclude,
+	}
+
+	/// <summary>
+	/// Stores a set of rules, similar to p4 path specifications, which can be used to efficiently filter files in or out of a set.
+	/// The rules are merged into a tree with nodes for each path fragments and rule numbers in each leaf, allowing tree traversal in an arbitrary order 
+	/// to how they are applied.
+	/// </summary>
+	public class FileFilter
+	{
+		/// <summary>
+		/// Root node for the tree.
+		/// </summary>
+		FileFilterNode RootNode;
+
+		/// <summary>
+		/// The default node, which will match any path
+		/// </summary>
+		FileFilterNode DefaultNode;
+
+		/// <summary>
+		/// Terminating nodes for each rule added to the filter
+		/// </summary>
+		List<FileFilterNode> Rules = new List<FileFilterNode>();
+
+		/// <summary>
+		/// Default constructor
+		/// </summary>
+		public FileFilter(FileFilterType DefaultType = FileFilterType.Exclude)
+		{
+			RootNode = new FileFilterNode(null, "");
+
+			DefaultNode = new FileFilterNode(RootNode, "...");
+			DefaultNode.Type = DefaultType;
+		}
+
+		/// <summary>
+		/// Copy constructor
+		/// </summary>
+		/// <param name="Other">Filter to copy from</param>
+		public FileFilter(FileFilter Other)
+			: this(Other.DefaultNode.Type)
+		{
+			foreach (FileFilterNode OtherRule in Other.Rules)
+			{
+				AddRule(OtherRule.ToString(), OtherRule.Type);
+			}
+		}
+
+		/// <summary>
+		/// Constructs a file filter from a p4-style filespec. Exclude lines are prefixed with a - character.
+		/// </summary>
+		/// <param name="Lines">Lines to parse rules from</param>
+		public FileFilter(IEnumerable<string> Lines)
+			: this()
+		{
+			foreach (string Line in Lines)
+			{
+				AddRule(Line);
+			}
+		}
+
+		/// <summary>
+		/// Adds an include or exclude rule to the filter
+		/// </summary>
+		/// <param name="Pattern">Pattern to match. See CreateRegex() for details.</param>
+		public void AddRule(string Rule)
+		{
+			if (Rule.StartsWith("-"))
+			{
+				Exclude(Rule.Substring(1).TrimStart());
+			}
+			else
+			{
+				Include(Rule);
+			}
+		}
+
+		/// <summary>
+		/// Adds an include or exclude rule to the filter. The rule may be 
+		/// </summary>
+		/// <param name="Pattern">Pattern to match. See CreateRegex() for details.</param>
+		public void AddRule(string Rule, params string[] AllowTags)
+		{
+			string CleanRule = Rule.Trim();
+			if(CleanRule.StartsWith("{"))
+			{
+				// Find the end of the condition
+				int ConditionEnd = CleanRule.IndexOf('}');
+				if(ConditionEnd == -1)
+				{
+					throw new Exception(String.Format("Missing closing parenthesis in rule: {0}", CleanRule));
+				}
+
+				// Check there's a matching tag
+				string[] RuleTags = CleanRule.Substring(1, ConditionEnd - 1).Split(',').Select(x => x.Trim()).ToArray();
+				if(!RuleTags.Any(x => AllowTags.Contains(x)))
+				{
+					return;
+				}
+
+				// Strip the condition from the rule
+				CleanRule = CleanRule.Substring(ConditionEnd + 1).TrimStart(); 
+			}
+			AddRule(CleanRule);
+		}
+
+		/// <summary>
+		/// Adds several rules to the filter
+		/// </summary>
+		/// <param name="Patterns">List of patterns to match.</param>
+		public void AddRules(IEnumerable<string> Rules)
+		{
+			foreach(string Rule in Rules)
+			{
+				AddRule(Rule);
+			}
+		}
+
+		/// <summary>
+		/// Adds several rules in the given lines. Rules may be prefixed with conditions of the syntax {Key=Value, Key2=Value2}, which
+		/// will be evaluated using variables in the given dictionary before being added.
+		/// </summary>
+		/// <param name="Lines"></param>
+		/// <param name="Variables">Lookup for variables to test against</param>
+		public void AddRules(IEnumerable<string> Rules, params string[] Tags)
+		{
+			foreach(string Rule in Rules)
+			{
+				AddRule(Rule, Tags);
+			}
+		}
+
+		/// <summary>
+		/// Reads a configuration file split into sections
+		/// </summary>
+		/// <param name="Filter"></param>
+		/// <param name="RulesFileName"></param>
+		/// <param name="Conditions"></param>
+		public void ReadRulesFromFile(string FileName, string SectionName, params string[] AllowTags)
+		{
+			bool bInSection = false;
+			foreach(string Line in File.ReadAllLines(FileName))
+			{
+				string TrimLine = Line.Trim();
+				if(!TrimLine.StartsWith(";") && TrimLine.Length > 0)
+				{
+					if(TrimLine.StartsWith("["))
+					{
+						bInSection = (TrimLine == "[" + SectionName + "]");
+					}
+					else if(bInSection)
+					{
+						AddRule(Line, AllowTags);
+					}
+				}
+			}
+		}
+
+		/// <summary>
+		/// Adds an include rule to the filter
+		/// </summary>
+		/// <param name="Pattern">Pattern to match. See CreateRegex() for details.</param>
+		public void Include(string Pattern)
+		{
+			AddRule(Pattern, FileFilterType.Include);
+		}
+
+		/// <summary>
+		/// Adds several exclude rules to the filter
+		/// </summary>
+		/// <param name="Patterns">Patterns to match. See CreateRegex() for details.</param>
+		public void Include(IEnumerable<string> Patterns)
+		{
+			foreach (string Pattern in Patterns)
+			{
+				Include(Pattern);
+			}
+		}
+
+		/// <summary>
+		/// Adds an exclude rule to the filter
+		/// </summary>
+		/// <param name="Pattern">Mask to match. See CreateRegex() for details.</param>
+		public void Exclude(string Pattern)
+		{
+			AddRule(Pattern, FileFilterType.Exclude);
+		}
+
+		/// <summary>
+		/// Adds several exclude rules to the filter
+		/// </summary>
+		/// <param name="Patterns">Patterns to match. See CreateRegex() for details.</param>
+		public void Exclude(IEnumerable<string> Patterns)
+		{
+			foreach (string Pattern in Patterns)
+			{
+				Exclude(Pattern);
+			}
+		}
+
+		/// <summary>
+		/// Adds an include or exclude rule to the filter
+		/// </summary>
+		/// <param name="Pattern">The pattern which the rule should match</param>
+		/// <param name="bInclude">Whether to include or exclude files matching this rule</param>
+		public void AddRule(string Pattern, FileFilterType Type)
+		{
+			string NormalizedPattern = Pattern.Replace('\\', '/');
+
+			// We don't want a slash at the start, but if there was not one specified, it's not anchored to the root of the tree.
+			if (NormalizedPattern.StartsWith("/"))
+			{
+				NormalizedPattern = NormalizedPattern.Substring(1);
+			}
+			else if(!NormalizedPattern.StartsWith("..."))
+			{
+				NormalizedPattern = ".../" + NormalizedPattern;
+			}
+
+			// All directories indicate a wildcard match
+			if (NormalizedPattern.EndsWith("/"))
+			{
+				NormalizedPattern += "...";
+			}
+
+			// Replace any directory wildcards mid-string
+			for (int Idx = NormalizedPattern.IndexOf("..."); Idx != -1; Idx = NormalizedPattern.IndexOf("...", Idx))
+			{
+				if (Idx > 0 && NormalizedPattern[Idx - 1] != '/')
+				{
+					NormalizedPattern = NormalizedPattern.Insert(Idx, "*/");
+					Idx++;
+				}
+
+				Idx += 3;
+
+				if (Idx < NormalizedPattern.Length && NormalizedPattern[Idx] != '/')
+				{
+					NormalizedPattern = NormalizedPattern.Insert(Idx, "/*");
+					Idx += 2;
+				}
+			}
+
+			// Split the pattern into fragments
+			string[] BranchPatterns = NormalizedPattern.Split('/');
+
+			// Add it into the tree
+			FileFilterNode LastNode = RootNode;
+			foreach (string BranchPattern in BranchPatterns)
+			{
+				FileFilterNode NextNode = LastNode.Branches.FirstOrDefault(x => x.Pattern == BranchPattern);
+				if (NextNode == null)
+				{
+					NextNode = new FileFilterNode(LastNode, BranchPattern);
+					LastNode.Branches.Add(NextNode);
+				}
+				LastNode = NextNode;
+			}
+
+			// We've reached the end of the pattern, so mark it as a leaf node
+			Rules.Add(LastNode);
+			LastNode.RuleNumber = Rules.Count - 1;
+			LastNode.Type = Type;
+
+			// Update the maximums along that path
+			for (FileFilterNode? UpdateNode = LastNode; UpdateNode != null; UpdateNode = UpdateNode.Parent)
+			{
+				if (Type == FileFilterType.Include)
+				{
+					UpdateNode.MaxIncludeRuleNumber = LastNode.RuleNumber;
+				}
+				else
+				{
+					UpdateNode.MaxExcludeRuleNumber = LastNode.RuleNumber;
+				}
+			}
+		}
+
+		/// <summary>
+		/// Excludes all confidential folders from the filter
+		/// </summary>
+		public void ExcludeConfidentialFolders()
+		{
+			AddRule(".../EpicInternal/...", FileFilterType.Exclude);
+			AddRule(".../CarefullyRedist/...", FileFilterType.Exclude);
+			AddRule(".../NotForLicensees/...", FileFilterType.Exclude);
+			AddRule(".../NoRedist/...", FileFilterType.Exclude);
+		}
+
+		/// <summary>
+		/// Determines whether the given file matches the filter
+		/// </summary>
+		/// <param name="FileName">File to match</param>
+		/// <returns>True if the file passes the filter</returns>
+		public bool Matches(string FileName)
+		{
+			string[] Tokens = FileName.TrimStart('/', '\\').Split('/', '\\');
+
+			FileFilterNode MatchingNode = FindMatchingNode(RootNode, Tokens, 0, DefaultNode);
+
+			return MatchingNode.Type == FileFilterType.Include;
+		}
+
+		/// <summary>
+		/// Determines whether it's possible for anything within the given folder name to match the filter. Useful to early out of recursive file searches.
+		/// </summary>
+		/// <param name="FolderName">File to match</param>
+		/// <returns>True if the file passes the filter</returns>
+		public bool PossiblyMatches(string FolderName)
+		{
+			string[] Tokens = FolderName.Trim('/', '\\').Split('/', '\\');
+
+			FileFilterNode MatchingNode = FindMatchingNode(RootNode, Tokens.Union(new string[]{ "" }).ToArray(), 0, DefaultNode);
+
+			return MatchingNode.Type == FileFilterType.Include || HighestPossibleIncludeMatch(RootNode, Tokens, 0, MatchingNode.RuleNumber) > MatchingNode.RuleNumber;
+		}
+
+		/// <summary>
+		/// Applies the filter to each element in a sequence, and returns the list of files that match
+		/// </summary>
+		/// <param name="FileNames">List of filenames</param>
+		/// <returns>List of filenames which match the filter</returns>
+		public IEnumerable<string> ApplyTo(IEnumerable<string> FileNames)
+		{
+			return FileNames.Where(x => Matches(x));
+		}
+
+		/// <summary>
+		/// Finds the node which matches a given list of tokens.
+		/// </summary>
+		/// <param name="CurrentNode"></param>
+		/// <param name="Tokens"></param>
+		/// <param name="TokenIdx"></param>
+		/// <param name="CurrentBestNode"></param>
+		/// <returns></returns>
+		FileFilterNode FindMatchingNode(FileFilterNode CurrentNode, string[] Tokens, int TokenIdx, FileFilterNode CurrentBestNode)
+		{
+			// If we've matched all the input tokens, check if this rule is better than any other we've seen
+			if (TokenIdx == Tokens.Length)
+			{
+				return (CurrentNode.RuleNumber > CurrentBestNode.RuleNumber) ? CurrentNode : CurrentBestNode;
+			}
+
+			// If there is no rule under the current node which is better than the current best node, early out
+			if (CurrentNode.MaxIncludeRuleNumber <= CurrentBestNode.RuleNumber && CurrentNode.MaxExcludeRuleNumber <= CurrentBestNode.RuleNumber)
+			{
+				return CurrentBestNode;
+			}
+
+			// Test all the branches for one that matches
+			FileFilterNode BestNode = CurrentBestNode;
+			foreach (FileFilterNode Branch in CurrentNode.Branches)
+			{
+				if (Branch.Pattern == "...")
+				{
+					for (int NextTokenIdx = Tokens.Length; NextTokenIdx >= TokenIdx; NextTokenIdx--)
+					{
+						BestNode = FindMatchingNode(Branch, Tokens, NextTokenIdx, BestNode);
+					}
+				}
+				else
+				{
+					if (Branch.IsMatch(Tokens[TokenIdx]))
+					{
+						BestNode = FindMatchingNode(Branch, Tokens, TokenIdx + 1, BestNode);
+					}
+				}
+			}
+			return BestNode;
+		}
+
+		/// <summary>
+		/// Returns the highest possible rule number which can match the given list of input tokens, assuming that the list of input tokens is incomplete.
+		/// </summary>
+		/// <param name="CurrentNode">The current node being checked</param>
+		/// <param name="Tokens">The tokens to match</param>
+		/// <param name="TokenIdx">Current token index</param>
+		/// <param name="CurrentBestRuleNumber">The highest rule number seen so far. Used to optimize tree traversals.</param>
+		/// <returns>New highest rule number</returns>
+		int HighestPossibleIncludeMatch(FileFilterNode CurrentNode, string[] Tokens, int TokenIdx, int CurrentBestRuleNumber)
+		{
+			// If we've matched all the input tokens, check if this rule is better than any other we've seen
+			if (TokenIdx == Tokens.Length)
+			{
+				return Math.Max(CurrentBestRuleNumber, CurrentNode.MaxIncludeRuleNumber);
+			}
+
+			// Test all the branches for one that matches
+			int BestRuleNumber = CurrentBestRuleNumber;
+			if (CurrentNode.MaxIncludeRuleNumber > BestRuleNumber)
+			{
+				foreach (FileFilterNode Branch in CurrentNode.Branches)
+				{
+					if (Branch.Pattern == "...")
+					{
+						if (Branch.MaxIncludeRuleNumber > BestRuleNumber)
+						{
+							BestRuleNumber = Branch.MaxIncludeRuleNumber;
+						}
+					}
+					else
+					{
+						if (Branch.IsMatch(Tokens[TokenIdx]))
+						{
+							BestRuleNumber = HighestPossibleIncludeMatch(Branch, Tokens, TokenIdx + 1, BestRuleNumber);
+						}
+					}
+				}
+			}
+			return BestRuleNumber;
+		}
+	}
+
+	/// <summary>
+	/// Node within a filter tree. Each node matches a single path fragment - a folder or file name, with an include or exclude rule consisting of a sequence of nodes. 
+	/// </summary>
+	class FileFilterNode
+	{
+		/// <summary>
+		/// Node which this is parented to. Null for the root node.
+		/// </summary>
+		public readonly FileFilterNode? Parent;
+
+		/// <summary>
+		/// Pattern to match for this node. May contain * or ? wildcards.
+		/// </summary>
+		public readonly string Pattern;
+
+		/// <summary>
+		/// Highest include rule number matched by this node or any child nodes.
+		/// </summary>
+		public int MaxIncludeRuleNumber = -1;
+
+		/// <summary>
+		/// Highest exclude rule number matched by this node or any child nodes.
+		/// </summary>
+		public int MaxExcludeRuleNumber = -1;
+
+		/// <summary>
+		/// Child branches of this node, distinct by the the pattern for each.
+		/// </summary>
+		public List<FileFilterNode> Branches = new List<FileFilterNode>();
+
+		/// <summary>
+		/// If a path matches the sequence of nodes ending in this node, the number of the rule which added it. -1 if this was not the last node in a rule.
+		/// </summary>
+		public int RuleNumber;
+
+		/// <summary>
+		/// Whether a rule terminating in this node should include files or exclude files.
+		/// </summary>
+		public FileFilterType Type;
+
+		/// <summary>
+		/// Default constructor.
+		/// </summary>
+		public FileFilterNode(FileFilterNode? InParent, string InPattern)
+		{
+			Parent = InParent;
+			Pattern = InPattern;
+			RuleNumber = -1;
+			Type = FileFilterType.Exclude;
+		}
+
+		/// <summary>
+		/// Determine if the given token matches the pattern in this node
+		/// </summary>
+		public bool IsMatch(string Token)
+		{
+			if(Pattern.EndsWith("."))
+			{
+				return !Token.Contains('.') && IsMatch(Token, 0, Pattern.Substring(0, Pattern.Length - 1), 0);
+			}
+			else
+			{
+				return IsMatch(Token, 0, Pattern, 0);
+			}
+		}
+
+		/// <summary>
+		/// Determine if a given token matches a pattern, starting from the given positions within each string.
+		/// </summary>
+		static bool IsMatch(string Token, int TokenIdx, string Pattern, int PatternIdx)
+		{
+			for (; ; )
+			{
+				if (PatternIdx == Pattern.Length)
+				{
+					return (TokenIdx == Token.Length);
+				}
+				else if (Pattern[PatternIdx] == '*')
+				{
+					for (int NextTokenIdx = Token.Length; NextTokenIdx >= TokenIdx; NextTokenIdx--)
+					{
+						if (IsMatch(Token, NextTokenIdx, Pattern, PatternIdx + 1))
+						{
+							return true;
+						}
+					}
+					return false;
+				}
+				else if (TokenIdx < Token.Length && (Char.ToLower(Token[TokenIdx]) == Char.ToLower(Pattern[PatternIdx]) || Pattern[PatternIdx] == '?'))
+				{
+					TokenIdx++;
+					PatternIdx++;
+				}
+				else
+				{
+					return false;
+				}
+			}
+		}
+
+		/// <summary>
+		/// Debugger visualization
+		/// </summary>
+		/// <returns>Path to this node</returns>
+		public override string ToString()
+		{
+			return (Parent == null) ? "/" : Parent.ToString().TrimEnd('/') + "/" + Pattern;
+		}
+	}
+}
diff --git a/Engine/Source/Programs/UnrealGameSync/UnrealGameSyncShared/OutputAdapters.cs b/Engine/Source/Programs/UnrealGameSync/UnrealGameSyncShared/OutputAdapters.cs
new file mode 100644
--- /dev/null
+++ b/Engine/Source/Programs/UnrealGameSync/UnrealGameSyncShared/OutputAdapters.cs
@@ -0,0 +1,279 @@
+// Copyright Epic Games, Inc. All Rights Reserved.
+
+using EpicGames.Core;
+using Microsoft.Extensions.Logging;
+using System;
+using System.Collections.Generic;
+using System.IO;
+using System.Linq;
+using System.Text;
+using System.Threading.Tasks;
+
+#nullable enable
+
+namespace UnrealGameSync
+{
+	class PrefixedTextWriter : ILogger
+	{
+		string Prefix;
+		ILogger Inner;
+
+		public PrefixedTextWriter(string InPrefix, ILogger InInner)
+		{
+			Prefix = InPrefix;
+			Inner = InInner;
+		}
+
+		public IDisposable BeginScope<TState>(TState State) => Inner.BeginScope(State);
+
+		public bool IsEnabled(LogLevel LogLevel) => Inner.IsEnabled(LogLevel);
+
+		public void Log<TState>(LogLevel LogLevel, EventId EventId, TState State, Exception Exception, Func<TState, Exception, string> Formatter)
+		{
+			Inner.Log(LogLevel, EventId, State, Exception, (State, Exception) => Prefix + Formatter(State, Exception));
+		}
+	}
+
+	public class ProgressValue
+	{
+		Tuple<string, float> State = null!;
+		Stack<Tuple<float, float>> Ranges = new Stack<Tuple<float,float>>();
+
+		public ProgressValue()
+		{
+			Clear();
+		}
+
+		public void Clear()
+		{
+			State = new Tuple<string,float>("Starting...", 0.0f);
+
+			Ranges.Clear();
+			Ranges.Push(new Tuple<float, float>(0.0f, 1.0f));
+		}
+
+		public Tuple<string, float> Current
+		{
+			get { return State; }
+		}
+
+		public void Set(string Message)
+		{
+			if(Ranges.Count == 1)
+			{
+				State = new Tuple<string,float>(Message, State.Item2);
+			}
+		}
+
+		public void Set(string Message, float Fraction)
+		{
+			if(Ranges.Count == 1)
+			{
+				State = new Tuple<string, float>(Message, RelativeToAbsoluteFraction(Fraction));
+			}
+			else
+			{
+				State = new Tuple<string, float>(State.Item1, RelativeToAbsoluteFraction(Fraction));
+			}
+		}
+
+		public void Set(float Fraction)
+		{
+			State = new Tuple<string, float>(State.Item1, RelativeToAbsoluteFraction(Fraction));
+		}
+
+		public void Increment(float Fraction)
+		{
+			Set(State.Item2 + RelativeToAbsoluteFraction(Fraction));
+		}
+
+		public void Push(float MaxFraction)
+		{
+			Ranges.Push(new Tuple<float,float>(State.Item2, RelativeToAbsoluteFraction(MaxFraction)));
+		}
+
+		public void Pop()
+		{
+			if(Ranges.Count > 1)
+			{
+				State = new Tuple<string,float>(State.Item1, Ranges.Pop().Item2);
+			}
+		}
+
+		float RelativeToAbsoluteFraction(float Fraction)
+		{
+			Tuple<float, float> Range = Ranges.Peek();
+			return Range.Item1 + (Range.Item2 - Range.Item1) * Fraction;
+		}
+	}
+
+	static class ProgressTextWriter
+	{
+		const string DirectivePrefix = "@progress ";
+
+		public static string? ParseLine(string Line, ProgressValue Value)
+		{
+			string TrimLine = Line.Trim();
+			if(TrimLine.StartsWith(DirectivePrefix))
+			{
+				// Line that just contains a progress directive
+				bool bSkipLine = false;
+				ProcessInternal(TrimLine.Substring(DirectivePrefix.Length), ref bSkipLine, Value);
+				return null;
+			}
+			else
+			{
+				bool bSkipLine = false;
+				string RemainingLine = Line;
+
+				// Look for a progress directive at the end of a line, in square brackets
+				if(TrimLine.EndsWith("]"))
+				{
+					for(int LastIdx = TrimLine.Length - 2; LastIdx >= 0 && TrimLine[LastIdx] != ']'; LastIdx--)
+					{
+						if(TrimLine[LastIdx] == '[')
+						{
+							string DirectiveSubstring = TrimLine.Substring(LastIdx + 1, TrimLine.Length - LastIdx - 2);
+							if(DirectiveSubstring.StartsWith(DirectivePrefix))
+							{
+								ProcessInternal(DirectiveSubstring.Substring(DirectivePrefix.Length), ref bSkipLine, Value);
+								RemainingLine = Line.Substring(0, LastIdx).TrimEnd();
+							}
+							break;
+						}
+					}
+				}
+
+				if (bSkipLine)
+				{
+					return null;
+				}
+				else
+				{
+					return RemainingLine;
+				}
+			}
+		}
+
+		static void ProcessInternal(string Line, ref bool bSkipLine, ProgressValue Value)
+		{
+			List<string> Tokens = ParseTokens(Line);
+			for(int TokenIdx = 0; TokenIdx < Tokens.Count; )
+			{
+				float Fraction;
+				if(ReadFraction(Tokens, ref TokenIdx, out Fraction))
+				{
+					Value.Set(Fraction);
+				}
+				else if(Tokens[TokenIdx] == "push")
+				{
+					TokenIdx++;
+					if(ReadFraction(Tokens, ref TokenIdx, out Fraction))
+					{
+						Value.Push(Fraction);
+					}
+				}
+				else if(Tokens[TokenIdx] == "pop")
+				{
+					TokenIdx++;
+					Value.Pop();
+				}
+				else if(Tokens[TokenIdx] == "increment")
+				{
+					TokenIdx++;
+					if(ReadFraction(Tokens, ref TokenIdx, out Fraction))
+					{
+						Value.Increment(Fraction);
+					}
+				}
+				else if(Tokens[TokenIdx] == "skipline")
+				{
+					TokenIdx++;
+					bSkipLine = true;
+				}
+				else if(Tokens[TokenIdx].Length >= 2 && (Tokens[TokenIdx][0] == '\'' || Tokens[TokenIdx][0] == '\"') && Tokens[TokenIdx].Last() == Tokens[TokenIdx].First())
+				{
+					string Message = Tokens[TokenIdx++];
+					Value.Set(Message.Substring(1, Message.Length - 2));
+				}
+				else
+				{
+					TokenIdx++;
+				}
+			}
+		}
+
+		static List<string> ParseTokens(string Line)
+		{
+			List<string> Tokens = new List<string>();
+			for(int Idx = 0;;)
+			{
+				// Skip whitespace
+				while(Idx < Line.Length && Char.IsWhiteSpace(Line[Idx]))
+				{
+					Idx++;
+				}
+				if(Idx == Line.Length)
+				{
+					break;
+				}
+
+				// Read the next token
+				if(Char.IsLetterOrDigit(Line[Idx]))
+				{
+					int StartIdx = Idx++;
+					while(Idx < Line.Length && Char.IsLetterOrDigit(Line[Idx]))
+					{
+						Idx++;
+					}
+					Tokens.Add(Line.Substring(StartIdx, Idx - StartIdx));
+				}
+				else if(Line[Idx] == '\'' || Line[Idx] == '\"')
+				{
+					int StartIdx = Idx++;
+					while(Idx < Line.Length && Line[Idx] != Line[StartIdx])
+					{
+						Idx++;
+					}
+					Tokens.Add(Line.Substring(StartIdx, ++Idx - StartIdx));
+				}
+				else
+				{
+					Tokens.Add(Line.Substring(Idx++, 1));
+				}
+			}
+			return Tokens;
+		}
+
+		static bool ReadFraction(List<string> Tokens, ref int TokenIdx, out float Fraction)
+		{
+			// Read a fraction in the form x%
+			if(TokenIdx + 2 <= Tokens.Count && Tokens[TokenIdx + 1] == "%")
+			{
+				int Numerator;
+				if(int.TryParse(Tokens[TokenIdx], out Numerator))
+				{
+					Fraction = (float)Numerator / 100.0f;
+					TokenIdx += 2;
+					return true;
+				}
+			}
+			
+			// Read a fraction in the form x/y
+			if(TokenIdx + 3 <= Tokens.Count && Tokens[TokenIdx + 1] == "/")
+			{
+				int Numerator, Denominator;
+				if(int.TryParse(Tokens[TokenIdx], out Numerator) && int.TryParse(Tokens[TokenIdx + 2], out Denominator))
+				{
+					Fraction = (float)Numerator / (float)Denominator;
+					TokenIdx += 3;
+					return true;
+				}
+			}
+
+			Fraction = 0.0f;
+			return false;
+		}
+	}
+}
+
diff --git a/Engine/Source/Programs/UnrealGameSync/UnrealGameSyncShared/Telemetry.cs b/Engine/Source/Programs/UnrealGameSync/UnrealGameSyncShared/Telemetry.cs
new file mode 100644
--- /dev/null
+++ b/Engine/Source/Programs/UnrealGameSync/UnrealGameSyncShared/Telemetry.cs
@@ -0,0 +1,231 @@
+// Copyright Epic Games, Inc. All Rights Reserved.
+
+using Microsoft.Extensions.Logging;
+using System;
+using System.Collections.Generic;
+using System.IO;
+using System.Net;
+using System.Text;
+using System.Text.Json;
+using System.Threading;
+using System.Web;
+
+namespace UnrealGameSync
+{
+	/// <summary>
+	/// Interface for a telemetry sink
+	/// </summary>
+	public interface ITelemetrySink : IDisposable
+	{
+		/// <summary>
+		/// Sends a telemetry event with the given information
+		/// </summary>
+		/// <param name="EventName">Name of the event</param>
+		/// <param name="Attributes">Arbitrary object to include in the payload</param>
+		void SendEvent(string EventName, object Attributes);
+	}
+
+	/// <summary>
+	/// Telemetry sink that discards all events
+	/// </summary>
+	public class NullTelemetrySink : ITelemetrySink
+	{
+		/// <inheritdoc/>
+		public void Dispose()
+		{
+		}
+
+		/// <inheritdoc/>
+		public void SendEvent(string EventName, object Attributes)
+		{
+		}
+	}
+
+	/// <summary>
+	/// Epic internal telemetry sink using the data router
+	/// </summary>
+	public class EpicTelemetrySink : ITelemetrySink
+	{
+		/// <summary>
+		/// Combined url to post event streams to
+		/// </summary>
+		string Url;
+
+		/// <summary>
+		/// Lock used to modify the event queue
+		/// </summary>
+		object LockObject = new object();
+
+		/// <summary>
+		/// Whether a flush is queued
+		/// </summary>
+		bool bHasPendingFlush = false;
+
+		/// <summary>
+		/// List of pending events
+		/// </summary>
+		List<string> PendingEvents = new List<string>();
+
+		/// <summary>
+		/// The log writer to use
+		/// </summary>
+		ILogger Logger;
+
+		/// <summary>
+		/// Constructor
+		/// </summary>
+		public EpicTelemetrySink(string Url, ILogger Logger)
+		{
+			this.Url = Url;
+			this.Logger = Logger;
+
+			Logger.LogInformation("Posting to URL: {Url}", Url);
+		}
+
+		/// <inheritdoc/>
+		public void Dispose()
+		{
+			Flush();
+		}
+
+		/// <inheritdoc/>
+		public void Flush()
+		{
+			for (; ; )
+			{
+				lock (LockObject)
+				{
+					if (!bHasPendingFlush)
+					{
+						break;
+					}
+				}
+				Thread.Sleep(10);
+			}
+		}
+
+		/// <inheritdoc/>
+		public void SendEvent(string EventName, object Attributes)
+		{
+			string AttributesText = JsonSerializer.Serialize(Attributes);
+			if (AttributesText[0] != '{')
+			{
+				throw new Exception("Expected event data with named properties");
+			}
+
+			string EventText = AttributesText.Insert(1, String.Format("\"EventName\":\"{0}\",", HttpUtility.JavaScriptStringEncode(EventName)));
+			lock (PendingEvents)
+			{
+				PendingEvents.Add(EventText);
+				if (!bHasPendingFlush)
+				{
+					ThreadPool.QueueUserWorkItem(Obj => BackgroundFlush());
+					bHasPendingFlush = true;
+				}
+			}
+		}
+
+		/// <summary>
+		/// Synchronously sends a telemetry event
+		/// </summary>
+		void BackgroundFlush()
+		{
+			for (; ; )
+			{
+				try
+				{
+					// Generate the content for this event
+					List<string> Events = new List<string>();
+					lock (LockObject)
+					{
+						if (PendingEvents.Count == 0)
+						{
+							bHasPendingFlush = false;
+							break;
+						}
+
+						Events.AddRange(PendingEvents);
+						PendingEvents.Clear();
+					}
+
+					// Print all the events we're sending
+					foreach (string Event in Events)
+					{
+						Logger.LogInformation("Sending Event: {0}", Event);
+					}
+
+					// Convert the content to UTF8
+					string ContentText = String.Format("{{\"Events\":[{0}]}}", String.Join(",", Events));
+					byte[] Content = Encoding.UTF8.GetBytes(ContentText);
+
+					// Post the event data
+					HttpWebRequest Request = (HttpWebRequest)WebRequest.Create(Url);
+					Request.Method = "POST";
+					Request.ContentType = "application/json";
+					Request.UserAgent = "ue/ugs";
+					Request.Timeout = 5000;
+					Request.ContentLength = Content.Length;
+					Request.ContentType = "application/json";
+					using (Stream RequestStream = Request.GetRequestStream())
+					{
+						RequestStream.Write(Content, 0, Content.Length);
+					}
+
+					// Wait for the response and dispose of it immediately
+					using (HttpWebResponse Response = (HttpWebResponse)Request.GetResponse())
+					{
+						Logger.LogInformation("Response: {StatusCode}", (int)Response.StatusCode);
+					}
+				}
+				catch (WebException Ex)
+				{
+					// Handle errors. Any non-200 responses automatically generate a WebException.
+					HttpWebResponse Response = (HttpWebResponse)Ex.Response;
+					if (Response == null)
+					{
+						Logger.LogError(Ex, "Exception while attempting to send event");
+					}
+					else
+					{
+						string ResponseText;
+						using (Stream ResponseStream = Response.GetResponseStream())
+						{
+							MemoryStream MemoryStream = new MemoryStream();
+							ResponseStream.CopyTo(MemoryStream);
+							ResponseText = Encoding.UTF8.GetString(MemoryStream.ToArray());
+						}
+						Logger.LogError("Failed to send analytics event. Code = {Code}. Desc = {Dec}. Response = {Response}.", (int)Response.StatusCode, Response.StatusDescription, ResponseText);
+					}
+				}
+				catch (Exception Ex)
+				{
+					Logger.LogError(Ex, "Exception while attempting to send event");
+				}
+			}
+		}
+	}
+
+	/// <summary>
+	/// Global telemetry static class
+	/// </summary>
+	public static class Telemetry
+	{
+		/// <summary>
+		/// The current telemetry provider
+		/// </summary>
+		public static ITelemetrySink? ActiveSink
+		{
+			get; set;
+		}
+
+		/// <summary>
+		/// Sends a telemetry event with the given information
+		/// </summary>
+		/// <param name="EventName">Name of the event</param>
+		/// <param name="Attributes">Arbitrary object to include in the payload</param>
+		public static void SendEvent(string EventName, object Attributes)
+		{
+			ActiveSink?.SendEvent(EventName, Attributes);
+		}
+	}
+}
diff --git a/Engine/Source/Programs/UnrealGameSync/UnrealGameSyncShared/TelemetryStopwatch.cs b/Engine/Source/Programs/UnrealGameSync/UnrealGameSyncShared/TelemetryStopwatch.cs
new file mode 100644
--- /dev/null
+++ b/Engine/Source/Programs/UnrealGameSync/UnrealGameSyncShared/TelemetryStopwatch.cs
@@ -0,0 +1,60 @@
+// Copyright Epic Games, Inc. All Rights Reserved.
+
+using System;
+using System.Collections.Generic;
+using System.Diagnostics;
+using System.Reflection;
+
+namespace UnrealGameSync
+{
+	class TelemetryStopwatch : IDisposable
+	{
+		readonly string EventName;
+		readonly Dictionary<string, object?> EventData;
+		readonly Stopwatch Timer;
+
+		public TelemetryStopwatch(string EventName, string Project)
+		{
+			this.EventName = EventName;
+
+			EventData = new Dictionary<string, object?>();
+			EventData["Project"] = Project;
+
+			Timer = Stopwatch.StartNew();
+		}
+
+		public void AddData(object Data)
+		{
+			foreach (PropertyInfo Property in Data.GetType().GetProperties())
+			{
+				EventData[Property.Name] = Property.GetValue(Data);
+			}
+		}
+
+		public TimeSpan Stop(string InResult)
+		{
+			if (Timer.IsRunning)
+			{
+				Timer.Stop();
+
+				EventData["Result"] = InResult;
+				EventData["TimeSeconds"] = Timer.Elapsed.TotalSeconds;
+			}
+			return Elapsed;
+		}
+
+		public void Dispose()
+		{
+			if (Timer.IsRunning)
+			{
+				Stop("Aborted");
+			}
+			Telemetry.SendEvent(EventName, EventData);
+		}
+
+		public TimeSpan Elapsed
+		{
+			get { return Timer.Elapsed; }
+		}
+	}
+}
diff --git a/Engine/Source/Programs/UnrealGameSync/UnrealGameSyncShared/UnrealGameSyncShared.csproj b/Engine/Source/Programs/UnrealGameSync/UnrealGameSyncShared/UnrealGameSyncShared.csproj
new file mode 100644
--- /dev/null
+++ b/Engine/Source/Programs/UnrealGameSync/UnrealGameSyncShared/UnrealGameSyncShared.csproj
@@ -0,0 +1,14 @@
+<Project Sdk="Microsoft.NET.Sdk">
+
+  <PropertyGroup>
+    <TargetFramework>netcoreapp3.1</TargetFramework>
+    <Nullable>enable</Nullable>
+    <RootNamespace>UnrealGameSync</RootNamespace>
+  </PropertyGroup>
+
+  <ItemGroup>
+    <ProjectReference Include="..\..\Shared\EpicGames.Core\EpicGames.Core.csproj" />
+    <ProjectReference Include="..\..\Shared\EpicGames.Perforce\EpicGames.Perforce.csproj" />
+  </ItemGroup>
+
+</Project>
diff --git a/Engine/Source/Programs/UnrealGameSync/UnrealGameSyncShared/UserSettings.cs b/Engine/Source/Programs/UnrealGameSync/UnrealGameSyncShared/UserSettings.cs
new file mode 100644
--- /dev/null
+++ b/Engine/Source/Programs/UnrealGameSync/UnrealGameSyncShared/UserSettings.cs
@@ -0,0 +1,1170 @@
+// Copyright Epic Games, Inc. All Rights Reserved.
+
+using EpicGames.Core;
+using Microsoft.Extensions.Logging;
+using System;
+using System.Collections.Generic;
+using System.Diagnostics.CodeAnalysis;
+using System.Drawing;
+using System.IO;
+using System.Linq;
+using System.Text;
+using System.Text.Json;
+using System.Text.Json.Serialization;
+using System.Text.RegularExpressions;
+using System.Threading.Tasks;
+
+namespace UnrealGameSync
+{
+	public enum BuildConfig
+	{
+		Debug,
+		DebugGame,
+		Development,
+	}
+
+	public enum TabLabels
+	{
+		Stream,
+		WorkspaceName,
+		WorkspaceRoot,
+		ProjectFile,
+	}
+
+	public enum BisectState
+	{
+		Include,
+		Exclude,
+		Pass,
+		Fail,
+	}
+
+	public enum UserSelectedProjectType
+	{
+		Client,
+		Local
+	}
+
+	public enum FilterType
+	{
+		None,
+		Code,
+		Content
+	}
+
+	public enum LatestChangeType
+	{
+		Any,
+		Good,
+		Starred,
+	}
+
+	public enum UserSettingsVersion
+	{
+		Initial = 0,
+		DefaultServerSettings = 1,
+		XgeShaderCompilation = 2,
+		DefaultNumberOfThreads = 3,
+		Latest = DefaultNumberOfThreads
+	}
+
+	public class ArchiveSettings
+	{
+		public bool bEnabled;
+		public string Type;
+		public List<string> Order;
+
+		public ArchiveSettings(bool bEnabled, string Type, IEnumerable<string> Order)
+		{
+			this.bEnabled = bEnabled;
+			this.Type = Type;
+			this.Order = new List<string>(Order);
+		}
+
+		public static bool TryParseConfigEntry(string Text, [NotNullWhen(true)] out ArchiveSettings? Settings)
+		{
+			ConfigObject Object = new ConfigObject(Text);
+
+			string? Type = Object.GetValue("Type", null);
+			if (Type == null)
+			{
+				Settings = null;
+				return false;
+			}
+			else
+			{
+				string[] Order = Object.GetValue("Order", "").Split(new char[] { ';' }, StringSplitOptions.RemoveEmptyEntries);
+				bool bEnabled = Object.GetValue("Enabled", 0) != 0;
+
+				Settings = new ArchiveSettings(bEnabled, Type, Order);
+				return true;
+			}
+		}
+
+		public string ToConfigEntry()
+		{
+			ConfigObject Object = new ConfigObject();
+
+			Object.SetValue("Enabled", bEnabled ? 1 : 0);
+			Object.SetValue("Type", Type);
+			Object.SetValue("Order", String.Join(";", Order));
+
+			return Object.ToString();
+		}
+
+		public override string ToString()
+		{
+			return ToConfigEntry();
+		}
+	}
+
+	public class UserSelectedProjectSettings
+	{
+		public readonly string? ServerAndPort;
+		public readonly string? UserName;
+		public readonly UserSelectedProjectType Type;
+		public readonly string? ClientPath;
+		public readonly string? LocalPath;
+
+		public UserSelectedProjectSettings(string? ServerAndPort, string? UserName, UserSelectedProjectType Type, string? ClientPath, string? LocalPath)
+		{
+			this.ServerAndPort = ServerAndPort;
+			this.UserName = UserName;
+			this.Type = Type;
+			this.ClientPath = ClientPath;
+			this.LocalPath = LocalPath;
+		}
+
+		public static bool TryParseConfigEntry(string Text, [NotNullWhen(true)] out UserSelectedProjectSettings? Project)
+		{
+			ConfigObject Object = new ConfigObject(Text);
+
+			UserSelectedProjectType Type;
+			if(Enum.TryParse(Object.GetValue("Type", ""), out Type))
+			{
+				string? ServerAndPort = Object.GetValue("ServerAndPort", null);
+				if(String.IsNullOrWhiteSpace(ServerAndPort))
+				{
+					ServerAndPort = null;
+				}
+
+				// Fixup for code that was saving server host name rather than DNS entry
+				if(ServerAndPort != null && ServerAndPort.Equals("p4-nodeb.epicgames.net:1666", StringComparison.OrdinalIgnoreCase))
+				{
+					ServerAndPort = "perforce:1666";
+				}
+
+				string? UserName = Object.GetValue("UserName", null);
+				if(String.IsNullOrWhiteSpace(UserName))
+				{
+					UserName = null;
+				}
+
+				string? LocalPath = Object.GetValue("LocalPath", null);
+				if(String.IsNullOrWhiteSpace(LocalPath))
+				{
+					LocalPath = null;
+				}
+
+				string? ClientPath = Object.GetValue("ClientPath", null);
+				if(String.IsNullOrWhiteSpace(ClientPath))
+				{
+					ClientPath = null;
+				}
+
+				if((Type == UserSelectedProjectType.Client && ClientPath != null) || (Type == UserSelectedProjectType.Local && LocalPath != null))
+				{
+					Project = new UserSelectedProjectSettings(ServerAndPort, UserName, Type, ClientPath, LocalPath);
+					return true;
+				}
+			}
+
+			Project = null;
+			return false;
+		}
+
+		public string ToConfigEntry()
+		{
+			ConfigObject Object = new ConfigObject();
+
+			if(ServerAndPort != null)
+			{
+				Object.SetValue("ServerAndPort", ServerAndPort);
+			}
+			if(UserName != null)
+			{
+				Object.SetValue("UserName", UserName);
+			}
+
+			Object.SetValue("Type", Type.ToString());
+
+			if(ClientPath != null)
+			{
+				Object.SetValue("ClientPath", ClientPath);
+			}
+			if(LocalPath != null)
+			{
+				Object.SetValue("LocalPath", LocalPath);
+			}
+
+			return Object.ToString();
+		}
+
+		public override string? ToString()
+		{
+			return LocalPath ?? ClientPath;
+		}
+	}
+
+	public class SyncCategory
+	{
+		public Guid Id { get; set; }
+		public bool Enable { get; set; }
+	}
+
+	public class BisectEntry
+	{
+		public int Change { get; set; }
+		public BisectState State { get; set; }
+	}
+
+	public class UserWorkspaceSettings
+	{
+		[JsonIgnore]
+		public FileReference? File { get; set; }
+
+		// Workspace specific SyncFilters
+		public string[] SyncView { get; set; } = Array.Empty<string>();
+		public List<SyncCategory> SyncCategories { get; set; } = new List<SyncCategory>();
+		public bool? bSyncAllProjects { get; set; }
+		public bool? bIncludeAllProjectsInSolution { get; set; }
+
+		[JsonIgnore]
+		public Dictionary<Guid, bool> SyncCategoriesDict
+		{
+			get => SyncCategories.ToDictionary(x => x.Id, x => x.Enable);
+			set => SyncCategories = value.Select(x => new SyncCategory { Id = x.Key, Enable = x.Value }).ToList();
+		}
+
+		public void Save() => Utility.SaveJson(File!, this);
+	}
+
+	public class UserWorkspaceState
+	{
+		[JsonIgnore]
+		public FileReference? File { get; set; }
+
+		// Client path to the current project
+		public string? SelectedClientFileName { get; set; }
+
+		// Settings for the currently synced project in this workspace. CurrentChangeNumber is only valid for this workspace if CurrentProjectPath is the current project.
+		public string? CurrentProjectIdentifier { get; set; }
+		public int CurrentChangeNumber { get; set; } = -1;
+		public string? CurrentSyncFilterHash { get; set; }
+		public List<int> AdditionalChangeNumbers { get; set; } = new List<int>();
+
+		// Settings for the last attempted sync. These values are set to persist error messages between runs.
+		public int LastSyncChangeNumber { get; set; }
+		public WorkspaceUpdateResult LastSyncResult { get; set; }
+		public string? LastSyncResultMessage { get; set; }
+		public DateTime? LastSyncTime { get; set; }
+		public int LastSyncDurationSeconds { get; set; }
+
+		// The last successful build, regardless of whether a failed sync has happened in the meantime. Used to determine whether to force a clean due to entries in the project config file.
+		public int LastBuiltChangeNumber { get; set; }
+
+		// Expanded archives in the workspace
+		public string[]? ExpandedArchiveTypes { get; set; }
+
+		// The changes that we're regressing at the moment
+		public List<BisectEntry> BisectChanges { get; set; } = new List<BisectEntry>();
+
+		public void SetBisectState(int Change, BisectState State)
+		{
+			BisectEntry Entry = BisectChanges.FirstOrDefault(x => x.Change == Change);
+			if (Entry == null)
+			{
+				Entry = new BisectEntry();
+				Entry.Change = Change;
+				BisectChanges.Add(Entry);
+			}
+			Entry.State = State;
+		}
+
+		public void Save()
+		{
+			lock(File!)
+			{
+				Utility.SaveJson(File, this);
+			}
+		}
+	}
+
+	public class UserProjectSettings
+	{
+		[JsonIgnore]
+		public FileReference? File { get; set; }
+
+		public List<ConfigObject> BuildSteps { get; set; } = new List<ConfigObject>();
+		public FilterType FilterType { get; set; }
+		public HashSet<string> FilterBadges { get; set; } = new HashSet<string>(StringComparer.OrdinalIgnoreCase);
+
+		public void Save() => Utility.SaveJson(File!, this);
+	}
+
+	public class UserSettings
+	{
+		/// <summary>
+		/// Enum that decribes which robomerge changes to show
+		/// </summary>
+		public enum RobomergeShowChangesOption
+		{
+			All,		// Show all changes from robomerge
+			Badged,		// Show only robomerge changes that have an associated badge
+			None		// Show no robomerge changes
+		};
+
+		FileReference FileName;
+		ConfigFile ConfigFile = new ConfigFile();
+
+		// General settings
+		public UserSettingsVersion Version = UserSettingsVersion.Latest;
+		public bool bBuildAfterSync;
+		public bool bRunAfterSync;
+		public bool bOpenSolutionAfterSync;
+		public bool bShowLogWindow;
+		public bool bAutoResolveConflicts;
+		public bool bShowUnreviewedChanges;
+		public bool bShowAutomatedChanges;
+		public RobomergeShowChangesOption ShowRobomerge;
+		public bool bAnnotateRobmergeChanges;
+		public bool bShowLocalTimes;
+		public bool bKeepInTray;
+		public Guid[] EnabledTools;
+		public bool bShowNetCoreInfo;
+		public int FilterIndex;
+		public UserSelectedProjectSettings? LastProject;
+		public List<UserSelectedProjectSettings> OpenProjects;
+		public List<UserSelectedProjectSettings> RecentProjects;
+		public string[] SyncView;
+		public Dictionary<Guid, bool> SyncCategories;
+		public bool bSyncAllProjects;
+		public bool bIncludeAllProjectsInSolution;
+		public LatestChangeType SyncType;
+		public BuildConfig CompiledEditorBuildConfig; // NB: This assumes not using precompiled editor. See CurrentBuildConfig.
+		public TabLabels TabLabels;
+
+		// Precompiled binaries
+		public List<ArchiveSettings> Archives = new List<ArchiveSettings>();
+
+		// OIDC Settings
+		public Dictionary<string, string> ProviderToRefreshTokens = new Dictionary<string, string>();
+
+		// Window settings
+		public bool bWindowVisible;
+		public string WindowState;
+		public Rectangle? WindowBounds;
+		
+		// Schedule settings
+		public bool bScheduleEnabled;
+		public TimeSpan ScheduleTime;
+		public LatestChangeType ScheduleChange;
+		public bool ScheduleAnyOpenProject;
+		public List<UserSelectedProjectSettings> ScheduleProjects;
+
+		// Run configuration
+		public List<Tuple<string, bool>> EditorArguments = new List<Tuple<string,bool>>();
+		public bool bEditorArgumentsPrompt;
+
+		// Notification settings
+		public List<string> NotifyProjects;
+		public int NotifyUnassignedMinutes;
+		public int NotifyUnacknowledgedMinutes;
+		public int NotifyUnresolvedMinutes;
+
+		// Project settings
+		Dictionary<FileReference, UserWorkspaceState> WorkspaceKeyToState = new Dictionary<FileReference, UserWorkspaceState>();
+		Dictionary<FileReference, UserWorkspaceSettings> WorkspaceKeyToSettings = new Dictionary<FileReference, UserWorkspaceSettings>();
+		Dictionary<FileReference, UserProjectSettings> ProjectKeyToSettings = new Dictionary<FileReference, UserProjectSettings>();
+
+		// Perforce settings
+		public PerforceSyncOptions SyncOptions = new PerforceSyncOptions();
+
+		private List<UserSelectedProjectSettings> ReadProjectList(string SettingName, string LegacySettingName)
+		{
+			List<UserSelectedProjectSettings> Projects = new List<UserSelectedProjectSettings>();
+
+			string[]? ProjectStrings = ConfigFile.GetValues(SettingName, null);
+			if(ProjectStrings != null)
+			{
+				foreach(string ProjectString in ProjectStrings)
+				{
+					UserSelectedProjectSettings? Project;
+					if(UserSelectedProjectSettings.TryParseConfigEntry(ProjectString, out Project))
+					{
+						Projects.Add(Project);
+					}
+				}
+			}
+			else if(LegacySettingName != null)
+			{
+				string[]? LegacyProjectStrings = ConfigFile.GetValues(LegacySettingName, null);
+				if(LegacyProjectStrings != null)
+				{
+					foreach(string LegacyProjectString in LegacyProjectStrings)
+					{
+						if(!String.IsNullOrWhiteSpace(LegacyProjectString))
+						{
+							Projects.Add(new UserSelectedProjectSettings(null, null, UserSelectedProjectType.Local, null, LegacyProjectString));
+						}
+					}
+				}
+			}
+
+			return Projects;
+		}
+
+		public UserSettings(FileReference InFileName, ILogger Logger)
+		{
+			FileName = InFileName;
+			ConfigFile.TryLoad(FileName, Logger);
+
+			// General settings
+			Version = (UserSettingsVersion)ConfigFile.GetValue("General.Version", (int)UserSettingsVersion.Initial);
+			bBuildAfterSync = (ConfigFile.GetValue("General.BuildAfterSync", "1") != "0");
+			bRunAfterSync = (ConfigFile.GetValue("General.RunAfterSync", "1") != "0");
+			bool bSyncPrecompiledEditor = (ConfigFile.GetValue("General.SyncPrecompiledEditor", "0") != "0");
+			bOpenSolutionAfterSync = (ConfigFile.GetValue("General.OpenSolutionAfterSync", "0") != "0");
+			bShowLogWindow = (ConfigFile.GetValue("General.ShowLogWindow", false));
+			bAutoResolveConflicts = (ConfigFile.GetValue("General.AutoResolveConflicts", "1") != "0");
+			bShowUnreviewedChanges = ConfigFile.GetValue("General.ShowUnreviewed", true);
+			bShowAutomatedChanges = ConfigFile.GetValue("General.ShowAutomated", false);
+
+			// safely parse the filter enum
+			ShowRobomerge = RobomergeShowChangesOption.All;
+			Enum.TryParse(ConfigFile.GetValue("General.RobomergeFilter", ""), out ShowRobomerge);
+
+			bAnnotateRobmergeChanges = ConfigFile.GetValue("General.AnnotateRobomerge", true);
+			bShowLocalTimes = ConfigFile.GetValue("General.ShowLocalTimes", false);
+			bKeepInTray = ConfigFile.GetValue("General.KeepInTray", true);
+
+			List<Guid> EnabledTools = ConfigFile.GetGuidValues("General.EnabledTools", new Guid[0]).ToList();
+			if (ConfigFile.GetValue("General.EnableP4VExtensions", false))
+			{
+				EnabledTools.Add(new Guid("963850A0-BF63-4E0E-B903-1C5954C7DCF8"));
+			}
+			if (ConfigFile.GetValue("General.EnableUshell", false))
+			{
+				EnabledTools.Add(new Guid("922EED87-E732-464C-92DC-5A8F7ED955E2"));
+			}
+			this.EnabledTools = EnabledTools.ToArray();
+
+			bShowNetCoreInfo = ConfigFile.GetValue("General.ShowNetCoreInfo", true);
+			int.TryParse(ConfigFile.GetValue("General.FilterIndex", "0"), out FilterIndex);
+
+			string? LastProjectString = ConfigFile.GetValue("General.LastProject", null);
+			if(LastProjectString != null)
+			{
+				UserSelectedProjectSettings.TryParseConfigEntry(LastProjectString, out LastProject);
+			}
+			else
+			{
+				string? LastProjectFileName = ConfigFile.GetValue("General.LastProjectFileName", null);
+				if(LastProjectFileName != null)
+				{
+					LastProject = new UserSelectedProjectSettings(null, null, UserSelectedProjectType.Local, null, LastProjectFileName);
+				}
+			}
+
+			OpenProjects = ReadProjectList("General.OpenProjects", "General.OpenProjectFileNames");
+			RecentProjects = ReadProjectList("General.RecentProjects", "General.OtherProjectFileNames");
+			SyncView = ConfigFile.GetValues("General.SyncFilter", new string[0]);
+			SyncCategories = GetCategorySettings(ConfigFile.FindSection("General"), "SyncIncludedCategories", "SyncExcludedCategories");
+
+			bSyncAllProjects = ConfigFile.GetValue("General.SyncAllProjects", false);
+			bIncludeAllProjectsInSolution = ConfigFile.GetValue("General.IncludeAllProjectsInSolution", false);
+			if(!Enum.TryParse(ConfigFile.GetValue("General.SyncType", ""), out SyncType))
+			{
+				SyncType = LatestChangeType.Good;
+			}
+
+			// Build configuration
+			string CompiledEditorBuildConfigName = ConfigFile.GetValue("General.BuildConfig", "");
+			if(!Enum.TryParse(CompiledEditorBuildConfigName, true, out CompiledEditorBuildConfig))
+			{
+				CompiledEditorBuildConfig = BuildConfig.DebugGame;
+			}
+
+			// Tab names
+			string TabLabelsValue = ConfigFile.GetValue("General.TabLabels", "");
+			if(!Enum.TryParse(TabLabelsValue, true, out TabLabels))
+			{
+				TabLabels = TabLabels.Stream;
+			}
+
+			// Editor arguments
+			string[] Arguments = ConfigFile.GetValues("General.EditorArguments", new string[]{ "0:-log", "0:-fastload" });
+			if (Version < UserSettingsVersion.XgeShaderCompilation)
+			{
+				Arguments = Enumerable.Concat(Arguments, new string[] { "0:-noxgeshadercompile" }).ToArray();
+			}
+			foreach(string Argument in Arguments)
+			{
+				if(Argument.StartsWith("0:"))
+				{
+					EditorArguments.Add(new Tuple<string,bool>(Argument.Substring(2), false));
+				}
+				else if(Argument.StartsWith("1:"))
+				{
+					EditorArguments.Add(new Tuple<string,bool>(Argument.Substring(2), true));
+				}
+				else
+				{
+					EditorArguments.Add(new Tuple<string,bool>(Argument, true));
+				}
+			}
+			bEditorArgumentsPrompt = ConfigFile.GetValue("General.EditorArgumentsPrompt", false);
+
+			// Precompiled binaries
+			string[] ArchiveValues = ConfigFile.GetValues("PrecompiledBinaries.Archives", new string[0]);
+			foreach (string ArchiveValue in ArchiveValues)
+			{
+				ArchiveSettings? Settings;
+				if (ArchiveSettings.TryParseConfigEntry(ArchiveValue, out Settings))
+				{
+					Archives.Add(Settings);
+				}
+			}
+
+			if (bSyncPrecompiledEditor)
+			{
+				Archives.Add(new ArchiveSettings(true, "Editor", new string[0]));
+			}
+
+			// OIDC Settings
+			string[] Tokens = ConfigFile.GetValues("OIDCProviders.Tokens", new string[0]);
+			foreach (string TokenValue in Tokens)
+			{
+				ConfigObject O = new ConfigObject(TokenValue);
+				string? Provider = O.GetValue("Provider");
+				string? Token = O.GetValue("Token");
+				if (Provider != null && Token != null)
+				{
+					ProviderToRefreshTokens.TryAdd(Provider, Token);
+				}
+			}
+
+			// Window settings
+			bWindowVisible = ConfigFile.GetValue("Window.Visible", true);
+			WindowState = ConfigFile.GetValue("Window.State", "");
+			WindowBounds = ParseRectangleValue(ConfigFile.GetValue("Window.Bounds", ""));
+
+			// Schedule settings
+			bScheduleEnabled = ConfigFile.GetValue("Schedule.Enabled", false);
+			if(!TimeSpan.TryParse(ConfigFile.GetValue("Schedule.Time", ""), out ScheduleTime))
+			{
+				ScheduleTime = new TimeSpan(6, 0, 0);
+			}
+			if(!Enum.TryParse(ConfigFile.GetValue("Schedule.Change", ""), out ScheduleChange))
+			{
+				ScheduleChange = LatestChangeType.Good;
+			}
+			ScheduleAnyOpenProject = ConfigFile.GetValue("Schedule.AnyOpenProject", true);
+			ScheduleProjects = ReadProjectList("Schedule.Projects", "Schedule.ProjectFileNames");
+
+			// Notification settings
+			NotifyProjects = ConfigFile.GetValues("Notifications.NotifyProjects", new string[0]).ToList();
+			NotifyUnassignedMinutes = ConfigFile.GetValue("Notifications.NotifyUnassignedMinutes", -1);
+			NotifyUnacknowledgedMinutes = ConfigFile.GetValue("Notifications.NotifyUnacknowledgedMinutes", -1);
+			NotifyUnresolvedMinutes = ConfigFile.GetValue("Notifications.NotifyUnresolvedMinutes", -1);
+
+			// Perforce settings
+			if(!int.TryParse(ConfigFile.GetValue("Perforce.NumRetries", "0"), out SyncOptions.NumRetries))
+			{
+				SyncOptions.NumRetries = 0;
+			}
+
+			int NumThreads;
+			if(int.TryParse(ConfigFile.GetValue("Perforce.NumThreads", "0"), out NumThreads) && NumThreads > 0)
+			{
+				if(Version >= UserSettingsVersion.DefaultNumberOfThreads || NumThreads > 1)
+				{
+					SyncOptions.NumThreads = NumThreads;
+				}
+			}
+
+			SyncOptions.TcpBufferSize = ConfigFile.GetValue("Perforce.TcpBufferSize", PerforceSyncOptions.DefaultTcpBufferSize);
+			SyncOptions.FileBufferSize = ConfigFile.GetValue("Perforce.FileBufferSize", PerforceSyncOptions.DefaultFileBufferSize);
+			SyncOptions.MaxCommandsPerBatch = ConfigFile.GetValue("Perforce.MaxCommandsPerBatch", PerforceSyncOptions.DefaultMaxCommandsPerBatch);
+			SyncOptions.MaxSizePerBatch = ConfigFile.GetValue("Perforce.MaxSizePerBatch", PerforceSyncOptions.DefaultMaxSizePerBatch);
+		}
+
+		static Dictionary<Guid, bool> GetCategorySettings(ConfigSection Section, string IncludedKey, string ExcludedKey)
+		{
+			Dictionary<Guid, bool> Result = new Dictionary<Guid, bool>();
+			if (Section != null)
+			{
+				foreach (Guid UniqueId in Section.GetValues(IncludedKey, new Guid[0]))
+				{
+					Result[UniqueId] = true;
+				}
+				foreach (Guid UniqueId in Section.GetValues(ExcludedKey, new Guid[0]))
+				{
+					Result[UniqueId] = false;
+				}
+			}
+			return Result;
+		}
+
+		static void SetCategorySettings(ConfigSection Section, string IncludedKey, string ExcludedKey, Dictionary<Guid, bool> Categories)
+		{
+			Guid[] IncludedCategories = Categories.Where(x => x.Value).Select(x => x.Key).ToArray();
+			if (IncludedCategories.Length > 0)
+			{
+				Section.SetValues(IncludedKey, IncludedCategories);
+			}
+
+			Guid[] ExcludedCategories = Categories.Where(x => !x.Value).Select(x => x.Key).ToArray();
+			if (ExcludedCategories.Length > 0)
+			{
+				Section.SetValues(ExcludedKey, ExcludedCategories);
+			}
+		}
+
+		static Rectangle? ParseRectangleValue(string Text)
+		{
+			ConfigObject Object = new ConfigObject(Text);
+
+			int X = Object.GetValue("X", -1);
+			int Y = Object.GetValue("Y", -1);
+			int W = Object.GetValue("W", -1);
+			int H = Object.GetValue("H", -1);
+
+			if(X == -1 || Y == -1 || W == -1 || H == -1)
+			{
+				return null;
+			}
+			else
+			{
+				return new Rectangle(X, Y, W, H);
+			}
+		}
+
+		static string FormatRectangleValue(Rectangle Value)
+		{
+			ConfigObject Object = new ConfigObject();
+
+			Object.SetValue("X", Value.X);
+			Object.SetValue("Y", Value.Y);
+			Object.SetValue("W", Value.Width);
+			Object.SetValue("H", Value.Height);
+
+			return Object.ToString();
+		}
+
+		public DirectoryReference GetConfigDir(DirectoryReference WorkspaceDir)
+		{
+			DirectoryReference ConfigDir = DirectoryReference.Combine(WorkspaceDir, ".ugs");
+
+			DirectoryInfo ConfigDirInfo = ConfigDir.ToDirectoryInfo();
+			if (!ConfigDirInfo.Exists)
+			{
+				ConfigDirInfo.Create();
+				ConfigDirInfo.Attributes = FileAttributes.Directory | FileAttributes.Hidden;
+			}
+
+			return ConfigDir;
+		}
+
+		public UserWorkspaceState FindOrAddWorkspaceState(DirectoryReference RootDir, string SelectedClientFileName, string SelectedProjectIdentifier)
+		{
+			FileReference ConfigFile = FileReference.Combine(GetConfigDir(RootDir), "global.json");
+			if (!WorkspaceKeyToState.TryGetValue(ConfigFile, out UserWorkspaceState? State))
+			{
+				if (!Utility.TryLoadJson(ConfigFile, out State))
+				{
+					string ClientBranchPath = Regex.Replace(SelectedClientFileName, "^(//[^/]+)/.*$", "$1");
+					State = ImportWorkspaceState(ClientBranchPath);
+					State.SelectedClientFileName = SelectedClientFileName;
+					Utility.SaveJson(ConfigFile, State);
+				}
+
+				State.File = ConfigFile;
+				WorkspaceKeyToState.Add(ConfigFile, State);
+			}
+			if (!String.Equals(State.SelectedClientFileName, SelectedClientFileName, StringComparison.Ordinal) || !String.Equals(State.CurrentProjectIdentifier, SelectedProjectIdentifier, StringComparison.Ordinal))
+			{
+				State = new UserWorkspaceState();
+				State.File = ConfigFile;
+				State.SelectedClientFileName = SelectedClientFileName;
+				State.CurrentProjectIdentifier = SelectedProjectIdentifier;
+				WorkspaceKeyToState[ConfigFile] = State;
+			}
+			return State;
+		}
+
+		public UserWorkspaceSettings FindOrAddWorkspaceSettings(DirectoryReference LocalRoot, string ClientBranchPath)
+		{
+			FileReference ConfigFile = FileReference.Combine(GetConfigDir(LocalRoot), $"client_{ClientBranchPath.Trim('/').Replace('/', '+')}.json");
+			if (!WorkspaceKeyToSettings.TryGetValue(ConfigFile, out UserWorkspaceSettings? Settings))
+			{
+				if (!Utility.TryLoadJson(ConfigFile, out Settings))
+				{
+					Settings = ImportWorkspaceSettings(ClientBranchPath);
+					Utility.SaveJson(ConfigFile, Settings);
+				}
+
+				Settings.File = ConfigFile;
+				WorkspaceKeyToSettings.Add(ConfigFile, Settings);
+			}
+			return Settings;
+		}
+
+		UserWorkspaceState ImportWorkspaceState(string ClientBranchPath)
+		{
+			UserWorkspaceState CurrentWorkspace = new UserWorkspaceState();
+
+			// Read the workspace settings
+			ConfigSection WorkspaceSection = ConfigFile.FindSection(ClientBranchPath.Trim('/'));
+			if(WorkspaceSection == null)
+			{
+				string LegacyBranchAndClientKey = ClientBranchPath.Trim('/');
+
+				int SlashIdx = LegacyBranchAndClientKey.IndexOf('/');
+				if(SlashIdx != -1)
+				{
+					LegacyBranchAndClientKey = LegacyBranchAndClientKey.Substring(0, SlashIdx) + "$" + LegacyBranchAndClientKey.Substring(SlashIdx + 1);
+				}
+
+				string? CurrentSync = ConfigFile.GetValue("Clients." + LegacyBranchAndClientKey, null);
+				if(CurrentSync != null)
+				{
+					int AtIdx = CurrentSync.LastIndexOf('@');
+					if(AtIdx != -1)
+					{
+						int ChangeNumber;
+						if(int.TryParse(CurrentSync.Substring(AtIdx + 1), out ChangeNumber))
+						{
+							CurrentWorkspace.CurrentProjectIdentifier = CurrentSync.Substring(0, AtIdx);
+							CurrentWorkspace.CurrentChangeNumber = ChangeNumber;
+						}
+					}
+				}
+
+				string? LastUpdateResultText = ConfigFile.GetValue("Clients." + LegacyBranchAndClientKey + "$LastUpdate", null);
+				if(LastUpdateResultText != null)
+				{
+					int ColonIdx = LastUpdateResultText.LastIndexOf(':');
+					if(ColonIdx != -1)
+					{
+						int ChangeNumber;
+						if(int.TryParse(LastUpdateResultText.Substring(0, ColonIdx), out ChangeNumber))
+						{
+							WorkspaceUpdateResult Result;
+							if(Enum.TryParse(LastUpdateResultText.Substring(ColonIdx + 1), out Result))
+							{
+								CurrentWorkspace.LastSyncChangeNumber = ChangeNumber;
+								CurrentWorkspace.LastSyncResult = Result;
+							}
+						}
+					}
+				}
+			}
+			else
+			{
+				CurrentWorkspace.CurrentProjectIdentifier = WorkspaceSection.GetValue("CurrentProjectPath");
+				CurrentWorkspace.CurrentChangeNumber = WorkspaceSection.GetValue("CurrentChangeNumber", -1);
+				CurrentWorkspace.CurrentSyncFilterHash = WorkspaceSection.GetValue("CurrentSyncFilterHash", null);
+				foreach(string AdditionalChangeNumberString in WorkspaceSection.GetValues("AdditionalChangeNumbers", new string[0]))
+				{
+					int AdditionalChangeNumber;
+					if(int.TryParse(AdditionalChangeNumberString, out AdditionalChangeNumber))
+					{
+						CurrentWorkspace.AdditionalChangeNumbers.Add(AdditionalChangeNumber);
+					}
+				}
+
+				WorkspaceUpdateResult LastSyncResult;
+				Enum.TryParse(WorkspaceSection.GetValue("LastSyncResult", ""), out LastSyncResult);
+				CurrentWorkspace.LastSyncResult = LastSyncResult;
+
+				CurrentWorkspace.LastSyncResultMessage = UnescapeText(WorkspaceSection.GetValue("LastSyncResultMessage"));
+				CurrentWorkspace.LastSyncChangeNumber = WorkspaceSection.GetValue("LastSyncChangeNumber", -1);
+
+				DateTime LastSyncTime;
+				if(DateTime.TryParse(WorkspaceSection.GetValue("LastSyncTime", ""), out LastSyncTime))
+				{
+					CurrentWorkspace.LastSyncTime = LastSyncTime;
+				}
+
+				CurrentWorkspace.LastSyncDurationSeconds = WorkspaceSection.GetValue("LastSyncDuration", 0);
+				CurrentWorkspace.LastBuiltChangeNumber = WorkspaceSection.GetValue("LastBuiltChangeNumber", 0);
+				CurrentWorkspace.ExpandedArchiveTypes = WorkspaceSection.GetValues("ExpandedArchiveName", new string[0]);
+
+				string[] BisectEntries = WorkspaceSection.GetValues("Bisect", new string[0]);
+				foreach(string BisectEntry in BisectEntries)
+				{
+					ConfigObject BisectEntryObject = new ConfigObject(BisectEntry);
+
+					int ChangeNumber = BisectEntryObject.GetValue("Change", -1);
+					if(ChangeNumber != -1)
+					{
+						BisectState State;
+						if(Enum.TryParse(BisectEntryObject.GetValue("State", ""), out State))
+						{
+							BisectEntry Entry = new BisectEntry();
+							Entry.Change = ChangeNumber;
+							Entry.State = State;
+							CurrentWorkspace.BisectChanges.Add(Entry);
+						}
+					}
+				}
+			}
+
+			return CurrentWorkspace;
+		}
+
+		UserWorkspaceSettings ImportWorkspaceSettings(string CurrentWorkspaceKey)
+		{
+			UserWorkspaceSettings CurrentWorkspace = new UserWorkspaceSettings();
+
+			ConfigSection WorkspaceSection = ConfigFile.FindSection(CurrentWorkspaceKey);
+			if (WorkspaceSection == null)
+			{
+				CurrentWorkspace.SyncView = new string[0];
+				CurrentWorkspace.SyncCategories = new List<SyncCategory>();
+				CurrentWorkspace.bSyncAllProjects = null;
+				CurrentWorkspace.bIncludeAllProjectsInSolution = null;
+			}
+			else
+			{
+				CurrentWorkspace.SyncView = WorkspaceSection.GetValues("SyncFilter", new string[0]);
+				CurrentWorkspace.SyncCategories = GetCategorySettings(WorkspaceSection, "SyncIncludedCategories", "SyncExcludedCategories").Select(x => new SyncCategory { Id = x.Key, Enable = x.Value }).ToList();
+
+				int SyncAllProjects = WorkspaceSection.GetValue("SyncAllProjects", -1);
+				CurrentWorkspace.bSyncAllProjects = (SyncAllProjects == 0) ? (bool?)false : (SyncAllProjects == 1) ? (bool?)true : (bool?)null;
+
+				int IncludeAllProjectsInSolution = WorkspaceSection.GetValue("IncludeAllProjectsInSolution", -1);
+				CurrentWorkspace.bIncludeAllProjectsInSolution = (IncludeAllProjectsInSolution == 0) ? (bool?)false : (IncludeAllProjectsInSolution == 1) ? (bool?)true : (bool?)null;
+			}
+
+			return CurrentWorkspace;
+		}
+
+		public UserProjectSettings FindOrAddProjectSettings(DirectoryReference LocalRoot, string ClientProjectFileName)
+		{
+			FileReference ConfigFile = FileReference.Combine(GetConfigDir(LocalRoot), $"project_{ClientProjectFileName.Trim('/').Replace('/', '+')}.json");
+			if (!ProjectKeyToSettings.TryGetValue(ConfigFile, out UserProjectSettings? Settings))
+			{
+				if (!Utility.TryLoadJson(ConfigFile, out Settings))
+				{
+					Settings = ImportProjectSettings(ClientProjectFileName);
+					Utility.SaveJson(ConfigFile, Settings);
+				}
+
+				Settings.File = ConfigFile;
+				ProjectKeyToSettings.Add(ConfigFile, Settings);
+			}
+			return Settings;
+		}
+
+		public UserProjectSettings ImportProjectSettings(string ClientProjectFileName)
+		{
+			// Read the project settings
+			UserProjectSettings CurrentProject = new UserProjectSettings();
+	
+			ConfigSection ProjectSection = ConfigFile.FindOrAddSection(ClientProjectFileName);
+			CurrentProject.BuildSteps.AddRange(ProjectSection.GetValues("BuildStep", new string[0]).Select(x => new ConfigObject(x)));
+
+			FilterType FilterType;
+			if (!Enum.TryParse(ProjectSection.GetValue("FilterType", ""), true, out FilterType))
+			{
+				FilterType = FilterType.None;
+			}
+
+			CurrentProject.FilterType = FilterType;
+			CurrentProject.FilterBadges.UnionWith(ProjectSection.GetValues("FilterBadges", new string[0]));
+
+			return CurrentProject;
+		}
+
+		public void Save()
+		{
+			// General settings
+			ConfigSection GeneralSection = ConfigFile.FindOrAddSection("General");
+			GeneralSection.Clear();
+			GeneralSection.SetValue("Version", (int)Version);
+			GeneralSection.SetValue("BuildAfterSync", bBuildAfterSync);
+			GeneralSection.SetValue("RunAfterSync", bRunAfterSync);
+			GeneralSection.SetValue("OpenSolutionAfterSync", bOpenSolutionAfterSync);
+			GeneralSection.SetValue("ShowLogWindow", bShowLogWindow);
+			GeneralSection.SetValue("AutoResolveConflicts", bAutoResolveConflicts);
+			GeneralSection.SetValue("ShowUnreviewed", bShowUnreviewedChanges);
+			GeneralSection.SetValue("ShowAutomated", bShowAutomatedChanges);
+			GeneralSection.SetValue("RobomergeFilter", ShowRobomerge.ToString());
+			GeneralSection.SetValue("AnnotateRobomerge", bAnnotateRobmergeChanges);
+			GeneralSection.SetValue("ShowLocalTimes", bShowLocalTimes);
+			if(LastProject != null)
+			{
+				GeneralSection.SetValue("LastProject", LastProject.ToConfigEntry());
+			}
+			GeneralSection.SetValues("OpenProjects", OpenProjects.Select(x => x.ToConfigEntry()).ToArray());
+			GeneralSection.SetValue("KeepInTray", bKeepInTray);
+			GeneralSection.SetValues("EnabledTools", EnabledTools);
+			GeneralSection.SetValue("ShowNetCoreInfo", bShowNetCoreInfo);
+			GeneralSection.SetValue("FilterIndex", FilterIndex);
+			GeneralSection.SetValues("RecentProjects", RecentProjects.Select(x => x.ToConfigEntry()).ToArray());
+			GeneralSection.SetValues("SyncFilter", SyncView);
+			SetCategorySettings(GeneralSection, "SyncIncludedCategories", "SyncExcludedCategories", SyncCategories);
+			GeneralSection.SetValue("SyncAllProjects", bSyncAllProjects);
+			GeneralSection.SetValue("IncludeAllProjectsInSolution", bIncludeAllProjectsInSolution);
+			GeneralSection.SetValue("SyncType", SyncType.ToString());
+
+			// Build configuration
+			GeneralSection.SetValue("BuildConfig", CompiledEditorBuildConfig.ToString());
+
+			// Tab labels
+			GeneralSection.SetValue("TabLabels", TabLabels.ToString());
+
+			// Editor arguments
+			List<string> EditorArgumentList = new List<string>();
+			foreach(Tuple<string, bool> EditorArgument in EditorArguments)
+			{
+				EditorArgumentList.Add(String.Format("{0}:{1}", EditorArgument.Item2? 1 : 0, EditorArgument.Item1));
+			}
+			GeneralSection.SetValues("EditorArguments", EditorArgumentList.ToArray());
+			GeneralSection.SetValue("EditorArgumentsPrompt", bEditorArgumentsPrompt);
+
+			// Schedule settings
+			ConfigSection ScheduleSection = ConfigFile.FindOrAddSection("Schedule");
+			ScheduleSection.Clear();
+			ScheduleSection.SetValue("Enabled", bScheduleEnabled);
+			ScheduleSection.SetValue("Time", ScheduleTime.ToString());
+			ScheduleSection.SetValue("Change", ScheduleChange.ToString());
+			ScheduleSection.SetValue("AnyOpenProject", ScheduleAnyOpenProject);
+			ScheduleSection.SetValues("Projects", ScheduleProjects.Select(x => x.ToConfigEntry()).ToArray());
+
+			// Precompiled binaries
+			ConfigSection ArchivesSection = ConfigFile.FindOrAddSection("PrecompiledBinaries");
+			ArchivesSection.SetValues("Archives", Archives.Select(x => x.ToConfigEntry()).ToArray());
+
+			// OIDC Settings
+			ConfigSection OIDCSection = ConfigFile.FindOrAddSection("OIDCProviders");
+			List<ConfigObject> TokenObjects = new List<ConfigObject>();
+			foreach (KeyValuePair<string, string> Pair in ProviderToRefreshTokens)
+			{
+				ConfigObject TokenEntryObject = new ConfigObject();
+				TokenEntryObject.SetValue("Provider", Pair.Key);
+				TokenEntryObject.SetValue("Token", Pair.Value);
+				TokenObjects.Add(TokenEntryObject);
+			}
+			OIDCSection.SetValues("Tokens", TokenObjects.Select(x => x.ToString()).ToArray());
+
+
+			// Window settings
+			ConfigSection WindowSection = ConfigFile.FindOrAddSection("Window");
+			WindowSection.Clear();
+			WindowSection.SetValue("Visible", bWindowVisible);
+			WindowSection.SetValue("State", WindowState);
+			if(WindowBounds != null)
+			{
+				WindowSection.SetValue("Bounds", FormatRectangleValue(WindowBounds.Value));
+			}
+
+			// Notification settings
+			ConfigSection NotificationSection = ConfigFile.FindOrAddSection("Notifications");
+			NotificationSection.Clear();
+			if (NotifyProjects.Count > 0)
+			{
+				NotificationSection.SetValues("NotifyProjects", NotifyProjects.ToArray());
+			}
+			if (NotifyUnassignedMinutes != -1)
+			{
+				NotificationSection.SetValue("NotifyUnassignedMinutes", NotifyUnassignedMinutes);
+			}
+			if (NotifyUnacknowledgedMinutes != -1)
+			{
+				NotificationSection.SetValue("NotifyUnacknowledgedMinutes", NotifyUnacknowledgedMinutes);
+			}
+			if (NotifyUnresolvedMinutes != -1)
+			{
+				NotificationSection.SetValue("NotifyUnresolvedMinutes", NotifyUnresolvedMinutes);
+			}
+
+			// Perforce settings
+			ConfigSection PerforceSection = ConfigFile.FindOrAddSection("Perforce");
+			PerforceSection.Clear();
+			if(SyncOptions.NumRetries > 0 && SyncOptions.NumRetries != PerforceSyncOptions.DefaultNumRetries)
+			{
+				PerforceSection.SetValue("NumRetries", SyncOptions.NumRetries);
+			}
+			if(SyncOptions.NumThreads > 0 && SyncOptions.NumThreads != PerforceSyncOptions.DefaultNumThreads)
+			{
+				PerforceSection.SetValue("NumThreads", SyncOptions.NumThreads);
+			}
+			if(SyncOptions.TcpBufferSize > 0 && SyncOptions.TcpBufferSize != PerforceSyncOptions.DefaultTcpBufferSize)
+			{
+				PerforceSection.SetValue("TcpBufferSize", SyncOptions.TcpBufferSize);
+			}
+			if (SyncOptions.FileBufferSize > 0 && SyncOptions.FileBufferSize != PerforceSyncOptions.DefaultFileBufferSize)
+			{
+				PerforceSection.SetValue("FileBufferSize", SyncOptions.FileBufferSize);
+			}
+			if (SyncOptions.MaxCommandsPerBatch > 0 && SyncOptions.MaxCommandsPerBatch != PerforceSyncOptions.DefaultMaxCommandsPerBatch)
+			{
+				PerforceSection.SetValue("MaxCommandsPerBatch", SyncOptions.MaxCommandsPerBatch);
+			}
+			if (SyncOptions.MaxSizePerBatch > 0 && SyncOptions.MaxSizePerBatch != PerforceSyncOptions.DefaultMaxSizePerBatch)
+			{
+				PerforceSection.SetValue("MaxSizePerBatch", SyncOptions.MaxSizePerBatch);
+			}
+
+			// Save the file
+			ConfigFile.Save(FileName);
+		}
+
+		public static string[] GetCombinedSyncFilter(Dictionary<Guid, WorkspaceSyncCategory> UniqueIdToFilter, string[] GlobalView, Dictionary<Guid, bool> GlobalCategoryIdToSetting, string[] WorkspaceView, Dictionary<Guid, bool> WorkspaceCategoryIdToSetting)
+		{
+			List<string> Lines = new List<string>();
+			foreach(string ViewLine in Enumerable.Concat(GlobalView, WorkspaceView).Select(x => x.Trim()).Where(x => x.Length > 0 && !x.StartsWith(";")))
+			{
+				Lines.Add(ViewLine);
+			}
+
+			HashSet<Guid> Enabled = new HashSet<Guid>();
+			foreach (WorkspaceSyncCategory Filter in UniqueIdToFilter.Values)
+			{
+				bool bEnable = Filter.bEnable;
+
+				bool bGlobalEnable;
+				if (GlobalCategoryIdToSetting.TryGetValue(Filter.UniqueId, out bGlobalEnable))
+				{
+					bEnable = bGlobalEnable;
+				}
+
+				bool bWorkspaceEnable;
+				if (WorkspaceCategoryIdToSetting.TryGetValue(Filter.UniqueId, out bWorkspaceEnable))
+				{
+					bEnable = bWorkspaceEnable;
+				}
+
+				if(bEnable)
+				{
+					EnableFilter(Filter.UniqueId, Enabled, UniqueIdToFilter);
+				}
+			}
+
+			foreach (WorkspaceSyncCategory Filter in UniqueIdToFilter.Values.OrderBy(x => x.Name))
+			{
+				if (!Enabled.Contains(Filter.UniqueId))
+				{
+					Lines.AddRange(Filter.Paths.Select(x => "-" + x.Trim()));
+				}
+			}
+
+			return Lines.ToArray();
+		}
+
+		static void EnableFilter(Guid UniqueId, HashSet<Guid> Enabled, Dictionary<Guid, WorkspaceSyncCategory> UniqueIdToFilter)
+		{
+			if(Enabled.Add(UniqueId))
+			{
+				WorkspaceSyncCategory? Category;
+				if(UniqueIdToFilter.TryGetValue(UniqueId, out Category))
+				{
+					foreach(Guid RequiresUniqueId in Category.Requires)
+					{
+						EnableFilter(RequiresUniqueId, Enabled, UniqueIdToFilter);
+					}
+				}
+			}
+		}
+
+		[return: NotNullIfNotNull("Text")]
+		static string? EscapeText(string? Text)
+		{
+			if(Text == null)
+			{
+				return null;
+			}
+
+			StringBuilder Result = new StringBuilder();
+			for(int Idx = 0; Idx < Text.Length; Idx++)
+			{
+				switch(Text[Idx])
+				{
+					case '\\':
+						Result.Append("\\\\");
+						break;
+					case '\t':
+						Result.Append("\\t");
+						break;
+					case '\r':
+						Result.Append("\\r");
+						break;
+					case '\n':
+						Result.Append("\\n");
+						break;
+					case '\'':
+						Result.Append("\\\'");
+						break;
+					case '\"':
+						Result.Append("\\\"");
+						break;
+					default:
+						Result.Append(Text[Idx]);
+						break;
+				}
+			}
+			return Result.ToString();
+		}
+
+		[return: NotNullIfNotNull("Text")]
+		static string? UnescapeText(string? Text)
+		{
+			if(Text == null)
+			{
+				return null;
+			}
+
+			StringBuilder Result = new StringBuilder();
+			for(int Idx = 0; Idx < Text.Length; Idx++)
+			{
+				if(Text[Idx] == '\\' && Idx + 1 < Text.Length)
+				{
+					switch(Text[++Idx])
+					{
+						case 't':
+							Result.Append('\t');
+							break;
+						case 'r':
+							Result.Append('\r');
+							break;
+						case 'n':
+							Result.Append('\n');
+							break;
+						case '\'':
+							Result.Append('\'');
+							break;
+						case '\"':
+							Result.Append('\"');
+							break;
+						default:
+							Result.Append(Text[Idx]);
+							break;
+					}
+				}
+				else
+				{
+					Result.Append(Text[Idx]);
+				}
+			}
+			return Result.ToString();
+		}
+
+		public IEnumerable<FileReference> GetCachedFilePaths()
+		{
+			return WorkspaceKeyToState.Keys.Concat(WorkspaceKeyToSettings.Keys).Concat(ProjectKeyToSettings.Keys);
+		}
+	}
+}
diff --git a/Engine/Source/Programs/UnrealGameSync/UnrealGameSyncShared/Utility.cs b/Engine/Source/Programs/UnrealGameSync/UnrealGameSyncShared/Utility.cs
new file mode 100644
--- /dev/null
+++ b/Engine/Source/Programs/UnrealGameSync/UnrealGameSyncShared/Utility.cs
@@ -0,0 +1,557 @@
+// Copyright Epic Games, Inc. All Rights Reserved.
+
+using EpicGames.Core;
+using EpicGames.Perforce;
+using Microsoft.Extensions.Logging;
+using Microsoft.Win32;
+using System;
+using System.Collections.Generic;
+using System.Diagnostics;
+using System.Diagnostics.CodeAnalysis;
+using System.Drawing;
+using System.IO;
+using System.Linq;	
+using System.Text;
+using System.Text.Json;
+using System.Text.Json.Serialization;
+using System.Text.RegularExpressions;
+using System.Threading;
+using System.Threading.Tasks;
+
+namespace UnrealGameSync
+{
+	public sealed class UserErrorException : Exception
+	{
+		public int Code { get; }
+
+		public UserErrorException(string Message, int Code = 1) : base(Message)
+		{
+			this.Code = Code;
+		}
+	}
+
+	public static class Utility
+	{
+		static JsonSerializerOptions GetDefaultJsonSerializerOptions()
+		{
+			JsonSerializerOptions Options = new JsonSerializerOptions();
+			Options.PropertyNameCaseInsensitive = true;
+			Options.Converters.Add(new JsonStringEnumConverter());
+			return Options;
+		}
+
+		public static JsonSerializerOptions DefaultJsonSerializerOptions { get; } = GetDefaultJsonSerializerOptions();
+
+		public static bool TryLoadJson<T>(FileReference File, [NotNullWhen(true)] out T? Object) where T : class
+		{
+			if (!FileReference.Exists(File))
+			{
+				Object = null;
+				return false;
+			}
+
+			try
+			{
+				Object = LoadJson<T>(File);
+				return true;
+			}
+			catch
+			{
+				Object = null;
+				return false;
+			}
+		}
+
+		public static T LoadJson<T>(FileReference File)
+		{
+			JsonSerializerOptions Options = new JsonSerializerOptions { AllowTrailingCommas = true, ReadCommentHandling = JsonCommentHandling.Skip };
+			Options.Converters.Add(new JsonStringEnumConverter());
+
+			byte[] Data = FileReference.ReadAllBytes(File);
+			return JsonSerializer.Deserialize<T>(Data, Options);
+		}
+
+		public static void SaveJson<T>(FileReference File, T Object)
+		{
+			JsonSerializerOptions Options = new JsonSerializerOptions { IgnoreNullValues = true, WriteIndented = true };
+			Options.Converters.Add(new JsonStringEnumConverter());
+
+			using (Stream Stream = FileReference.Open(File, FileMode.Create, FileAccess.Write, FileShare.Read))
+			{
+				using (Utf8JsonWriter Writer = new Utf8JsonWriter(Stream, new JsonWriterOptions { Indented = true }))
+				{
+					JsonSerializer.Serialize(Writer, Object, Options);
+				}
+			}
+		}
+
+		public static string GetPathWithCorrectCase(FileInfo Info)
+		{
+			DirectoryInfo ParentInfo = Info.Directory;
+			if(Info.Exists)
+			{
+				return Path.Combine(GetPathWithCorrectCase(ParentInfo), ParentInfo.GetFiles(Info.Name)[0].Name); 
+			}
+			else
+			{
+				return Path.Combine(GetPathWithCorrectCase(ParentInfo), Info.Name);
+			}
+		}
+
+		public static string GetPathWithCorrectCase(DirectoryInfo Info)
+		{
+			DirectoryInfo ParentInfo = Info.Parent;
+			if(ParentInfo == null)
+			{
+				return Info.FullName.ToUpperInvariant();
+			}
+			else if(Info.Exists)
+			{
+				return Path.Combine(GetPathWithCorrectCase(ParentInfo), ParentInfo.GetDirectories(Info.Name)[0].Name);
+			}
+			else
+			{
+				return Path.Combine(GetPathWithCorrectCase(ParentInfo), Info.Name);
+			}
+		}
+
+		public static void ForceDeleteFile(string FileName)
+		{
+			if(File.Exists(FileName))
+			{
+				File.SetAttributes(FileName, File.GetAttributes(FileName) & ~FileAttributes.ReadOnly);
+				File.Delete(FileName);
+			}
+		}
+
+		public static bool SpawnProcess(string FileName, string CommandLine)
+		{
+			using(Process ChildProcess = new Process())
+			{
+				ChildProcess.StartInfo.FileName = FileName;
+				ChildProcess.StartInfo.Arguments = String.IsNullOrEmpty(CommandLine) ? "" : CommandLine;
+				ChildProcess.StartInfo.UseShellExecute = false;
+				return ChildProcess.Start();
+			}
+		}
+
+		public static bool SpawnHiddenProcess(string FileName, string CommandLine)
+		{
+			using(Process ChildProcess = new Process())
+			{
+				ChildProcess.StartInfo.FileName = FileName;
+				ChildProcess.StartInfo.Arguments = String.IsNullOrEmpty(CommandLine) ? "" : CommandLine;
+				ChildProcess.StartInfo.UseShellExecute = false;
+				ChildProcess.StartInfo.RedirectStandardOutput = true;
+				ChildProcess.StartInfo.RedirectStandardError = true;
+				ChildProcess.StartInfo.CreateNoWindow = true;
+				try
+				{
+					return ChildProcess.Start();
+				}
+				catch
+				{
+					return false;
+				}
+			}
+		}
+
+		public static async Task<int> ExecuteProcessAsync(string FileName, string? WorkingDir, string CommandLine, Action<string> OutputLine, CancellationToken CancellationToken)
+		{
+			using (ManagedProcess NewProcess = new ManagedProcess(null, FileName, CommandLine, WorkingDir, null, ProcessPriorityClass.Normal))
+			{
+				for (; ; )
+				{
+					string? Line = await NewProcess.ReadLineAsync(CancellationToken);
+					if (Line == null)
+					{
+						NewProcess.WaitForExit();
+						return NewProcess.ExitCode;
+					}
+					OutputLine(Line);
+				}
+			}
+		}
+
+		public static bool SafeIsFileUnderDirectory(string FileName, string DirectoryName)
+		{
+			try
+			{
+				string FullDirectoryName = Path.GetFullPath(DirectoryName).TrimEnd(Path.DirectorySeparatorChar) + Path.DirectorySeparatorChar;
+				string FullFileName = Path.GetFullPath(FileName);
+				return FullFileName.StartsWith(FullDirectoryName, StringComparison.InvariantCultureIgnoreCase);
+			}
+			catch(Exception)
+			{
+				return false;
+			}
+		}
+
+		/// <summary>
+		/// Expands variables in $(VarName) format in the given string. Variables are retrieved from the given dictionary, or through the environment of the current process.
+		/// Any unknown variables are ignored.
+		/// </summary>
+		/// <param name="InputString">String to search for variable names</param>
+		/// <param name="AdditionalVariables">Lookup of variable names to values</param>
+		/// <returns>String with all variables replaced</returns>
+		public static string ExpandVariables(string InputString, Dictionary<string, string>? AdditionalVariables = null)
+		{
+			string Result = InputString;
+			for (int Idx = Result.IndexOf("$("); Idx != -1; Idx = Result.IndexOf("$(", Idx))
+			{
+				// Find the end of the variable name
+				int EndIdx = Result.IndexOf(')', Idx + 2);
+				if (EndIdx == -1)
+				{
+					break;
+				}
+
+				// Extract the variable name from the string
+				string Name = Result.Substring(Idx + 2, EndIdx - (Idx + 2));
+
+				// Strip the format from the name
+				string? Format = null;
+				int FormatIdx = Name.IndexOf(':');
+				if(FormatIdx != -1)
+				{ 
+					Format = Name.Substring(FormatIdx + 1);
+					Name = Name.Substring(0, FormatIdx);
+				}
+
+				// Find the value for it, either from the dictionary or the environment block
+				string? Value;
+				if (AdditionalVariables == null || !AdditionalVariables.TryGetValue(Name, out Value))
+				{
+					Value = Environment.GetEnvironmentVariable(Name);
+					if (Value == null)
+					{
+						Idx = EndIdx + 1;
+						continue;
+					}
+				}
+
+				// Encode the variable if necessary
+				if(Format != null)
+				{
+					if(String.Equals(Format, "URI", StringComparison.InvariantCultureIgnoreCase))
+					{
+						Value = Uri.EscapeDataString(Value);
+					}
+				}
+
+				// Replace the variable, or skip past it
+				Result = Result.Substring(0, Idx) + Value + Result.Substring(EndIdx + 1);
+			}
+			return Result;
+		}
+
+		class ProjectJson
+		{
+			public bool Enterprise { get; set; }
+		}
+
+		/// <summary>
+		/// Determines if a project is an enterprise project
+		/// </summary>
+		/// <param name="FileName">Path to the project file</param>
+		/// <returns>True if the given filename is an enterprise project</returns>
+		public static bool IsEnterpriseProjectFromText(string Text)
+		{
+			try
+			{
+				JsonSerializerOptions Options = new JsonSerializerOptions();
+				Options.PropertyNameCaseInsensitive = true;
+				Options.Converters.Add(new JsonStringEnumConverter());
+
+				ProjectJson Project = JsonSerializer.Deserialize<ProjectJson>(Text, Options);
+
+				return Project.Enterprise;
+			}
+			catch
+			{
+				return false;
+			}
+		}
+
+		/******/
+
+		private static void AddLocalConfigPaths_WithSubFolders(DirectoryInfo BaseDir, string FileName, List<FileInfo> Files)
+		{
+			if(BaseDir.Exists)
+			{
+				FileInfo BaseFileInfo = new FileInfo(Path.Combine(BaseDir.FullName, FileName));
+				if(BaseFileInfo.Exists)
+				{
+					Files.Add(BaseFileInfo);
+				}
+
+				foreach (DirectoryInfo SubDirInfo in BaseDir.EnumerateDirectories())
+				{
+					FileInfo SubFile = new FileInfo(Path.Combine(SubDirInfo.FullName, FileName));
+					if (SubFile.Exists)
+					{
+						Files.Add(SubFile);
+					}
+				}
+			}
+		}
+
+		private static void AddLocalConfigPaths_WithExtensionDirs(DirectoryInfo BaseDir, string RelativePath, string FileName, List<FileInfo> Files)
+		{
+			if (BaseDir.Exists)
+			{
+				AddLocalConfigPaths_WithSubFolders(new DirectoryInfo(Path.Combine(BaseDir.FullName, RelativePath)), FileName, Files);
+
+				DirectoryInfo PlatformExtensionsDir = new DirectoryInfo(Path.Combine(BaseDir.FullName, "Platforms"));
+				if (PlatformExtensionsDir.Exists)
+				{
+					foreach (DirectoryInfo PlatformExtensionDir in PlatformExtensionsDir.EnumerateDirectories())
+					{
+						AddLocalConfigPaths_WithSubFolders(new DirectoryInfo(Path.Combine(PlatformExtensionDir.FullName, RelativePath)), FileName, Files);
+					}
+				}
+
+				DirectoryInfo RestrictedBaseDir = new DirectoryInfo(Path.Combine(BaseDir.FullName, "Restricted"));
+				if (RestrictedBaseDir.Exists)
+				{
+					foreach (DirectoryInfo RestrictedDir in RestrictedBaseDir.EnumerateDirectories())
+					{
+						AddLocalConfigPaths_WithSubFolders(new DirectoryInfo(Path.Combine(RestrictedDir.FullName, RelativePath)), FileName, Files);
+					}
+				}
+			}
+		}
+
+		public static List<FileInfo> GetLocalConfigPaths(DirectoryInfo EngineDir, FileInfo ProjectFile)
+		{
+			List<FileInfo> SearchPaths = new List<FileInfo>();
+			AddLocalConfigPaths_WithExtensionDirs(EngineDir, "Programs/UnrealGameSync", "UnrealGameSync.ini", SearchPaths);
+
+			if (ProjectFile.Name.EndsWith(".uproject", StringComparison.OrdinalIgnoreCase))
+			{
+				AddLocalConfigPaths_WithExtensionDirs(ProjectFile.Directory, "Build", "UnrealGameSync.ini", SearchPaths);
+			}
+			else
+			{
+				AddLocalConfigPaths_WithExtensionDirs(EngineDir, "Programs/UnrealGameSync", "DefaultEngine.ini", SearchPaths);
+			}
+			return SearchPaths;
+		}
+
+		/******/
+
+		private static void AddDepotConfigPaths_PlatformFolders(string BasePath, string FileName, List<string> SearchPaths)
+		{
+			SearchPaths.Add(String.Format("{0}/{1}", BasePath, FileName));
+			SearchPaths.Add(String.Format("{0}/*/{1}", BasePath, FileName));
+		}
+
+		private static void AddDepotConfigPaths_PlatformExtensions(string BasePath, string RelativePath, string FileName, List<string> SearchPaths)
+		{
+			AddDepotConfigPaths_PlatformFolders(BasePath + RelativePath, FileName, SearchPaths);
+			AddDepotConfigPaths_PlatformFolders(BasePath + "/Platforms/*" + RelativePath, FileName, SearchPaths);
+			AddDepotConfigPaths_PlatformFolders(BasePath + "/Restricted/*" + RelativePath, FileName, SearchPaths);
+		}
+
+		public static List<string> GetDepotConfigPaths(string EnginePath, string ProjectPath)
+		{
+			List<string> SearchPaths = new List<string>();
+			AddDepotConfigPaths_PlatformExtensions(EnginePath, "/Programs/UnrealGameSync", "UnrealGameSync.ini", SearchPaths);
+
+			if (ProjectPath.EndsWith(".uproject", StringComparison.OrdinalIgnoreCase))
+			{
+				AddDepotConfigPaths_PlatformExtensions(ProjectPath.Substring(0, ProjectPath.LastIndexOf('/')), "/Build", "UnrealGameSync.ini", SearchPaths);
+			}
+			else
+			{
+				AddDepotConfigPaths_PlatformExtensions(EnginePath, "/Programs/UnrealGameSync", "DefaultEngine.ini", SearchPaths);
+			}
+			return SearchPaths;
+		}
+
+		/******/
+
+		public static async Task<string[]?> TryPrintFileUsingCacheAsync(IPerforceConnection Perforce, string DepotPath, DirectoryReference CacheFolder, string? Digest, ILogger Logger, CancellationToken CancellationToken)
+		{
+			if(Digest == null)
+			{
+				PerforceResponse<PrintRecord<string[]>> Response = await Perforce.TryPrintLinesAsync(DepotPath, CancellationToken);
+				if (Response.Succeeded)
+				{
+					return Response.Data.Contents;
+				}
+				else
+				{
+					return null;
+				}
+			}
+
+			FileReference CacheFile = FileReference.Combine(CacheFolder, Digest);
+			if(FileReference.Exists(CacheFile))
+			{
+				Logger.LogInformation("Reading cached copy of {DepotFile} from {LocalFile}", DepotPath, CacheFile);
+				string[] Lines = FileReference.ReadAllLines(CacheFile);
+				try
+				{
+					FileReference.SetLastWriteTimeUtc(CacheFile, DateTime.UtcNow);
+				}
+				catch(Exception Ex)
+				{
+					Logger.LogWarning(Ex, "Exception touching cache file {LocalFile}", CacheFile);
+				}
+				return Lines;
+			}
+			else
+			{
+				FileReference TempFile = new FileReference(String.Format("{0}.{1}.temp", CacheFile.FullName, Guid.NewGuid()));
+				PerforceResponse<PrintRecord> Response = await Perforce.TryPrintAsync(TempFile.FullName, DepotPath, CancellationToken);
+				if (!Response.Succeeded)
+				{
+					return null;
+				}
+				else
+				{
+					string[] Lines = await FileReference.ReadAllLinesAsync(TempFile);
+					try
+					{
+						FileReference.SetAttributes(TempFile, FileAttributes.Normal);
+						FileReference.SetLastWriteTimeUtc(TempFile, DateTime.UtcNow);
+						FileReference.Move(TempFile, CacheFile);
+					}
+					catch
+					{
+						try
+						{
+							FileReference.Delete(TempFile);
+						}
+						catch
+						{
+						}
+					}
+					return Lines;
+				}
+			}
+		}
+
+		public static void ClearPrintCache(DirectoryReference CacheFolder)
+		{
+			DirectoryInfo CacheDir = CacheFolder.ToDirectoryInfo();
+			if(CacheDir.Exists)
+			{
+				DateTime DeleteTime = DateTime.UtcNow - TimeSpan.FromDays(5.0);
+				foreach(FileInfo CacheFile in CacheDir.EnumerateFiles())
+				{
+					if(CacheFile.LastWriteTimeUtc < DeleteTime || CacheFile.Name.EndsWith(".temp", StringComparison.OrdinalIgnoreCase))
+					{
+						try
+						{
+							CacheFile.Attributes = FileAttributes.Normal;
+							CacheFile.Delete();
+						}
+						catch
+						{
+						}
+					}
+				}
+			}
+		}
+
+		public static Color Blend(Color First, Color Second, float T)
+		{
+			return Color.FromArgb((int)(First.R + (Second.R - First.R) * T), (int)(First.G + (Second.G - First.G) * T), (int)(First.B + (Second.B - First.B) * T));
+		}
+
+		public static PerforceSettings OverridePerforceSettings(IPerforceSettings DefaultConnection, string? ServerAndPort, string? UserName)
+		{
+			PerforceSettings NewSettings = new PerforceSettings(DefaultConnection);
+			if(!String.IsNullOrWhiteSpace(ServerAndPort))
+			{
+				NewSettings.ServerAndPort = ServerAndPort;
+			}
+			if (!String.IsNullOrWhiteSpace(UserName))
+			{
+				NewSettings.UserName = UserName;
+			}
+			return NewSettings;
+		}
+
+		public static string FormatRecentDateTime(DateTime Date)
+		{
+			DateTime Now = DateTime.Now;
+
+			DateTime Midnight = new DateTime(Now.Year, Now.Month, Now.Day);
+			DateTime MidnightTonight = Midnight + TimeSpan.FromDays(1.0);
+
+			if(Date > MidnightTonight)
+			{
+				return String.Format("{0} at {1}", Date.ToLongDateString(), Date.ToShortTimeString());
+			}
+			else if(Date >= Midnight)
+			{
+				return String.Format("today at {0}", Date.ToShortTimeString());
+			}
+			else if(Date >= Midnight - TimeSpan.FromDays(1.0))
+			{
+				return String.Format("yesterday at {0}", Date.ToShortTimeString());
+			}
+			else if(Date >= Midnight - TimeSpan.FromDays(5.0))
+			{
+				return String.Format("{0:dddd} at {1}", Date, Date.ToShortTimeString());
+			}
+			else
+			{
+				return String.Format("{0} at {1}", Date.ToLongDateString(), Date.ToShortTimeString());
+			}
+		}
+
+		public static string FormatDurationMinutes(TimeSpan Duration)
+		{
+			return FormatDurationMinutes((int)(Duration.TotalMinutes + 1));
+		}
+
+		public static string FormatDurationMinutes(int TotalMinutes)
+		{
+			if(TotalMinutes > 24 * 60)
+			{
+				return String.Format("{0}d {1}h", TotalMinutes / (24 * 60), (TotalMinutes / 60) % 24);
+			}
+			else if(TotalMinutes > 60)
+			{
+				return String.Format("{0}h {1}m", TotalMinutes / 60, TotalMinutes % 60);
+			}
+			else
+			{
+				return String.Format("{0}m", TotalMinutes);
+			}
+		}
+
+		public static string FormatUserName(string UserName)
+		{
+			StringBuilder NormalUserName = new StringBuilder();
+			for(int Idx = 0; Idx < UserName.Length; Idx++)
+			{
+				if(Idx == 0 || UserName[Idx - 1] == '.')
+				{
+					NormalUserName.Append(Char.ToUpper(UserName[Idx]));
+				}
+				else if(UserName[Idx] == '.')
+				{
+					NormalUserName.Append(' ');
+				}
+				else
+				{
+					NormalUserName.Append(Char.ToLower(UserName[Idx]));
+				}
+			}
+			return NormalUserName.ToString();
+		}
+
+		public static void OpenUrl(string Url)
+		{
+			ProcessStartInfo StartInfo = new ProcessStartInfo();
+			StartInfo.FileName = Url;
+			StartInfo.UseShellExecute = true;
+			using Process _ = Process.Start(StartInfo);
+		}
+	}
+}
diff --git a/Engine/Source/Programs/UnrealGameSync/UnrealGameSyncShared/WorkspaceUpdate.cs b/Engine/Source/Programs/UnrealGameSync/UnrealGameSyncShared/WorkspaceUpdate.cs
new file mode 100644
--- /dev/null
+++ b/Engine/Source/Programs/UnrealGameSync/UnrealGameSyncShared/WorkspaceUpdate.cs
@@ -0,0 +1,1825 @@
+// Copyright Epic Games, Inc. All Rights Reserved.
+
+using EpicGames.Core;
+using EpicGames.Perforce;
+using Microsoft.Extensions.DependencyInjection;
+using Microsoft.Extensions.Logging;
+using System;
+using System.Collections.Concurrent;
+using System.Collections.Generic;
+using System.Diagnostics;
+using System.Diagnostics.CodeAnalysis;
+using System.Globalization;
+using System.IO;
+using System.IO.Compression;
+using System.Linq;
+using System.Security.Cryptography;
+using System.Text;
+using System.Text.Encodings.Web;
+using System.Text.Json;
+using System.Threading;
+using System.Threading.Tasks;
+
+namespace UnrealGameSync
+{
+	[Flags]
+	public enum WorkspaceUpdateOptions
+	{
+		Sync = 0x01,
+		SyncSingleChange = 0x02,
+		AutoResolveChanges = 0x04,
+		GenerateProjectFiles = 0x08,
+		SyncArchives = 0x10,
+		Build = 0x20,
+		Clean = 0x40,
+		ScheduledBuild = 0x80,
+		RunAfterSync = 0x100,
+		OpenSolutionAfterSync = 0x200,
+		ContentOnly = 0x400,
+		UpdateFilter = 0x800,
+		SyncAllProjects = 0x1000,
+		IncludeAllProjectsInSolution = 0x2000,
+	}
+
+	public enum WorkspaceUpdateResult
+	{
+		Canceled,
+		FailedToSync,
+		FailedToSyncLoginExpired,
+		FilesToDelete,
+		FilesToResolve,
+		FilesToClobber,
+		FailedToCompile,
+		FailedToCompileWithCleanWorkspace,
+		Success,
+	}
+
+	public class PerforceSyncOptions
+	{
+		public const int DefaultNumRetries = 0;
+		public const int DefaultNumThreads = 2;
+		public const int DefaultTcpBufferSize = 0;
+		public const int DefaultFileBufferSize = 0;
+
+		public const int DefaultMaxCommandsPerBatch = 200;
+		public const int DefaultMaxSizePerBatch = 128 * 1024 * 1024;
+
+		public int NumRetries = DefaultNumRetries;
+		public int NumThreads = DefaultNumThreads;
+		public int TcpBufferSize = DefaultTcpBufferSize;
+		public int FileBufferSize = DefaultFileBufferSize;
+
+		public int MaxCommandsPerBatch = DefaultMaxCommandsPerBatch;
+		public int MaxSizePerBatch = DefaultMaxSizePerBatch;
+
+
+		public PerforceSyncOptions Clone()
+		{
+			return (PerforceSyncOptions)MemberwiseClone();
+		}
+	}
+
+	public interface IArchiveInfo
+	{
+		string Name { get; }
+		string Type { get; }
+		string BasePath { get; }
+		string? Target { get; }
+		bool Exists();
+		bool TryGetArchiveKeyForChangeNumber(int ChangeNumber, [NotNullWhen(true)] out string? ArchiveKey);
+		Task<bool> DownloadArchive(IPerforceConnection Perforce, string ArchiveKey, DirectoryReference LocalRootPath, FileReference ManifestFileName, ILogger Logger, ProgressValue Progress, CancellationToken CancellationToken);
+	}
+
+	public class WorkspaceUpdateContext
+	{
+		public DateTime StartTime = DateTime.UtcNow;
+		public int ChangeNumber;
+		public WorkspaceUpdateOptions Options;
+		public string[]? SyncFilter;
+		public Dictionary<string, Tuple<IArchiveInfo, string>?> ArchiveTypeToArchive = new Dictionary<string, Tuple<IArchiveInfo, string>?>();
+		public Dictionary<string, bool> DeleteFiles = new Dictionary<string,bool>();
+		public Dictionary<string, bool> ClobberFiles = new Dictionary<string,bool>();
+		public Dictionary<Guid,ConfigObject> DefaultBuildSteps;
+		public List<ConfigObject> UserBuildStepObjects;
+		public HashSet<Guid>? CustomBuildSteps;
+		public Dictionary<string, string> Variables;
+		public PerforceSyncOptions? PerforceSyncOptions;
+		public List<HaveRecord>? HaveFiles; // Cached when sync filter has changed
+
+		// May be updated during sync
+		public ConfigFile? ProjectConfigFile;
+		public IReadOnlyList<string>? ProjectStreamFilter;
+
+		public WorkspaceUpdateContext(int InChangeNumber, WorkspaceUpdateOptions InOptions, string[]? InSyncFilter, Dictionary<Guid, ConfigObject> InDefaultBuildSteps, List<ConfigObject> InUserBuildSteps, HashSet<Guid>? InCustomBuildSteps, Dictionary<string, string> InVariables)
+		{
+			ChangeNumber = InChangeNumber;
+			Options = InOptions;
+			SyncFilter = InSyncFilter;
+			DefaultBuildSteps = InDefaultBuildSteps;
+			UserBuildStepObjects = InUserBuildSteps;
+			CustomBuildSteps = InCustomBuildSteps;
+			Variables = InVariables;
+		}
+	}
+
+	public class WorkspaceSyncCategory
+	{
+		public Guid UniqueId;
+		public bool bEnable;
+		public string Name;
+		public string[] Paths;
+		public bool bHidden;
+		public Guid[] Requires;
+
+		public WorkspaceSyncCategory(Guid UniqueId) : this(UniqueId, "Unnamed")
+		{
+		}
+
+		public WorkspaceSyncCategory(Guid UniqueId, string Name, params string[] Paths)
+		{
+			this.UniqueId = UniqueId;
+			this.bEnable = true;
+			this.Name = Name;
+			this.Paths = Paths;
+			this.Requires = new Guid[0];
+		}
+
+		public static Dictionary<Guid, bool> GetDefault(IEnumerable<WorkspaceSyncCategory> Categories)
+		{
+			return Categories.ToDictionary(x => x.UniqueId, x => x.bEnable);
+		}
+
+		public static Dictionary<Guid, bool> GetDelta(Dictionary<Guid, bool> Source, Dictionary<Guid, bool> Target)
+		{
+			Dictionary<Guid, bool> Changes = new Dictionary<Guid, bool>();
+			foreach (KeyValuePair<Guid, bool> Pair in Target)
+			{
+				bool bValue;
+				if (!Source.TryGetValue(Pair.Key, out bValue) || bValue != Pair.Value)
+				{
+					Changes[Pair.Key] = Pair.Value;
+				}
+			}
+			return Changes;
+		}
+
+		public static void ApplyDelta(Dictionary<Guid, bool> Categories, Dictionary<Guid, bool> Delta)
+		{
+			foreach(KeyValuePair<Guid, bool> Pair in Delta)
+			{
+				Categories[Pair.Key] = Pair.Value;
+			}
+		}
+
+		public override string ToString()
+		{
+			return Name;
+		}
+	}
+
+	public class ProjectInfo
+	{
+		public DirectoryReference LocalRootPath { get; }
+		public FileReference LocalFileName { get; }
+		public string ClientRootPath { get; }
+		public string ClientFileName { get; }
+		public string ProjectIdentifier { get; }
+		public string TelemetryProjectIdentifier { get; }
+		public bool bIsEnterpriseProject { get; }
+
+		public ProjectInfo(DirectoryReference LocalRootPath, FileReference LocalFileName, string ClientRootPath, string ClientFileName, string ProjectIdentifier, bool bIsEnterpriseProject)
+		{
+			this.LocalRootPath = LocalRootPath;
+			this.LocalFileName = LocalFileName;
+			this.ClientRootPath = ClientRootPath;
+			this.ClientFileName = ClientFileName;
+			this.ProjectIdentifier = ProjectIdentifier;
+			this.TelemetryProjectIdentifier = PerforceUtils.GetClientOrDepotDirectoryName(ProjectIdentifier);
+			this.bIsEnterpriseProject = bIsEnterpriseProject;
+		}
+	}
+
+	public class WorkspaceUpdate : IAsyncDisposable
+	{
+		const string BuildVersionFileName = "/Engine/Build/Build.version";
+		const string VersionHeaderFileName = "/Engine/Source/Runtime/Launch/Resources/Version.h";
+		const string ObjectVersionFileName = "/Engine/Source/Runtime/Core/Private/UObject/ObjectVersion.cpp";
+
+		static readonly string LocalVersionHeaderFileName = VersionHeaderFileName.Replace('/', '\\');
+		static readonly string LocalObjectVersionFileName = ObjectVersionFileName.Replace('/', '\\');
+
+		static SemaphoreSlim UpdateSemaphore = new SemaphoreSlim(1);
+
+		class RecordCounter : IDisposable
+		{
+			ProgressValue Progress;
+			string Message;
+			int Count;
+			Stopwatch Timer = Stopwatch.StartNew();
+
+			public RecordCounter(ProgressValue Progress, string Message)
+			{
+				this.Progress = Progress;
+				this.Message = Message;
+
+				Progress.Set(Message);
+			}
+
+			public void Dispose()
+			{
+				UpdateMessage();
+			}
+
+			public void Increment()
+			{
+				Count++;
+				if (Timer.ElapsedMilliseconds > 250)
+				{
+					UpdateMessage();
+				}
+			}
+
+			public void UpdateMessage()
+			{
+				Progress.Set(String.Format("{0} ({1:N0})", Message, Count));
+				Timer.Restart();
+			}
+		}
+
+		class SyncBatchBuilder
+		{
+			public int MaxCommandsPerList { get; }
+			public long MaxSizePerList { get; }
+			public Queue<List<string>> Batches { get; }
+
+			List<string>? Commands;
+			List<string>? DeleteCommands;
+			long Size;
+
+			public SyncBatchBuilder(int? MaxCommandsPerList, long? MaxSizePerList)
+			{
+				this.MaxCommandsPerList = MaxCommandsPerList ?? PerforceSyncOptions.DefaultMaxCommandsPerBatch;
+				this.MaxSizePerList = MaxSizePerList ?? PerforceSyncOptions.DefaultMaxSizePerBatch;
+				this.Batches = new Queue<List<string>>();
+			}
+
+			public void Add(string NewCommand, long NewSize)
+			{
+				if (NewSize == 0)
+				{
+					if (DeleteCommands == null || DeleteCommands.Count >= MaxCommandsPerList)
+					{
+						DeleteCommands = new List<string>();
+						Batches.Enqueue(DeleteCommands);
+					}
+
+					DeleteCommands.Add(NewCommand);
+				}
+				else
+				{
+					if (Commands == null || Commands.Count >= MaxCommandsPerList || Size + NewSize >= MaxSizePerList)
+					{
+						Commands = new List<string>();
+						Batches.Enqueue(Commands);
+						Size = 0;
+					}
+
+					Commands.Add(NewCommand);
+					Size += NewSize;
+				}
+			}
+		}
+
+		class SyncTree
+		{
+			public bool bCanUseWildcard;
+			public int TotalIncludedFiles;
+			public long TotalSize;
+			public int TotalExcludedFiles;
+			public Dictionary<string, long> IncludedFiles = new Dictionary<string, long>();
+			public Dictionary<string, SyncTree> NameToSubTree = new Dictionary<string, SyncTree>(StringComparer.OrdinalIgnoreCase);
+
+			public SyncTree(bool bCanUseWildcard)
+			{
+				this.bCanUseWildcard = bCanUseWildcard;
+			}
+
+			public SyncTree FindOrAddSubTree(string Name)
+			{
+				SyncTree? Result;
+				if (!NameToSubTree.TryGetValue(Name, out Result))
+				{
+					Result = new SyncTree(bCanUseWildcard);
+					NameToSubTree.Add(Name, Result);
+				}
+				return Result;
+			}
+
+			public void IncludeFile(string Path, long Size)
+			{
+				int Idx = Path.IndexOf('/');
+				if (Idx == -1)
+				{
+					IncludedFiles.Add(Path, Size);
+				}
+				else
+				{
+					SyncTree SubTree = FindOrAddSubTree(Path.Substring(0, Idx));
+					SubTree.IncludeFile(Path.Substring(Idx + 1), Size);
+				}
+				TotalIncludedFiles++;
+				TotalSize += Size;
+			}
+
+			public void ExcludeFile(string Path)
+			{
+				int Idx = Path.IndexOf('/');
+				if (Idx != -1)
+				{
+					SyncTree SubTree = FindOrAddSubTree(Path.Substring(0, Idx));
+					SubTree.ExcludeFile(Path.Substring(Idx + 1));
+				}
+				TotalExcludedFiles++;
+			}
+
+			public void GetOptimizedSyncCommands(string Prefix, int ChangeNumber, SyncBatchBuilder Builder)
+			{
+				if (bCanUseWildcard && TotalExcludedFiles == 0 && TotalSize < Builder.MaxSizePerList)
+				{
+					Builder.Add(String.Format("{0}/...@{1}", Prefix, ChangeNumber), TotalSize);
+				}
+				else
+				{
+					foreach (KeyValuePair<string, long> File in IncludedFiles)
+					{
+						Builder.Add(String.Format("{0}/{1}@{2}", Prefix, File.Key, ChangeNumber), File.Value);
+					}
+					foreach (KeyValuePair<string, SyncTree> Pair in NameToSubTree)
+					{
+						Pair.Value.GetOptimizedSyncCommands(String.Format("{0}/{1}", Prefix, PerforceUtils.EscapePath(Pair.Key)), ChangeNumber, Builder);
+					}
+				}
+			}
+		}
+
+		CancellationTokenSource CancellationSource;
+		public WorkspaceUpdateContext Context { get; }
+		public ProgressValue Progress { get; } = new ProgressValue();
+		public Task UpdateTask { get; }
+
+		public WorkspaceUpdate(WorkspaceUpdateContext Context, Func<WorkspaceUpdate, CancellationToken, Task> TaskFunc)
+		{
+			this.CancellationSource = new CancellationTokenSource();
+			this.Context = Context;
+			this.UpdateTask = Task.Run(() => TaskFunc(this, CancellationSource.Token));
+		}
+
+		public void Cancel()
+		{
+			CancellationSource.Cancel();
+		}
+
+		public ValueTask DisposeAsync()
+		{
+			Task DisposeTask = Task.CompletedTask;
+			if (UpdateTask != null)
+			{
+				CancellationSource.Cancel();
+				DisposeTask = UpdateTask.ContinueWith(Task => CancellationSource.Dispose());
+			}
+			return new ValueTask(DisposeTask);
+		}
+
+		class SyncFile
+		{
+			public string DepotFile;
+			public string RelativePath;
+			public long Size;
+
+			public SyncFile(string DepotFile, string RelativePath, long Size)
+			{
+				this.DepotFile = DepotFile;
+				this.RelativePath = RelativePath;
+				this.Size = Size;
+			}
+		};
+
+		class SemaphoreScope : IDisposable
+		{
+			SemaphoreSlim Semaphore;
+			public bool HasLock { get; private set; }
+
+			public SemaphoreScope(SemaphoreSlim Semaphore)
+			{
+				this.Semaphore = Semaphore;
+			}
+
+			public bool TryAcquire()
+			{
+				HasLock = HasLock || Semaphore.Wait(0);
+				return HasLock;
+			}
+
+			public async Task AcquireAsync(CancellationToken CancellationToken)
+			{
+				await Semaphore.WaitAsync(CancellationToken);
+				HasLock = true;
+			}
+
+			public void Release()
+			{
+				if (HasLock)
+				{
+					Semaphore.Release();
+					HasLock = false;
+				}
+			}
+
+			public void Dispose() => Release();
+		}
+
+		public async Task<(WorkspaceUpdateResult, string)> ExecuteAsync(IPerforceSettings PerforceSettings, ProjectInfo Project, UserWorkspaceState State, ILogger Logger, CancellationToken CancellationToken)
+		{
+			using IPerforceConnection Perforce = await PerforceConnection.CreateAsync(PerforceSettings, Logger);
+
+			string CmdExe = Path.Combine(Environment.GetFolderPath(Environment.SpecialFolder.System), "cmd.exe");
+			if (!File.Exists(CmdExe))
+			{
+				return (WorkspaceUpdateResult.FailedToSync, $"Missing {CmdExe}.");
+			}
+
+			List<Tuple<string, TimeSpan>> Times = new List<Tuple<string, TimeSpan>>();
+
+			int NumFilesSynced = 0;
+			if (Context.Options.HasFlag(WorkspaceUpdateOptions.Sync) || Context.Options.HasFlag(WorkspaceUpdateOptions.SyncSingleChange))
+			{
+				using (TelemetryStopwatch SyncTelemetryStopwatch = new TelemetryStopwatch("Workspace_Sync", Project.TelemetryProjectIdentifier))
+				{
+					Logger.LogInformation("Syncing to {Change}...", Context.ChangeNumber);
+
+					// Make sure we're logged in
+					PerforceResponse<LoginRecord> LoginResponse = await Perforce.TryGetLoginStateAsync(CancellationToken);
+					if (!LoginResponse.Succeeded)
+					{
+						return (WorkspaceUpdateResult.FailedToSyncLoginExpired, "User is not logged in.");
+					}
+
+					// Figure out which paths to sync
+					List<string> RelativeSyncPaths = GetRelativeSyncPaths(Project, (Context.Options & WorkspaceUpdateOptions.SyncAllProjects) != 0, Context.SyncFilter);
+					List<string> SyncPaths = new List<string>(RelativeSyncPaths.Select(x => Project.ClientRootPath + x));
+
+					// Get the user's sync filter
+					FileFilter UserFilter = new FileFilter(FileFilterType.Include);
+					if (Context.SyncFilter != null)
+					{
+						UserFilter.AddRules(Context.SyncFilter.Select(x => x.Trim()).Where(x => x.Length > 0 && !x.StartsWith(";") && !x.StartsWith("#")));
+					}
+
+					// Check if the new sync filter matches the previous one. If not, we'll enumerate all files in the workspace and make sure there's nothing extra there.
+					string? NextSyncFilterHash = null;
+					using (SHA1Managed SHA = new SHA1Managed())
+					{
+						StringBuilder CombinedFilter = new StringBuilder();
+						foreach (string RelativeSyncPath in RelativeSyncPaths)
+						{
+							CombinedFilter.AppendFormat("{0}\n", RelativeSyncPath);
+						}
+						if (Context.SyncFilter != null)
+						{
+							CombinedFilter.Append("--FROM--\n");
+							CombinedFilter.Append(String.Join("\n", Context.SyncFilter));
+						}
+						NextSyncFilterHash = BitConverter.ToString(SHA.ComputeHash(Encoding.UTF8.GetBytes(CombinedFilter.ToString()))).Replace("-", "");
+					}
+
+					// If the hash differs, enumerate everything in the workspace to find what needs to be removed
+					if (NextSyncFilterHash != State.CurrentSyncFilterHash)
+					{
+						using (TelemetryStopwatch FilterStopwatch = new TelemetryStopwatch("Workspace_Sync_FilterChanged", Project.TelemetryProjectIdentifier))
+						{
+							Logger.LogInformation("Filter has changed ({PrevHash} -> {NextHash}); finding files in workspace that need to be removed.", (String.IsNullOrEmpty(State.CurrentSyncFilterHash)) ? "None" : State.CurrentSyncFilterHash, NextSyncFilterHash);
+
+							// Find all the files that are in this workspace
+							List<HaveRecord>? HaveFiles = Context.HaveFiles;
+							if (HaveFiles == null)
+							{
+								HaveFiles = new List<HaveRecord>();
+								using (RecordCounter HaveCounter = new RecordCounter(Progress, "Sync filter changed; checking workspace..."))
+								{
+									await foreach (PerforceResponse<HaveRecord> Record in Perforce.TryHaveAsync(FileSpecList.Any, CancellationToken))
+									{
+										if (Record.Succeeded)
+										{
+											HaveFiles.Add(Record.Data);
+											HaveCounter.Increment();
+										}
+										else
+										{
+											return (WorkspaceUpdateResult.FailedToSync, $"Unable to query files ({Record}).");
+										}
+									}
+								}
+								Context.HaveFiles = HaveFiles;
+							}
+
+							// Build a filter for the current sync paths
+							FileFilter SyncPathsFilter = new FileFilter(FileFilterType.Exclude);
+							foreach (string RelativeSyncPath in RelativeSyncPaths)
+							{
+								SyncPathsFilter.Include(RelativeSyncPath);
+							}
+
+							// Remove all the files that are not included by the filter
+							List<string> RemoveDepotPaths = new List<string>();
+							foreach (HaveRecord HaveFile in HaveFiles)
+							{
+								try
+								{
+									FileReference FullPath = new FileReference(HaveFile.Path);
+									if (MatchFilter(Project, FullPath, SyncPathsFilter) && !MatchFilter(Project, FullPath, UserFilter))
+									{
+										Logger.LogInformation("  {DepotFile}", HaveFile.DepotFile);
+										RemoveDepotPaths.Add(HaveFile.DepotFile);
+									}
+								}
+								catch (PathTooLongException)
+								{
+									// We don't actually care about this when looking for files to remove. Perforce may think that it's synced the path, and silently failed. Just ignore it.
+								}
+							}
+
+							// Check if there are any paths outside the regular sync paths
+							if (RemoveDepotPaths.Count > 0)
+							{
+								bool bDeleteListMatches = true;
+
+								Dictionary<string, bool> NewDeleteFiles = new Dictionary<string, bool>(StringComparer.OrdinalIgnoreCase);
+								foreach (string RemoveDepotPath in RemoveDepotPaths)
+								{
+									bool bDelete;
+									if (!Context.DeleteFiles.TryGetValue(RemoveDepotPath, out bDelete))
+									{
+										bDeleteListMatches = false;
+										bDelete = true;
+									}
+									NewDeleteFiles[RemoveDepotPath] = bDelete;
+								}
+								Context.DeleteFiles = NewDeleteFiles;
+
+								if (!bDeleteListMatches)
+								{
+									return (WorkspaceUpdateResult.FilesToDelete, $"Cancelled after finding {NewDeleteFiles.Count} files excluded by filter");
+								}
+
+								RemoveDepotPaths.RemoveAll(x => !Context.DeleteFiles[x]);
+							}
+
+							// Actually delete any files that we don't want
+							if (RemoveDepotPaths.Count > 0)
+							{
+								// Clear the current sync filter hash. If the sync is canceled, we'll be in an indeterminate state, and we should always clean next time round.
+								State.CurrentSyncFilterHash = "INVALID";
+								State.Save();
+
+								// Find all the depot paths that will be synced
+								HashSet<string> RemainingDepotPathsToRemove = new HashSet<string>(StringComparer.OrdinalIgnoreCase);
+								RemainingDepotPathsToRemove.UnionWith(RemoveDepotPaths);
+
+								// Build the list of revisions to sync
+								List<string> RevisionsToRemove = new List<string>();
+								RevisionsToRemove.AddRange(RemoveDepotPaths.Select(x => String.Format("{0}#0", x)));
+
+								(WorkspaceUpdateResult, string) RemoveResult = await SyncFileRevisions(Perforce, "Removing files...", Context, RevisionsToRemove, RemainingDepotPathsToRemove, Progress, Logger, CancellationToken);
+								if (RemoveResult.Item1 != WorkspaceUpdateResult.Success)
+								{
+									return RemoveResult;
+								}
+							}
+
+							// Update the sync filter hash. We've removed any files we need to at this point.
+							State.CurrentSyncFilterHash = NextSyncFilterHash;
+							State.Save();
+						}
+					}
+
+					// Create a filter for all the files we don't want
+					FileFilter Filter = new FileFilter(UserFilter);
+					Filter.Exclude(BuildVersionFileName);
+					if (Context.Options.HasFlag(WorkspaceUpdateOptions.ContentOnly))
+					{
+						Filter.Exclude("*.usf");
+						Filter.Exclude("*.ush");
+					}
+
+					// Create a tree to store the sync path
+					SyncTree SyncTree = new SyncTree(false);
+					if (!Context.Options.HasFlag(WorkspaceUpdateOptions.SyncSingleChange))
+					{
+						foreach (string RelativeSyncPath in RelativeSyncPaths)
+						{
+							const string WildcardSuffix = "/...";
+							if (RelativeSyncPath.EndsWith(WildcardSuffix, StringComparison.Ordinal))
+							{
+								SyncTree Leaf = SyncTree;
+
+								string[] Fragments = RelativeSyncPath.Split('/');
+								for (int Idx = 1; Idx < Fragments.Length - 1; Idx++)
+								{
+									Leaf = Leaf.FindOrAddSubTree(Fragments[Idx]);
+								}
+
+								Leaf.bCanUseWildcard = true;
+							}
+						}
+					}
+
+					// Find all the server changes, and anything that's opened for edit locally. We need to sync files we have open to schedule a resolve.
+					SyncBatchBuilder BatchBuilder = new SyncBatchBuilder(Context.PerforceSyncOptions?.MaxCommandsPerBatch, Context.PerforceSyncOptions?.MaxSizePerBatch);
+					List<string> SyncDepotPaths = new List<string>();
+					using (RecordCounter Counter = new RecordCounter(Progress, "Filtering files..."))
+					{
+						// Track the total new bytes that will be required on disk when syncing. Add an extra 100MB for padding.
+						long RequiredFreeSpace = 100 * 1024 * 1024;
+
+						foreach (string SyncPath in SyncPaths)
+						{
+							string SyncFilter = Context.Options.HasFlag(WorkspaceUpdateOptions.Sync) ? $"{Context.ChangeNumber}" : $"={Context.ChangeNumber}";
+
+							List<SyncFile> SyncFiles = new List<SyncFile>();
+							await foreach (PerforceResponse<SyncRecord> Response in Perforce.TrySyncAsync(SyncOptions.PreviewOnly, -1, $"{SyncPath}@{SyncFilter}", CancellationToken))
+							{
+								if (!Response.Succeeded)
+								{
+									return (WorkspaceUpdateResult.FailedToSync, $"Couldn't enumerate changes matching {SyncPath}.");
+								}
+
+								SyncRecord Record = Response.Data;
+
+								// Get the full local path
+								string RelativePath;
+								try
+								{
+									FileReference SyncFile = new FileReference(Record.Path);
+									RelativePath = PerforceUtils.GetClientRelativePath(Project.LocalRootPath, SyncFile);
+								}
+								catch (PathTooLongException)
+								{
+									Logger.LogInformation("The local path for {Path} exceeds the maximum allowed by Windows. Re-sync your workspace to a directory with a shorter name, or delete the file from the server.", Record.Path);
+									return (WorkspaceUpdateResult.FailedToSync, "File exceeds maximum path length allowed by Windows.");
+								}
+
+								// Create the sync record
+								long SyncSize = (Record.Action == SyncAction.Deleted) ? 0 : Record.FileSize;
+								SyncFiles.Add(new SyncFile(Record.DepotFile, RelativePath, SyncSize));
+								Counter.Increment();
+							}
+
+							// Also sync the currently open files
+							await foreach (PerforceResponse<OpenedRecord> Response in Perforce.TryOpenedAsync(OpenedOptions.None, FileSpecList.Any, CancellationToken))
+							{
+								if (!Response.Succeeded)
+								{
+									return (WorkspaceUpdateResult.FailedToSync, $"Couldn't enumerate changes matching {SyncPath}.");
+								}
+
+								OpenedRecord Record = Response.Data;
+								if (Record.Action != FileAction.Add || Record.Action != FileAction.Branch || Record.Action != FileAction.MoveAdd)
+								{
+									string RelativePath = PerforceUtils.GetClientRelativePath(Record.ClientFile);
+									SyncFiles.Add(new SyncFile(Record.DepotFile, RelativePath, 0));
+								}
+							}
+
+							// Enumerate all the files to be synced. NOTE: depotPath is escaped, whereas clientPath is not.
+							foreach (SyncFile SyncRecord in SyncFiles)
+							{
+								if (Filter.Matches(SyncRecord.RelativePath))
+								{
+									SyncTree.IncludeFile(PerforceUtils.EscapePath(SyncRecord.RelativePath), SyncRecord.Size);
+									SyncDepotPaths.Add(SyncRecord.DepotFile);
+									RequiredFreeSpace += SyncRecord.Size;
+
+									// If the file exists the required free space can be reduced as those bytes will be replaced.
+									FileInfo LocalFileInfo = FileReference.Combine(Project.LocalRootPath, SyncRecord.RelativePath).ToFileInfo();
+									if (LocalFileInfo.Exists)
+									{
+										RequiredFreeSpace -= LocalFileInfo.Length;
+									}
+								}
+								else
+								{
+									SyncTree.ExcludeFile(PerforceUtils.EscapePath(SyncRecord.RelativePath));
+								}
+							}
+						}
+
+						try
+						{
+							DirectoryInfo LocalRootInfo = Project.LocalRootPath.ToDirectoryInfo();
+							DriveInfo Drive = new DriveInfo(LocalRootInfo.Root.FullName);
+
+							if (Drive.AvailableFreeSpace < RequiredFreeSpace)
+							{
+								Logger.LogInformation("Syncing requires {RequiredSpace} which exceeds the {AvailableSpace} available free space on {Drive}.", StringUtils.FormatBytesString(RequiredFreeSpace), StringUtils.FormatBytesString(Drive.AvailableFreeSpace), Drive.Name);
+								return (WorkspaceUpdateResult.FailedToSync, "Not enough available free space.");
+							}
+						}
+						catch (SystemException)
+						{
+							Logger.LogInformation("Unable to check available free space for {RootPath}.", Project.LocalRootPath);
+						}
+					}
+					SyncTree.GetOptimizedSyncCommands(Project.ClientRootPath, Context.ChangeNumber, BatchBuilder);
+
+					// Clear the current sync changelist, in case we cancel
+					if (!Context.Options.HasFlag(WorkspaceUpdateOptions.SyncSingleChange))
+					{
+						State.CurrentChangeNumber = -1;
+						State.AdditionalChangeNumbers.Clear();
+						State.Save();
+					}
+
+					// Find all the depot paths that will be synced
+					HashSet<string> RemainingDepotPaths = new HashSet<string>(StringComparer.OrdinalIgnoreCase);
+					RemainingDepotPaths.UnionWith(SyncDepotPaths);
+
+					using (TelemetryStopwatch TransferStopwatch = new TelemetryStopwatch("Workspace_Sync_TransferFiles", Project.TelemetryProjectIdentifier))
+					{
+						TransferStopwatch.AddData(new { MachineName = Environment.MachineName, DomainName = Environment.UserDomainName, ServerAndPort = Perforce.Settings.ServerAndPort, UserName = Perforce.Settings.UserName, IncludedFiles = SyncTree.TotalIncludedFiles, ExcludedFiles = SyncTree.TotalExcludedFiles, Size = SyncTree.TotalSize, NumThreads = Context.PerforceSyncOptions?.NumThreads ?? PerforceSyncOptions.DefaultNumThreads });
+
+						(WorkspaceUpdateResult, string) SyncResult = await SyncFileRevisions(Perforce, "Syncing files...", Context, BatchBuilder.Batches, RemainingDepotPaths, Progress, Logger, CancellationToken);
+						if (SyncResult.Item1 != WorkspaceUpdateResult.Success)
+						{
+							TransferStopwatch.AddData(new { SyncResult = SyncResult.Item2.ToString(), CompletedFilesFiles = SyncDepotPaths.Count - RemainingDepotPaths.Count });
+							return SyncResult;
+						}
+
+						TransferStopwatch.Stop("Ok");
+						TransferStopwatch.AddData(new { TransferRate = SyncTree.TotalSize / Math.Max(TransferStopwatch.Elapsed.TotalSeconds, 0.0001f) });
+					}
+
+					int VersionChangeNumber = -1;
+					if (Context.Options.HasFlag(WorkspaceUpdateOptions.Sync) && !Context.Options.HasFlag(WorkspaceUpdateOptions.UpdateFilter))
+					{
+						// Read the new config file
+						Context.ProjectConfigFile = await ReadProjectConfigFile(Project.LocalRootPath, Project.LocalFileName, Logger);
+						Context.ProjectStreamFilter = await ReadProjectStreamFilter(Perforce, Context.ProjectConfigFile, Logger, CancellationToken);
+
+						// Get the branch name
+						string? BranchOrStreamName = await Perforce.GetCurrentStreamAsync(CancellationToken);
+						if (BranchOrStreamName != null)
+						{
+							// If it's a virtual stream, take the concrete parent stream instead
+							for (; ; )
+							{
+								StreamRecord StreamSpec = await Perforce.GetStreamAsync(BranchOrStreamName, false, CancellationToken);
+								if (StreamSpec.Type != "virtual" || StreamSpec.Parent == "none" || StreamSpec.Parent == null)
+								{
+									break;
+								}
+								BranchOrStreamName = StreamSpec.Parent;
+							}
+						}
+						else
+						{
+							// Otherwise use the depot path for GenerateProjectFiles.bat in the root of the workspace
+							List<WhereRecord> Files = await Perforce.WhereAsync(Project.ClientRootPath + "/GenerateProjectFiles.bat").Where(x => !x.Unmap).ToListAsync(CancellationToken);
+							if (Files.Count != 1)
+							{
+								return (WorkspaceUpdateResult.FailedToSync, $"Couldn't determine branch name for {Project.ClientFileName}.");
+							}
+							BranchOrStreamName = PerforceUtils.GetClientOrDepotDirectoryName(Files[0].DepotFile);
+						}
+
+						// Find the last code change before this changelist. For consistency in versioning between local builds and precompiled binaries, we need to use the last submitted code changelist as our version number.
+						string[] CodeFilter = new string[] { ".cs", ".h", ".cpp", ".usf", ".ush", ".uproject", ".uplugin" }.SelectMany(x => SyncPaths.Select(y => String.Format("{0}{1}@<={2}", y, x, Context.ChangeNumber))).ToArray();
+
+						PerforceResponseList<ChangesRecord> CodeChanges = await Perforce.TryGetChangesAsync(ChangesOptions.None, 1, ChangeStatus.Submitted, CodeFilter, CancellationToken);
+						if (!CodeChanges.Succeeded)
+						{
+							return (WorkspaceUpdateResult.FailedToSync, $"Couldn't determine last code changelist before CL {Context.ChangeNumber}.");
+						}
+						if (CodeChanges.Count == 0)
+						{
+							return (WorkspaceUpdateResult.FailedToSync, $"Could not find any code changes before CL {Context.ChangeNumber}.");
+						}
+
+						// Get the last code change
+						if (Context.ProjectConfigFile.GetValue("Options.VersionToLastCodeChange", true))
+						{
+							VersionChangeNumber = CodeChanges.Max(x => x.Data.Number);
+						}
+						else
+						{
+							VersionChangeNumber = Context.ChangeNumber;
+						}
+
+						// Update the version files
+						if (Context.ProjectConfigFile.GetValue("Options.UseFastModularVersioningV2", false))
+						{
+							bool bIsLicenseeVersion = await IsLicenseeVersion(Perforce, Project, Logger, CancellationToken);
+							if (!await UpdateVersionFile(Perforce, Project.ClientRootPath + BuildVersionFileName, Context.ChangeNumber, Text => UpdateBuildVersion(Text, Context.ChangeNumber, VersionChangeNumber, BranchOrStreamName, bIsLicenseeVersion), Logger, CancellationToken))
+							{
+								return (WorkspaceUpdateResult.FailedToSync, $"Failed to update {BuildVersionFileName}.");
+							}
+						}
+						else if (Context.ProjectConfigFile.GetValue("Options.UseFastModularVersioning", false))
+						{
+							bool bIsLicenseeVersion = await IsLicenseeVersion(Perforce, Project, Logger, CancellationToken);
+							if (!await UpdateVersionFile(Perforce, Project.ClientRootPath + BuildVersionFileName, Context.ChangeNumber, Text => UpdateBuildVersion(Text, Context.ChangeNumber, VersionChangeNumber, BranchOrStreamName, bIsLicenseeVersion), Logger, CancellationToken))
+							{
+								return (WorkspaceUpdateResult.FailedToSync, $"Failed to update {BuildVersionFileName}");
+							}
+
+							Dictionary<string, string> VersionHeaderStrings = new Dictionary<string, string>();
+							VersionHeaderStrings["#define ENGINE_IS_PROMOTED_BUILD"] = " (0)";
+							VersionHeaderStrings["#define BUILT_FROM_CHANGELIST"] = " 0";
+							VersionHeaderStrings["#define BRANCH_NAME"] = " \"" + BranchOrStreamName.Replace('/', '+') + "\"";
+							if (!await UpdateVersionFile(Perforce, Project.ClientRootPath + VersionHeaderFileName, VersionHeaderStrings, Context.ChangeNumber, Logger, CancellationToken))
+							{
+								return (WorkspaceUpdateResult.FailedToSync, $"Failed to update {VersionHeaderFileName}.");
+							}
+							if (!await UpdateVersionFile(Perforce, Project.ClientRootPath + ObjectVersionFileName, new Dictionary<string, string>(), Context.ChangeNumber, Logger, CancellationToken))
+							{
+								return (WorkspaceUpdateResult.FailedToSync, $"Failed to update {ObjectVersionFileName}.");
+							}
+						}
+						else
+						{
+							if (!await UpdateVersionFile(Perforce, Project.ClientRootPath + BuildVersionFileName, new Dictionary<string, string>(), Context.ChangeNumber, Logger, CancellationToken))
+							{
+								return (WorkspaceUpdateResult.FailedToSync, $"Failed to update {BuildVersionFileName}");
+							}
+
+							Dictionary<string, string> VersionStrings = new Dictionary<string, string>();
+							VersionStrings["#define ENGINE_VERSION"] = " " + VersionChangeNumber.ToString();
+							VersionStrings["#define ENGINE_IS_PROMOTED_BUILD"] = " (0)";
+							VersionStrings["#define BUILT_FROM_CHANGELIST"] = " " + VersionChangeNumber.ToString();
+							VersionStrings["#define BRANCH_NAME"] = " \"" + BranchOrStreamName.Replace('/', '+') + "\"";
+							if (!await UpdateVersionFile(Perforce, Project.ClientRootPath + VersionHeaderFileName, VersionStrings, Context.ChangeNumber, Logger, CancellationToken))
+							{
+								return (WorkspaceUpdateResult.FailedToSync, $"Failed to update {VersionHeaderFileName}");
+							}
+							if (!await UpdateVersionFile(Perforce, Project.ClientRootPath + ObjectVersionFileName, VersionStrings, Context.ChangeNumber, Logger, CancellationToken))
+							{
+								return (WorkspaceUpdateResult.FailedToSync, $"Failed to update {ObjectVersionFileName}");
+							}
+						}
+
+						// Remove all the receipts for build targets in this branch
+						if (Project.ClientFileName.EndsWith(".uproject", StringComparison.OrdinalIgnoreCase))
+						{
+							await Perforce.SyncQuietAsync(SyncOptions.None, -1, PerforceUtils.GetClientOrDepotDirectoryName(Project.ClientFileName) + "/Build/Receipts/...#0", CancellationToken);
+						}
+					}
+
+					// Check if there are any files which need resolving
+					List<FStatRecord> UnresolvedFiles = await FindUnresolvedFiles(Perforce, SyncPaths, CancellationToken);
+					if (UnresolvedFiles.Count > 0 && Context.Options.HasFlag(WorkspaceUpdateOptions.AutoResolveChanges))
+					{
+						foreach (FStatRecord UnresolvedFile in UnresolvedFiles)
+						{
+							if (UnresolvedFile.DepotFile != null)
+							{
+								await Perforce.ResolveAsync(-1, ResolveOptions.Automatic, UnresolvedFile.DepotFile, CancellationToken);
+							}
+						}
+						UnresolvedFiles = await FindUnresolvedFiles(Perforce, SyncPaths, CancellationToken);
+					}
+					if (UnresolvedFiles.Count > 0)
+					{
+						Logger.LogInformation("{NumFiles} files need resolving:", UnresolvedFiles.Count);
+						foreach (FStatRecord UnresolvedFile in UnresolvedFiles)
+						{
+							Logger.LogInformation("  {ClientFile}", UnresolvedFile.ClientFile);
+						}
+						return (WorkspaceUpdateResult.FilesToResolve, "Files need resolving.");
+					}
+
+					// Continue processing sync-only actions
+					if (Context.Options.HasFlag(WorkspaceUpdateOptions.Sync) && !Context.Options.HasFlag(WorkspaceUpdateOptions.UpdateFilter))
+					{
+						Context.ProjectConfigFile ??= await ReadProjectConfigFile(Project.LocalRootPath, Project.LocalFileName, Logger);
+
+						// Execute any project specific post-sync steps
+						string[]? PostSyncSteps = Context.ProjectConfigFile.GetValues("Sync.Step", null);
+						if (PostSyncSteps != null)
+						{
+							Logger.LogInformation("");
+							Logger.LogInformation("Executing post-sync steps...");
+
+							Dictionary<string, string> PostSyncVariables = new Dictionary<string, string>(Context.Variables);
+							PostSyncVariables["Change"] = Context.ChangeNumber.ToString();
+							PostSyncVariables["CodeChange"] = VersionChangeNumber.ToString();
+
+							foreach (string PostSyncStep in PostSyncSteps.Select(x => x.Trim()))
+							{
+								ConfigObject PostSyncStepObject = new ConfigObject(PostSyncStep);
+
+								string ToolFileName = Utility.ExpandVariables(PostSyncStepObject.GetValue("FileName", ""), PostSyncVariables);
+								if (ToolFileName != null)
+								{
+									string ToolArguments = Utility.ExpandVariables(PostSyncStepObject.GetValue("Arguments", ""), PostSyncVariables);
+
+									Logger.LogInformation("post-sync> Running {FileName} {Arguments}", ToolFileName, ToolArguments);
+
+									int ResultFromTool = await Utility.ExecuteProcessAsync(ToolFileName, null, ToolArguments, Line => ProcessOutput(Line, "post-sync> ", Progress, Logger), CancellationToken);
+									if (ResultFromTool != 0)
+									{
+										return (WorkspaceUpdateResult.FailedToSync, $"Post-sync step terminated with exit code {ResultFromTool}.");
+									}
+								}
+							}
+						}
+					}
+
+					// Update the current state
+					if (Context.Options.HasFlag(WorkspaceUpdateOptions.SyncSingleChange))
+					{
+						State.AdditionalChangeNumbers.Add(Context.ChangeNumber);
+					}
+					else
+					{
+						State.CurrentChangeNumber = Context.ChangeNumber;
+					}
+					State.Save();
+
+					// Update the timing info
+					Times.Add(new Tuple<string, TimeSpan>("Sync", SyncTelemetryStopwatch.Stop("Success")));
+
+					// Save the number of files synced
+					NumFilesSynced = SyncDepotPaths.Count;
+					Logger.LogInformation("");
+				}
+			}
+
+			// Extract an archive from the depot path
+			if (Context.Options.HasFlag(WorkspaceUpdateOptions.SyncArchives))
+			{
+				using (TelemetryStopwatch Stopwatch = new TelemetryStopwatch("Workspace_SyncArchives", Project.TelemetryProjectIdentifier))
+				{
+					// Create the directory for extracted archive manifests
+					DirectoryReference ManifestDirectoryName;
+					if (Project.LocalFileName.HasExtension(".uproject"))
+					{
+						ManifestDirectoryName = DirectoryReference.Combine(Project.LocalFileName.Directory, "Saved", "UnrealGameSync");
+					}
+					else
+					{
+						ManifestDirectoryName = DirectoryReference.Combine(Project.LocalFileName.Directory, "Engine", "Saved", "UnrealGameSync");
+					}
+					DirectoryReference.CreateDirectory(ManifestDirectoryName);
+
+					// Sync and extract (or just remove) the given archives
+					foreach (KeyValuePair<string, Tuple<IArchiveInfo, string>?> ArchiveTypeAndArchive in Context.ArchiveTypeToArchive)
+					{
+						string ArchiveType = ArchiveTypeAndArchive.Key;
+
+						// Remove any existing binaries
+						FileReference ManifestFileName = FileReference.Combine(ManifestDirectoryName, String.Format("{0}.zipmanifest", ArchiveType));
+						if (FileReference.Exists(ManifestFileName))
+						{
+							Logger.LogInformation("Removing {ArchiveType} binaries...", ArchiveType);
+							Progress.Set(String.Format("Removing {0} binaries...", ArchiveType), 0.0f);
+							ArchiveUtils.RemoveExtractedFiles(Project.LocalRootPath, ManifestFileName, Progress, Logger);
+							FileReference.Delete(ManifestFileName);
+							Logger.LogInformation("");
+						}
+
+						// If we have a new depot path, sync it down and extract it
+						if (ArchiveTypeAndArchive.Value != null)
+						{
+							IArchiveInfo ArchiveInfo = ArchiveTypeAndArchive.Value.Item1;
+							string ArchiveKey = ArchiveTypeAndArchive.Value.Item2;
+
+							Logger.LogInformation("Syncing {ArchiveType} binaries...", ArchiveType.ToLowerInvariant());
+							Progress.Set(String.Format("Syncing {0} binaries...", ArchiveType.ToLowerInvariant()), 0.0f);
+							if (!ArchiveInfo.DownloadArchive(Perforce, ArchiveKey, Project.LocalRootPath, ManifestFileName, Logger, Progress, CancellationToken.None).Result)
+							{
+								return (WorkspaceUpdateResult.FailedToSync, $"Couldn't read {ArchiveKey}");
+							}
+						}
+					}
+
+					// Update the state
+					State.ExpandedArchiveTypes = Context.ArchiveTypeToArchive.Where(x => x.Value != null).Select(x => x.Key).ToArray();
+					State.Save();
+
+					// Add the finish time
+					Times.Add(new Tuple<string, TimeSpan>("Archive", Stopwatch.Stop("Success")));
+				}
+			}
+
+			// Take the lock before doing anything else. Building and generating project files can only be done on one workspace at a time.
+			using SemaphoreScope Scope = new SemaphoreScope(UpdateSemaphore);
+			if (Context.Options.HasFlag(WorkspaceUpdateOptions.GenerateProjectFiles) || Context.Options.HasFlag(WorkspaceUpdateOptions.Build))
+			{
+				if (!Scope.TryAcquire())
+				{
+					Logger.LogInformation("Waiting for other workspaces to finish...");
+					await Scope.AcquireAsync(CancellationToken);
+				}
+			}
+
+			// Generate project files in the workspace
+			if (Context.Options.HasFlag(WorkspaceUpdateOptions.GenerateProjectFiles))
+			{
+				using (TelemetryStopwatch Stopwatch = new TelemetryStopwatch("Workspace_GenerateProjectFiles", Project.TelemetryProjectIdentifier))
+				{
+					Progress.Set("Generating project files...", 0.0f);
+
+					StringBuilder CommandLine = new StringBuilder();
+					CommandLine.AppendFormat("/C \"\"{0}\"", FileReference.Combine(Project.LocalRootPath, "GenerateProjectFiles.bat"));
+					if ((Context.Options & WorkspaceUpdateOptions.SyncAllProjects) == 0 && (Context.Options & WorkspaceUpdateOptions.IncludeAllProjectsInSolution) == 0)
+					{
+						if (Project.LocalFileName.HasExtension(".uproject"))
+						{
+							CommandLine.AppendFormat(" \"{0}\"", Project.LocalFileName);
+						}
+					}
+					CommandLine.Append(" -progress\"");
+
+					Logger.LogInformation("Generating project files...");
+					Logger.LogInformation("gpf> Running {FileName} {Arguments}", CmdExe, CommandLine);
+
+					int GenerateProjectFilesResult = await Utility.ExecuteProcessAsync(CmdExe, null, CommandLine.ToString(), Line => ProcessOutput(Line, "gpf> ", Progress, Logger), CancellationToken);
+					if (GenerateProjectFilesResult != 0)
+					{
+						return (WorkspaceUpdateResult.FailedToCompile, $"Failed to generate project files (exit code {GenerateProjectFilesResult}).");
+					}
+
+					Logger.LogInformation("");
+					Times.Add(new Tuple<string, TimeSpan>("Prj gen", Stopwatch.Stop("Success")));
+				}
+			}
+
+			// Build everything using MegaXGE
+			if (Context.Options.HasFlag(WorkspaceUpdateOptions.Build))
+			{
+				Context.ProjectConfigFile ??= await ReadProjectConfigFile(Project.LocalRootPath, Project.LocalFileName, Logger);
+
+				// Compile all the build steps together
+				Dictionary<Guid, ConfigObject> BuildStepObjects = Context.DefaultBuildSteps.ToDictionary(x => x.Key, x => new ConfigObject(x.Value));
+				BuildStep.MergeBuildStepObjects(BuildStepObjects, Context.ProjectConfigFile.GetValues("Build.Step", new string[0]).Select(x => new ConfigObject(x)));
+				BuildStep.MergeBuildStepObjects(BuildStepObjects, Context.UserBuildStepObjects);
+
+				// Construct build steps from them
+				List<BuildStep> BuildSteps = BuildStepObjects.Values.Select(x => new BuildStep(x)).OrderBy(x => (x.OrderIndex == -1) ? 10000 : x.OrderIndex).ToList();
+				if (Context.CustomBuildSteps != null && Context.CustomBuildSteps.Count > 0)
+				{
+					BuildSteps.RemoveAll(x => !Context.CustomBuildSteps.Contains(x.UniqueId));
+				}
+				else if (Context.Options.HasFlag(WorkspaceUpdateOptions.ScheduledBuild))
+				{
+					BuildSteps.RemoveAll(x => !x.bScheduledSync);
+				}
+				else
+				{
+					BuildSteps.RemoveAll(x => !x.bNormalSync);
+				}
+
+				// Check if the last successful build was before a change that we need to force a clean for
+				bool bForceClean = false;
+				if (State.LastBuiltChangeNumber != 0)
+				{
+					foreach (string CleanBuildChange in Context.ProjectConfigFile.GetValues("ForceClean.Changelist", new string[0]))
+					{
+						int ChangeNumber;
+						if (int.TryParse(CleanBuildChange, out ChangeNumber))
+						{
+							if ((State.LastBuiltChangeNumber >= ChangeNumber) != (State.CurrentChangeNumber >= ChangeNumber))
+							{
+								Logger.LogInformation("Forcing clean build due to changelist {Change}.", ChangeNumber);
+								Logger.LogInformation("");
+								bForceClean = true;
+								break;
+							}
+						}
+					}
+				}
+
+				// Execute them all
+				using (TelemetryStopwatch Stopwatch = new TelemetryStopwatch("Workspace_Build", Project.TelemetryProjectIdentifier))
+				{
+					Progress.Set("Starting build...", 0.0f);
+
+					// Execute all the steps
+					float MaxProgressFraction = 0.0f;
+					foreach (BuildStep Step in BuildSteps)
+					{
+						MaxProgressFraction += (float)Step.EstimatedDuration / (float)Math.Max(BuildSteps.Sum(x => x.EstimatedDuration), 1);
+
+						Progress.Set(Step.StatusText ?? "Executing build step");
+						Progress.Push(MaxProgressFraction);
+
+						Logger.LogInformation("{Status}", Step.StatusText);
+
+						if (Step.IsValid())
+						{
+							switch (Step.Type)
+							{
+								case BuildStepType.Compile:
+									using (TelemetryStopwatch StepStopwatch = new TelemetryStopwatch("Workspace_Execute_Compile", Project.TelemetryProjectIdentifier))
+									{
+										StepStopwatch.AddData(new { Target = Step.Target });
+
+										FileReference BuildBat = FileReference.Combine(Project.LocalRootPath, "Engine", "Build", "BatchFiles", "Build.bat");
+										string CommandLine = String.Format("{0} {1} {2} {3} -NoHotReloadFromIDE", Step.Target, Step.Platform, Step.Configuration, Utility.ExpandVariables(Step.Arguments ?? "", Context.Variables));
+										if (Context.Options.HasFlag(WorkspaceUpdateOptions.Clean) || bForceClean)
+										{
+											Logger.LogInformation("ubt> Running {FileName} {Arguments}", BuildBat, CommandLine + " -clean");
+											await Utility.ExecuteProcessAsync(CmdExe, null, "/C \"\"" + BuildBat + "\" " + CommandLine + " -clean\"", Line => ProcessOutput(Line, "ubt> ", Progress, Logger), CancellationToken);
+										}
+
+										Logger.LogInformation("ubt> Running {FileName} {Arguments}", BuildBat, CommandLine + " -progress");
+
+										int ResultFromBuild = await Utility.ExecuteProcessAsync(CmdExe, null, "/C \"\"" + BuildBat + "\" " + CommandLine + " -progress\"", Line => ProcessOutput(Line, "ubt> ", Progress, Logger), CancellationToken);
+										if (ResultFromBuild != 0)
+										{
+											StepStopwatch.Stop("Failed");
+
+											WorkspaceUpdateResult Result;
+											if (await HasModifiedSourceFiles(Perforce, Project, CancellationToken) || Context.UserBuildStepObjects.Count > 0)
+											{
+												Result = WorkspaceUpdateResult.FailedToCompile;
+											}
+											else
+											{
+												Result = WorkspaceUpdateResult.FailedToCompileWithCleanWorkspace;
+											}
+
+											return (Result, $"Failed to compile {Step.Target}");
+										}
+
+										StepStopwatch.Stop("Success");
+									}
+									break;
+								case BuildStepType.Cook:
+									using (TelemetryStopwatch StepStopwatch = new TelemetryStopwatch("Workspace_Execute_Cook", Project.TelemetryProjectIdentifier))
+									{
+										StepStopwatch.AddData(new { Project = Path.GetFileNameWithoutExtension(Step.FileName) });
+
+										FileReference LocalRunUAT = FileReference.Combine(Project.LocalRootPath, "Engine", "Build", "BatchFiles", "RunUAT.bat");
+										string Arguments = String.Format("/C \"\"{0}\" -profile=\"{1}\"\"", LocalRunUAT, FileReference.Combine(Project.LocalRootPath, Step.FileName ?? "unknown"));
+										Logger.LogInformation("uat> Running {FileName} {Argument}", LocalRunUAT, Arguments);
+
+										int ResultFromUAT = await Utility.ExecuteProcessAsync(CmdExe, null, Arguments, Line => ProcessOutput(Line, "uat> ", Progress, Logger), CancellationToken);
+										if (ResultFromUAT != 0)
+										{
+											StepStopwatch.Stop("Failed");
+											return (WorkspaceUpdateResult.FailedToCompile, $"Cook failed. ({ResultFromUAT})");
+										}
+
+										StepStopwatch.Stop("Success");
+									}
+									break;
+								case BuildStepType.Other:
+									using (TelemetryStopwatch StepStopwatch = new TelemetryStopwatch("Workspace_Execute_Custom", Project.TelemetryProjectIdentifier))
+									{
+										StepStopwatch.AddData(new { FileName = Path.GetFileNameWithoutExtension(Step.FileName) });
+
+										FileReference ToolFileName = FileReference.Combine(Project.LocalRootPath, Utility.ExpandVariables(Step.FileName ?? "unknown", Context.Variables));
+										string ToolWorkingDir = String.IsNullOrWhiteSpace(Step.WorkingDir) ? ToolFileName.Directory.FullName : Utility.ExpandVariables(Step.WorkingDir, Context.Variables);
+										string ToolArguments = Utility.ExpandVariables(Step.Arguments ?? "", Context.Variables);
+										Logger.LogInformation("tool> Running {0} {1}", ToolFileName, ToolArguments);
+
+										if (Step.bUseLogWindow)
+										{
+											int ResultFromTool = await Utility.ExecuteProcessAsync(ToolFileName.FullName, ToolWorkingDir, ToolArguments, Line => ProcessOutput(Line, "tool> ", Progress, Logger), CancellationToken);
+											if (ResultFromTool != 0)
+											{
+												StepStopwatch.Stop("Failed");
+												return (WorkspaceUpdateResult.FailedToCompile, $"Tool terminated with exit code {ResultFromTool}.");
+											}
+										}
+										else
+										{
+											ProcessStartInfo StartInfo = new ProcessStartInfo(ToolFileName.FullName, ToolArguments);
+											StartInfo.WorkingDirectory = ToolWorkingDir;
+											using (Process.Start(StartInfo))
+											{
+											}
+										}
+
+										StepStopwatch.Stop("Success");
+									}
+									break;
+							}
+						}
+
+						Logger.LogInformation("");
+						Progress.Pop();
+					}
+
+					Times.Add(new Tuple<string, TimeSpan>("Build", Stopwatch.Stop("Success")));
+				}
+
+				// Update the last successful build change number
+				if (Context.CustomBuildSteps == null || Context.CustomBuildSteps.Count == 0)
+				{
+					State.LastBuiltChangeNumber = State.CurrentChangeNumber;
+					State.Save();
+				}
+			}
+
+			// Write out all the timing information
+			Logger.LogInformation("Total time: {TotalTime}", FormatTime(Times.Sum(x => (long)(x.Item2.TotalMilliseconds / 1000))));
+			foreach (Tuple<string, TimeSpan> Time in Times)
+			{
+				Logger.LogInformation("   {Name,-8}: {TimeSeconds}", Time.Item1, FormatTime((long)(Time.Item2.TotalMilliseconds / 1000)));
+			}
+			if (NumFilesSynced > 0)
+			{
+				Logger.LogInformation("{NumFiles} files synced.", NumFilesSynced);
+			}
+
+			DateTime FinishTime = DateTime.Now;
+			Logger.LogInformation("");
+			Logger.LogInformation("UPDATE SUCCEEDED ({FinishDate} {FinishTime})", FinishTime.ToShortDateString(), FinishTime.ToShortTimeString());
+
+			return (WorkspaceUpdateResult.Success, "Update succeeded");
+		}
+
+		static void ProcessOutput(string Line, string Prefix, ProgressValue Progress, ILogger Logger)
+		{
+			string? ParsedLine = ProgressTextWriter.ParseLine(Line, Progress);
+			if (ParsedLine != null)
+			{
+				Logger.LogInformation("{Prefix}{Message}", Prefix, Line);
+			}
+		}
+
+		static async Task<bool> IsLicenseeVersion(IPerforceConnection Perforce, ProjectInfo Project, ILogger Logger, CancellationToken CancellationToken)
+		{
+			string[] Files =
+			{
+				Project.ClientRootPath + "/Engine/Build/NotForLicensees/EpicInternal.txt",
+				Project.ClientRootPath + "/Engine/Restricted/NotForLicensees/Build/EpicInternal.txt"
+			};
+
+			List<FStatRecord> Records = await Perforce.FStatAsync(Files, CancellationToken).ToListAsync(CancellationToken);
+			return !Records.Any(x => x.IsMapped);
+		}
+
+		public static List<string> GetSyncPaths(ProjectInfo Project, bool bSyncAllProjects, string[] SyncFilter)
+		{
+			List<string> SyncPaths = GetRelativeSyncPaths(Project, bSyncAllProjects, SyncFilter);
+			return SyncPaths.Select(x => Project.ClientRootPath + x).ToList();
+		}
+
+		public static List<string> GetRelativeSyncPaths(ProjectInfo Project, bool bSyncAllProjects, string[]? SyncFilter)
+		{
+			List<string> SyncPaths = new List<string>();
+
+			// Check the client path is formatted correctly
+			if (!Project.ClientFileName.StartsWith(Project.ClientRootPath + "/"))
+			{
+				throw new Exception(String.Format("Expected '{0}' to start with '{1}'", Project.ClientFileName, Project.ClientRootPath));
+			}
+
+			// Add the default project paths
+			int LastSlashIdx = Project.ClientFileName.LastIndexOf('/');
+			if (bSyncAllProjects || !Project.ClientFileName.EndsWith(".uproject", StringComparison.OrdinalIgnoreCase) || LastSlashIdx <= Project.ClientRootPath.Length)
+			{
+				SyncPaths.Add("/...");
+			}
+			else
+			{
+				SyncPaths.Add("/*");
+				SyncPaths.Add("/Engine/...");
+				if (Project.bIsEnterpriseProject)
+				{
+					SyncPaths.Add("/Enterprise/...");
+				}
+				SyncPaths.Add(Project.ClientFileName.Substring(Project.ClientRootPath.Length, LastSlashIdx - Project.ClientRootPath.Length) + "/...");
+			}
+
+			// Apply the sync filter to that list. We only want inclusive rules in the output list, but we can manually apply exclusions to previous entries.
+			if (SyncFilter != null)
+			{
+				foreach (string SyncPath in SyncFilter)
+				{
+					string TrimSyncPath = SyncPath.Trim();
+					if (TrimSyncPath.StartsWith("/"))
+					{
+						SyncPaths.Add(TrimSyncPath);
+					}
+					else if (TrimSyncPath.StartsWith("-/") && TrimSyncPath.EndsWith("..."))
+					{
+						SyncPaths.RemoveAll(x => x.StartsWith(TrimSyncPath.Substring(1, TrimSyncPath.Length - 4)));
+					}
+				}
+			}
+
+			// Sort the remaining paths by length, and remove any paths which are included twice
+			SyncPaths = SyncPaths.OrderBy(x => x.Length).ToList();
+			for (int Idx = 0; Idx < SyncPaths.Count; Idx++)
+			{
+				string SyncPath = SyncPaths[Idx];
+				if (SyncPath.EndsWith("..."))
+				{
+					string SyncPathPrefix = SyncPath.Substring(0, SyncPath.Length - 3);
+					for (int OtherIdx = SyncPaths.Count - 1; OtherIdx > Idx; OtherIdx--)
+					{
+						if (SyncPaths[OtherIdx].StartsWith(SyncPathPrefix))
+						{
+							SyncPaths.RemoveAt(OtherIdx);
+						}
+					}
+				}
+			}
+
+			return SyncPaths;
+		}
+
+		public static bool MatchFilter(ProjectInfo Project, FileReference FileName, FileFilter Filter)
+		{
+			bool bMatch = true;
+			if (FileName.IsUnderDirectory(Project.LocalRootPath))
+			{
+				string RelativePath = FileName.MakeRelativeTo(Project.LocalRootPath);
+				if (!Filter.Matches(RelativePath))
+				{
+					bMatch = false;
+				}
+			}
+			return bMatch;
+		}
+
+		class SyncState
+		{
+			public int TotalDepotPaths;
+			public HashSet<string> RemainingDepotPaths;
+			public Queue<List<string>> SyncCommandLists;
+			public string StatusMessage;
+			public WorkspaceUpdateResult Result = WorkspaceUpdateResult.Success;
+
+			public SyncState(HashSet<string> RemainingDepotPaths, Queue<List<string>> SyncCommandLists)
+			{
+				this.TotalDepotPaths = RemainingDepotPaths.Count;
+				this.RemainingDepotPaths = RemainingDepotPaths;
+				this.SyncCommandLists = SyncCommandLists;
+				this.StatusMessage = "Succeeded.";
+			}
+		}
+
+		static Task<(WorkspaceUpdateResult, string)> SyncFileRevisions(IPerforceConnection Perforce, string Prefix, WorkspaceUpdateContext Context, List<string> SyncCommands, HashSet<string> RemainingDepotPaths, ProgressValue Progress, ILogger Logger, CancellationToken CancellationToken)
+		{
+			Queue<List<string>> SyncCommandLists = new Queue<List<string>>();
+			SyncCommandLists.Enqueue(SyncCommands);
+			return SyncFileRevisions(Perforce, Prefix, Context, SyncCommandLists, RemainingDepotPaths, Progress, Logger, CancellationToken);
+		}
+
+		static async Task<(WorkspaceUpdateResult, string)> SyncFileRevisions(IPerforceConnection Perforce, string Prefix, WorkspaceUpdateContext Context, Queue<List<string>> SyncCommandLists, HashSet<string> RemainingDepotPaths, ProgressValue Progress, ILogger Logger, CancellationToken CancellationToken)
+		{
+			// Figure out the number of additional background threads we want to run with. We can run worker on the current thread.
+			int NumThreads = Context.PerforceSyncOptions?.NumThreads ?? PerforceSyncOptions.DefaultNumThreads;
+			int NumExtraThreads = Math.Max(Math.Min(SyncCommandLists.Count, NumThreads) - 1, 0);
+
+			List<IPerforceConnection> ChildConnections = new List<IPerforceConnection>();
+			List<Task> ChildTasks = new List<Task>(NumExtraThreads);
+			try
+			{
+				// Create the state object shared by all the worker threads
+				SyncState State = new SyncState(RemainingDepotPaths, SyncCommandLists);
+
+				// Wrapper writer around the log class to prevent multiple threads writing to it at once
+				ILogger LogWrapper = Logger;
+
+				// Initialize Sync Progress
+				UpdateSyncState(Prefix, State, Progress);
+
+				// Delegate for updating the sync state after a file has been synced
+				Action<SyncRecord, ILogger> SyncOutput = (Record, LocalLog) => { UpdateSyncState(Prefix, Record, State, Progress, LocalLog); };
+
+				// Create all the child threads
+				for (int ThreadIdx = 0; ThreadIdx < NumExtraThreads; ThreadIdx++)
+				{
+					int ThreadNumber = ThreadIdx + 2;
+
+					// Create connection
+					IPerforceConnection ChildConnection = await PerforceConnection.CreateAsync(Perforce.Settings, Perforce.Logger);
+					ChildConnections.Add(ChildConnection);
+
+					Task ChildTask = Task.Run(() => StaticSyncWorker(ThreadNumber, ChildConnection, Context, State, SyncOutput, LogWrapper, CancellationToken));
+					ChildTasks.Add(ChildTask);
+				}
+
+				// Run one worker on the current thread
+				await StaticSyncWorker(1, Perforce, Context, State, SyncOutput, LogWrapper, CancellationToken);
+
+				// Return the result that was set on the state object
+				return (State.Result, State.StatusMessage);
+			}
+			finally
+			{
+				await Task.WhenAll(ChildTasks);
+
+				foreach (IPerforceConnection ChildConnection in ChildConnections)
+				{
+					ChildConnection.Dispose();
+				}
+
+				foreach (Task ChildTask in ChildTasks)
+				{
+					await ChildTask;
+				}
+			}
+		}
+
+		static void UpdateSyncState(string Prefix, SyncState State, ProgressValue Progress)
+		{
+			lock (State)
+			{
+				string Message = String.Format("{0} ({1}/{2})", Prefix, State.TotalDepotPaths - State.RemainingDepotPaths.Count, State.TotalDepotPaths);
+				float Fraction = Math.Min((float)(State.TotalDepotPaths - State.RemainingDepotPaths.Count) / (float)State.TotalDepotPaths, 1.0f);
+				Progress.Set(Message, Fraction);
+			}
+		}
+
+		static void UpdateSyncState(string Prefix, SyncRecord Record, SyncState State, ProgressValue Progress, ILogger Logger)
+		{
+			lock (State)
+			{
+				State.RemainingDepotPaths.Remove(Record.DepotFile);
+
+				string Message = String.Format("{0} ({1}/{2})", Prefix, State.TotalDepotPaths - State.RemainingDepotPaths.Count, State.TotalDepotPaths);
+				float Fraction = Math.Min((float)(State.TotalDepotPaths - State.RemainingDepotPaths.Count) / (float)State.TotalDepotPaths, 1.0f);
+				Progress.Set(Message, Fraction);
+
+				Logger.LogInformation("p4>   {Action} {Path}", Record.Action, Record.Path);
+			}
+		}
+
+		static async Task StaticSyncWorker(int ThreadNumber, IPerforceConnection Perforce, WorkspaceUpdateContext Context, SyncState State, Action<SyncRecord, ILogger> SyncOutput, ILogger GlobalLog, CancellationToken CancellationToken)
+		{
+			PrefixedTextWriter ThreadLog = new PrefixedTextWriter(String.Format("{0}:", ThreadNumber), GlobalLog);
+			for (; ; )
+			{
+				// Remove the next batch that needs to be synced
+				List<string> SyncCommands;
+				lock (State)
+				{
+					if (State.Result == WorkspaceUpdateResult.Success && State.SyncCommandLists.Count > 0)
+					{
+						SyncCommands = State.SyncCommandLists.Dequeue();
+					}
+					else
+					{
+						break;
+					}
+				}
+
+				// Sync the files
+				(WorkspaceUpdateResult Result, string StatusMessage) = await StaticSyncFileRevisions(Perforce, Context, SyncCommands, Record => SyncOutput(Record, ThreadLog), CancellationToken);
+
+				// If it failed, try to set it on the state if nothing else has failed first
+				if (Result != WorkspaceUpdateResult.Success)
+				{
+					lock (State)
+					{
+						if (State.Result == WorkspaceUpdateResult.Success)
+						{
+							State.Result = Result;
+							State.StatusMessage = StatusMessage;
+						}
+					}
+					break;
+				}
+			}
+		}
+
+		static async Task<(WorkspaceUpdateResult, string)> StaticSyncFileRevisions(IPerforceConnection Perforce, WorkspaceUpdateContext Context, List<string> SyncCommands, Action<SyncRecord> SyncOutput, CancellationToken CancellationToken)
+		{
+			// Sync them all
+			List<PerforceResponse<SyncRecord>> Responses = await Perforce.TrySyncAsync(SyncOptions.None, -1, SyncCommands, CancellationToken).ToListAsync(CancellationToken);
+
+			List<string> TamperedFiles = new List<string>();
+			foreach (PerforceResponse<SyncRecord> Response in Responses)
+			{
+				const string NoClobberPrefix = "Can't clobber writable file ";
+				if (Response.Succeeded)
+				{
+					SyncOutput(Response.Data);
+				}
+				else if (Response.Error != null && Response.Error.Generic == PerforceGenericCode.Client && Response.Error.Data.StartsWith(NoClobberPrefix, StringComparison.OrdinalIgnoreCase))
+				{
+					TamperedFiles.Add(Response.Error.Data.Substring(NoClobberPrefix.Length).Trim());
+				}
+				else
+				{
+					return (WorkspaceUpdateResult.FailedToSync, "Aborted sync due to errors.");
+				}
+			}
+
+			// If any files need to be clobbered, defer to the main thread to figure out which ones
+			if (TamperedFiles.Count > 0)
+			{
+				int NumNewFilesToClobber = 0;
+				foreach (string TamperedFile in TamperedFiles)
+				{
+					if (!Context.ClobberFiles.ContainsKey(TamperedFile))
+					{
+						Context.ClobberFiles[TamperedFile] = true;
+						if (TamperedFile.EndsWith(LocalObjectVersionFileName, StringComparison.OrdinalIgnoreCase) || TamperedFile.EndsWith(LocalVersionHeaderFileName, StringComparison.OrdinalIgnoreCase))
+						{
+							// Hack for UseFastModularVersioningV2; we don't need to update these files any more.
+							continue;
+						}
+						NumNewFilesToClobber++;
+					}
+				}
+				if (NumNewFilesToClobber > 0)
+				{
+					return (WorkspaceUpdateResult.FilesToClobber, $"Cancelled sync after checking files to clobber ({NumNewFilesToClobber} new files).");
+				}
+				foreach (string TamperedFile in TamperedFiles)
+				{
+					if (Context.ClobberFiles[TamperedFile])
+					{
+						List<PerforceResponse<SyncRecord>> Response = await Perforce.TrySyncAsync(SyncOptions.Force, -1, TamperedFile, CancellationToken).ToListAsync();
+						if (!Response.Succeeded())
+						{
+							return (WorkspaceUpdateResult.FailedToSync, $"Couldn't sync {TamperedFile}.");
+						}
+					}
+				}
+			}
+
+			// All succeeded
+			return (WorkspaceUpdateResult.Success, "Succeeded.");
+		}
+
+		public static async Task<ConfigFile> ReadProjectConfigFile(DirectoryReference LocalRootPath, FileReference SelectedLocalFileName, ILogger Logger)
+		{
+			// Find the valid config file paths
+			DirectoryInfo EngineDir = DirectoryReference.Combine(LocalRootPath, "Engine").ToDirectoryInfo();
+			List<FileInfo> LocalConfigFiles = Utility.GetLocalConfigPaths(EngineDir, SelectedLocalFileName.ToFileInfo());
+
+			// Read them in
+			ConfigFile ProjectConfig = new ConfigFile();
+			foreach (FileInfo LocalConfigFile in LocalConfigFiles)
+			{
+				try
+				{
+					string[] Lines = await File.ReadAllLinesAsync(LocalConfigFile.FullName);
+					ProjectConfig.Parse(Lines);
+					Logger.LogInformation("Read config file from {FileName}", LocalConfigFile.FullName);
+				}
+				catch (Exception Ex)
+				{
+					Logger.LogInformation(Ex, "Failed to read config file from {FileName}", LocalConfigFile.FullName);
+				}
+			}
+			return ProjectConfig;
+		}
+
+		public static async Task<IReadOnlyList<string>?> ReadProjectStreamFilter(IPerforceConnection Perforce, ConfigFile ProjectConfigFile, ILogger Logger, CancellationToken CancellationToken)
+		{
+			string? StreamListDepotPath = ProjectConfigFile.GetValue("Options.QuickSelectStreamList", null);
+			if (StreamListDepotPath == null)
+			{
+				return null;
+			}
+
+			PerforceResponse<PrintRecord<string[]>> Response = await Perforce.TryPrintLinesAsync(StreamListDepotPath, CancellationToken);
+			if (!Response.Succeeded)
+			{
+				return null;
+			}
+
+			return Response.Data.Contents.Select(x => x.Trim()).Where(x => x.Length > 0).ToList().AsReadOnly();
+		}
+
+		static string FormatTime(long Seconds)
+		{
+			if (Seconds >= 60)
+			{
+				return String.Format("{0,3}m {1:00}s", Seconds / 60, Seconds % 60);
+			}
+			else
+			{
+				return String.Format("     {0,2}s", Seconds);
+			}
+		}
+
+		static async Task<bool> HasModifiedSourceFiles(IPerforceConnection Perforce, ProjectInfo Project, CancellationToken CancellationToken)
+		{
+			List<OpenedRecord> OpenFiles = await Perforce.OpenedAsync(OpenedOptions.None, Project.ClientRootPath + "/...", CancellationToken).ToListAsync(CancellationToken);
+			if (OpenFiles.Any(x => x.DepotFile.IndexOf("/Source/", StringComparison.OrdinalIgnoreCase) != -1))
+			{
+				return true;
+			}
+			return false;
+		}
+
+		static async Task<List<FStatRecord>> FindUnresolvedFiles(IPerforceConnection Perforce, IEnumerable<string> SyncPaths, CancellationToken CancellationToken)
+		{
+			List<FStatRecord> UnresolvedFiles = new List<FStatRecord>();
+			foreach (string SyncPath in SyncPaths)
+			{
+				List<FStatRecord> Records = await Perforce.FStatAsync(FStatOptions.OnlyUnresolved, SyncPath, CancellationToken).ToListAsync(CancellationToken);
+				UnresolvedFiles.AddRange(Records);
+			}
+			return UnresolvedFiles;
+		}
+
+		static Task<bool> UpdateVersionFile(IPerforceConnection Perforce, string ClientPath, Dictionary<string, string> VersionStrings, int ChangeNumber, ILogger Logger, CancellationToken CancellationToken)
+		{
+			return UpdateVersionFile(Perforce, ClientPath, ChangeNumber, Text => UpdateVersionStrings(Text, VersionStrings), Logger, CancellationToken);
+		}
+
+		static async Task<bool> UpdateVersionFile(IPerforceConnection Perforce, string ClientPath, int ChangeNumber, Func<string, string> Update, ILogger Logger, CancellationToken CancellationToken)
+		{
+			List<PerforceResponse<FStatRecord>> Records = await Perforce.TryFStatAsync(FStatOptions.None, ClientPath, CancellationToken).ToListAsync(CancellationToken);
+			if (!Records.Succeeded())
+			{
+				Logger.LogInformation("Failed to query records for {ClientPath}", ClientPath);
+				return false;
+			}
+			if (Records.Count > 1)
+			{
+				// Attempt to remove any existing file which is synced
+				await Perforce.SyncAsync(SyncOptions.Force, -1, $"{ClientPath}#0", CancellationToken).ToListAsync(CancellationToken);
+
+				// Try to get the mapped files again
+				Records = await Perforce.TryFStatAsync(FStatOptions.None, ClientPath, CancellationToken).ToListAsync(CancellationToken);
+				if (!Records.Succeeded())
+				{
+					Logger.LogInformation("Failed to query records for {ClientPath}", ClientPath);
+					return false;
+				}
+			}
+			if (Records.Count == 0)
+			{
+				Logger.LogInformation("Ignoring {ClientPath}; not found on server.", ClientPath);
+				return true;
+			}
+
+			FStatRecord Record = Records[0].Data;
+			string? LocalPath = Record.ClientFile; // Actually a filesystem path
+			if (LocalPath == null)
+			{
+				Logger.LogInformation("Version file is not mapped to workspace ({ClientFile})", ClientPath);
+				return false;
+			}
+			string? DepotPath = Record.DepotFile;
+			if (DepotPath == null)
+			{
+				Logger.LogInformation("Version file does not exist in depot ({ClientFile})", ClientPath);
+				return false;
+			}
+
+			PerforceResponse<PrintRecord<string[]>> Response = await Perforce.TryPrintLinesAsync($"{DepotPath}@{ChangeNumber}", CancellationToken);
+			if (!Response.Succeeded)
+			{
+				Logger.LogInformation("Couldn't get default contents of {DepotPath}", DepotPath);
+				return false;
+			}
+
+			string[]? Contents = Response.Data.Contents;
+			if (Contents == null)
+			{
+				Logger.LogInformation("No data returned for {DepotPath}", DepotPath);
+				return false;
+			}
+
+			string Text = String.Join("\n", Contents);
+			Text = Update(Text);
+			return await WriteVersionFile(Perforce, LocalPath, DepotPath, Text, Logger, CancellationToken);
+		}
+
+		static string UpdateVersionStrings(string Text, Dictionary<string, string> VersionStrings)
+		{
+			StringWriter Writer = new StringWriter();
+			foreach (string Line in Text.Split('\n'))
+			{
+				string NewLine = Line;
+				foreach (KeyValuePair<string, string> VersionString in VersionStrings)
+				{
+					if (UpdateVersionLine(ref NewLine, VersionString.Key, VersionString.Value))
+					{
+						break;
+					}
+				}
+				Writer.WriteLine(NewLine);
+			}
+			return Writer.ToString();
+		}
+
+		static string UpdateBuildVersion(string Text, int Changelist, int CodeChangelist, string BranchOrStreamName, bool bIsLicenseeVersion)
+		{
+			Dictionary<string, object> Object = JsonSerializer.Deserialize<Dictionary<string, object>>(Text, Utility.DefaultJsonSerializerOptions);
+
+			int PrevCompatibleChangelist = 0;
+			if (Object.TryGetValue("CompatibleChangelist", out object? PrevCompatibleChangelistObj))
+			{
+				int.TryParse(PrevCompatibleChangelistObj?.ToString(), out PrevCompatibleChangelist);
+			}
+
+			int PrevIsLicenseeVersion = 0;
+			if (Object.TryGetValue("IsLicenseeVersion", out object? PrevIsLicenseeVersionObj))
+			{
+				int.TryParse(PrevIsLicenseeVersionObj?.ToString(), out PrevIsLicenseeVersion);
+			}
+
+			Object["Changelist"] = Changelist;
+			if (PrevCompatibleChangelist == 0 || (PrevIsLicenseeVersion != 0) != bIsLicenseeVersion)
+			{
+				// Don't overwrite the compatible changelist if we're in a hotfix release
+				Object["CompatibleChangelist"] = CodeChangelist;
+			}
+			Object["BranchName"] = BranchOrStreamName.Replace('/', '+');
+			Object["IsPromotedBuild"] = 0;
+			Object["IsLicenseeVersion"] = bIsLicenseeVersion ? 1 : 0;
+
+			return JsonSerializer.Serialize(Object, new JsonSerializerOptions
+			{
+				WriteIndented = true,
+				// do not escape +
+				Encoder = JavaScriptEncoder.UnsafeRelaxedJsonEscaping
+			});
+		}
+
+		static async Task<bool> WriteVersionFile(IPerforceConnection Perforce, string LocalPath, string DepotPath, string NewText, ILogger Logger, CancellationToken CancellationToken)
+		{
+			try
+			{
+				if (File.Exists(LocalPath) && File.ReadAllText(LocalPath) == NewText)
+				{
+					Logger.LogInformation("Ignored {FileName}; contents haven't changed", LocalPath);
+				}
+				else
+				{
+					Directory.CreateDirectory(Path.GetDirectoryName(LocalPath));
+					Utility.ForceDeleteFile(LocalPath);
+					if (DepotPath != null)
+					{
+						await Perforce.SyncAsync(DepotPath + "#0", CancellationToken).ToListAsync(CancellationToken);
+					}
+					File.WriteAllText(LocalPath, NewText);
+					Logger.LogInformation("Written {FileName}", LocalPath);
+				}
+				return true;
+			}
+			catch (Exception Ex)
+			{
+				Logger.LogError(Ex, "Failed to write to {FileName}.", LocalPath);
+				return false;
+			}
+		}
+
+		static bool UpdateVersionLine(ref string Line, string Prefix, string Suffix)
+		{
+			int LineIdx = 0;
+			int PrefixIdx = 0;
+			for (; ; )
+			{
+				string? PrefixToken = ReadToken(Prefix, ref PrefixIdx);
+				if (PrefixToken == null)
+				{
+					break;
+				}
+
+				string? LineToken = ReadToken(Line, ref LineIdx);
+				if (LineToken == null || LineToken != PrefixToken)
+				{
+					return false;
+				}
+			}
+			Line = Line.Substring(0, LineIdx) + Suffix;
+			return true;
+		}
+
+		static string? ReadToken(string Line, ref int LineIdx)
+		{
+			for (; ; LineIdx++)
+			{
+				if (LineIdx == Line.Length)
+				{
+					return null;
+				}
+				else if (!Char.IsWhiteSpace(Line[LineIdx]))
+				{
+					break;
+				}
+			}
+
+			int StartIdx = LineIdx++;
+			if (Char.IsLetterOrDigit(Line[StartIdx]) || Line[StartIdx] == '_')
+			{
+				while (LineIdx < Line.Length && (Char.IsLetterOrDigit(Line[LineIdx]) || Line[LineIdx] == '_'))
+				{
+					LineIdx++;
+				}
+			}
+
+			return Line.Substring(StartIdx, LineIdx - StartIdx);
+		}
+
+		public Tuple<string, float> CurrentProgress
+		{
+			get { return Progress.Current; }
+		}
+	}
+}
