UGS: Save connection information into workspace settings file, stored in the per-workspace config directory. Fix support for clients containing multiple branches and standardize/validate format of project and branch paths.

diff --git a/Engine/Source/Programs/UnrealGameSync/UnrealGameSync/Controls/WorkspaceControl.cs b/Engine/Source/Programs/UnrealGameSync/UnrealGameSync/Controls/WorkspaceControl.cs
--- a/Engine/Source/Programs/UnrealGameSync/UnrealGameSync/Controls/WorkspaceControl.cs
+++ b/Engine/Source/Programs/UnrealGameSync/UnrealGameSync/Controls/WorkspaceControl.cs
@@ -20,27 +20,21 @@ using EpicGames.Perforce;
 using System.Threading.Tasks;
 using Microsoft.Extensions.Logging;
 using Microsoft.Extensions.DependencyInjection;
+using System.Diagnostics.CodeAnalysis;
 
 namespace UnrealGameSync
 {
 	public class UncontrolledChangelist
 	{
-		public string GUID { get; set; }
-		public string Name { get; set; }
-		public List<string> Files { get; set; }
+		public string GUID { get; set; } = String.Empty;
+		public string Name { get; set; } = String.Empty;
+		public List<string> Files { get; set; } = new List<string>();
 	}
 
 	public class UncontrolledChangelistPersistency
 	{
 		public int Version { get; set; }
-		public List<UncontrolledChangelist> Changelists { get; set; }
-	}
-
-	public enum LatestChangeType
-	{
-		Any,
-		Good,
-		Starred,
+		public List<UncontrolledChangelist> Changelists { get; set; } = new List<UncontrolledChangelist>();
 	}
 
 	interface IWorkspaceControlOwner
@@ -58,7 +52,7 @@ namespace UnrealGameSync
 		void UpdateAlertWindows();
 		void UpdateTintColors();
 
-		IssueMonitor CreateIssueMonitor(string ApiUrl, string UserName);
+		IssueMonitor CreateIssueMonitor(string? ApiUrl, string UserName);
 		void ReleaseIssueMonitor(IssueMonitor IssueMonitor);
 	}
 
@@ -84,22 +78,22 @@ namespace UnrealGameSync
 		class BadgeInfo
 		{
 			public string Label;
-			public string Group;
-			public string UniqueId;
+			public string? Group;
+			public string? UniqueId;
 			public int Offset;
 			public int Width;
 			public int Height;
 			public Color BackgroundColor;
 			public Color HoverBackgroundColor;
-			public Action ClickHandler;
-			public string ToolTip;
+			public Action? ClickHandler;
+			public string? ToolTip;
 
 			public BadgeInfo(string Label, string Group, Color BadgeColor)
 				: this(Label, Group, null, BadgeColor, BadgeColor, null)
 			{
 			}
 
-			public BadgeInfo(string Label, string Group, string UniqueId, Color BackgroundColor, Color HoverBackgroundColor, Action ClickHandler)
+			public BadgeInfo(string Label, string? Group, string? UniqueId, Color BackgroundColor, Color HoverBackgroundColor, Action? ClickHandler)
 			{
 				this.Label = Label;
 				this.Group = Group;
@@ -126,22 +120,22 @@ namespace UnrealGameSync
 		class ServiceBadgeInfo : IEquatable<ServiceBadgeInfo>
 		{
 			public string Name;
-			public string Url;
+			public string? Url;
 			public BadgeResult Result;
 
-			public ServiceBadgeInfo(string Name, string Url, BadgeResult Result)
+			public ServiceBadgeInfo(string Name, string? Url, BadgeResult Result)
 			{
 				this.Name = Name;
 				this.Url = Url;
 				this.Result = Result;
 			}
 
-			public bool Equals(ServiceBadgeInfo OtherBadge)
+			public bool Equals(ServiceBadgeInfo? OtherBadge)
 			{
 				return OtherBadge != null && Name == OtherBadge.Name && Url == OtherBadge.Url && Result == OtherBadge.Result;
 			}
 
-			public override bool Equals(object Other)
+			public override bool Equals(object? Other)
 			{
 				return Equals(Other as ServiceBadgeInfo);
 			}
@@ -154,18 +148,18 @@ namespace UnrealGameSync
 
 		class ChangeLayoutInfo
 		{
-			public List<BadgeInfo> DescriptionBadges;
-			public List<BadgeInfo> TypeBadges;
-			public List<BadgeInfo> BuildBadges;
-			public Dictionary<string, List<BadgeInfo>> CustomBadges;
+			public List<BadgeInfo> DescriptionBadges = new List<BadgeInfo>();
+			public List<BadgeInfo> TypeBadges = new List<BadgeInfo>();
+			public List<BadgeInfo> BuildBadges = new List<BadgeInfo>();
+			public Dictionary<string, List<BadgeInfo>> CustomBadges = new Dictionary<string, List<BadgeInfo>>();
 		}
 
 		class BuildListItemComparer : System.Collections.IComparer
 		{
-			public int Compare(object A, object B)
+			public int Compare(object? A, object? B)
 			{
-				ChangesRecord ChangeA = ((ListViewItem)A).Tag as ChangesRecord;
-				ChangesRecord ChangeB = ((ListViewItem)B).Tag as ChangesRecord;
+				ChangesRecord? ChangeA = ((ListViewItem?)A)?.Tag as ChangesRecord;
+				ChangesRecord? ChangeB = ((ListViewItem?)B)?.Tag as ChangesRecord;
 				if (ChangeA == null)
 				{
 					return +1;
@@ -209,7 +203,7 @@ namespace UnrealGameSync
 		const string EditorArchiveType = "Editor";
 
 		IWorkspaceControlOwner Owner;
-		string ApiUrl;
+		string? ApiUrl;
 		DirectoryReference DataFolder;
 		IServiceProvider ServiceProvider;
 		ILogger Logger;
@@ -270,7 +264,7 @@ namespace UnrealGameSync
 		HashSet<int> PromotedChangeNumbers = new HashSet<int>();
 		List<int> ListIndexToChangeIndex = new List<int>();
 		List<int> SortedChangeNumbers = new List<int>();
-		Dictionary<string, Dictionary<int, string>> ArchiveToChangeNumberToArchiveKey = new Dictionary<string, Dictionary<int, string>>();
+		Dictionary<string, Dictionary<int, string?>> ArchiveToChangeNumberToArchiveKey = new Dictionary<string, Dictionary<int, string?>>();
 		Dictionary<int, ChangeLayoutInfo> ChangeNumberToLayoutInfo = new Dictionary<int, ChangeLayoutInfo>();
 		List<ContextMenuStrip> CustomStatusPanelMenus = new List<ContextMenuStrip>();
 		List<(string, Action<Point, Rectangle>)> CustomStatusPanelLinks = new List<(string, Action<Point, Rectangle>)>();
@@ -282,12 +276,12 @@ namespace UnrealGameSync
 
 		bool bMouseOverExpandLink;
 
-		string HoverBadgeUniqueId = null;
+		string? HoverBadgeUniqueId = null;
 		bool bHoverSync;
-		ChangesRecord ContextMenuChange;
-		Font BuildFont;
-		Font SelectedBuildFont;
-		Font BadgeFont;
+		ChangesRecord? ContextMenuChange;
+		Font? BuildFont;
+		Font? SelectedBuildFont;
+		Font? BadgeFont;
 		List<KeyValuePair<string, string>> BadgeNameAndGroupPairs = new List<KeyValuePair<string, string>>();
 		Dictionary<string, Size> BadgeLabelToSize = new Dictionary<string, Size>();
 		List<ServiceBadgeInfo> ServiceBadges = new List<ServiceBadgeInfo>();
@@ -303,23 +297,23 @@ namespace UnrealGameSync
 		}
 
 		int BuildListWidth;
-		float[] ColumnWidths;
-		float[] ColumnWeights;
-		int[] MinColumnWidths;
-		int[] DesiredColumnWidths;
-		string LastColumnSettings;
-		List<ColumnHeader> CustomColumns;
+		float[] ColumnWidths = Array.Empty<float>();
+		float[] ColumnWeights = Array.Empty<float>();
+		int[] MinColumnWidths = Array.Empty<int>();
+		int[] DesiredColumnWidths = Array.Empty<int>();
+		string? LastColumnSettings;
+		List<ColumnHeader> CustomColumns = new List<ColumnHeader>();
 		int MaxBuildBadgeChars;
-		ListViewItem ExpandItem;
+		ListViewItem? ExpandItem;
 
 		bool bUpdateBuildListPosted;
 		bool bUpdateBuildMetadataPosted;
 		bool bUpdateReviewsPosted;
 
-		WorkspaceUpdateCallback UpdateCallback;
+		WorkspaceUpdateCallback? UpdateCallback;
 
-		System.Threading.Timer StartupTimer;
-		List<WorkspaceStartupCallback> StartupCallbacks;
+		System.Threading.Timer? StartupTimer;
+		List<WorkspaceStartupCallback>? StartupCallbacks;
 
 		// When an author filter is applied this will be non-empty and != AuthorFilterPlaceholderText
 		string AuthorFilterText = "";
@@ -327,11 +321,11 @@ namespace UnrealGameSync
 		// Placeholder text that is in the control and cleared when the user starts editing.
 		static string AuthorFilterPlaceholderText = "<username>";
 
-		public WorkspaceControl(IWorkspaceControlOwner InOwner, string InApiUrl, string InOriginalExecutableFileName, bool bInUnstable, WorkspaceSettings WorkspaceSettings, IServiceProvider InServiceProvider, UserSettings InSettings, OIDCTokenManager InOidcTokenManager)
+		public WorkspaceControl(IWorkspaceControlOwner InOwner, string? InApiUrl, string InOriginalExecutableFileName, bool bInUnstable, WorkspaceSettings WorkspaceSettings, IServiceProvider InServiceProvider, UserSettings InSettings, OIDCTokenManager? InOidcTokenManager)
 		{
 			InitializeComponent();
 
-			MainThreadSynchronizationContext = SynchronizationContext.Current;
+			MainThreadSynchronizationContext = SynchronizationContext.Current!;
 
 			Owner = InOwner;
 			ApiUrl = InApiUrl;
@@ -343,20 +337,20 @@ namespace UnrealGameSync
 			PerforceSettings = WorkspaceSettings.PerforceSettings;
 			ProjectInfo = WorkspaceSettings.ProjectInfo;
 			Settings = InSettings;
-			WorkspaceState = InSettings.FindOrAddWorkspaceState(WorkspaceSettings.ProjectInfo.LocalRootPath, WorkspaceSettings.ProjectInfo.ClientFileName, WorkspaceSettings.ProjectInfo.ProjectIdentifier);
-			this.WorkspaceSettings = InSettings.FindOrAddWorkspaceSettings(WorkspaceSettings.ProjectInfo.LocalRootPath, WorkspaceSettings.ProjectInfo.ClientFileName);
+			this.WorkspaceSettings = InSettings.FindOrAddWorkspaceSettings(ProjectInfo.LocalRootPath, PerforceSettings.ServerAndPort, PerforceSettings.UserName, ProjectInfo.ClientName, ProjectInfo.BranchPath, ProjectInfo.ProjectPath);
+			WorkspaceState = InSettings.FindOrAddWorkspaceState(this.WorkspaceSettings);
 
 			string ProjectPath = Regex.Replace(WorkspaceSettings.ProjectInfo.ClientFileName, "^//[^/]+/", "");
 			ProjectSettings = InSettings.FindOrAddProjectSettings(WorkspaceSettings.ProjectInfo.LocalRootPath, ProjectPath);
 
 			DesiredTaskbarState = Tuple.Create(TaskbarState.NoProgress, 0.0f);
 
-			System.Reflection.PropertyInfo DoubleBufferedProperty = typeof(Control).GetProperty("DoubleBuffered", System.Reflection.BindingFlags.NonPublic | System.Reflection.BindingFlags.Instance);
+			System.Reflection.PropertyInfo DoubleBufferedProperty = typeof(Control).GetProperty("DoubleBuffered", System.Reflection.BindingFlags.NonPublic | System.Reflection.BindingFlags.Instance)!;
 			DoubleBufferedProperty.SetValue(BuildList, true, null);
 
 			// force the height of the rows
 			BuildList.SmallImageList = new ImageList() { ImageSize = new Size(1, 20) };
-			BuildList_FontChanged(null, null);
+			BuildList_FontChanged(null, null!);
 			BuildList.OnScroll += BuildList_OnScroll;
 			BuildList.ListViewItemSorter = new BuildListItemComparer();
 
@@ -393,7 +387,7 @@ namespace UnrealGameSync
 
 			// Commit all the new project info
 			IPerforceSettings PerforceClientSettings = WorkspaceSettings.PerforceSettings;
-			ClientName = PerforceUtils.GetClientOrDepotDirectoryName(WorkspaceSettings.ProjectInfo.ClientFileName);
+			ClientName = PerforceClientSettings.ClientName!;
 			SelectedProject = WorkspaceSettings.SelectedProject;
 			SelectedProjectIdentifier = WorkspaceSettings.NewSelectedProjectIdentifier;
 			SelectedFileName = WorkspaceSettings.ProjectInfo.LocalFileName;
@@ -406,7 +400,7 @@ namespace UnrealGameSync
 
 			// Check if we've the project we've got open in this workspace is the one we're actually synced to
 			int CurrentChangeNumber = -1;
-			string CurrentSyncFilterHash = null;
+			string? CurrentSyncFilterHash = null;
 			if (String.Compare(WorkspaceState.CurrentProjectIdentifier, SelectedProjectIdentifier, true) == 0)
 			{
 				CurrentChangeNumber = WorkspaceState.CurrentChangeNumber;
@@ -414,7 +408,7 @@ namespace UnrealGameSync
 			}
 
 			// Figure out which API server to use
-			string NewApiUrl;
+			string? NewApiUrl;
 			if (TryGetProjectSetting(WorkspaceSettings.LatestProjectConfigFile, "ApiUrl", out NewApiUrl))
 			{
 				ApiUrl = NewApiUrl;
@@ -441,9 +435,9 @@ namespace UnrealGameSync
 			ILogger<JupiterMonitor> JupiterLogger = ServiceProvider.GetRequiredService<ILogger<JupiterMonitor>>();
 			JupiterMonitor = JupiterMonitor.CreateFromConfigFile(InOidcTokenManager, JupiterLogger, WorkspaceSettings.LatestProjectConfigFile, SelectedProjectIdentifier);
 
-			UpdateColumnSettings();
+			UpdateColumnSettings(true);
 
-			SyncLog.OpenFile(ProjectLogBaseName + ".sync.log");
+			SyncLog.OpenFile(ProjectLogBaseName);
 
 			Splitter.SetLogVisibility(Settings.bShowLogWindow);
 
@@ -460,7 +454,7 @@ namespace UnrealGameSync
 			StartupTimer = new System.Threading.Timer(x => MainThreadSynchronizationContext.Post((o) => { if (!IsDisposed) { StartupTimerElapsed(false); } }, null), null, TimeSpan.FromSeconds(20.0), TimeSpan.FromMilliseconds(-1.0));
 			StartupCallbacks = new List<WorkspaceStartupCallback>();
 
-			string IssuesApiUrl = GetIssuesApiUrl();
+			string? IssuesApiUrl = GetIssuesApiUrl();
 			IssueMonitor = InOwner.CreateIssueMonitor(IssuesApiUrl, WorkspaceSettings.PerforceSettings.UserName);
 			IssueMonitor.OnIssuesChanged += IssueMonitor_OnIssuesChangedAsync;
 
@@ -544,32 +538,32 @@ namespace UnrealGameSync
 			}
 			else
 			{
-				StartupCallbacks.Add(StartupCallback);
+				StartupCallbacks!.Add(StartupCallback);
 			}
 		}
 
-		private void UpdateColumnSettings()
+		private void UpdateColumnSettings(bool bForceUpdate)
 		{
-			string NextColumnSettings;
+			string? NextColumnSettings;
 			TryGetProjectSetting(PerforceMonitor.LatestProjectConfigFile, "Columns", out NextColumnSettings);
 
-			if (CustomColumns == null || NextColumnSettings != LastColumnSettings)
+			if (bForceUpdate || NextColumnSettings != LastColumnSettings)
 			{
 				LastColumnSettings = NextColumnSettings;
 
-				if (CustomColumns != null)
+				foreach (ColumnHeader CustomColumn in CustomColumns)
 				{
-					foreach (ColumnHeader CustomColumn in CustomColumns)
-					{
-						BuildList.Columns.Remove(CustomColumn);
-					}
+					BuildList.Columns.Remove(CustomColumn);
 				}
 
 				Dictionary<string, ColumnHeader> NameToColumn = new Dictionary<string, ColumnHeader>();
-				foreach (ColumnHeader Column in BuildList.Columns)
+				foreach (ColumnHeader? Column in BuildList.Columns)
 				{
-					NameToColumn[Column.Text] = Column;
-					Column.Tag = null;
+					if (Column != null)
+					{
+						NameToColumn[Column.Text] = Column;
+						Column.Tag = null;
+					}
 				}
 
 				CustomColumns = new List<ColumnHeader>();
@@ -579,10 +573,10 @@ namespace UnrealGameSync
 					{
 						ConfigObject ColumnConfig = new ConfigObject(CustomColumn);
 
-						string Name = ColumnConfig.GetValue("Name", null);
+						string? Name = ColumnConfig.GetValue("Name", null);
 						if (Name != null)
 						{
-							ColumnHeader Column;
+							ColumnHeader? Column;
 							if (NameToColumn.TryGetValue(Name, out Column))
 							{
 								Column.Tag = ColumnConfig;
@@ -640,14 +634,17 @@ namespace UnrealGameSync
 					ColumnWeights[DescriptionColumn.Index] = 1.25f;
 					ColumnWeights[CISColumn.Index] = 1.5f;
 
-					foreach (ColumnHeader Column in BuildList.Columns)
+					foreach (ColumnHeader? Column in BuildList.Columns)
 					{
-						ConfigObject ColumnConfig = (ConfigObject)Column.Tag;
-						if (ColumnConfig != null)
+						if (Column != null)
 						{
-							MinColumnWidths[Column.Index] = (int)(ColumnConfig.GetValue("MinWidth", MinColumnWidths[Column.Index]) * DpiScaleX);
-							DesiredColumnWidths[Column.Index] = (int)(ColumnConfig.GetValue("DesiredWidth", DesiredColumnWidths[Column.Index]) * DpiScaleX);
-							ColumnWeights[Column.Index] = ColumnConfig.GetValue("Weight", MinColumnWidths[Column.Index]);
+							ConfigObject? ColumnConfig = (ConfigObject?)Column.Tag;
+							if (ColumnConfig != null)
+							{
+								MinColumnWidths[Column.Index] = (int)(ColumnConfig.GetValue("MinWidth", MinColumnWidths[Column.Index]) * DpiScaleX);
+								DesiredColumnWidths[Column.Index] = (int)(ColumnConfig.GetValue("DesiredWidth", DesiredColumnWidths[Column.Index]) * DpiScaleX);
+								ColumnWeights[Column.Index] = ColumnConfig.GetValue("Weight", MinColumnWidths[Column.Index]);
+							}
 						}
 					}
 
@@ -656,7 +653,7 @@ namespace UnrealGameSync
 					{
 						if (!String.IsNullOrEmpty(BuildList.Columns[Idx].Text))
 						{
-							string StringValue;
+							string? StringValue;
 							if (TryGetProjectSetting(ProjectConfigFile, String.Format("ColumnWidth_{0}", BuildList.Columns[Idx].Text), out StringValue))
 							{
 								int IntValue;
@@ -678,10 +675,10 @@ namespace UnrealGameSync
 
 		private bool UpdateServiceBadges()
 		{
-			string[] ServiceBadgeNames;
+			string[]? ServiceBadgeNames;
 			if (!TryGetProjectSetting(PerforceMonitor.LatestProjectConfigFile, "ServiceBadges", out ServiceBadgeNames))
 			{
-				ServiceBadgeNames = new string[0];
+				ServiceBadgeNames = Array.Empty<string>();
 			}
 
 			List<ServiceBadgeInfo> PrevServiceBadges = ServiceBadges;
@@ -689,7 +686,7 @@ namespace UnrealGameSync
 			ServiceBadges = new List<ServiceBadgeInfo>();
 			foreach (string ServiceBadgeName in ServiceBadgeNames)
 			{
-				BadgeData LatestBuild;
+				BadgeData? LatestBuild;
 				if (EventMonitor.TryGetLatestBadge(ServiceBadgeName, out LatestBuild))
 				{
 					ServiceBadges.Add(new ServiceBadgeInfo(ServiceBadgeName, LatestBuild.Url, LatestBuild.Result));
@@ -750,59 +747,59 @@ namespace UnrealGameSync
 				{
 					StartupCallback(this, true);
 				}
-				StartupCallbacks = null;
+				StartupCallbacks = null!;
 			}
 
 			if (IssueMonitor != null)
 			{
 				IssueMonitor.OnIssuesChanged -= IssueMonitor_OnIssuesChangedAsync;
 				Owner.ReleaseIssueMonitor(IssueMonitor);
-				IssueMonitor = null;
+				IssueMonitor = null!;
 			}
 			if (StartupTimer != null)
 			{
 				StartupTimer.Dispose();
-				StartupTimer = null;
+				StartupTimer = null!;
 			}
 			if (NotificationWindow != null)
 			{
 				NotificationWindow.Dispose();
-				NotificationWindow = null;
+				NotificationWindow = null!;
 			}
 			if (PerforceMonitor != null)
 			{
 				PerforceMonitor.Dispose();
-				PerforceMonitor = null;
+				PerforceMonitor = null!;
 			}
 			if (Workspace != null)
 			{
 				Workspace.Dispose();
-				Workspace = null;
+				Workspace = null!;
 			}
 			if (EventMonitor != null)
 			{
 				EventMonitor.Dispose();
-				EventMonitor = null;
+				EventMonitor = null!;
 			}
 			if (BuildFont != null)
 			{
 				BuildFont.Dispose();
-				BuildFont = null;
+				BuildFont = null!;
 			}
 			if (SelectedBuildFont != null)
 			{
 				SelectedBuildFont.Dispose();
-				SelectedBuildFont = null;
+				SelectedBuildFont = null!;
 			}
 			if (BadgeFont != null)
 			{
 				BadgeFont.Dispose();
-				BadgeFont = null;
+				BadgeFont = null!;
 			}
 			if (JupiterMonitor != null)
 			{
 				JupiterMonitor.Dispose();
-				JupiterMonitor = null;
+				JupiterMonitor = null!;
 			}
 
 			base.Dispose(disposing);
@@ -917,7 +914,7 @@ namespace UnrealGameSync
 			StartSync(ChangeNumber, false, null);
 		}
 
-		void StartSync(int ChangeNumber, bool bSyncOnly, WorkspaceUpdateCallback Callback)
+		void StartSync(int ChangeNumber, bool bSyncOnly, WorkspaceUpdateCallback? Callback)
 		{
 			WorkspaceUpdateOptions Options = WorkspaceUpdateOptions.Sync | WorkspaceUpdateOptions.SyncArchives | WorkspaceUpdateOptions.GenerateProjectFiles;
 			if (!bSyncOnly)
@@ -943,7 +940,7 @@ namespace UnrealGameSync
 			StartWorkspaceUpdate(ChangeNumber, Options, null);
 		}
 
-		void StartWorkspaceUpdate(int ChangeNumber, WorkspaceUpdateOptions Options, WorkspaceUpdateCallback Callback)
+		void StartWorkspaceUpdate(int ChangeNumber, WorkspaceUpdateOptions Options, WorkspaceUpdateCallback? Callback)
 		{
 			if ((Options & (WorkspaceUpdateOptions.Sync | WorkspaceUpdateOptions.Build)) != 0 && GetProcessesRunningInWorkspace().Length > 0)
 			{
@@ -983,7 +980,7 @@ namespace UnrealGameSync
 				List<IArchiveInfo> SelectedArchives = GetSelectedArchives(Archives);
 				foreach (IArchiveInfo Archive in SelectedArchives)
 				{
-					string ArchivePath = GetArchiveKeyForChangeNumber(Archive, ChangeNumber);
+					string? ArchivePath = GetArchiveKeyForChangeNumber(Archive, ChangeNumber);
 					if (ArchivePath == null)
 					{
 						MessageBox.Show(String.Format("There are no compiled {0} binaries for this change. To sync it, you must disable syncing of precompiled editor binaries.", Archive.Name));
@@ -995,7 +992,7 @@ namespace UnrealGameSync
 						Context.Options &= ~(WorkspaceUpdateOptions.Build | WorkspaceUpdateOptions.GenerateProjectFiles | WorkspaceUpdateOptions.OpenSolutionAfterSync);
 					}
 
-					string[] ZippedBinariesSyncFilter;
+					string[]? ZippedBinariesSyncFilter;
 					if (TryGetProjectSetting(PerforceMonitor.LatestProjectConfigFile, "ZippedBinariesSyncFilter", out ZippedBinariesSyncFilter) && ZippedBinariesSyncFilter.Length > 0)
 					{
 						Context.SyncFilter = Enumerable.Concat(Context.SyncFilter, ZippedBinariesSyncFilter).ToArray();
@@ -1007,7 +1004,7 @@ namespace UnrealGameSync
 			StartWorkspaceUpdate(Context, Callback);
 		}
 
-		void StartWorkspaceUpdate(WorkspaceUpdateContext Context, WorkspaceUpdateCallback Callback)
+		void StartWorkspaceUpdate(WorkspaceUpdateContext Context, WorkspaceUpdateCallback? Callback)
 		{
 			if (Settings.bAutoResolveConflicts)
 			{
@@ -1041,7 +1038,7 @@ namespace UnrealGameSync
 			{
 				if (!Context.Options.HasFlag(WorkspaceUpdateOptions.ContentOnly) && (Context.CustomBuildSteps == null || Context.CustomBuildSteps.Count == 0))
 				{
-					foreach (BuildConfig Config in Enum.GetValues(typeof(BuildConfig)))
+					foreach (BuildConfig Config in Enum.GetValues(typeof(BuildConfig)).OfType<BuildConfig>())
 					{
 						List<FileReference> EditorReceiptPaths = GetEditorReceiptPaths(Config);
 						foreach (FileReference EditorReceiptPath in EditorReceiptPaths)
@@ -1123,7 +1120,7 @@ namespace UnrealGameSync
 
 				HashSet<string> UncontrolledFiles = new HashSet<string>();
 				
-				if (SelectedProject.LocalPath.EndsWith(".uprojectdirs", StringComparison.InvariantCultureIgnoreCase))
+				if (SelectedProject.LocalPath != null && SelectedProject.LocalPath.EndsWith(".uprojectdirs", StringComparison.InvariantCultureIgnoreCase))
 				{
 					List<string> ProjectRoots = GetProjectRoots(SelectedProject.LocalPath);
 
@@ -1134,7 +1131,7 @@ namespace UnrealGameSync
 				}
 				else
 				{
-					ParseUncontrolledChangelistsPersistencyFile(Path.GetDirectoryName(SelectedProject.LocalPath), UncontrolledFiles);
+					ParseUncontrolledChangelistsPersistencyFile(Path.GetDirectoryName(SelectedProject.LocalPath)!, UncontrolledFiles);
 				}
 
 				ClobberWindow Window = new ClobberWindow(Context.ClobberFiles, UncontrolledFiles);
@@ -1274,14 +1271,14 @@ namespace UnrealGameSync
 			Dictionary<int, ChangesRecord> ChangeNumberToSummary = Changes.ToDictionary(x => x.Number, x => x);
 			for (int Idx = BuildList.Items.Count - 1; Idx >= 0; Idx--)
 			{
-				ChangesRecord Change = BuildList.Items[Idx].Tag as ChangesRecord;
+				ChangesRecord? Change = BuildList.Items[Idx].Tag as ChangesRecord;
 				if (Change != null)
 				{
-					ChangesRecord Summary;
+					ChangesRecord? Summary;
 					if (ChangeNumberToSummary.TryGetValue(Change.Number, out Summary))
 					{
 						// Update
-						BuildList.Items[Idx].SubItems[DescriptionColumn.Index].Text = Change.Description.Replace('\n', ' ');
+						BuildList.Items[Idx].SubItems[DescriptionColumn.Index].Text = Change.Description!.Replace('\n', ' ');
 						ChangeNumberToSummary.Remove(Change.Number);
 					}
 					else
@@ -1299,7 +1296,7 @@ namespace UnrealGameSync
 			}
 
 			// Figure out which group the expand item should be in
-			ListViewGroup NewExpandItemGroup = null;
+			ListViewGroup? NewExpandItemGroup = null;
 			for (int Idx = BuildList.Items.Count - 1; Idx >= 0; Idx--)
 			{
 				ListViewItem Item = BuildList.Items[Idx];
@@ -1358,7 +1355,7 @@ namespace UnrealGameSync
 					return NextGroup;
 				}
 
-				ChangesRecord LastChange = null;
+				ChangesRecord? LastChange = null;
 				for (int Idx = NextGroup.Items.Count - 1; Idx >= 0 && LastChange == null; Idx--)
 				{
 					LastChange = NextGroup.Items[Idx].Tag as ChangesRecord;
@@ -1396,7 +1393,7 @@ namespace UnrealGameSync
 			Columns[ChangeColumn.Index] = String.Format("{0}", Change.Number);
 
 			Columns[TimeColumn.Index] = DisplayTime.ToShortTimeString();
-			string UserName = FormatUserName(Change.User);
+			string UserName = FormatUserName(Change.User!);
 
 			// If annotate robomerge is on, add a prefix to the user name
 			if (Settings.bAnnotateRobmergeChanges && IsRobomergeChange(Change))
@@ -1405,7 +1402,7 @@ namespace UnrealGameSync
 			}
 
 			Columns[AuthorColumn.Index] = UserName;
-			Columns[DescriptionColumn.Index] = Change.Description.Replace('\n', ' ');
+			Columns[DescriptionColumn.Index] = Change.Description!.Replace('\n', ' ');
 
 			for (int ColumnIdx = 1; ColumnIdx < BuildList.Columns.Count; ColumnIdx++)
 			{
@@ -1416,7 +1413,7 @@ namespace UnrealGameSync
 			int GroupInsertIdx = 0;
 			for (; GroupInsertIdx < Group.Items.Count; GroupInsertIdx++)
 			{
-				ChangesRecord OtherChange = Group.Items[GroupInsertIdx].Tag as ChangesRecord;
+				ChangesRecord? OtherChange = Group.Items[GroupInsertIdx].Tag as ChangesRecord;
 				if (OtherChange == null || Change.Number >= OtherChange.Number)
 				{
 					break;
@@ -1437,14 +1434,14 @@ namespace UnrealGameSync
 		{
 			// If hiding robomerge, filter out based on workspace name. Note - don't look at the description because we
 			// *do* want to see conflicts that were manually merged by someone
-			return Change.Client.IndexOf("ROBOMERGE", StringComparison.OrdinalIgnoreCase) >= 0;
+			return Change.Client!.IndexOf("ROBOMERGE", StringComparison.OrdinalIgnoreCase) >= 0;
 		}
 
 		bool ShouldShowChange(ChangesRecord Change, string[] ExcludeChanges)
 		{
 			if (ProjectSettings.FilterType != FilterType.None)
 			{
-				PerforceChangeDetails Details;
+				PerforceChangeDetails? Details;
 				if (!PerforceMonitor.TryGetChangeDetails(Change.Number, out Details))
 				{
 					return false;
@@ -1462,7 +1459,7 @@ namespace UnrealGameSync
 			// if filtering by user, only show changes where the author contains the filter text
 			if (!string.IsNullOrEmpty(this.AuthorFilterText)
 				&& this.AuthorFilterText != AuthorFilterPlaceholderText
-				&& Change.User.IndexOf(this.AuthorFilterText, StringComparison.OrdinalIgnoreCase) < 0)
+				&& Change.User!.IndexOf(this.AuthorFilterText, StringComparison.OrdinalIgnoreCase) < 0)
 			{
 				return false;
 			}
@@ -1477,7 +1474,7 @@ namespace UnrealGameSync
 				else if (Settings.ShowRobomerge == UserSettings.RobomergeShowChangesOption.Badged)
 				{
 					// if this change has any badges, we'll show it unless it's later filtered out
-					EventSummary Summary = EventMonitor.GetSummaryForChange(Change.Number);
+					EventSummary? Summary = EventMonitor.GetSummaryForChange(Change.Number);
 					if (Summary == null || !Summary.Badges.Any())
 					{
 						return false;
@@ -1487,7 +1484,7 @@ namespace UnrealGameSync
 
 			if (ProjectSettings.FilterBadges.Count > 0)
 			{
-				EventSummary Summary = EventMonitor.GetSummaryForChange(Change.Number);
+				EventSummary? Summary = EventMonitor.GetSummaryForChange(Change.Number);
 				if (Summary == null || !Summary.Badges.Any(x => ProjectSettings.FilterBadges.Contains(x.BadgeName)))
 				{
 					return false;
@@ -1503,7 +1500,7 @@ namespace UnrealGameSync
 					}
 				}
 
-				if (String.Compare(Change.User, "buildmachine", true) == 0 && Change.Description.IndexOf("lightmaps", StringComparison.InvariantCultureIgnoreCase) == -1)
+				if (String.Compare(Change.User, "buildmachine", true) == 0 && Change.Description!.IndexOf("lightmaps", StringComparison.InvariantCultureIgnoreCase) == -1)
 				{
 					return false;
 				}
@@ -1524,9 +1521,9 @@ namespace UnrealGameSync
 			}
 		}
 
-		string GetIssuesApiUrl()
+		string? GetIssuesApiUrl()
 		{
-			string IssuesApiUrl;
+			string? IssuesApiUrl;
 			if (!TryGetProjectSetting(PerforceMonitor.LatestProjectConfigFile, "IssuesApiUrl", out IssuesApiUrl))
 			{
 				IssuesApiUrl = ApiUrl;
@@ -1537,7 +1534,7 @@ namespace UnrealGameSync
 		void UpdateBuildMetadata()
 		{
 			// Refresh the issue monitor if it's changed
-			string IssuesApiUrl = GetIssuesApiUrl();
+			string? IssuesApiUrl = GetIssuesApiUrl();
 			if (IssuesApiUrl != IssueMonitor.ApiUrl)
 			{
 				Logger.LogInformation("Changing issues API url from {OldApiUrl} to {ApiUrl}", IssueMonitor.ApiUrl, IssuesApiUrl);
@@ -1547,7 +1544,7 @@ namespace UnrealGameSync
 			}
 
 			// Update the column settings first, since this may affect the badges we hide
-			UpdateColumnSettings();
+			UpdateColumnSettings(false);
 
 			// Clear the badge size cache
 			BadgeLabelToSize.Clear();
@@ -1556,7 +1553,7 @@ namespace UnrealGameSync
 			Dictionary<string, string> BadgeNameToGroup = new Dictionary<string, string>();
 
 			// Read the group mappings from the config file
-			string GroupDefinitions;
+			string? GroupDefinitions;
 			if (TryGetProjectSetting(PerforceMonitor.LatestProjectConfigFile, "BadgeGroups", out GroupDefinitions))
 			{
 				string[] GroupDefinitionsArray = GroupDefinitions.Split('\n');
@@ -1571,11 +1568,11 @@ namespace UnrealGameSync
 			}
 
 			// Add a dummy group for any other badges we have
-			foreach (ListViewItem Item in BuildList.Items)
+			foreach (ListViewItem? Item in BuildList.Items)
 			{
-				if (Item.Tag != null)
+				if (Item != null && Item.Tag != null)
 				{
-					EventSummary Summary = EventMonitor.GetSummaryForChange(((ChangesRecord)Item.Tag).Number);
+					EventSummary? Summary = EventMonitor.GetSummaryForChange(((ChangesRecord)Item.Tag).Number);
 					if (Summary != null)
 					{
 						foreach (BadgeData Badge in Summary.Badges)
@@ -1689,7 +1686,7 @@ namespace UnrealGameSync
 				return true;
 			}
 
-			EventSummary Review = EventMonitor.GetSummaryForChange(ChangeNumber);
+			EventSummary? Review = EventMonitor.GetSummaryForChange(ChangeNumber);
 			if (Review != null)
 			{
 				if (Review.LastStarReview != null && Review.LastStarReview.Type == EventType.Starred)
@@ -1720,7 +1717,7 @@ namespace UnrealGameSync
 					Guid UniqueId;
 					if (Guid.TryParse(Object.GetValue("UniqueId", ""), out UniqueId))
 					{
-						WorkspaceSyncCategory Category;
+						WorkspaceSyncCategory? Category;
 						if (!UniqueIdToCategory.TryGetValue(UniqueId, out Category))
 						{
 							Category = new WorkspaceSyncCategory(UniqueId);
@@ -1765,7 +1762,7 @@ namespace UnrealGameSync
 				return ProjectRoots;
 			}
 
-			string UProjectDirsRoot = Path.GetDirectoryName(InUProjectDirsPath);
+			string UProjectDirsRoot = Path.GetDirectoryName(InUProjectDirsPath)!;
 
 			foreach (string Line in File.ReadLines(InUProjectDirsPath))
 			{
@@ -1809,7 +1806,7 @@ namespace UnrealGameSync
 				try
 				{
 					string JsonString = File.ReadAllText(UncontrolledChangelistPersistencyFilePath);
-					UncontrolledChangelistPersistency UCLPersistency = JsonSerializer.Deserialize<UncontrolledChangelistPersistency>(JsonString, Program.DefaultJsonSerializerOptions);
+					UncontrolledChangelistPersistency UCLPersistency = JsonSerializer.Deserialize<UncontrolledChangelistPersistency>(JsonString, Utility.DefaultJsonSerializerOptions);
 
 					foreach (UncontrolledChangelist UCL in UCLPersistency.Changelists)
 					{
@@ -1871,7 +1868,7 @@ namespace UnrealGameSync
 			Dictionary<string, BadgeData> TypeToLastSucceededBuild = new Dictionary<string, BadgeData>();
 			for (int Idx = SortedChangeNumbers.Count - 1; Idx >= 0; Idx--)
 			{
-				EventSummary Summary = EventMonitor.GetSummaryForChange(SortedChangeNumbers[Idx]);
+				EventSummary? Summary = EventMonitor.GetSummaryForChange(SortedChangeNumbers[Idx]);
 				if (Summary != null)
 				{
 					foreach (BadgeData Badge in Summary.Badges)
@@ -1910,7 +1907,7 @@ namespace UnrealGameSync
 					if (LastChangeByCurrentUserOfType > 0 && LastBuild.ChangeNumber >= LastChangeByCurrentUserOfType)
 					{
 						// And check that there wasn't a successful build after we submitted (if there was, we're in the clear)
-						BadgeData LastSuccessfulBuild;
+						BadgeData? LastSuccessfulBuild;
 						if (!TypeToLastSucceededBuild.TryGetValue(LastBuild.BuildType, out LastSuccessfulBuild) || LastSuccessfulBuild.ChangeNumber < LastChangeByCurrentUserOfType)
 						{
 							// Add it to the list of notifications
@@ -1975,9 +1972,9 @@ namespace UnrealGameSync
 
 		class ExpandRowLayout
 		{
-			public string MainText;
+			public string? MainText;
 			public Rectangle MainRect;
-			public string LinkText;
+			public string? LinkText;
 			public Rectangle LinkRect;
 		}
 
@@ -2072,20 +2069,20 @@ namespace UnrealGameSync
 			}
 		}
 
-		private string GetArchiveKeyForChangeNumber(IArchiveInfo Archive, int ChangeNumber)
+		private string? GetArchiveKeyForChangeNumber(IArchiveInfo Archive, int ChangeNumber)
 		{
-			string ArchivePath;
+			string? ArchivePath;
 
-			Dictionary<int, string> ChangeNumberToArchivePath;
+			Dictionary<int, string?>? ChangeNumberToArchivePath;
 			if (!ArchiveToChangeNumberToArchiveKey.TryGetValue(Archive.Name, out ChangeNumberToArchivePath))
 			{
-				ChangeNumberToArchivePath = new Dictionary<int, string>();
+				ChangeNumberToArchivePath = new Dictionary<int, string?>();
 				ArchiveToChangeNumberToArchiveKey[Archive.Name] = ChangeNumberToArchivePath;
 			}
 
 			if (!ChangeNumberToArchivePath.TryGetValue(ChangeNumber, out ArchivePath))
 			{
-				PerforceChangeDetails Details;
+				PerforceChangeDetails? Details;
 				if (PerforceMonitor.TryGetChangeDetails(ChangeNumber, out Details))
 				{
 					// Try to get the archive for this CL
@@ -2123,7 +2120,7 @@ namespace UnrealGameSync
 
 		private ChangeLayoutInfo GetChangeLayoutInfo(ChangesRecord Change)
 		{
-			ChangeLayoutInfo LayoutInfo;
+			ChangeLayoutInfo? LayoutInfo;
 			if (!ChangeNumberToLayoutInfo.TryGetValue(Change.Number, out LayoutInfo))
 			{
 				LayoutInfo = new ChangeLayoutInfo();
@@ -2131,7 +2128,7 @@ namespace UnrealGameSync
 				LayoutInfo.DescriptionBadges = CreateDescriptionBadges(Change);
 				LayoutInfo.TypeBadges = CreateTypeBadges(Change.Number);
 
-				EventSummary Summary = EventMonitor.GetSummaryForChange(Change.Number);
+				EventSummary? Summary = EventMonitor.GetSummaryForChange(Change.Number);
 				LayoutInfo.BuildBadges = CreateBuildBadges(Change.Number, Summary);
 				LayoutInfo.CustomBadges = CreateCustomBadges(Change.Number, Summary);
 
@@ -2182,11 +2179,11 @@ namespace UnrealGameSync
 			Format.FormatFlags = StringFormatFlags.NoWrap;
 			Format.Trimming = StringTrimming.EllipsisCharacter;
 
-			Font CurrentFont = (Change.Number == Workspace.PendingChangeNumber || Change.Number == Workspace.CurrentChangeNumber) ? SelectedBuildFont : BuildFont;
+			Font CurrentFont = (Change.Number == Workspace.PendingChangeNumber || Change.Number == Workspace.CurrentChangeNumber) ? SelectedBuildFont! : BuildFont!;
 
 			bool bAllowSync = CanSyncChange(Change.Number);
 			int BadgeAlpha = bAllowSync ? 255 : 128;
-			Color TextColor = (bAllowSync || Change.Number == Workspace.PendingChangeNumber || Change.Number == Workspace.CurrentChangeNumber || (WorkspaceState != null && WorkspaceState.AdditionalChangeNumbers.Contains(Change.Number))) ? SystemColors.WindowText : Blend(SystemColors.Window, SystemColors.WindowText, 0.25f);
+			Color TextColor = (bAllowSync || Change.Number == Workspace.PendingChangeNumber || Change.Number == Workspace.CurrentChangeNumber || WorkspaceState.AdditionalChangeNumbers.Contains(Change.Number)) ? SystemColors.WindowText : Blend(SystemColors.Window, SystemColors.WindowText, 0.25f);
 
 			const int FadeRange = 6;
 			if (e.ItemIndex >= BuildList.Items.Count - FadeRange && NumChanges >= PerforceMonitor.CurrentMaxChanges && !IsBisectModeEnabled())
@@ -2198,7 +2195,7 @@ namespace UnrealGameSync
 
 			if (e.ColumnIndex == IconColumn.Index)
 			{
-				EventSummary Summary = EventMonitor.GetSummaryForChange(Change.Number);
+				EventSummary? Summary = EventMonitor.GetSummaryForChange(Change.Number);
 
 				float MinX = 4 * DpiScaleX;
 				if ((Summary != null && EventMonitor.WasSyncedByCurrentUser(Summary.ChangeNumber)) || (Workspace != null && Workspace.CurrentChangeNumber == Change.Number))
@@ -2229,7 +2226,7 @@ namespace UnrealGameSync
 						{
 							QualityIcon = new Rectangle(0, 0, 0, 0);
 						}
-						if (Entry.State == BisectState.Pass)
+						else if (Entry.State == BisectState.Pass)
 						{
 							QualityIcon = BisectPassIcon;
 						}
@@ -2354,7 +2351,7 @@ namespace UnrealGameSync
 
 					if (Change.Number == Workspace.CurrentChangeNumber)
 					{
-						EventData Review = EventMonitor.GetReviewByCurrentUser(Change.Number);
+						EventData? Review = EventMonitor.GetReviewByCurrentUser(Change.Number);
 
 						MaxX -= FrownIcon.Width * DpiScaleX;
 						e.Graphics.DrawImage(Properties.Resources.Icons, MaxX, IconY, (Review == null || !EventMonitor.IsPositiveReview(Review.Type)) ? FrownIcon : DisabledFrownIcon, GraphicsUnit.Pixel);
@@ -2369,12 +2366,12 @@ namespace UnrealGameSync
 						MaxX = SyncBadgeRectangle.Left - (int)(2 * DpiScaleX);
 					}
 
-					string SummaryText;
+					string? SummaryText;
 					if (WorkspaceState.LastSyncChangeNumber == -1 || WorkspaceState.LastSyncChangeNumber != Change.Number || !GetLastUpdateMessage(WorkspaceState.LastSyncResult, WorkspaceState.LastSyncResultMessage, out SummaryText))
 					{
 						StringBuilder SummaryTextBuilder = new StringBuilder();
 
-						EventSummary Summary = EventMonitor.GetSummaryForChange(Change.Number);
+						EventSummary? Summary = EventMonitor.GetSummaryForChange(Change.Number);
 
 						AppendItemList(SummaryTextBuilder, " ", "Under investigation by {0}.", EventMonitor.GetInvestigatingUsers(Change.Number).Select(x => FormatUserName(x)));
 
@@ -2409,7 +2406,7 @@ namespace UnrealGameSync
 				{
 					ChangeLayoutInfo Layout = GetChangeLayoutInfo(Change);
 
-					List<BadgeInfo> Badges;
+					List<BadgeInfo>? Badges;
 					if (Layout.CustomBadges.TryGetValue(Column.Text, out Badges) && Badges.Count > 0)
 					{
 						e.Graphics.IntersectClip(e.Bounds);
@@ -2500,9 +2497,9 @@ namespace UnrealGameSync
 
 		private List<BadgeInfo> CreateDescriptionBadges(ChangesRecord Change)
 		{
-			string Description = Change.Description;
+			string Description = Change.Description ?? String.Empty;
 
-			PerforceChangeDetails Details;
+			PerforceChangeDetails? Details;
 			if (PerforceMonitor.TryGetChangeDetails(Change.Number, out Details))
 			{
 				Description = Details.Description;
@@ -2519,40 +2516,43 @@ namespace UnrealGameSync
 					foreach (string BadgeDefinition in BadgeDefinitions.Distinct())
 					{
 						ConfigObject BadgeDefinitionObject = new ConfigObject(BadgeDefinition);
-						string Pattern = BadgeDefinitionObject.GetValue("Pattern", null);
-						string Name = BadgeDefinitionObject.GetValue("Name", null);
-						string Group = BadgeDefinitionObject.GetValue("Group", null);
+						string? Pattern = BadgeDefinitionObject.GetValue("Pattern", null);
+						string? Name = BadgeDefinitionObject.GetValue("Name", null);
+						string? Group = BadgeDefinitionObject.GetValue("Group", null);
 						string Color = BadgeDefinitionObject.GetValue("Color", "#909090");
 						string HoverColor = BadgeDefinitionObject.GetValue("HoverColor", "#b0b0b0");
-						string Url = BadgeDefinitionObject.GetValue("Url", null);
-						string Arguments = BadgeDefinitionObject.GetValue("Arguments", null);
+						string? Url = BadgeDefinitionObject.GetValue("Url", null);
+						string? Arguments = BadgeDefinitionObject.GetValue("Arguments", null);
 						if (!String.IsNullOrEmpty(Name) && !String.IsNullOrEmpty(Pattern))
 						{
-							foreach (Match MatchResult in Regex.Matches(Description, Pattern, RegexOptions.Multiline))
+							foreach (Match? MatchResult in Regex.Matches(Description, Pattern, RegexOptions.Multiline))
 							{
-								Color BadgeColor = System.Drawing.ColorTranslator.FromHtml(Color);
-								Color HoverBadgeColor = System.Drawing.ColorTranslator.FromHtml(HoverColor);
+								if (MatchResult != null)
+								{
+									Color BadgeColor = System.Drawing.ColorTranslator.FromHtml(Color);
+									Color HoverBadgeColor = System.Drawing.ColorTranslator.FromHtml(HoverColor);
 
-								string UniqueId = String.IsNullOrEmpty(Url) ? null : String.Format("Description:{0}:{1}", Change.Number, Badges.Count);
+									string? UniqueId = String.IsNullOrEmpty(Url) ? null : String.Format("Description:{0}:{1}", Change.Number, Badges.Count);
 
-								string ExpandedUrl = ReplaceRegexMatches(Url, MatchResult);
-								string ExpandedArguments = ReplaceRegexMatches(Arguments, MatchResult);
+									string? ExpandedUrl = ReplaceRegexMatches(Url, MatchResult);
+									string? ExpandedArguments = ReplaceRegexMatches(Arguments, MatchResult);
 
-								Action ClickHandler;
-								if (String.IsNullOrEmpty(ExpandedUrl))
-								{
-									ClickHandler = null;
-								}
-								else if (String.IsNullOrEmpty(ExpandedArguments))
-								{
-									ClickHandler = () => SafeProcessStart(ExpandedUrl);
-								}
-								else
-								{
-									ClickHandler = () => SafeProcessStart(ExpandedUrl, ExpandedArguments);
-								}
+									Action? ClickHandler;
+									if (String.IsNullOrEmpty(ExpandedUrl))
+									{
+										ClickHandler = null;
+									}
+									else if (String.IsNullOrEmpty(ExpandedArguments))
+									{
+										ClickHandler = () => SafeProcessStart(ExpandedUrl);
+									}
+									else
+									{
+										ClickHandler = () => SafeProcessStart(ExpandedUrl, ExpandedArguments);
+									}
 
-								Badges.Add(new BadgeInfo(ReplaceRegexMatches(Name, MatchResult), Group, UniqueId, BadgeColor, HoverBadgeColor, ClickHandler));
+									Badges.Add(new BadgeInfo(ReplaceRegexMatches(Name, MatchResult), Group, UniqueId, BadgeColor, HoverBadgeColor, ClickHandler));
+								}
 							}
 						}
 					}
@@ -2632,7 +2632,8 @@ namespace UnrealGameSync
 			UpdateStatusPanel();
 		}
 
-		private string ReplaceRegexMatches(string Text, Match MatchResult)
+		[return: NotNullIfNotNull("Text")]
+		private string? ReplaceRegexMatches(string? Text, Match MatchResult)
 		{
 			if (Text != null)
 			{
@@ -2649,7 +2650,7 @@ namespace UnrealGameSync
 		{
 			List<BadgeInfo> Badges = new List<BadgeInfo>();
 
-			PerforceChangeDetails Details;
+			PerforceChangeDetails? Details;
 			if (PerforceMonitor.TryGetChangeDetails(ChangeNumber, out Details))
 			{
 				if (Details.bContainsCode)
@@ -2670,7 +2671,7 @@ namespace UnrealGameSync
 			return Badges;
 		}
 
-		private bool TryGetProjectSetting(ConfigFile ProjectConfigFile, string Name, out string Value)
+		private bool TryGetProjectSetting(ConfigFile ProjectConfigFile, string Name, [NotNullWhen(true)] out string? Value)
 		{
 			string Path = SelectedProjectIdentifier;
 			for(; ;)
@@ -2678,7 +2679,7 @@ namespace UnrealGameSync
 				ConfigSection ProjectSection = ProjectConfigFile.FindSection(Path);
 				if (ProjectSection != null)
 				{
-					string NewValue = ProjectSection.GetValue(Name, null);
+					string? NewValue = ProjectSection.GetValue(Name, null);
 					if (NewValue != null)
 					{
 						Value = NewValue;
@@ -2698,7 +2699,7 @@ namespace UnrealGameSync
 			ConfigSection DefaultSection = ProjectConfigFile.FindSection("Default");
 			if (DefaultSection != null)
 			{
-				string NewValue = DefaultSection.GetValue(Name, null);
+				string? NewValue = DefaultSection.GetValue(Name, null);
 				if (NewValue != null)
 				{
 					Value = NewValue;
@@ -2710,9 +2711,9 @@ namespace UnrealGameSync
 			return false;
 		}
 
-		private bool TryGetProjectSetting(ConfigFile ProjectConfigFile, string Name, out string[] Values)
+		private bool TryGetProjectSetting(ConfigFile ProjectConfigFile, string Name, [NotNullWhen(true)] out string[]? Values)
 		{
-			string ValueList;
+			string? ValueList;
 			if (TryGetProjectSetting(ProjectConfigFile, Name, out ValueList))
 			{
 				Values = ValueList.Split('\n').Select(x => x.Trim()).Where(x => x.Length > 0).ToArray();
@@ -2725,9 +2726,9 @@ namespace UnrealGameSync
 			}
 		}
 
-		private bool TryGetProjectSetting(ConfigFile ProjectConfigFile, string Name, string LegacyName, out string Value)
+		private bool TryGetProjectSetting(ConfigFile ProjectConfigFile, string Name, string LegacyName, [NotNullWhen(true)] out string? Value)
 		{
-			string NewValue;
+			string? NewValue;
 			if (TryGetProjectSetting(ProjectConfigFile, Name, out NewValue))
 			{
 				Value = NewValue;
@@ -2745,7 +2746,7 @@ namespace UnrealGameSync
 			return false;
 		}
 
-		private List<BadgeInfo> CreateBuildBadges(int ChangeNumber, EventSummary Summary)
+		private List<BadgeInfo> CreateBuildBadges(int ChangeNumber, EventSummary? Summary)
 		{
 			List<BadgeInfo> Badges = new List<BadgeInfo>();
 
@@ -2761,7 +2762,7 @@ namespace UnrealGameSync
 				// Add all the badges, sorted by group
 				foreach (KeyValuePair<string, string> BadgeNameAndGroup in BadgeNameAndGroupPairs)
 				{
-					BadgeData BadgeData;
+					BadgeData? BadgeData;
 					BadgeNameToBuildData.TryGetValue(BadgeNameAndGroup.Key, out BadgeData);
 
 					BadgeInfo BadgeInfo = CreateBadge(ChangeNumber, BadgeNameAndGroup.Key, BadgeNameAndGroup.Value, BadgeData);
@@ -2779,7 +2780,7 @@ namespace UnrealGameSync
 			return Badges;
 		}
 
-		private BadgeInfo CreateBadge(int ChangeNumber, string BadgeName, string BadgeGroup, BadgeData BadgeData)
+		private BadgeInfo CreateBadge(int ChangeNumber, string BadgeName, string BadgeGroup, BadgeData? BadgeData)
 		{
 			string BadgeLabel = BadgeName;
 			Color BadgeColor = Color.FromArgb(0, Color.White);
@@ -2792,7 +2793,7 @@ namespace UnrealGameSync
 
 			Color HoverBadgeColor = Color.FromArgb(BadgeColor.A, Math.Min(BadgeColor.R + 32, 255), Math.Min(BadgeColor.G + 32, 255), Math.Min(BadgeColor.B + 32, 255));
 
-			Action ClickHandler;
+			Action? ClickHandler;
 			if (BadgeData == null || String.IsNullOrEmpty(BadgeData.Url))
 			{
 				ClickHandler = null;
@@ -2806,7 +2807,7 @@ namespace UnrealGameSync
 			return new BadgeInfo(BadgeLabel, BadgeGroup, UniqueId, BadgeColor, HoverBadgeColor, ClickHandler);
 		}
 
-		private Dictionary<string, List<BadgeInfo>> CreateCustomBadges(int ChangeNumber, EventSummary Summary)
+		private Dictionary<string, List<BadgeInfo>> CreateCustomBadges(int ChangeNumber, EventSummary? Summary)
 		{
 			Dictionary<string, List<BadgeInfo>> ColumnNameToBadges = new Dictionary<string, List<BadgeInfo>>();
 			if (Summary != null && Summary.Badges.Count > 0)
@@ -2903,7 +2904,7 @@ namespace UnrealGameSync
 			if (!BadgeLabelToSize.TryGetValue(BadgeText, out BadgeSize))
 			{
 				Size LabelSize = TextRenderer.MeasureText(BadgeText, BadgeFont);
-				int BadgeHeight = BadgeFont.Height + 1;
+				int BadgeHeight = BadgeFont!.Height + 1;
 
 				BadgeSize = new Size(LabelSize.Width + BadgeHeight - 4, BadgeHeight);
 				BadgeLabelToSize[BadgeText] = BadgeSize;
@@ -2911,7 +2912,7 @@ namespace UnrealGameSync
 			return BadgeSize;
 		}
 
-		private static bool GetLastUpdateMessage(WorkspaceUpdateResult Result, string ResultMessage, out string Message)
+		private static bool GetLastUpdateMessage(WorkspaceUpdateResult Result, string? ResultMessage, [NotNullWhen(true)] out string? Message)
 		{
 			if (Result != WorkspaceUpdateResult.Success && ResultMessage != null)
 			{
@@ -2921,7 +2922,7 @@ namespace UnrealGameSync
 			return GetGenericLastUpdateMessage(Result, out Message);
 		}
 
-		private static bool GetGenericLastUpdateMessage(WorkspaceUpdateResult Result, out string Message)
+		private static bool GetGenericLastUpdateMessage(WorkspaceUpdateResult Result, [NotNullWhen(true)] out string? Message)
 		{
 			switch (Result)
 			{
@@ -2972,7 +2973,7 @@ namespace UnrealGameSync
 
 		private static void AppendItemList(StringBuilder Builder, string Separator, string Format, IEnumerable<string> Items)
 		{
-			string ItemList = FormatItemList(Format, Items);
+			string? ItemList = FormatItemList(Format, Items);
 			if (ItemList != null)
 			{
 				if (Builder.Length > 0)
@@ -2983,7 +2984,7 @@ namespace UnrealGameSync
 			}
 		}
 
-		private static string FormatItemList(string Format, IEnumerable<string> Items)
+		private static string? FormatItemList(string Format, IEnumerable<string> Items)
 		{
 			string[] ItemsArray = Items.Distinct().ToArray();
 			if (ItemsArray.Length == 0)
@@ -3102,7 +3103,7 @@ namespace UnrealGameSync
 
 		private class ReceiptJsonObject
 		{
-			public string Launch { get; set; }
+			public string? Launch { get; set; }
 		}
 
 		private FileReference GetEditorExePath(BuildConfig Config)
@@ -3117,9 +3118,9 @@ namespace UnrealGameSync
 					try
 					{
 						string Text = FileReference.ReadAllText(ReceiptFileName);
-						ReceiptJsonObject Receipt = JsonSerializer.Deserialize<ReceiptJsonObject>(Text, Program.DefaultJsonSerializerOptions);
+						ReceiptJsonObject Receipt = JsonSerializer.Deserialize<ReceiptJsonObject>(Text, Utility.DefaultJsonSerializerOptions);
 
-						string LaunchFileName = Receipt.Launch;
+						string? LaunchFileName = Receipt.Launch;
 						if (LaunchFileName != null)
 						{
 							LaunchFileName = LaunchFileName.Replace("$(EngineDir)", DirectoryReference.Combine(BranchDirectoryName, "Engine").FullName);
@@ -3251,7 +3252,7 @@ namespace UnrealGameSync
 			return PossiblePaths;
 		}
 
-		private void TimerCallback(object Sender, EventArgs Args)
+		private void TimerCallback(object? Sender, EventArgs Args)
 		{
 			Tuple<string, float> Progress = Workspace.CurrentProgress;
 			if (Progress != null && Progress.Item2 > 0.0f)
@@ -3269,15 +3270,15 @@ namespace UnrealGameSync
 			BuildList.Refresh();
 		}
 
-		private string GetDefaultIssueFilter()
+		private string? GetDefaultIssueFilter()
 		{
-			string BuildHealthProject;
+			string? BuildHealthProject;
 			if (TryGetProjectSetting(PerforceMonitor.LatestProjectConfigFile, "BuildHealthProject", out BuildHealthProject))
 			{
 				return BuildHealthProject;
 			}
 
-			string DefaultIssueFilter;
+			string? DefaultIssueFilter;
 			if (TryGetProjectSetting(PerforceMonitor.LatestProjectConfigFile, "DefaultIssueFilter", out DefaultIssueFilter))
 			{
 				return DefaultIssueFilter;
@@ -3290,13 +3291,13 @@ namespace UnrealGameSync
 		{
 			Dictionary<string, Func<IssueData, bool>> CustomFilters = new Dictionary<string, Func<IssueData, bool>>();
 
-			string BuildHealthProject;
+			string? BuildHealthProject;
 			if (TryGetProjectSetting(PerforceMonitor.LatestProjectConfigFile, "BuildHealthProject", out BuildHealthProject))
 			{
 				CustomFilters[BuildHealthProject] = x => !String.IsNullOrEmpty(x.Project) && x.Project.Equals(BuildHealthProject, StringComparison.OrdinalIgnoreCase);
 			}
 
-			string[] Filters;
+			string[]? Filters;
 			if (TryGetProjectSetting(PerforceMonitor.LatestProjectConfigFile, "IssueFilters", out Filters))
 			{
 				foreach (string Filter in Filters)
@@ -3305,12 +3306,12 @@ namespace UnrealGameSync
 					{
 						ConfigObject Config = new ConfigObject(Filter);
 
-						string Name = Config.GetValue("Name", null);
-						string Pattern = Config.GetValue("Pattern", null);
+						string? Name = Config.GetValue("Name", null);
+						string? Pattern = Config.GetValue("Pattern", null);
 
 						if(Name != null && Pattern != null)
 						{
-							string DepotName;
+							string? DepotName;
 							if (PerforceUtils.TryGetDepotName(StreamName, out DepotName))
 							{
 								Pattern = Regex.Replace(Pattern, @"\$\(Depot\)", DepotName, RegexOptions.IgnoreCase);
@@ -3344,16 +3345,18 @@ namespace UnrealGameSync
 
 		void RunTool(ToolDefinition Tool, ToolLink Link)
 		{
-			DirectoryReference ToolDir = Owner.ToolUpdateMonitor.GetToolPath(Tool.Name);
-
-			Dictionary<string, string> Variables = GetWorkspaceVariables(Workspace.CurrentChangeNumber);
-			Variables["ToolDir"] = ToolDir.FullName;
+			DirectoryReference? ToolDir = Owner.ToolUpdateMonitor.GetToolPath(Tool.Name);
+			if (ToolDir != null)
+			{
+				Dictionary<string, string> Variables = GetWorkspaceVariables(Workspace.CurrentChangeNumber);
+				Variables["ToolDir"] = ToolDir.FullName;
 
-			string FileName = Utility.ExpandVariables(Link.FileName, Variables);
-			string Arguments = Utility.ExpandVariables(Link.Arguments, Variables);
-			string WorkingDir = Utility.ExpandVariables(Link.WorkingDir ?? ToolDir.FullName, Variables);
+				string FileName = Utility.ExpandVariables(Link.FileName, Variables);
+				string Arguments = Utility.ExpandVariables(Link.Arguments ?? String.Empty, Variables);
+				string WorkingDir = Utility.ExpandVariables(Link.WorkingDir ?? ToolDir.FullName, Variables);
 
-			SafeProcessStart(FileName, Arguments, WorkingDir);
+				SafeProcessStart(FileName, Arguments, WorkingDir);
+			}
 		}
 
 		private void UpdateStatusPanel_CrossThread()
@@ -3479,7 +3482,7 @@ namespace UnrealGameSync
 					ProgramsLine.AddText("  |  ");
 				}
 
-				string[] SdkInfoEntries;
+				string[]? SdkInfoEntries;
 				if (TryGetProjectSetting(PerforceMonitor.LatestProjectConfigFile, "SdkInfo", out SdkInfoEntries))
 				{
 					ProgramsLine.AddLink("SDK Info", FontStyle.Regular, () => { ShowRequiredSdkInfo(); });
@@ -3547,7 +3550,7 @@ namespace UnrealGameSync
 				// Get the summary of the last sync
 				if (WorkspaceState.LastSyncChangeNumber > 0)
 				{
-					string SummaryText;
+					string? SummaryText;
 					if (WorkspaceState.LastSyncChangeNumber == Workspace.CurrentChangeNumber && WorkspaceState.LastSyncResult == WorkspaceUpdateResult.Success && WorkspaceState.LastSyncTime.HasValue)
 					{
 						Lines.Add(new StatusLine() { LineHeight = 0.5f });
@@ -3587,26 +3590,26 @@ namespace UnrealGameSync
 				}
 			}
 
-			StatusLine Caption = null;
+			StatusLine? Caption = null;
 			if (StreamName != null && !Workspace.IsBusy())
 			{
 				Caption = new StatusLine();
 				Caption.AddLink(StreamName + "\u25BE", FontStyle.Bold, (P, R) => { SelectOtherStream(R); });
 			}
 
-			StatusLine Alert = null;
+			StatusLine? Alert = null;
 			Color? TintColor = null;
 
 			ConfigFile ProjectConfigFile = PerforceMonitor.LatestProjectConfigFile;
 			if (ProjectConfigFile != null)
 			{
-				string Message;
+				string? Message;
 				if (TryGetProjectSetting(ProjectConfigFile, "Message", out Message))
 				{
 					Alert = CreateStatusLineFromMarkdown(Message);
 				}
 
-				string StatusPanelColor;
+				string? StatusPanelColor;
 				if (TryGetProjectSetting(ProjectConfigFile, "StatusPanelColor", out StatusPanelColor))
 				{
 					TintColor = System.Drawing.ColorTranslator.FromHtml(StatusPanelColor);
@@ -3638,8 +3641,9 @@ namespace UnrealGameSync
 				BuildHealthContextMenu.Items.RemoveAt(MinSeparatorIdx + 1);
 			}
 
-			Func<IssueData, bool> Predicate;
-			if(!GetCustomIssueFilters().TryGetValue(GetDefaultIssueFilter(), out Predicate))
+			Func<IssueData, bool>? Predicate = null;
+			string? DefaultIssueFilter = GetDefaultIssueFilter();
+			if (DefaultIssueFilter == null || !GetCustomIssueFilters().TryGetValue(DefaultIssueFilter, out Predicate))
 			{
 				Predicate = x => true;
 			}
@@ -3737,14 +3741,14 @@ namespace UnrealGameSync
 
 		private void BuildHealthContextMenu_Browse_Click(object sender, EventArgs e)
 		{
-			string DefaultFilter = GetDefaultIssueFilter();
+			string? DefaultFilter = GetDefaultIssueFilter();
 			Dictionary<string, Func<IssueData, bool>> CustomFilters = GetCustomIssueFilters();
 			IssueBrowserWindow.Show(ParentForm, IssueMonitor, PerforceSettings, GetServerTimeOffset(), ServiceProvider, StreamName, CustomFilters, DefaultFilter);
 		}
 
 		private void BuildHealthContextMenu_Settings_Click(object sender, EventArgs e)
 		{
-			string BuildHealthProject;
+			string? BuildHealthProject;
 			TryGetProjectSetting(PerforceMonitor.LatestProjectConfigFile, "BuildHealthProject", out BuildHealthProject);
 
 			IssueSettingsWindow IssueSettings = new IssueSettingsWindow(Settings, BuildHealthProject ?? "");
@@ -3756,11 +3760,11 @@ namespace UnrealGameSync
 
 		private void ShowRequiredSdkInfo()
 		{
-			string[] SdkInfoEntries;
+			string[]? SdkInfoEntries;
 			if (TryGetProjectSetting(PerforceMonitor.LatestProjectConfigFile, "SdkInfo", out SdkInfoEntries))
 			{
 				Dictionary<string, string> Variables = GetWorkspaceVariables(-1);
-				SdkInfoWindow Window = new SdkInfoWindow(SdkInfoEntries, Variables, BadgeFont);
+				SdkInfoWindow Window = new SdkInfoWindow(SdkInfoEntries, Variables, BadgeFont!);
 				Window.ShowDialog();
 			}
 		}
@@ -3866,7 +3870,7 @@ namespace UnrealGameSync
 			bool bShownContextMenu = false;
 			if (StreamName != null)
 			{
-				IReadOnlyList<string> OtherStreamNames = Workspace.ProjectStreamFilter;
+				IReadOnlyList<string>? OtherStreamNames = Workspace.ProjectStreamFilter;
 				if (OtherStreamNames != null)
 				{
 					StreamContextMenu.Items.Clear();
@@ -3909,7 +3913,7 @@ namespace UnrealGameSync
 
 		private void SelectOtherStreamDialog()
 		{
-			string NewStreamName;
+			string? NewStreamName;
 			if (SelectStreamWindow.ShowModal(this, PerforceSettings, StreamName, ServiceProvider, out NewStreamName))
 			{
 				SelectStream(NewStreamName);
@@ -3932,14 +3936,14 @@ namespace UnrealGameSync
 						{
 							if (Idx == 1 || !await Perforce.OpenedAsync(OpenedOptions.None, FileSpecList.Any, CancellationToken).AnyAsync())
 							{
-								await Perforce.SwitchClientToStreamAsync(Perforce.Settings.ClientName, NewStreamName, SwitchClientOptions.IgnoreOpenFiles, CancellationToken);
+								await Perforce.SwitchClientToStreamAsync(Perforce.Settings.ClientName!, NewStreamName, SwitchClientOptions.IgnoreOpenFiles, CancellationToken);
 								return true;
 							}
 							return false;
 						};
 
-						ModalTask<bool> SwitchTask = PerforceModalTask.Execute<bool>(this, "Switching streams", "Please wait...", PerforceSettings, SwitchFunc, Logger);
-						if (!SwitchTask.Succeeded)
+						ModalTask<bool>? SwitchTask = PerforceModalTask.Execute<bool>(this, "Switching streams", "Please wait...", PerforceSettings, SwitchFunc, Logger);
+						if (SwitchTask == null || !SwitchTask.Succeeded)
 						{
 							break;
 						}
@@ -3957,10 +3961,10 @@ namespace UnrealGameSync
 
 		private void ViewLastSyncStatus()
 		{
-			string SummaryText;
+			string? SummaryText;
 			if (GetLastUpdateMessage(WorkspaceState.LastSyncResult, WorkspaceState.LastSyncResultMessage, out SummaryText))
 			{
-				string CaptionText;
+				string? CaptionText;
 				if (!GetGenericLastUpdateMessage(WorkspaceState.LastSyncResult, out CaptionText))
 				{
 					CaptionText = "Sync error";
@@ -3994,14 +3998,17 @@ namespace UnrealGameSync
 
 			PendingSelectedChangeNumber = -1;
 
-			foreach (ListViewItem Item in BuildList.Items)
+			foreach (ListViewItem? Item in BuildList.Items)
 			{
-				ChangesRecord Summary = (ChangesRecord)Item.Tag;
-				if (Summary != null && Summary.Number <= ChangeNumber)
+				if (Item != null)
 				{
-					Item.Selected = true;
-					Item.EnsureVisible();
-					return;
+					ChangesRecord Summary = (ChangesRecord)Item.Tag;
+					if (Summary != null && Summary.Number <= ChangeNumber)
+					{
+						Item.Selected = true;
+						Item.EnsureVisible();
+						return;
+					}
 				}
 			}
 
@@ -4114,7 +4121,7 @@ namespace UnrealGameSync
 								Point BuildListLocation = GetBadgeListLocation(LayoutInfo.BuildBadges, HitTest.SubItem.Bounds, HorizontalAlign.Center, VerticalAlignment.Middle);
 								BuildListLocation.X = Math.Max(BuildListLocation.X, HitTest.SubItem.Bounds.Left);
 
-								BadgeInfo BadgeInfo = HitTestBadge(Args.Location, LayoutInfo.BuildBadges, BuildListLocation);
+								BadgeInfo? BadgeInfo = HitTestBadge(Args.Location, LayoutInfo.BuildBadges, BuildListLocation);
 								if (BadgeInfo != null && BadgeInfo.ClickHandler != null)
 								{
 									BadgeInfo.ClickHandler();
@@ -4128,12 +4135,12 @@ namespace UnrealGameSync
 							{
 								ChangeLayoutInfo LayoutInfo = GetChangeLayoutInfo((ChangesRecord)HitTest.Item.Tag);
 
-								List<BadgeInfo> Badges;
+								List<BadgeInfo>? Badges;
 								if (LayoutInfo.CustomBadges.TryGetValue(CustomColumn.Text, out Badges) && Badges.Count > 0)
 								{
 									Point ListLocation = GetBadgeListLocation(Badges, HitTest.SubItem.Bounds, HorizontalAlign.Center, VerticalAlignment.Middle);
 
-									BadgeInfo BadgeInfo = HitTestBadge(Args.Location, Badges, ListLocation);
+									BadgeInfo? BadgeInfo = HitTestBadge(Args.Location, Badges, ListLocation);
 									if (BadgeInfo != null && BadgeInfo.ClickHandler != null)
 									{
 										BadgeInfo.ClickHandler();
@@ -4168,7 +4175,7 @@ namespace UnrealGameSync
 						BuildListContextMenu_StartInvestigating.Visible = !EventMonitor.IsUnderInvestigationByCurrentUser(ContextMenuChange.Number);
 						BuildListContextMenu_FinishInvestigating.Visible = EventMonitor.IsUnderInvestigation(ContextMenuChange.Number);
 
-						string CommentText;
+						string? CommentText;
 						bool bHasExistingComment = EventMonitor.GetCommentByCurrentUser(ContextMenuChange.Number, out CommentText);
 						BuildListContextMenu_LeaveComment.Visible = !bHasExistingComment;
 						BuildListContextMenu_EditComment.Visible = bHasExistingComment;
@@ -4199,7 +4206,7 @@ namespace UnrealGameSync
 						BuildListContextMenu_MarkBad.Visible = !bIsBisectMode;
 						BuildListContextMenu_WithdrawReview.Visible = !bIsBisectMode;
 
-						EventSummary Summary = EventMonitor.GetSummaryForChange(ContextMenuChange.Number);
+						EventSummary? Summary = EventMonitor.GetSummaryForChange(ContextMenuChange.Number);
 						bool bStarred = (Summary != null && Summary.LastStarReview != null && Summary.LastStarReview.Type == EventType.Starred);
 						BuildListContextMenu_AddStar.Visible = !bStarred;
 						BuildListContextMenu_RemoveStar.Visible = bStarred;
@@ -4219,7 +4226,7 @@ namespace UnrealGameSync
 							CustomToolEnd--;
 						}
 
-						ConfigFile ProjectConfigFile = PerforceMonitor.LatestProjectConfigFile;
+						ConfigFile? ProjectConfigFile = PerforceMonitor.LatestProjectConfigFile;
 						if (ProjectConfigFile != null)
 						{
 							Dictionary<string, string> Variables = GetWorkspaceVariables(ContextMenuChange.Number);
@@ -4229,9 +4236,9 @@ namespace UnrealGameSync
 							{
 								ConfigObject Object = new ConfigObject(ChangeContextMenuEntry);
 
-								string Label = Object.GetValue("Label");
-								string Execute = Object.GetValue("Execute");
-								string Arguments = Object.GetValue("Arguments");
+								string? Label = Object.GetValue("Label");
+								string? Execute = Object.GetValue("Execute");
+								string? Arguments = Object.GetValue("Arguments");
 
 								if (Label != null && Execute != null)
 								{
@@ -4249,8 +4256,8 @@ namespace UnrealGameSync
 
 						BuildListContextMenu_CustomTool_End.Visible = (CustomToolEnd > CustomToolStart);
 
-						string SwarmURL;
-						BuildListContextMenu_ViewInSwarm.Visible = TryGetProjectSetting(ProjectConfigFile, "SwarmURL", out SwarmURL);
+						string? SwarmURL;
+						BuildListContextMenu_ViewInSwarm.Visible = ProjectConfigFile != null && TryGetProjectSetting(ProjectConfigFile, "SwarmURL", out SwarmURL);
 							
 						BuildListContextMenu.Show(BuildList, Args.Location);
 					}
@@ -4270,7 +4277,7 @@ namespace UnrealGameSync
 				BuildList.Invalidate();
 			}
 
-			string NewHoverBadgeUniqueId = null;
+			string? NewHoverBadgeUniqueId = null;
 			if (HitTest.Item != null && HitTest.Item.Tag is ChangesRecord)
 			{
 				int ColumnIndex = HitTest.Item.SubItems.IndexOf(HitTest.SubItem);
@@ -4291,7 +4298,7 @@ namespace UnrealGameSync
 						Point BuildListLocation = GetBadgeListLocation(LayoutInfo.BuildBadges, HitTest.SubItem.Bounds, HorizontalAlign.Center, VerticalAlignment.Middle);
 						BuildListLocation.X = Math.Max(BuildListLocation.X, HitTest.SubItem.Bounds.Left);
 
-						BadgeInfo Badge = HitTestBadge(e.Location, LayoutInfo.BuildBadges, BuildListLocation);
+						BadgeInfo? Badge = HitTestBadge(e.Location, LayoutInfo.BuildBadges, BuildListLocation);
 						NewHoverBadgeUniqueId = (Badge != null) ? Badge.UniqueId : null;
 
 						if (HoverBadgeUniqueId != NewHoverBadgeUniqueId)
@@ -4313,7 +4320,7 @@ namespace UnrealGameSync
 
 					ChangeLayoutInfo LayoutInfo = GetChangeLayoutInfo((ChangesRecord)HitTest.Item.Tag);
 
-					List<BadgeInfo> Badges;
+					List<BadgeInfo>? Badges;
 					if (LayoutInfo.CustomBadges.TryGetValue(Column.Text, out Badges) && Badges.Count > 0)
 					{
 						Point ListLocation = GetBadgeListLocation(Badges, HitTest.SubItem.Bounds, HorizontalAlign.Center, VerticalAlignment.Middle);
@@ -4339,7 +4346,7 @@ namespace UnrealGameSync
 			}
 		}
 
-		private BadgeInfo HitTestBadge(Point Location, List<BadgeInfo> BadgeList, Point ListLocation)
+		private BadgeInfo? HitTestBadge(Point Location, List<BadgeInfo> BadgeList, Point ListLocation)
 		{
 			foreach (BadgeInfo Badge in BadgeList)
 			{
@@ -4482,14 +4489,14 @@ namespace UnrealGameSync
 
 		private IReadOnlyList<IArchiveInfo> GetArchives()
 		{
-			IReadOnlyList<IArchiveInfo> AvailableArchives = JupiterMonitor?.AvailableArchives;
+			IReadOnlyList<IArchiveInfo>? AvailableArchives = JupiterMonitor?.AvailableArchives;
 			if (AvailableArchives != null && AvailableArchives.Count != 0)
 				return AvailableArchives;
 
 			// if jupiter had no archives we fallback to the perforce monitor
 			if (PerforceMonitor != null)
 			{
-				return PerforceMonitor?.AvailableArchives;
+				return PerforceMonitor.AvailableArchives;
 			}
 
 			return new List<IArchiveInfo>();
@@ -4541,7 +4548,7 @@ namespace UnrealGameSync
 			}
 		}
 
-		private void BuildList_FontChanged(object sender, EventArgs e)
+		private void BuildList_FontChanged(object? sender, EventArgs e)
 		{
 			if (BuildFont != null)
 			{
@@ -4578,7 +4585,7 @@ namespace UnrealGameSync
 			}
 		}
 
-		public void SyncChange(int ChangeNumber, bool bSyncOnly, WorkspaceUpdateCallback Callback)
+		public void SyncChange(int ChangeNumber, bool bSyncOnly, WorkspaceUpdateCallback? Callback)
 		{
 			if(Workspace != null)
 			{
@@ -4593,7 +4600,7 @@ namespace UnrealGameSync
 			SyncLatestChange(null);
 		}
 
-		public void SyncLatestChange(WorkspaceUpdateCallback Callback)
+		public void SyncLatestChange(WorkspaceUpdateCallback? Callback)
 		{
 			if (Workspace != null)
 			{
@@ -4615,32 +4622,50 @@ namespace UnrealGameSync
 
 		private void BuildListContextMenu_MarkGood_Click(object sender, EventArgs e)
 		{
-			EventMonitor.PostEvent(ContextMenuChange.Number, EventType.Good);
+			if (ContextMenuChange != null)
+			{
+				EventMonitor.PostEvent(ContextMenuChange.Number, EventType.Good);
+			}
 		}
 
 		private void BuildListContextMenu_MarkBad_Click(object sender, EventArgs e)
 		{
-			EventMonitor.PostEvent(ContextMenuChange.Number, EventType.Bad);
+			if (ContextMenuChange != null)
+			{
+				EventMonitor.PostEvent(ContextMenuChange.Number, EventType.Bad);
+			}
 		}
 
 		private void BuildListContextMenu_WithdrawReview_Click(object sender, EventArgs e)
 		{
-			EventMonitor.PostEvent(ContextMenuChange.Number, EventType.Unknown);
+			if (ContextMenuChange != null)
+			{
+				EventMonitor.PostEvent(ContextMenuChange.Number, EventType.Unknown);
+			}
 		}
 
 		private void BuildListContextMenu_Sync_Click(object sender, EventArgs e)
 		{
-			StartSync(ContextMenuChange.Number);
+			if (ContextMenuChange != null)
+			{
+				StartSync(ContextMenuChange.Number);
+			}
 		}
 
 		private void BuildListContextMenu_SyncContentOnly_Click(object sender, EventArgs e)
 		{
-			StartWorkspaceUpdate(ContextMenuChange.Number, WorkspaceUpdateOptions.Sync | WorkspaceUpdateOptions.ContentOnly);
+			if (ContextMenuChange != null)
+			{
+				StartWorkspaceUpdate(ContextMenuChange.Number, WorkspaceUpdateOptions.Sync | WorkspaceUpdateOptions.ContentOnly);
+			}
 		}
 
 		private void BuildListContextMenu_SyncOnlyThisChange_Click(object sender, EventArgs e)
 		{
-			StartWorkspaceUpdate(ContextMenuChange.Number, WorkspaceUpdateOptions.SyncSingleChange);
+			if (ContextMenuChange != null)
+			{
+				StartWorkspaceUpdate(ContextMenuChange.Number, WorkspaceUpdateOptions.SyncSingleChange);
+			}
 		}
 
 		private void BuildListContextMenu_Build_Click(object sender, EventArgs e)
@@ -4665,44 +4690,59 @@ namespace UnrealGameSync
 
 		private void BuildListContextMenu_MoreInfo_Click(object sender, EventArgs e)
 		{
-			Utility.SpawnP4VC(String.Format("{0} change {1}", PerforceSettings.GetArgumentsForExternalProgram(true), ContextMenuChange.Number));
+			if (ContextMenuChange != null)
+			{
+				Program.SpawnP4VC(String.Format("{0} change {1}", PerforceSettings.GetArgumentsForExternalProgram(true), ContextMenuChange.Number));
+			}
 		}
 
 		private void BuildListContextMenu_ViewInSwarm_Click(object sender, EventArgs e)
 		{
-			string SwarmURL;
-			if (TryGetProjectSetting(PerforceMonitor.LatestProjectConfigFile, "SwarmURL", out SwarmURL))
-			{
-				Utility.OpenUrl(String.Format("{0}/changes/{1}", SwarmURL, ContextMenuChange.Number));
-			}
-			else
+			if (ContextMenuChange != null)
 			{
-				MessageBox.Show("Swarm URL is not configured.");
+				string? SwarmURL;
+				if (TryGetProjectSetting(PerforceMonitor.LatestProjectConfigFile, "SwarmURL", out SwarmURL))
+				{
+					Utility.OpenUrl(String.Format("{0}/changes/{1}", SwarmURL, ContextMenuChange.Number));
+				}
+				else
+				{
+					MessageBox.Show("Swarm URL is not configured.");
+				}
 			}
 		}
 
 		private void BuildListContextMenu_CopyChangelistNumber_Click(object sender, EventArgs e)
 		{
-			System.Windows.Forms.Clipboard.SetText(ContextMenuChange.Number.ToString());
+			if (ContextMenuChange != null)
+			{
+				System.Windows.Forms.Clipboard.SetText(ContextMenuChange.Number.ToString());
+			}
 		}
 
 		private void BuildListContextMenu_AddStar_Click(object sender, EventArgs e)
 		{
-			if (MessageBox.Show("Starred builds are meant to convey a stable, verified build to the rest of the team. Do not star a build unless it has been fully tested.\n\nAre you sure you want to star this build?", "Confirm star", MessageBoxButtons.YesNo) == DialogResult.Yes)
+			if (ContextMenuChange != null)
 			{
-				EventMonitor.PostEvent(ContextMenuChange.Number, EventType.Starred);
+				if (MessageBox.Show("Starred builds are meant to convey a stable, verified build to the rest of the team. Do not star a build unless it has been fully tested.\n\nAre you sure you want to star this build?", "Confirm star", MessageBoxButtons.YesNo) == DialogResult.Yes)
+				{
+					EventMonitor.PostEvent(ContextMenuChange.Number, EventType.Starred);
+				}
 			}
 		}
 
 		private void BuildListContextMenu_RemoveStar_Click(object sender, EventArgs e)
 		{
-			EventSummary Summary = EventMonitor.GetSummaryForChange(ContextMenuChange.Number);
-			if (Summary != null && Summary.LastStarReview != null && Summary.LastStarReview.Type == EventType.Starred)
+			if (ContextMenuChange != null)
 			{
-				string Message = String.Format("This change was starred by {0}. Are you sure you want to remove it?", FormatUserName(Summary.LastStarReview.UserName));
-				if (MessageBox.Show(Message, "Confirm removing star", MessageBoxButtons.YesNo) == DialogResult.Yes)
+				EventSummary? Summary = EventMonitor.GetSummaryForChange(ContextMenuChange.Number);
+				if (Summary != null && Summary.LastStarReview != null && Summary.LastStarReview.Type == EventType.Starred)
 				{
-					EventMonitor.PostEvent(ContextMenuChange.Number, EventType.Unstarred);
+					string Message = String.Format("This change was starred by {0}. Are you sure you want to remove it?", FormatUserName(Summary.LastStarReview.UserName));
+					if (MessageBox.Show(Message, "Confirm removing star", MessageBoxButtons.YesNo) == DialogResult.Yes)
+					{
+						EventMonitor.PostEvent(ContextMenuChange.Number, EventType.Unstarred);
+					}
 				}
 			}
 		}
@@ -4714,33 +4754,39 @@ namespace UnrealGameSync
 
 		private void BuildListContextMenu_StartInvestigating_Click(object sender, EventArgs e)
 		{
-			string Message = String.Format("All changes from {0} onwards will be marked as bad while you are investigating an issue.\n\nAre you sure you want to continue?", ContextMenuChange.Number);
-			if (MessageBox.Show(Message, "Confirm investigating", MessageBoxButtons.YesNo) == DialogResult.Yes)
+			if (ContextMenuChange != null)
 			{
-				EventMonitor.StartInvestigating(ContextMenuChange.Number);
+				string Message = String.Format("All changes from {0} onwards will be marked as bad while you are investigating an issue.\n\nAre you sure you want to continue?", ContextMenuChange.Number);
+				if (MessageBox.Show(Message, "Confirm investigating", MessageBoxButtons.YesNo) == DialogResult.Yes)
+				{
+					EventMonitor.StartInvestigating(ContextMenuChange.Number);
+				}
 			}
 		}
 
 		private void BuildListContextMenu_FinishInvestigating_Click(object sender, EventArgs e)
 		{
-			int StartChangeNumber = EventMonitor.GetInvestigationStartChangeNumber(ContextMenuChange.Number);
-			if (StartChangeNumber != -1)
+			if (ContextMenuChange != null)
 			{
-				if (ContextMenuChange.Number > StartChangeNumber)
+				int StartChangeNumber = EventMonitor.GetInvestigationStartChangeNumber(ContextMenuChange.Number);
+				if (StartChangeNumber != -1)
 				{
-					string Message = String.Format("Mark all changes between {0} and {1} as bad?", StartChangeNumber, ContextMenuChange.Number);
-					if (MessageBox.Show(Message, "Finish investigating", MessageBoxButtons.YesNo) == DialogResult.Yes)
+					if (ContextMenuChange.Number > StartChangeNumber)
 					{
-						foreach (ChangesRecord Change in PerforceMonitor.GetChanges())
+						string Message = String.Format("Mark all changes between {0} and {1} as bad?", StartChangeNumber, ContextMenuChange.Number);
+						if (MessageBox.Show(Message, "Finish investigating", MessageBoxButtons.YesNo) == DialogResult.Yes)
 						{
-							if (Change.Number >= StartChangeNumber && Change.Number < ContextMenuChange.Number)
+							foreach (ChangesRecord Change in PerforceMonitor.GetChanges())
 							{
-								EventMonitor.PostEvent(Change.Number, EventType.Bad);
+								if (Change.Number >= StartChangeNumber && Change.Number < ContextMenuChange.Number)
+								{
+									EventMonitor.PostEvent(Change.Number, EventType.Bad);
+								}
 							}
 						}
 					}
+					EventMonitor.FinishInvestigating(ContextMenuChange.Number);
 				}
-				EventMonitor.FinishInvestigating(ContextMenuChange.Number);
 			}
 		}
 
@@ -4801,7 +4847,7 @@ namespace UnrealGameSync
 						}
 					}
 
-					EventSummary Summary = EventMonitor.GetSummaryForChange(Change.Number);
+					EventSummary? Summary = EventMonitor.GetSummaryForChange(Change.Number);
 					if (Summary != null)
 					{
 						StringBuilder SummaryText = new StringBuilder();
@@ -5026,7 +5072,7 @@ namespace UnrealGameSync
 					}
 					else if (ChangeType == LatestChangeType.Good)
 					{
-						EventSummary Summary = EventMonitor.GetSummaryForChange(Change.Number);
+						EventSummary? Summary = EventMonitor.GetSummaryForChange(Change.Number);
 						if (Summary != null && Summary.Verdict == ReviewVerdict.Good && CanSyncChange(Change.Number))
 						{
 							ChangeNumber = FindNewestGoodContentChange(Change.Number);
@@ -5035,7 +5081,7 @@ namespace UnrealGameSync
 					}
 					else if (ChangeType == LatestChangeType.Starred)
 					{
-						EventSummary Summary = EventMonitor.GetSummaryForChange(Change.Number);
+						EventSummary? Summary = EventMonitor.GetSummaryForChange(Change.Number);
 						if (((Summary != null && Summary.LastStarReview != null && Summary.LastStarReview.Type == EventType.Starred) || PromotedChangeNumbers.Contains(Change.Number)) && CanSyncChange(Change.Number))
 						{
 							ChangeNumber = FindNewestGoodContentChange(Change.Number);
@@ -5061,13 +5107,13 @@ namespace UnrealGameSync
 			{
 				int NextChangeNumber = SortedChangeNumbers[NextIndex];
 
-				PerforceChangeDetails Details;
+				PerforceChangeDetails? Details;
 				if (!PerforceMonitor.TryGetChangeDetails(NextChangeNumber, out Details) || Details.bContainsCode)
 				{
 					break;
 				}
 
-				EventSummary Summary = EventMonitor.GetSummaryForChange(NextChangeNumber);
+				EventSummary? Summary = EventMonitor.GetSummaryForChange(NextChangeNumber);
 				if (Summary != null && Summary.Verdict == ReviewVerdict.Bad)
 				{
 					break;
@@ -5085,7 +5131,7 @@ namespace UnrealGameSync
 			{
 				LeaveCommentWindow LeaveComment = new LeaveCommentWindow();
 
-				string CommentText;
+				string? CommentText;
 				if (EventMonitor.GetCommentByCurrentUser(ContextMenuChange.Number, out CommentText))
 				{
 					LeaveComment.CommentTextBox.Text = CommentText;
@@ -5121,20 +5167,20 @@ namespace UnrealGameSync
 				return;
 			}
 
-			string ExtraSafeToDeleteFolders;
+			string? ExtraSafeToDeleteFolders;
 			if (!TryGetProjectSetting(Workspace.ProjectConfigFile, "SafeToDeleteFolders", out ExtraSafeToDeleteFolders))
 			{
 				ExtraSafeToDeleteFolders = "";
 			}
 
-			string ExtraSafeToDeleteExtensions;
+			string? ExtraSafeToDeleteExtensions;
 			if (!TryGetProjectSetting(Workspace.ProjectConfigFile, "SafeToDeleteExtensions", out ExtraSafeToDeleteExtensions))
 			{
 				ExtraSafeToDeleteExtensions = "";
 			}
 
 			string[] CombinedSyncFilter = UserSettings.GetCombinedSyncFilter(GetSyncCategories(), Settings.SyncView, Settings.SyncCategories, WorkspaceSettings.SyncView, WorkspaceSettings.SyncCategoriesDict);
-			List<string> SyncPaths = Workspace.GetSyncPaths(Workspace.Project, WorkspaceSettings.bSyncAllProjects ?? Settings.bSyncAllProjects, CombinedSyncFilter);
+			List<string> SyncPaths = WorkspaceUpdate.GetSyncPaths(Workspace.Project, WorkspaceSettings.bSyncAllProjects ?? Settings.bSyncAllProjects, CombinedSyncFilter);
 
 			CleanWorkspaceWindow.DoClean(ParentForm, PerforceSettings, BranchDirectoryName, Workspace.Project.ClientRootPath, SyncPaths, ExtraSafeToDeleteFolders.Split('\n'), ExtraSafeToDeleteExtensions.Split('\n'), ServiceProvider.GetRequiredService<ILogger<CleanWorkspaceWindow>>());
 		}
@@ -5190,7 +5236,7 @@ namespace UnrealGameSync
 								}
 								else
 								{
-									ContextMenuStrip Menu;
+									ContextMenuStrip? Menu;
 									if (!NameToMenu.TryGetValue(MenuName, out Menu))
 									{
 										Menu = new ContextMenuStrip();
@@ -5248,10 +5294,14 @@ namespace UnrealGameSync
 					Dictionary<string, string> Variables = GetWorkspaceVariables(Workspace.CurrentChangeNumber);
 					if (Step.ToolId != Guid.Empty)
 					{
-						string ToolName = Owner.ToolUpdateMonitor.GetToolName(Step.ToolId);
+						string? ToolName = Owner.ToolUpdateMonitor.GetToolName(Step.ToolId);
 						if (ToolName != null)
 						{
-							Variables["ToolDir"] = Owner.ToolUpdateMonitor.GetToolPath(ToolName).FullName;
+							DirectoryReference? ToolPath = Owner.ToolUpdateMonitor.GetToolPath(ToolName);
+							if (ToolPath != null)
+							{
+								Variables["ToolDir"] = ToolPath.FullName;
+							}
 						}
 					}
 
@@ -5265,7 +5315,7 @@ namespace UnrealGameSync
 		{
 			BuildConfig EditorBuildConfig = GetEditorBuildConfig();
 
-			string SdkInstallerDir;
+			string? SdkInstallerDir;
 			TryGetProjectSetting(PerforceMonitor.LatestProjectConfigFile, "SdkInstallerDir", out SdkInstallerDir);
 
 			Dictionary<string, string> Variables = new Dictionary<string, string>();
@@ -5320,10 +5370,10 @@ namespace UnrealGameSync
 				{
 					if (Step.IsValid())
 					{
-						ConfigObject DefaultObject;
+						ConfigObject? DefaultObject;
 						ProjectBuildStepObjects.TryGetValue(Step.UniqueId, out DefaultObject);
 
-						ConfigObject UserConfigObject = Step.ToConfigObject(DefaultObject);
+						ConfigObject? UserConfigObject = Step.ToConfigObject(DefaultObject);
 						if (UserConfigObject != null && UserConfigObject.Pairs.Any(x => x.Key != "UniqueId"))
 						{
 							ModifiedBuildSteps.Add(UserConfigObject);
@@ -5350,7 +5400,7 @@ namespace UnrealGameSync
 				// Add or apply Object to the list of steps in Steps. Do not modify Object; make a copy first.
 				ConfigObject NewObject = new ConfigObject(Object);
 
-				ConfigObject DefaultObject;
+				ConfigObject? DefaultObject;
 				if (Steps.TryGetValue(UniqueId, out DefaultObject))
 				{
 					NewObject.SetDefaults(DefaultObject);
@@ -5372,7 +5422,7 @@ namespace UnrealGameSync
 
 		private string GetDefaultEditorTargetName()
 		{
-			string EditorTarget;
+			string? EditorTarget;
 			if (!TryGetProjectSetting(PerforceMonitor.LatestProjectConfigFile, "EditorTarget", out EditorTarget))
 			{
 				EditorTarget = "UE4Editor";
@@ -5803,10 +5853,13 @@ namespace UnrealGameSync
 			if (BuildList.SelectedItems.Count >= 2)
 			{
 				Dictionary<int, BisectState> ChangeNumberToBisectState = new Dictionary<int, BisectState>();
-				foreach (ListViewItem SelectedItem in BuildList.SelectedItems)
+				foreach (ListViewItem? SelectedItem in BuildList.SelectedItems)
 				{
-					ChangesRecord Change = (ChangesRecord)SelectedItem.Tag;
-					ChangeNumberToBisectState[Change.Number] = BisectState.Include;
+					if (SelectedItem != null)
+					{
+						ChangesRecord Change = (ChangesRecord)SelectedItem.Tag;
+						ChangeNumberToBisectState[Change.Number] = BisectState.Include;
+					}
 				}
 
 				ChangeNumberToBisectState[ChangeNumberToBisectState.Keys.Min()] = BisectState.Pass;
@@ -5831,9 +5884,9 @@ namespace UnrealGameSync
 
 		private void SetBisectStateForSelection(BisectState State)
 		{
-			foreach (ListViewItem SelectedItem in BuildList.SelectedItems)
+			foreach (ListViewItem? SelectedItem in BuildList.SelectedItems)
 			{
-				ChangesRecord Change = (ChangesRecord)SelectedItem.Tag;
+				ChangesRecord? Change = SelectedItem?.Tag as ChangesRecord;
 				if (Change != null)
 				{
 					WorkspaceState.SetBisectState(Change.Number, State);
@@ -5925,7 +5978,7 @@ namespace UnrealGameSync
 			List<KeyValuePair<string, string>> DisplayBadgeNameAndGroupPairs = new List<KeyValuePair<string, string>>(BadgeNameAndGroupPairs);
 			DisplayBadgeNameAndGroupPairs.AddRange(BadgeNames.Select(x => new KeyValuePair<string, string>(x, "User")));
 
-			string LastGroup = null;
+			string? LastGroup = null;
 			foreach (KeyValuePair<string, string> BadgeNameAndGroupPair in DisplayBadgeNameAndGroupPairs)
 			{
 				if (LastGroup != BadgeNameAndGroupPair.Value)
@@ -6137,7 +6190,7 @@ namespace UnrealGameSync
 					return null;
 				}
 
-				string OverrideValue = Section.GetValue("EditorMainWindowBackgroundOverride", null);
+				string? OverrideValue = Section.GetValue("EditorMainWindowBackgroundOverride", null);
 				if (OverrideValue == null)
 				{
 					return null;
@@ -6145,7 +6198,7 @@ namespace UnrealGameSync
 
 				ConfigObject OverrideObject = new ConfigObject(OverrideValue);
 
-				string TintColorValue = OverrideObject.GetValue("TintColor");
+				string? TintColorValue = OverrideObject.GetValue("TintColor");
 				if (TintColorValue == null)
 				{
 					return null;
@@ -6157,7 +6210,13 @@ namespace UnrealGameSync
 					return null;
 				}
 
-				ConfigObject SpecifiedColorObject = new ConfigObject(TintColorObject.GetValue("SpecifiedColor"));
+				string? SpecifiedColorValue = TintColorObject.GetValue("SpecifiedColor");
+				if(SpecifiedColorValue == null)
+				{
+					return null;
+				}
+
+				ConfigObject SpecifiedColorObject = new ConfigObject(SpecifiedColorValue);
 
 				float R, G, B;
 				if (!float.TryParse(SpecifiedColorObject.GetValue("R", ""), out R) || !float.TryParse(SpecifiedColorObject.GetValue("G", ""), out G) || !float.TryParse(SpecifiedColorObject.GetValue("B", ""), out B))
diff --git a/Engine/Source/Programs/UnrealGameSync/UnrealGameSync/DetectProjectSettingsTask.cs b/Engine/Source/Programs/UnrealGameSync/UnrealGameSync/DetectProjectSettingsTask.cs
--- a/Engine/Source/Programs/UnrealGameSync/UnrealGameSync/DetectProjectSettingsTask.cs
+++ b/Engine/Source/Programs/UnrealGameSync/UnrealGameSync/DetectProjectSettingsTask.cs
@@ -12,8 +12,6 @@ using System.Text.RegularExpressions;
 using System.Threading;
 using System.Threading.Tasks;
 
-#nullable enable
-
 namespace UnrealGameSync
 {
 	public class WorkspaceSettings
@@ -108,8 +106,6 @@ namespace UnrealGameSync
 				}
 				else
 				{
-					PerforceSettings = DefaultConnection.Settings;
-
 					// Get the perforce server settings
 					InfoRecord PerforceInfo = await DefaultConnection.GetInfoAsync(InfoOptions.ShortOutput, CancellationToken);
 
@@ -151,7 +147,8 @@ namespace UnrealGameSync
 					}
 
 					// Take the client we've chosen
-					PerforceClient = await PerforceConnection.CreateAsync(CandidateClients[0], Logger);
+					PerforceSettings = CandidateClients[0];
+					PerforceClient = await PerforceConnection.CreateAsync(PerforceSettings, Logger);
 
 					// Get the client path for the project file
 					List<WhereRecord> Records = await PerforceClient.WhereAsync(NewSelectedFileName.FullName, CancellationToken).Where(x => !x.Unmap).ToListAsync(CancellationToken);
@@ -297,13 +294,17 @@ namespace UnrealGameSync
 
 				// Read the initial config file
 				List<KeyValuePair<FileReference, DateTime>> LocalConfigFiles = new List<KeyValuePair<FileReference, DateTime>>();
-				ConfigFile LatestProjectConfigFile = await PerforceMonitor.ReadProjectConfigFileAsync(PerforceClient, BranchClientPath, NewSelectedClientFileName, GetCacheFolder(BranchDirectoryName), LocalConfigFiles, Logger, CancellationToken);
+				ConfigFile LatestProjectConfigFile = await ConfigUtils.ReadProjectConfigFileAsync(PerforceClient, BranchClientPath, NewSelectedClientFileName, GetCacheFolder(BranchDirectoryName), LocalConfigFiles, Logger, CancellationToken);
 
 				// Get the local config file and stream filter
-				ConfigFile WorkspaceProjectConfigFile = await Workspace.ReadProjectConfigFile(BranchDirectoryName, NewSelectedFileName, Logger);
-				IReadOnlyList<string>? WorkspaceProjectStreamFilter = await Workspace.ReadProjectStreamFilter(PerforceClient, WorkspaceProjectConfigFile, Logger, CancellationToken.None);
+				ConfigFile WorkspaceProjectConfigFile = await WorkspaceUpdate.ReadProjectConfigFile(BranchDirectoryName, NewSelectedFileName, Logger);
+				IReadOnlyList<string>? WorkspaceProjectStreamFilter = await WorkspaceUpdate.ReadProjectStreamFilter(PerforceClient, WorkspaceProjectConfigFile, Logger, CancellationToken);
+
+				int BranchIdx = BranchClientPath.IndexOf('/', 2);
+				string BranchPath = (BranchIdx == -1) ? String.Empty : BranchClientPath.Substring(BranchIdx);
+				string ProjectPath = NewSelectedClientFileName.Substring(BranchClientPath.Length);
 
-				ProjectInfo ProjectInfo = new ProjectInfo(BranchDirectoryName, NewSelectedFileName, BranchClientPath, NewSelectedClientFileName, NewSelectedProjectIdentifier, bIsEnterpriseProject);
+				ProjectInfo ProjectInfo = new ProjectInfo(BranchDirectoryName, PerforceSettings.ClientName!, BranchPath, ProjectPath, NewSelectedProjectIdentifier, bIsEnterpriseProject);
 
 				WorkspaceSettings WorkspaceSettings = new WorkspaceSettings(PerforceSettings, SelectedProject, ProjectInfo, NewSelectedProjectIdentifier, NewProjectEditorTarget!, StreamName!, LatestProjectConfigFile, WorkspaceProjectConfigFile, WorkspaceProjectStreamFilter, LocalConfigFiles);
 				DirectoryReference.CreateDirectory(WorkspaceSettings.DataFolder);
@@ -312,7 +313,7 @@ namespace UnrealGameSync
 				// Run any event hooks
 				if (DeploymentSettings.OnDetectProjectSettings != null)
 				{
-					string Message;
+					string? Message;
 					if (!DeploymentSettings.OnDetectProjectSettings(WorkspaceSettings, Logger, out Message))
 					{
 						throw new UserErrorException(Message);
diff --git a/Engine/Source/Programs/UnrealGameSync/UnrealGameSync/Forms/MainWindow.cs b/Engine/Source/Programs/UnrealGameSync/UnrealGameSync/Forms/MainWindow.cs
--- a/Engine/Source/Programs/UnrealGameSync/UnrealGameSync/Forms/MainWindow.cs
+++ b/Engine/Source/Programs/UnrealGameSync/UnrealGameSync/Forms/MainWindow.cs
@@ -19,6 +19,7 @@ using EpicGames.Perforce;
 using System.Threading.Tasks;
 using Microsoft.Extensions.Logging;
 using Microsoft.Extensions.DependencyInjection;
+using System.Diagnostics.CodeAnalysis;
 
 namespace UnrealGameSync
 {
@@ -58,6 +59,11 @@ namespace UnrealGameSync
 		{
 			public IssueMonitor IssueMonitor;
 			public int RefCount;
+
+			public WorkspaceIssueMonitor(IssueMonitor IssueMonitor)
+			{
+				this.IssueMonitor = IssueMonitor;
+			}
 		}
 
 		[Flags]
@@ -80,7 +86,7 @@ namespace UnrealGameSync
 		List<IssueMonitor> DefaultIssueMonitors = new List<IssueMonitor>();
 		List<WorkspaceIssueMonitor> WorkspaceIssueMonitors = new List<WorkspaceIssueMonitor>();
 
-		string ApiUrl;
+		string? ApiUrl;
 		DirectoryReference DataFolder;
 		DirectoryReference CacheFolder;
 		IPerforceSettings DefaultPerforceSettings;
@@ -94,15 +100,15 @@ namespace UnrealGameSync
 
 		bool bRestoreStateOnLoad;
 
-		OIDCTokenManager OIDCTokenManager;
+		OIDCTokenManager? OIDCTokenManager;
 
-		System.Threading.Timer ScheduleTimer;
-		System.Threading.Timer ScheduleSettledTimer;
+		System.Threading.Timer? ScheduleTimer;
+		System.Threading.Timer? ScheduleSettledTimer;
 
 		string OriginalExecutableFileName;
 		bool bUnstable;
 
-		IMainWindowTabPanel CurrentTabPanel;
+		IMainWindowTabPanel? CurrentTabPanel;
 
 		AutomationServer AutomationServer;
 
@@ -113,9 +119,7 @@ namespace UnrealGameSync
 
 		public ToolUpdateMonitor ToolUpdateMonitor { get; private set; }
 
-		NetCoreWindow NetCoreWindow;
-
-		public MainWindow(UpdateMonitor InUpdateMonitor, string InApiUrl, DirectoryReference InDataFolder, DirectoryReference InCacheFolder, bool bInRestoreStateOnLoad, string InOriginalExecutableFileName, bool bInUnstable, List<(UserSelectedProjectSettings, ModalTask<WorkspaceSettings>)> StartupTasks, IPerforceSettings InDefaultPerforceSettings, IServiceProvider InServiceProvider, UserSettings InSettings, string InUri, OIDCTokenManager InOidcTokenManager)
+		public MainWindow(UpdateMonitor InUpdateMonitor, string? InApiUrl, DirectoryReference InDataFolder, DirectoryReference InCacheFolder, bool bInRestoreStateOnLoad, string InOriginalExecutableFileName, bool bInUnstable, List<(UserSelectedProjectSettings, ModalTask<WorkspaceSettings>)> StartupTasks, IPerforceSettings InDefaultPerforceSettings, IServiceProvider InServiceProvider, UserSettings InSettings, string? InUri, OIDCTokenManager? InOidcTokenManager)
 		{
 			ServiceProvider = InServiceProvider;
 			Logger = ServiceProvider.GetRequiredService<ILogger<MainWindow>>();
@@ -125,7 +129,7 @@ namespace UnrealGameSync
 			InitializeComponent();
 
 			UpdateMonitor = InUpdateMonitor;
-			MainThreadSynchronizationContext = SynchronizationContext.Current;
+			MainThreadSynchronizationContext = SynchronizationContext.Current!;
 			ApiUrl = InApiUrl;
 			DataFolder = InDataFolder;
 			CacheFolder = InCacheFolder;
@@ -220,7 +224,7 @@ namespace UnrealGameSync
 				}
 				else if (Request.Input.Type == AutomationRequestType.SyncProject)
 				{
-					AutomationRequestOutput Output = StartAutomatedSync(Request, true);
+					AutomationRequestOutput? Output = StartAutomatedSync(Request, true);
 					if (Output != null)
 					{
 						Request.SetOutput(Output);
@@ -233,7 +237,7 @@ namespace UnrealGameSync
 				}
 				else if (Request.Input.Type == AutomationRequestType.OpenProject)
 				{
-					AutomationRequestOutput Output = StartAutomatedSync(Request, false);
+					AutomationRequestOutput? Output = StartAutomatedSync(Request, false);
 					if (Output != null)
 					{
 						Request.SetOutput(Output);
@@ -269,13 +273,13 @@ namespace UnrealGameSync
 			string Command = Reader.ReadString();
 			string ProjectPath = Reader.ReadString();
 
-			AutomatedBuildWindow.BuildInfo BuildInfo;
+			AutomatedBuildWindow.BuildInfo? BuildInfo;
 			if (!AutomatedBuildWindow.ShowModal(this, DefaultPerforceSettings, StreamName, ProjectPath, Changelist, Command.ToString(), Settings, ServiceProvider, out BuildInfo))
 			{
 				return new AutomationRequestOutput(AutomationRequestResult.Canceled);
 			}
 
-			WorkspaceControl Workspace;
+			WorkspaceControl? Workspace;
 			if (!OpenWorkspaceForAutomation(BuildInfo.SelectedWorkspaceInfo, StreamName, BuildInfo.ProjectPath, out Workspace))
 			{
 				return new AutomationRequestOutput(AutomationRequestResult.Error);
@@ -304,12 +308,12 @@ namespace UnrealGameSync
 		{
 			if (Result == WorkspaceUpdateResult.Success && Command != null)
 			{
-				string CmdExe = Environment.GetEnvironmentVariable("COMSPEC");
+				string CmdExe = Environment.GetEnvironmentVariable("COMSPEC") ?? "C:\\Windows\\System32\\cmd.exe";
 				Workspace.ExecCommand("Run build command", "Running build command", CmdExe, String.Format("/c {0}", Command), Workspace.BranchDirectoryName.FullName, true);
 			}
 		}
 
-		AutomationRequestOutput StartAutomatedSync(AutomationRequest Request, bool bForceSync)
+		AutomationRequestOutput? StartAutomatedSync(AutomationRequest Request, bool bForceSync)
 		{
 			ShowAndActivate();
 
@@ -317,13 +321,13 @@ namespace UnrealGameSync
 			string StreamName = Reader.ReadString();
 			string ProjectPath = Reader.ReadString();
 
-			AutomatedSyncWindow.WorkspaceInfo WorkspaceInfo;
+			AutomatedSyncWindow.WorkspaceInfo? WorkspaceInfo;
 			if(!AutomatedSyncWindow.ShowModal(this, DefaultPerforceSettings, StreamName, ProjectPath, out WorkspaceInfo, ServiceProvider))
 			{
 				return new AutomationRequestOutput(AutomationRequestResult.Canceled);
 			}
 
-			WorkspaceControl Workspace;
+			WorkspaceControl? Workspace;
 			if(!OpenWorkspaceForAutomation(WorkspaceInfo, StreamName, ProjectPath, out Workspace))
 			{
 				return new AutomationRequestOutput(AutomationRequestResult.Error);
@@ -338,14 +342,14 @@ namespace UnrealGameSync
 			return null;
 		}
 
-		private bool OpenWorkspaceForAutomation(AutomatedSyncWindow.WorkspaceInfo WorkspaceInfo, string StreamName, string ProjectPath, out WorkspaceControl OutWorkspace)
+		private bool OpenWorkspaceForAutomation(AutomatedSyncWindow.WorkspaceInfo WorkspaceInfo, string StreamName, string ProjectPath, [NotNullWhen(true)] out WorkspaceControl? OutWorkspace)
 		{
 			if (WorkspaceInfo.bRequiresStreamSwitch)
 			{
 				// Close any tab containing this window
 				for (int ExistingTabIdx = 0; ExistingTabIdx < TabControl.GetTabCount(); ExistingTabIdx++)
 				{
-					WorkspaceControl ExistingWorkspace = TabControl.GetTabData(ExistingTabIdx) as WorkspaceControl;
+					WorkspaceControl? ExistingWorkspace = TabControl.GetTabData(ExistingTabIdx) as WorkspaceControl;
 					if (ExistingWorkspace != null && ExistingWorkspace.ClientName.Equals(WorkspaceInfo.WorkspaceName))
 					{
 						TabControl.RemoveTab(ExistingTabIdx);
@@ -356,16 +360,16 @@ namespace UnrealGameSync
 				// Switch the stream
 				Func<IPerforceConnection, CancellationToken, Task> SwitchTask = async (Connection, CancellationToken) =>
 				{
-					await Connection.SwitchClientToStreamAsync(Connection.Settings.ClientName, StreamName, SwitchClientOptions.None, CancellationToken);
+					await Connection.SwitchClientToStreamAsync(Connection.Settings.ClientName!, StreamName, SwitchClientOptions.None, CancellationToken);
 				};
 
 				PerforceSettings Settings = new PerforceSettings(WorkspaceInfo.ServerAndPort, WorkspaceInfo.UserName);
 				Settings.ClientName = WorkspaceInfo.WorkspaceName;
 
-				ModalTask Result = PerforceModalTask.Execute(Owner, "Please wait", "Switching streams, please wait...", Settings, SwitchTask, Logger, ModalTaskFlags.Quiet);
+				ModalTask? Result = PerforceModalTask.Execute(Owner, "Please wait", "Switching streams, please wait...", Settings, SwitchTask, Logger, ModalTaskFlags.Quiet);
 				if (Result == null || !Result.Succeeded)
 				{
-					Logger.LogError("Unable to switch stream ({Message})", Result.Error);
+					Logger.LogError("Unable to switch stream ({Message})", Result?.Error ?? "Operation cancelled");
 					OutWorkspace = null;
 					return false;
 				}
@@ -381,7 +385,7 @@ namespace UnrealGameSync
 				return false;
 			}
 
-			WorkspaceControl Workspace = TabControl.GetTabData(TabIdx) as WorkspaceControl;
+			WorkspaceControl? Workspace = TabControl.GetTabData(TabIdx) as WorkspaceControl;
 			if (Workspace == null)
 			{
 				Logger.LogError("Workspace was unable to open");
@@ -430,10 +434,10 @@ namespace UnrealGameSync
 
 			for(int ExistingTabIdx = 0; ExistingTabIdx < TabControl.GetTabCount(); ExistingTabIdx++)
 			{
-				WorkspaceControl ExistingWorkspace = TabControl.GetTabData(ExistingTabIdx) as WorkspaceControl;
+				WorkspaceControl? ExistingWorkspace = TabControl.GetTabData(ExistingTabIdx) as WorkspaceControl;
 				if(ExistingWorkspace != null && String.Compare(ExistingWorkspace.StreamName, StreamName, StringComparison.OrdinalIgnoreCase) == 0 && ExistingWorkspace.SelectedProject != null)
 				{
-					string ClientPath = ExistingWorkspace.SelectedProject.ClientPath;
+					string? ClientPath = ExistingWorkspace.SelectedProject.ClientPath;
 					if(ClientPath != null && ClientPath.StartsWith("//"))
 					{
 						int SlashIdx = ClientPath.IndexOf('/', 2);
@@ -459,7 +463,7 @@ namespace UnrealGameSync
 
 			for (int ExistingTabIdx = 0; ExistingTabIdx < TabControl.GetTabCount(); ExistingTabIdx++)
 			{
-				WorkspaceControl ExistingWorkspace = TabControl.GetTabData(ExistingTabIdx) as WorkspaceControl;
+				WorkspaceControl? ExistingWorkspace = TabControl.GetTabData(ExistingTabIdx) as WorkspaceControl;
 				if (ExistingWorkspace != null)
 				{
 					IssueMonitor IssueMonitor = ExistingWorkspace.GetIssueMonitor();
@@ -469,7 +473,7 @@ namespace UnrealGameSync
 						try
 						{
 							Func<CancellationToken, Task<IssueData>> Func = x => RESTApi.GetAsync<IssueData>($"{IssueMonitor.ApiUrl}/api/issues/{IssueId}", x);
-							ModalTask<IssueData> IssueTask = ModalTask.Execute(this, "Finding Issue", "Querying issue data, please wait...", Func);
+							ModalTask<IssueData>? IssueTask = ModalTask.Execute(this, "Finding Issue", "Querying issue data, please wait...", Func);
 							if (IssueTask == null)
 							{
 								Logger.LogInformation("Operation cancelled");
@@ -512,7 +516,7 @@ namespace UnrealGameSync
 			}
 		}
 
-		void TabControl_OnTabClicked(object TabData, Point Location, MouseButtons Buttons)
+		void TabControl_OnTabClicked(object? TabData, Point Location, MouseButtons Buttons)
 		{
 			if(Buttons == System.Windows.Forms.MouseButtons.Right)
 			{
@@ -605,13 +609,13 @@ namespace UnrealGameSync
 			if(AutomationServer != null)
 			{
 				AutomationServer.Dispose();
-				AutomationServer = null;
+				AutomationServer = null!;
 			}
 
 			if (ToolUpdateMonitor != null)
 			{
 				ToolUpdateMonitor.Dispose();
-				ToolUpdateMonitor = null;
+				ToolUpdateMonitor = null!;
 			}
 
 			base.Dispose(disposing);
@@ -643,7 +647,7 @@ namespace UnrealGameSync
 				{
 					for (int Idx = 0; Idx < TabControl.GetTabCount(); Idx++)
 					{
-						IMainWindowTabPanel TabPanel = (IMainWindowTabPanel)TabControl.GetTabData(Idx);
+						IMainWindowTabPanel TabPanel = (IMainWindowTabPanel)TabControl.GetTabData(Idx)!;
 						if (!TabPanel.CanClose())
 						{
 							EventArgs.Cancel = true;
@@ -656,7 +660,7 @@ namespace UnrealGameSync
 			}
 
 			Settings.bWindowVisible = Visible;
-			Settings.WindowState = WindowState;
+			Settings.WindowState = WindowState.ToString();
 			if(WindowState == FormWindowState.Normal)
 			{
 				Settings.WindowBounds = new Rectangle(Location, Size);
@@ -790,7 +794,7 @@ namespace UnrealGameSync
 
 			for (int TabIdx = 0; TabIdx < TabControl.GetTabCount(); TabIdx++)
 			{
-				IMainWindowTabPanel TabPanel = (IMainWindowTabPanel)TabControl.GetTabData(TabIdx);
+				IMainWindowTabPanel TabPanel = (IMainWindowTabPanel)TabControl.GetTabData(TabIdx)!;
 				if(TabPanel.IsBusy())
 				{
 					return false;
@@ -954,11 +958,11 @@ namespace UnrealGameSync
 			Logger.LogInformation("Schedule: Starting Sync");
 			for(int Idx = 0; Idx < TabControl.GetTabCount(); Idx++)
 			{
-				WorkspaceControl Workspace = TabControl.GetTabData(Idx) as WorkspaceControl;
+				WorkspaceControl? Workspace = TabControl.GetTabData(Idx) as WorkspaceControl;
 				if(Workspace != null)
 				{
 					Logger.LogInformation("Schedule: Considering {File}", Workspace.SelectedFileName);
-					if(Settings.ScheduleAnyOpenProject || Settings.ScheduleProjects.Any(x => x.LocalPath.Equals(Workspace.SelectedProject.LocalPath, StringComparison.OrdinalIgnoreCase)))
+					if(Settings.ScheduleAnyOpenProject || Settings.ScheduleProjects.Any(x => x.LocalPath != null && x.LocalPath.Equals(Workspace.SelectedProject.LocalPath, StringComparison.OrdinalIgnoreCase)))
 					{
 						Logger.LogInformation("Schedule: Starting Sync");
 						Workspace.ScheduleTimerElapsed();
@@ -967,7 +971,7 @@ namespace UnrealGameSync
 			}
 		}
 
-		void TabControl_OnTabChanged(object NewTabData)
+		void TabControl_OnTabChanged(object? NewTabData)
 		{
 			if(IsHandleCreated)
 			{
@@ -1031,7 +1035,7 @@ namespace UnrealGameSync
 
 		public void OpenNewProject()
 		{
-			WorkspaceSettings WorkspaceSettings = OpenProjectWindow.ShowModal(this, null, Settings, DataFolder, CacheFolder, DefaultPerforceSettings, ServiceProvider);
+			WorkspaceSettings? WorkspaceSettings = OpenProjectWindow.ShowModal(this, null, Settings, DataFolder, CacheFolder, DefaultPerforceSettings, ServiceProvider);
 			if(WorkspaceSettings != null)
 			{
 				int NewTabIdx = TryOpenProject(WorkspaceSettings, -1, OpenProjectOptions.None);
@@ -1046,7 +1050,7 @@ namespace UnrealGameSync
 
 		void UpdateRecentProjectsList(UserSelectedProjectSettings DetectedProjectSettings)
 		{
-			Settings.RecentProjects.RemoveAll(x => x.LocalPath.Equals(DetectedProjectSettings.LocalPath, StringComparison.OrdinalIgnoreCase));
+			Settings.RecentProjects.RemoveAll(x => x.LocalPath != null && x.LocalPath.Equals(DetectedProjectSettings.LocalPath, StringComparison.OrdinalIgnoreCase));
 			Settings.RecentProjects.Insert(0, DetectedProjectSettings);
 
 			const int MaxRecentProjects = 10;
@@ -1093,7 +1097,7 @@ namespace UnrealGameSync
 
 		public void EditSelectedProject(int TabIdx, UserSelectedProjectSettings SelectedProject)
 		{
-			WorkspaceSettings WorkspaceSettings = OpenProjectWindow.ShowModal(this, SelectedProject, Settings, DataFolder, CacheFolder, DefaultPerforceSettings, ServiceProvider);
+			WorkspaceSettings? WorkspaceSettings = OpenProjectWindow.ShowModal(this, SelectedProject, Settings, DataFolder, CacheFolder, DefaultPerforceSettings, ServiceProvider);
 			if(WorkspaceSettings != null)
 			{
 				int NewTabIdx = TryOpenProject(WorkspaceSettings, TabIdx, OpenProjectOptions.None);
@@ -1116,10 +1120,10 @@ namespace UnrealGameSync
 				TaskFlags |= ModalTaskFlags.Quiet;
 			}
 
-			ModalTask<WorkspaceSettings> SettingsTask = PerforceModalTask.Execute(this, "Opening Project", "Opening project, please wait...", DefaultPerforceSettings, (p, c) => WorkspaceSettings.CreateAsync(p, Project, ProjectLogger, c), ProjectLogger, TaskFlags);
-			if(SettingsTask == null || !SettingsTask.Succeeded)
+			ModalTask<WorkspaceSettings>? SettingsTask = PerforceModalTask.Execute(this, "Opening Project", "Opening project, please wait...", DefaultPerforceSettings, (p, c) => OpenProjectWindow.DetectSettingsAsync(p, Project, ProjectLogger, c), ProjectLogger, TaskFlags);
+			if (SettingsTask == null || SettingsTask.Failed)
 			{
-				CreateErrorPanel(ReplaceTabIdx, Project, SettingsTask.Error);
+				if(SettingsTask != null) CreateErrorPanel(ReplaceTabIdx, Project, SettingsTask.Error);
 				return -1;
 			}
 
@@ -1135,7 +1139,7 @@ namespace UnrealGameSync
 			{
 				if(ReplaceTabIdx != TabIdx)
 				{
-					WorkspaceControl Workspace = TabControl.GetTabData(TabIdx) as WorkspaceControl;
+					WorkspaceControl? Workspace = TabControl.GetTabData(TabIdx) as WorkspaceControl;
 					if(Workspace != null)
 					{
 						if(Workspace.SelectedFileName == WorkspaceSettings.ProjectInfo.LocalFileName)
@@ -1170,7 +1174,7 @@ namespace UnrealGameSync
 			// Remove the current tab. We need to ensure the workspace has been shut down before creating a new one with the same log files, etc...
 			if(ReplaceTabIdx != -1)
 			{
-				WorkspaceControl OldWorkspace = TabControl.GetTabData(ReplaceTabIdx) as WorkspaceControl;
+				WorkspaceControl? OldWorkspace = TabControl.GetTabData(ReplaceTabIdx) as WorkspaceControl;
 				if(OldWorkspace != null)
 				{
 					OldWorkspace.Hide();
@@ -1236,7 +1240,7 @@ namespace UnrealGameSync
 			Settings.OpenProjects.Clear();
 			for(int TabIdx = 0; TabIdx < TabControl.GetTabCount(); TabIdx++)
 			{
-				IMainWindowTabPanel TabPanel = (IMainWindowTabPanel)TabControl.GetTabData(TabIdx);
+				IMainWindowTabPanel TabPanel = (IMainWindowTabPanel)TabControl.GetTabData(TabIdx)!;
 				Settings.OpenProjects.Add(TabPanel.SelectedProject);
 			}
 			Settings.Save();
@@ -1290,7 +1294,7 @@ namespace UnrealGameSync
 
 				for(int Idx = 0; Idx < TabControl.GetTabCount(); Idx++)
 				{
-					WorkspaceControl Workspace = TabControl.GetTabData(Idx) as WorkspaceControl;
+					WorkspaceControl? Workspace = TabControl.GetTabData(Idx) as WorkspaceControl;
 					if(Workspace != null)
 					{
 						TabControl.SetTabName(Idx, GetTabName(Workspace));
@@ -1354,7 +1358,7 @@ namespace UnrealGameSync
 
 			for(int Idx = 0; Idx < TabControl.GetTabCount(); Idx++)
 			{
-				IMainWindowTabPanel TabPanel = (IMainWindowTabPanel)TabControl.GetTabData(Idx);
+				IMainWindowTabPanel TabPanel = (IMainWindowTabPanel)TabControl.GetTabData(Idx)!;
 
 				Tuple<TaskbarState, float> DesiredTaskbarState = TabPanel.DesiredTaskbarState;
 				if(DesiredTaskbarState.Item1 == TaskbarState.Error)
@@ -1397,7 +1401,7 @@ namespace UnrealGameSync
 		{
 			for (int Idx = 0; Idx < TabControl.GetTabCount(); Idx++)
 			{
-				IMainWindowTabPanel TabPanel = (IMainWindowTabPanel)TabControl.GetTabData(Idx);
+				IMainWindowTabPanel TabPanel = (IMainWindowTabPanel)TabControl.GetTabData(Idx)!;
 				TabControl.SetTint(Idx, TabPanel.TintColor);
 			}
 		}
@@ -1412,7 +1416,7 @@ namespace UnrealGameSync
 
 			for (int Idx = 0; Idx < TabControl.GetTabCount(); Idx++)
 			{
-				IMainWindowTabPanel TabPanel = (IMainWindowTabPanel)TabControl.GetTabData(Idx);
+				IMainWindowTabPanel TabPanel = (IMainWindowTabPanel)TabControl.GetTabData(Idx)!;
 				TabPanel.UpdateSettings();
 			}
 		}
@@ -1435,7 +1439,12 @@ namespace UnrealGameSync
 					}
 				}
 			}
-			WindowState = Settings.WindowState;
+
+			FormWindowState NewWindowState;
+			if (Enum.TryParse(Settings.WindowState, true, out NewWindowState))
+			{
+				WindowState = NewWindowState;
+			}
 		}
 
 		bool ShowNotificationForIssue(IssueData Issue)
@@ -1588,7 +1597,7 @@ namespace UnrealGameSync
 		{
 			for(int Idx = 0; Idx < TabControl.GetTabCount(); Idx++)
 			{
-				WorkspaceControl Workspace = TabControl.GetTabData(Idx) as WorkspaceControl;
+				WorkspaceControl? Workspace = TabControl.GetTabData(Idx) as WorkspaceControl;
 				if(Workspace != null && Workspace.GetIssueMonitor() == Alert.IssueMonitor)
 				{
 					Workspace.ShowIssueDetails(Alert.Issue);
@@ -1640,7 +1649,7 @@ namespace UnrealGameSync
 			SetAlertWindowPositions();
 		}
 
-		public IssueMonitor CreateIssueMonitor(string ApiUrl, string UserName)
+		public IssueMonitor CreateIssueMonitor(string? ApiUrl, string UserName)
 		{
 			WorkspaceIssueMonitor WorkspaceIssueMonitor = WorkspaceIssueMonitors.FirstOrDefault(x => String.Compare(x.IssueMonitor.ApiUrl, ApiUrl, StringComparison.OrdinalIgnoreCase) == 0 && String.Compare(x.IssueMonitor.UserName, UserName, StringComparison.OrdinalIgnoreCase) == 0);
 			if (WorkspaceIssueMonitor == null)
@@ -1650,10 +1659,10 @@ namespace UnrealGameSync
 
 				FileReference LogFileName = FileReference.Combine(DataFolder, String.Format("IssueMonitor-{0}-{1}.log", ServerId, UserName));
 
-				WorkspaceIssueMonitor = new WorkspaceIssueMonitor();
-				WorkspaceIssueMonitor.IssueMonitor = new IssueMonitor(ApiUrl, UserName, TimeSpan.FromSeconds(60.0), ServiceProvider);
-				WorkspaceIssueMonitor.IssueMonitor.OnIssuesChanged += IssueMonitor_OnUpdateAsync;
+				IssueMonitor IssueMonitor = new IssueMonitor(ApiUrl, UserName, TimeSpan.FromSeconds(60.0), ServiceProvider);
+				IssueMonitor.OnIssuesChanged += IssueMonitor_OnUpdateAsync;
 
+				WorkspaceIssueMonitor = new WorkspaceIssueMonitor(IssueMonitor);
 				WorkspaceIssueMonitors.Add(WorkspaceIssueMonitor);
 			}
 
@@ -1692,18 +1701,5 @@ namespace UnrealGameSync
 				}
 			}
 		}
-
-		private void NetCoreTimer_Tick(object sender, EventArgs e)
-		{
-			if (NetCoreWindow == null || !NetCoreWindow.Created)
-			{
-				if (!EventMonitor.HasNetCore3() && Settings.bShowNetCoreInfo)
-				{
-					NetCoreWindow = new NetCoreWindow(Settings);
-					NetCoreWindow.Location = new Point(Location.X + (Width - NetCoreWindow.Width) / 2, Location.Y + (Height - NetCoreWindow.Height) / 2);
-					NetCoreWindow.Show(this);
-				}
-			}
-		}
 	}
 }
diff --git a/Engine/Source/Programs/UnrealGameSync/UnrealGameSync/Forms/OpenProjectWindow.cs b/Engine/Source/Programs/UnrealGameSync/UnrealGameSync/Forms/OpenProjectWindow.cs
--- a/Engine/Source/Programs/UnrealGameSync/UnrealGameSync/Forms/OpenProjectWindow.cs
+++ b/Engine/Source/Programs/UnrealGameSync/UnrealGameSync/Forms/OpenProjectWindow.cs
@@ -30,7 +30,7 @@ namespace UnrealGameSync
 		IServiceProvider ServiceProvider;
 		UserSettings Settings;
 
-		private OpenProjectWindow(UserSelectedProjectSettings Project, UserSettings Settings, IPerforceSettings DefaultSettings, IServiceProvider ServiceProvider)
+		private OpenProjectWindow(UserSelectedProjectSettings? Project, UserSettings Settings, IPerforceSettings DefaultSettings, IServiceProvider ServiceProvider)
 		{
 			InitializeComponent();
 
@@ -90,7 +90,7 @@ namespace UnrealGameSync
 			get => Utility.OverridePerforceSettings(DefaultSettings, ServerAndPortOverride, UserNameOverride);
 		}
 
-		public static WorkspaceSettings? ShowModal(IWin32Window Owner, UserSelectedProjectSettings Project, UserSettings Settings, DirectoryReference DataFolder, DirectoryReference CacheFolder, IPerforceSettings DefaultPerforceSettings, IServiceProvider ServiceProvider)
+		public static WorkspaceSettings? ShowModal(IWin32Window Owner, UserSelectedProjectSettings? Project, UserSettings Settings, DirectoryReference DataFolder, DirectoryReference CacheFolder, IPerforceSettings DefaultPerforceSettings, IServiceProvider ServiceProvider)
 		{
 			OpenProjectWindow Window = new OpenProjectWindow(Project, Settings, DefaultPerforceSettings, ServiceProvider);
 			if(Window.ShowDialog(Owner) == DialogResult.OK)
@@ -286,7 +286,7 @@ namespace UnrealGameSync
 			{
 				ILogger<WorkspaceSettings> Logger = ServiceProvider.GetRequiredService<ILogger<WorkspaceSettings>>();
 
-				ModalTask<WorkspaceSettings>? NewDetectedProjectSettings = PerforceModalTask.Execute(this, "Opening project", "Opening project, please wait...", DefaultSettings, (x, y) => WorkspaceSettings.CreateAsync(x, SelectedProject, Logger, y), Logger);
+				ModalTask<WorkspaceSettings>? NewDetectedProjectSettings = PerforceModalTask.Execute(this, "Opening project", "Opening project, please wait...", DefaultSettings, (x, y) => DetectSettingsAsync(x, SelectedProject, Logger, y), Logger);
 				if (NewDetectedProjectSettings != null && NewDetectedProjectSettings.Succeeded)
 				{
 					DetectedProjectSettings = NewDetectedProjectSettings.Result;
@@ -296,6 +296,20 @@ namespace UnrealGameSync
 			}
 		}
 
+		public static async Task<WorkspaceSettings> DetectSettingsAsync(IPerforceConnection Perforce, UserSelectedProjectSettings SelectedProject, ILogger<WorkspaceSettings> Logger, CancellationToken CancellationToken)
+		{
+			WorkspaceSettings Settings = await WorkspaceSettings.CreateAsync(Perforce, SelectedProject, Logger, CancellationToken);
+			if (DeploymentSettings.OnDetectProjectSettings != null)
+			{
+				string? Message;
+				if (!DeploymentSettings.OnDetectProjectSettings(Settings, Logger, out Message))
+				{
+					throw new UserErrorException(Message);
+				}
+			}
+			return Settings;
+		}
+
 		private void ChangeLink_LinkClicked(object sender, LinkLabelLinkClickedEventArgs e)
 		{
 			if(ConnectWindow.ShowModal(this, DefaultSettings, ref ServerAndPortOverride, ref UserNameOverride, ServiceProvider))
diff --git a/Engine/Source/Programs/UnrealGameSync/UnrealGameSync/PerforceMonitor.cs b/Engine/Source/Programs/UnrealGameSync/UnrealGameSync/PerforceMonitor.cs
--- a/Engine/Source/Programs/UnrealGameSync/UnrealGameSync/PerforceMonitor.cs
+++ b/Engine/Source/Programs/UnrealGameSync/UnrealGameSync/PerforceMonitor.cs
@@ -16,8 +16,6 @@ using System.Text;
 using System.Threading;
 using System.Threading.Tasks;
 
-#nullable enable
-
 namespace UnrealGameSync
 {
 	class PerforceChangeDetails
@@ -78,17 +76,6 @@ namespace UnrealGameSync
 		IReadOnlyList<IArchiveInfo> AvailableArchives { get; }
 	}
 
-	public interface IArchiveInfo
-	{
-		string Name { get; }
-		string Type { get; }
-		string BasePath { get; }
-		string? Target { get;  }
-		bool Exists();
-		bool TryGetArchiveKeyForChangeNumber(int ChangeNumber, [NotNullWhen(true)] out string? ArchiveKey);
-		Task<bool> DownloadArchive(IPerforceConnection Perforce, string ArchiveKey, DirectoryReference LocalRootPath, FileReference ManifestFileName, ILogger Logger, ProgressValue Progress, CancellationToken CancellationToken);
-	}
-
 	class PerforceMonitor : IDisposable, IArchiveInfoSource
 	{
 		internal class PerforceArchiveInfo : IArchiveInfo
@@ -693,67 +680,7 @@ namespace UnrealGameSync
 		async Task UpdateProjectConfigFileAsync(IPerforceConnection Perforce, CancellationToken CancellationToken)
 		{
 			LocalConfigFiles.Clear();
-			LatestProjectConfigFile = await ReadProjectConfigFileAsync(Perforce, BranchClientPath, SelectedClientFileName, CacheFolder, LocalConfigFiles, Logger, CancellationToken);
-		}
-
-		public static async Task<ConfigFile> ReadProjectConfigFileAsync(IPerforceConnection Perforce, string BranchClientPath, string SelectedClientFileName, DirectoryReference CacheFolder, List<KeyValuePair<FileReference, DateTime>> LocalConfigFiles, ILogger Logger, CancellationToken CancellationToken)
-		{
-			List<string> ConfigFilePaths = Utility.GetDepotConfigPaths(BranchClientPath + "/Engine", SelectedClientFileName);
-
-			ConfigFile ProjectConfig = new ConfigFile();
-
-			List<PerforceResponse<FStatRecord>> Responses = await Perforce.TryFStatAsync(FStatOptions.IncludeFileSizes, ConfigFilePaths, CancellationToken).ToListAsync(CancellationToken);
-			foreach(PerforceResponse<FStatRecord> Response in Responses)
-			{
-				if (Response.Succeeded)
-				{
-					string[]? Lines = null;
-
-					// Skip file records which are still in the workspace, but were synced from a different branch. For these files, the action seems to be empty, so filter against that.
-					FStatRecord FileRecord = Response.Data;
-					if (FileRecord.HeadAction == FileAction.None)
-					{
-						continue;
-					}
-
-					// If this file is open for edit, read the local version
-					string? LocalFileName = FileRecord.ClientFile;
-					if(LocalFileName != null && File.Exists(LocalFileName) && (File.GetAttributes(LocalFileName) & FileAttributes.ReadOnly) == 0)
-					{
-						try
-						{
-							DateTime LastModifiedTime = File.GetLastWriteTimeUtc(LocalFileName);
-							LocalConfigFiles.Add(new KeyValuePair<FileReference, DateTime>(new FileReference(LocalFileName), LastModifiedTime));
-							Lines = await File.ReadAllLinesAsync(LocalFileName, CancellationToken);
-						}
-						catch(Exception Ex)
-						{
-							Logger.LogInformation(Ex, "Failed to read local config file for {Path}", LocalFileName);
-						}
-					}
-
-					// Otherwise try to get it from perforce
-					if(Lines == null && FileRecord.DepotFile != null)
-					{
-						Lines = await Utility.TryPrintFileUsingCacheAsync(Perforce, FileRecord.DepotFile, CacheFolder, FileRecord.Digest, Logger, CancellationToken);
-					}
-
-					// Merge the text with the config file
-					if(Lines != null)
-					{
-						try
-						{
-							ProjectConfig.Parse(Lines.ToArray());
-							Logger.LogInformation("Read config file from {DepotFile}", FileRecord.DepotFile);
-						}
-						catch(Exception Ex)
-						{
-							Logger.LogInformation(Ex, "Failed to read config file from {DepotFile}", FileRecord.DepotFile);
-						}
-					}
-				}
-			}
-			return ProjectConfig;
+			LatestProjectConfigFile = await ConfigUtils.ReadProjectConfigFileAsync(Perforce, BranchClientPath, SelectedClientFileName, CacheFolder, LocalConfigFiles, Logger, CancellationToken);
 		}
 
 		public List<ChangesRecord> GetChanges()
diff --git a/Engine/Source/Programs/UnrealGameSync/UnrealGameSync/UnrealGameSync.csproj b/Engine/Source/Programs/UnrealGameSync/UnrealGameSync/UnrealGameSync.csproj
--- a/Engine/Source/Programs/UnrealGameSync/UnrealGameSync/UnrealGameSync.csproj
+++ b/Engine/Source/Programs/UnrealGameSync/UnrealGameSync/UnrealGameSync.csproj
@@ -12,6 +12,7 @@
     <StartupObject>UnrealGameSync.Program</StartupObject>
     <ApplicationIcon>Resources\Icon.ico</ApplicationIcon>
     <ApplicationManifest>Properties\app.manifest</ApplicationManifest>
+    <Nullable>enable</Nullable>
   </PropertyGroup>
   <PropertyGroup Condition="'$(Configuration)|$(Platform)'=='Release|AnyCPU'">
     <DefineConstants>TRACE;DEBUG;WITH_TELEMETRY</DefineConstants>
@@ -32,6 +33,7 @@
     <ProjectReference Include="..\..\Shared\EpicGames.Core\EpicGames.Core.csproj" />
     <ProjectReference Include="..\..\Shared\EpicGames.Jupiter\EpicGames.Jupiter.csproj" />
     <ProjectReference Include="..\..\Shared\EpicGames.Perforce\EpicGames.Perforce.csproj" />
+    <ProjectReference Include="..\UnrealGameSyncShared\UnrealGameSyncShared.csproj" />
   </ItemGroup>
   <ItemGroup>
     <Reference Include="Microsoft.VisualStudio.Setup.Configuration.Interop">
diff --git a/Engine/Source/Programs/UnrealGameSync/UnrealGameSyncShared/ConfigUtils.cs b/Engine/Source/Programs/UnrealGameSync/UnrealGameSyncShared/ConfigUtils.cs
new file mode 100644
--- /dev/null
+++ b/Engine/Source/Programs/UnrealGameSync/UnrealGameSyncShared/ConfigUtils.cs
@@ -0,0 +1,78 @@
+// Copyright Epic Games, Inc. All Rights Reserved.
+
+using EpicGames.Core;
+using EpicGames.Perforce;
+using Microsoft.Extensions.Logging;
+using System;
+using System.Collections.Generic;
+using System.IO;
+using System.Linq;
+using System.Text;
+using System.Threading;
+using System.Threading.Tasks;
+
+namespace UnrealGameSync
+{
+	public static class ConfigUtils
+	{
+		public static async Task<ConfigFile> ReadProjectConfigFileAsync(IPerforceConnection Perforce, string BranchClientPath, string SelectedClientFileName, DirectoryReference CacheFolder, List<KeyValuePair<FileReference, DateTime>> LocalConfigFiles, ILogger Logger, CancellationToken CancellationToken)
+		{
+			List<string> ConfigFilePaths = Utility.GetDepotConfigPaths(BranchClientPath + "/Engine", SelectedClientFileName);
+
+			ConfigFile ProjectConfig = new ConfigFile();
+
+			List<PerforceResponse<FStatRecord>> Responses = await Perforce.TryFStatAsync(FStatOptions.IncludeFileSizes, ConfigFilePaths, CancellationToken).ToListAsync(CancellationToken);
+			foreach (PerforceResponse<FStatRecord> Response in Responses)
+			{
+				if (Response.Succeeded)
+				{
+					string[]? Lines = null;
+
+					// Skip file records which are still in the workspace, but were synced from a different branch. For these files, the action seems to be empty, so filter against that.
+					FStatRecord FileRecord = Response.Data;
+					if (FileRecord.HeadAction == FileAction.None)
+					{
+						continue;
+					}
+
+					// If this file is open for edit, read the local version
+					string? LocalFileName = FileRecord.ClientFile;
+					if (LocalFileName != null && File.Exists(LocalFileName) && (File.GetAttributes(LocalFileName) & FileAttributes.ReadOnly) == 0)
+					{
+						try
+						{
+							DateTime LastModifiedTime = File.GetLastWriteTimeUtc(LocalFileName);
+							LocalConfigFiles.Add(new KeyValuePair<FileReference, DateTime>(new FileReference(LocalFileName), LastModifiedTime));
+							Lines = await File.ReadAllLinesAsync(LocalFileName, CancellationToken);
+						}
+						catch (Exception Ex)
+						{
+							Logger.LogInformation(Ex, "Failed to read local config file for {Path}", LocalFileName);
+						}
+					}
+
+					// Otherwise try to get it from perforce
+					if (Lines == null && FileRecord.DepotFile != null)
+					{
+						Lines = await Utility.TryPrintFileUsingCacheAsync(Perforce, FileRecord.DepotFile, CacheFolder, FileRecord.Digest, Logger, CancellationToken);
+					}
+
+					// Merge the text with the config file
+					if (Lines != null)
+					{
+						try
+						{
+							ProjectConfig.Parse(Lines.ToArray());
+							Logger.LogInformation("Read config file from {DepotFile}", FileRecord.DepotFile);
+						}
+						catch (Exception Ex)
+						{
+							Logger.LogInformation(Ex, "Failed to read config file from {DepotFile}", FileRecord.DepotFile);
+						}
+					}
+				}
+			}
+			return ProjectConfig;
+		}
+	}
+}
diff --git a/Engine/Source/Programs/UnrealGameSync/UnrealGameSyncShared/UserSettings.cs b/Engine/Source/Programs/UnrealGameSync/UnrealGameSyncShared/UserSettings.cs
new file mode 100644
--- /dev/null
+++ b/Engine/Source/Programs/UnrealGameSync/UnrealGameSyncShared/UserSettings.cs
@@ -0,0 +1,1247 @@
+// Copyright Epic Games, Inc. All Rights Reserved.
+
+using EpicGames.Core;
+using EpicGames.Perforce;
+using Microsoft.Extensions.Logging;
+using System;
+using System.Collections.Generic;
+using System.Diagnostics.CodeAnalysis;
+using System.Drawing;
+using System.IO;
+using System.Linq;
+using System.Text;
+using System.Text.Json;
+using System.Text.Json.Serialization;
+using System.Text.RegularExpressions;
+using System.Threading.Tasks;
+
+namespace UnrealGameSync
+{
+	public enum BuildConfig
+	{
+		Debug,
+		DebugGame,
+		Development,
+	}
+
+	public enum TabLabels
+	{
+		Stream,
+		WorkspaceName,
+		WorkspaceRoot,
+		ProjectFile,
+	}
+
+	public enum BisectState
+	{
+		Include,
+		Exclude,
+		Pass,
+		Fail,
+	}
+
+	public enum UserSelectedProjectType
+	{
+		Client,
+		Local
+	}
+
+	public enum FilterType
+	{
+		None,
+		Code,
+		Content
+	}
+
+	public enum LatestChangeType
+	{
+		Any,
+		Good,
+		Starred,
+	}
+
+	public enum UserSettingsVersion
+	{
+		Initial = 0,
+		DefaultServerSettings = 1,
+		XgeShaderCompilation = 2,
+		DefaultNumberOfThreads = 3,
+		Latest = DefaultNumberOfThreads
+	}
+
+	public class ArchiveSettings
+	{
+		public bool bEnabled;
+		public string Type;
+		public List<string> Order;
+
+		public ArchiveSettings(bool bEnabled, string Type, IEnumerable<string> Order)
+		{
+			this.bEnabled = bEnabled;
+			this.Type = Type;
+			this.Order = new List<string>(Order);
+		}
+
+		public static bool TryParseConfigEntry(string Text, [NotNullWhen(true)] out ArchiveSettings? Settings)
+		{
+			ConfigObject Object = new ConfigObject(Text);
+
+			string? Type = Object.GetValue("Type", null);
+			if (Type == null)
+			{
+				Settings = null;
+				return false;
+			}
+			else
+			{
+				string[] Order = Object.GetValue("Order", "").Split(new char[] { ';' }, StringSplitOptions.RemoveEmptyEntries);
+				bool bEnabled = Object.GetValue("Enabled", 0) != 0;
+
+				Settings = new ArchiveSettings(bEnabled, Type, Order);
+				return true;
+			}
+		}
+
+		public string ToConfigEntry()
+		{
+			ConfigObject Object = new ConfigObject();
+
+			Object.SetValue("Enabled", bEnabled ? 1 : 0);
+			Object.SetValue("Type", Type);
+			Object.SetValue("Order", String.Join(";", Order));
+
+			return Object.ToString();
+		}
+
+		public override string ToString()
+		{
+			return ToConfigEntry();
+		}
+	}
+
+	public class UserSelectedProjectSettings
+	{
+		[JsonIgnore]
+		FileReference? File { get; set; }
+
+		public string? ServerAndPort { get; set; }
+		public string? UserName { get; set; }
+		public UserSelectedProjectType Type { get; set; }
+		public string? ClientPath { get; set; }
+		public string? LocalPath { get; set; }
+
+		public UserSelectedProjectSettings(string? ServerAndPort, string? UserName, UserSelectedProjectType Type, string? ClientPath, string? LocalPath)
+		{
+			this.ServerAndPort = ServerAndPort;
+			this.UserName = UserName;
+			this.Type = Type;
+			this.ClientPath = ClientPath;
+			this.LocalPath = LocalPath;
+		}
+
+		public static bool TryParseConfigEntry(string Text, [NotNullWhen(true)] out UserSelectedProjectSettings? Project)
+		{
+			ConfigObject Object = new ConfigObject(Text);
+
+			UserSelectedProjectType Type;
+			if(Enum.TryParse(Object.GetValue("Type", ""), out Type))
+			{
+				string? ServerAndPort = Object.GetValue("ServerAndPort", null);
+				if(String.IsNullOrWhiteSpace(ServerAndPort))
+				{
+					ServerAndPort = null;
+				}
+
+				// Fixup for code that was saving server host name rather than DNS entry
+				if(ServerAndPort != null && ServerAndPort.Equals("p4-nodeb.epicgames.net:1666", StringComparison.OrdinalIgnoreCase))
+				{
+					ServerAndPort = "perforce:1666";
+				}
+
+				string? UserName = Object.GetValue("UserName", null);
+				if(String.IsNullOrWhiteSpace(UserName))
+				{
+					UserName = null;
+				}
+
+				string? LocalPath = Object.GetValue("LocalPath", null);
+				if(String.IsNullOrWhiteSpace(LocalPath))
+				{
+					LocalPath = null;
+				}
+
+				string? ClientPath = Object.GetValue("ClientPath", null);
+				if(String.IsNullOrWhiteSpace(ClientPath))
+				{
+					ClientPath = null;
+				}
+
+				if((Type == UserSelectedProjectType.Client && ClientPath != null) || (Type == UserSelectedProjectType.Local && LocalPath != null))
+				{
+					Project = new UserSelectedProjectSettings(ServerAndPort, UserName, Type, ClientPath, LocalPath);
+					return true;
+				}
+			}
+
+			Project = null;
+			return false;
+		}
+
+		public string ToConfigEntry()
+		{
+			ConfigObject Object = new ConfigObject();
+
+			if(ServerAndPort != null)
+			{
+				Object.SetValue("ServerAndPort", ServerAndPort);
+			}
+			if(UserName != null)
+			{
+				Object.SetValue("UserName", UserName);
+			}
+
+			Object.SetValue("Type", Type.ToString());
+
+			if(ClientPath != null)
+			{
+				Object.SetValue("ClientPath", ClientPath);
+			}
+			if(LocalPath != null)
+			{
+				Object.SetValue("LocalPath", LocalPath);
+			}
+
+			return Object.ToString();
+		}
+
+		public override string? ToString()
+		{
+			return LocalPath ?? ClientPath;
+		}
+	}
+
+	public class SyncCategory
+	{
+		public Guid Id { get; set; }
+		public bool Enable { get; set; }
+	}
+
+	public class BisectEntry
+	{
+		public int Change { get; set; }
+		public BisectState State { get; set; }
+	}
+
+	public class UserWorkspaceSettings
+	{
+		[JsonIgnore]
+		public DirectoryReference RootDir { get; set; } = null!;
+
+		// Connection settings
+		public string? ServerAndPort { get; set; }
+		public string? UserName { get; set; }
+		public string ClientName { get; set; } = String.Empty;
+
+		// Path to the root of the branch within this client, with a trailing slash if non-empty
+		public string BranchPath { get; set; } = String.Empty;
+
+		// The currently selected project
+		public string ProjectPath { get; set; } = String.Empty;
+
+		// Workspace specific SyncFilters
+		public string[] SyncView { get; set; } = Array.Empty<string>();
+		public List<SyncCategory> SyncCategories { get; set; } = new List<SyncCategory>();
+		public bool? bSyncAllProjects { get; set; }
+		public bool? bIncludeAllProjectsInSolution { get; set; }
+
+		[JsonIgnore]
+		public Dictionary<Guid, bool> SyncCategoriesDict
+		{
+			get => SyncCategories.ToDictionary(x => x.Id, x => x.Enable);
+			set => SyncCategories = value.Select(x => new SyncCategory { Id = x.Key, Enable = x.Value }).ToList();
+		}
+
+		[JsonIgnore]
+		public FileReference ConfigFile => GetConfigFile(RootDir);
+
+		[JsonIgnore]
+		public string ClientProjectPath => $"//{ClientName}{BranchPath}{ProjectPath}";
+
+		public static bool TryLoad(DirectoryReference RootDir, [NotNullWhen(true)] out UserWorkspaceSettings? Settings)
+		{
+			FileReference ConfigFile = GetConfigFile(RootDir);
+			if (Utility.TryLoadJson(ConfigFile, out Settings))
+			{
+				Settings.RootDir = RootDir;
+				return true;
+			}
+			else
+			{
+				Settings = null;
+				return false;
+			}
+		}
+
+		static object SyncRoot = new object();
+
+		public void Save()
+		{
+			lock (SyncRoot)
+			{
+				Utility.SaveJson(ConfigFile, this);
+			}
+		}
+
+		public static FileReference GetConfigFile(DirectoryReference RootDir)
+		{
+			return FileReference.Combine(UserSettings.GetConfigDir(RootDir), "settings.json");
+		}
+	}
+
+	public class UserWorkspaceState
+	{
+		[JsonIgnore]
+		public DirectoryReference RootDir { get; set; } = null!;
+
+		// Settings for the currently synced project in this workspace. CurrentChangeNumber is only valid for this workspace if CurrentProjectPath is the current project.
+		public string? CurrentProjectIdentifier { get; set; }
+		public int CurrentChangeNumber { get; set; } = -1;
+		public string? CurrentSyncFilterHash { get; set; }
+		public List<int> AdditionalChangeNumbers { get; set; } = new List<int>();
+
+		// Settings for the last attempted sync. These values are set to persist error messages between runs.
+		public int LastSyncChangeNumber { get; set; }
+		public WorkspaceUpdateResult LastSyncResult { get; set; }
+		public string? LastSyncResultMessage { get; set; }
+		public DateTime? LastSyncTime { get; set; }
+		public int LastSyncDurationSeconds { get; set; }
+
+		// The last successful build, regardless of whether a failed sync has happened in the meantime. Used to determine whether to force a clean due to entries in the project config file.
+		public int LastBuiltChangeNumber { get; set; }
+
+		// Expanded archives in the workspace
+		public string[]? ExpandedArchiveTypes { get; set; }
+
+		// The changes that we're regressing at the moment
+		public List<BisectEntry> BisectChanges { get; set; } = new List<BisectEntry>();
+
+		// Path to the config file
+		public FileReference ConfigFile => GetConfigFile(RootDir);
+
+		public void SetBisectState(int Change, BisectState State)
+		{
+			BisectEntry Entry = BisectChanges.FirstOrDefault(x => x.Change == Change);
+			if (Entry == null)
+			{
+				Entry = new BisectEntry();
+				Entry.Change = Change;
+				BisectChanges.Add(Entry);
+			}
+			Entry.State = State;
+		}
+
+		public static bool TryLoad(DirectoryReference RootDir, [NotNullWhen(true)] out UserWorkspaceState? State)
+		{
+			FileReference ConfigFile = GetConfigFile(RootDir);
+			if (Utility.TryLoadJson(ConfigFile, out State))
+			{
+				State.RootDir = RootDir;
+				return true;
+			}
+			else
+			{
+				State = null;
+				return false;
+			}
+		}
+
+		static object SyncRoot = new object();
+
+		public void Save()
+		{
+			lock (SyncRoot)
+			{
+				Utility.SaveJson(ConfigFile, this);
+			}
+		}
+
+		public static FileReference GetConfigFile(DirectoryReference RootDir) => FileReference.Combine(UserSettings.GetConfigDir(RootDir), "state.json");
+	}
+
+	public class UserProjectSettings
+	{
+		[JsonIgnore]
+		public FileReference? File { get; set; }
+
+		public List<ConfigObject> BuildSteps { get; set; } = new List<ConfigObject>();
+		public FilterType FilterType { get; set; }
+		public HashSet<string> FilterBadges { get; set; } = new HashSet<string>(StringComparer.OrdinalIgnoreCase);
+
+		public void Save() => Utility.SaveJson(File!, this);
+	}
+
+	public class UserSettings
+	{
+		/// <summary>
+		/// Enum that decribes which robomerge changes to show
+		/// </summary>
+		public enum RobomergeShowChangesOption
+		{
+			All,		// Show all changes from robomerge
+			Badged,		// Show only robomerge changes that have an associated badge
+			None		// Show no robomerge changes
+		};
+
+		FileReference FileName;
+		ConfigFile ConfigFile = new ConfigFile();
+
+		// General settings
+		public UserSettingsVersion Version = UserSettingsVersion.Latest;
+		public bool bBuildAfterSync;
+		public bool bRunAfterSync;
+		public bool bOpenSolutionAfterSync;
+		public bool bShowLogWindow;
+		public bool bAutoResolveConflicts;
+		public bool bShowUnreviewedChanges;
+		public bool bShowAutomatedChanges;
+		public RobomergeShowChangesOption ShowRobomerge;
+		public bool bAnnotateRobmergeChanges;
+		public bool bShowLocalTimes;
+		public bool bKeepInTray;
+		public Guid[] EnabledTools;
+		public bool bShowNetCoreInfo;
+		public int FilterIndex;
+		public UserSelectedProjectSettings? LastProject;
+		public List<UserSelectedProjectSettings> OpenProjects;
+		public List<UserSelectedProjectSettings> RecentProjects;
+		public string[] SyncView;
+		public Dictionary<Guid, bool> SyncCategories;
+		public bool bSyncAllProjects;
+		public bool bIncludeAllProjectsInSolution;
+		public LatestChangeType SyncType;
+		public BuildConfig CompiledEditorBuildConfig; // NB: This assumes not using precompiled editor. See CurrentBuildConfig.
+		public TabLabels TabLabels;
+
+		// Precompiled binaries
+		public List<ArchiveSettings> Archives = new List<ArchiveSettings>();
+
+		// OIDC Settings
+		public Dictionary<string, string> ProviderToRefreshTokens = new Dictionary<string, string>();
+
+		// Window settings
+		public bool bWindowVisible;
+		public string WindowState;
+		public Rectangle? WindowBounds;
+		
+		// Schedule settings
+		public bool bScheduleEnabled;
+		public TimeSpan ScheduleTime;
+		public LatestChangeType ScheduleChange;
+		public bool ScheduleAnyOpenProject;
+		public List<UserSelectedProjectSettings> ScheduleProjects;
+
+		// Run configuration
+		public List<Tuple<string, bool>> EditorArguments = new List<Tuple<string,bool>>();
+		public bool bEditorArgumentsPrompt;
+
+		// Notification settings
+		public List<string> NotifyProjects;
+		public int NotifyUnassignedMinutes;
+		public int NotifyUnacknowledgedMinutes;
+		public int NotifyUnresolvedMinutes;
+
+		// Project settings
+		Dictionary<DirectoryReference, UserWorkspaceState> WorkspaceDirToState = new Dictionary<DirectoryReference, UserWorkspaceState>();
+		Dictionary<DirectoryReference, UserWorkspaceSettings> WorkspaceDirToSettings = new Dictionary<DirectoryReference, UserWorkspaceSettings>();
+		Dictionary<FileReference, UserProjectSettings> ProjectKeyToSettings = new Dictionary<FileReference, UserProjectSettings>();
+
+		// Perforce settings
+		public PerforceSyncOptions SyncOptions = new PerforceSyncOptions();
+
+		private List<UserSelectedProjectSettings> ReadProjectList(string SettingName, string LegacySettingName)
+		{
+			List<UserSelectedProjectSettings> Projects = new List<UserSelectedProjectSettings>();
+
+			string[]? ProjectStrings = ConfigFile.GetValues(SettingName, null);
+			if(ProjectStrings != null)
+			{
+				foreach(string ProjectString in ProjectStrings)
+				{
+					UserSelectedProjectSettings? Project;
+					if(UserSelectedProjectSettings.TryParseConfigEntry(ProjectString, out Project))
+					{
+						Projects.Add(Project);
+					}
+				}
+			}
+			else if(LegacySettingName != null)
+			{
+				string[]? LegacyProjectStrings = ConfigFile.GetValues(LegacySettingName, null);
+				if(LegacyProjectStrings != null)
+				{
+					foreach(string LegacyProjectString in LegacyProjectStrings)
+					{
+						if(!String.IsNullOrWhiteSpace(LegacyProjectString))
+						{
+							Projects.Add(new UserSelectedProjectSettings(null, null, UserSelectedProjectType.Local, null, LegacyProjectString));
+						}
+					}
+				}
+			}
+
+			return Projects;
+		}
+
+		public UserSettings(FileReference InFileName, ILogger Logger)
+		{
+			FileName = InFileName;
+			ConfigFile.TryLoad(FileName, Logger);
+
+			// General settings
+			Version = (UserSettingsVersion)ConfigFile.GetValue("General.Version", (int)UserSettingsVersion.Initial);
+			bBuildAfterSync = (ConfigFile.GetValue("General.BuildAfterSync", "1") != "0");
+			bRunAfterSync = (ConfigFile.GetValue("General.RunAfterSync", "1") != "0");
+			bool bSyncPrecompiledEditor = (ConfigFile.GetValue("General.SyncPrecompiledEditor", "0") != "0");
+			bOpenSolutionAfterSync = (ConfigFile.GetValue("General.OpenSolutionAfterSync", "0") != "0");
+			bShowLogWindow = (ConfigFile.GetValue("General.ShowLogWindow", false));
+			bAutoResolveConflicts = (ConfigFile.GetValue("General.AutoResolveConflicts", "1") != "0");
+			bShowUnreviewedChanges = ConfigFile.GetValue("General.ShowUnreviewed", true);
+			bShowAutomatedChanges = ConfigFile.GetValue("General.ShowAutomated", false);
+
+			// safely parse the filter enum
+			ShowRobomerge = RobomergeShowChangesOption.All;
+			Enum.TryParse(ConfigFile.GetValue("General.RobomergeFilter", ""), out ShowRobomerge);
+
+			bAnnotateRobmergeChanges = ConfigFile.GetValue("General.AnnotateRobomerge", true);
+			bShowLocalTimes = ConfigFile.GetValue("General.ShowLocalTimes", false);
+			bKeepInTray = ConfigFile.GetValue("General.KeepInTray", true);
+
+			List<Guid> EnabledTools = ConfigFile.GetGuidValues("General.EnabledTools", new Guid[0]).ToList();
+			if (ConfigFile.GetValue("General.EnableP4VExtensions", false))
+			{
+				EnabledTools.Add(new Guid("963850A0-BF63-4E0E-B903-1C5954C7DCF8"));
+			}
+			if (ConfigFile.GetValue("General.EnableUshell", false))
+			{
+				EnabledTools.Add(new Guid("922EED87-E732-464C-92DC-5A8F7ED955E2"));
+			}
+			this.EnabledTools = EnabledTools.ToArray();
+
+			bShowNetCoreInfo = ConfigFile.GetValue("General.ShowNetCoreInfo", true);
+			int.TryParse(ConfigFile.GetValue("General.FilterIndex", "0"), out FilterIndex);
+
+			string? LastProjectString = ConfigFile.GetValue("General.LastProject", null);
+			if(LastProjectString != null)
+			{
+				UserSelectedProjectSettings.TryParseConfigEntry(LastProjectString, out LastProject);
+			}
+			else
+			{
+				string? LastProjectFileName = ConfigFile.GetValue("General.LastProjectFileName", null);
+				if(LastProjectFileName != null)
+				{
+					LastProject = new UserSelectedProjectSettings(null, null, UserSelectedProjectType.Local, null, LastProjectFileName);
+				}
+			}
+
+			OpenProjects = ReadProjectList("General.OpenProjects", "General.OpenProjectFileNames");
+			RecentProjects = ReadProjectList("General.RecentProjects", "General.OtherProjectFileNames");
+			SyncView = ConfigFile.GetValues("General.SyncFilter", new string[0]);
+			SyncCategories = GetCategorySettings(ConfigFile.FindSection("General"), "SyncIncludedCategories", "SyncExcludedCategories");
+
+			bSyncAllProjects = ConfigFile.GetValue("General.SyncAllProjects", false);
+			bIncludeAllProjectsInSolution = ConfigFile.GetValue("General.IncludeAllProjectsInSolution", false);
+			if(!Enum.TryParse(ConfigFile.GetValue("General.SyncType", ""), out SyncType))
+			{
+				SyncType = LatestChangeType.Good;
+			}
+
+			// Build configuration
+			string CompiledEditorBuildConfigName = ConfigFile.GetValue("General.BuildConfig", "");
+			if(!Enum.TryParse(CompiledEditorBuildConfigName, true, out CompiledEditorBuildConfig))
+			{
+				CompiledEditorBuildConfig = BuildConfig.DebugGame;
+			}
+
+			// Tab names
+			string TabLabelsValue = ConfigFile.GetValue("General.TabLabels", "");
+			if(!Enum.TryParse(TabLabelsValue, true, out TabLabels))
+			{
+				TabLabels = TabLabels.Stream;
+			}
+
+			// Editor arguments
+			string[] Arguments = ConfigFile.GetValues("General.EditorArguments", new string[]{ "0:-log", "0:-fastload" });
+			if (Version < UserSettingsVersion.XgeShaderCompilation)
+			{
+				Arguments = Enumerable.Concat(Arguments, new string[] { "0:-noxgeshadercompile" }).ToArray();
+			}
+			foreach(string Argument in Arguments)
+			{
+				if(Argument.StartsWith("0:"))
+				{
+					EditorArguments.Add(new Tuple<string,bool>(Argument.Substring(2), false));
+				}
+				else if(Argument.StartsWith("1:"))
+				{
+					EditorArguments.Add(new Tuple<string,bool>(Argument.Substring(2), true));
+				}
+				else
+				{
+					EditorArguments.Add(new Tuple<string,bool>(Argument, true));
+				}
+			}
+			bEditorArgumentsPrompt = ConfigFile.GetValue("General.EditorArgumentsPrompt", false);
+
+			// Precompiled binaries
+			string[] ArchiveValues = ConfigFile.GetValues("PrecompiledBinaries.Archives", new string[0]);
+			foreach (string ArchiveValue in ArchiveValues)
+			{
+				ArchiveSettings? Settings;
+				if (ArchiveSettings.TryParseConfigEntry(ArchiveValue, out Settings))
+				{
+					Archives.Add(Settings);
+				}
+			}
+
+			if (bSyncPrecompiledEditor)
+			{
+				Archives.Add(new ArchiveSettings(true, "Editor", new string[0]));
+			}
+
+			// OIDC Settings
+			string[] Tokens = ConfigFile.GetValues("OIDCProviders.Tokens", new string[0]);
+			foreach (string TokenValue in Tokens)
+			{
+				ConfigObject O = new ConfigObject(TokenValue);
+				string? Provider = O.GetValue("Provider");
+				string? Token = O.GetValue("Token");
+				if (Provider != null && Token != null)
+				{
+					ProviderToRefreshTokens.TryAdd(Provider, Token);
+				}
+			}
+
+			// Window settings
+			bWindowVisible = ConfigFile.GetValue("Window.Visible", true);
+			WindowState = ConfigFile.GetValue("Window.State", "");
+			WindowBounds = ParseRectangleValue(ConfigFile.GetValue("Window.Bounds", ""));
+
+			// Schedule settings
+			bScheduleEnabled = ConfigFile.GetValue("Schedule.Enabled", false);
+			if(!TimeSpan.TryParse(ConfigFile.GetValue("Schedule.Time", ""), out ScheduleTime))
+			{
+				ScheduleTime = new TimeSpan(6, 0, 0);
+			}
+			if(!Enum.TryParse(ConfigFile.GetValue("Schedule.Change", ""), out ScheduleChange))
+			{
+				ScheduleChange = LatestChangeType.Good;
+			}
+			ScheduleAnyOpenProject = ConfigFile.GetValue("Schedule.AnyOpenProject", true);
+			ScheduleProjects = ReadProjectList("Schedule.Projects", "Schedule.ProjectFileNames");
+
+			// Notification settings
+			NotifyProjects = ConfigFile.GetValues("Notifications.NotifyProjects", new string[0]).ToList();
+			NotifyUnassignedMinutes = ConfigFile.GetValue("Notifications.NotifyUnassignedMinutes", -1);
+			NotifyUnacknowledgedMinutes = ConfigFile.GetValue("Notifications.NotifyUnacknowledgedMinutes", -1);
+			NotifyUnresolvedMinutes = ConfigFile.GetValue("Notifications.NotifyUnresolvedMinutes", -1);
+
+			// Perforce settings
+			if(!int.TryParse(ConfigFile.GetValue("Perforce.NumRetries", "0"), out SyncOptions.NumRetries))
+			{
+				SyncOptions.NumRetries = 0;
+			}
+
+			int NumThreads;
+			if(int.TryParse(ConfigFile.GetValue("Perforce.NumThreads", "0"), out NumThreads) && NumThreads > 0)
+			{
+				if(Version >= UserSettingsVersion.DefaultNumberOfThreads || NumThreads > 1)
+				{
+					SyncOptions.NumThreads = NumThreads;
+				}
+			}
+
+			SyncOptions.TcpBufferSize = ConfigFile.GetValue("Perforce.TcpBufferSize", PerforceSyncOptions.DefaultTcpBufferSize);
+			SyncOptions.FileBufferSize = ConfigFile.GetValue("Perforce.FileBufferSize", PerforceSyncOptions.DefaultFileBufferSize);
+			SyncOptions.MaxCommandsPerBatch = ConfigFile.GetValue("Perforce.MaxCommandsPerBatch", PerforceSyncOptions.DefaultMaxCommandsPerBatch);
+			SyncOptions.MaxSizePerBatch = ConfigFile.GetValue("Perforce.MaxSizePerBatch", PerforceSyncOptions.DefaultMaxSizePerBatch);
+		}
+
+		static Dictionary<Guid, bool> GetCategorySettings(ConfigSection Section, string IncludedKey, string ExcludedKey)
+		{
+			Dictionary<Guid, bool> Result = new Dictionary<Guid, bool>();
+			if (Section != null)
+			{
+				foreach (Guid UniqueId in Section.GetValues(IncludedKey, new Guid[0]))
+				{
+					Result[UniqueId] = true;
+				}
+				foreach (Guid UniqueId in Section.GetValues(ExcludedKey, new Guid[0]))
+				{
+					Result[UniqueId] = false;
+				}
+			}
+			return Result;
+		}
+
+		static void SetCategorySettings(ConfigSection Section, string IncludedKey, string ExcludedKey, Dictionary<Guid, bool> Categories)
+		{
+			Guid[] IncludedCategories = Categories.Where(x => x.Value).Select(x => x.Key).ToArray();
+			if (IncludedCategories.Length > 0)
+			{
+				Section.SetValues(IncludedKey, IncludedCategories);
+			}
+
+			Guid[] ExcludedCategories = Categories.Where(x => !x.Value).Select(x => x.Key).ToArray();
+			if (ExcludedCategories.Length > 0)
+			{
+				Section.SetValues(ExcludedKey, ExcludedCategories);
+			}
+		}
+
+		static Rectangle? ParseRectangleValue(string Text)
+		{
+			ConfigObject Object = new ConfigObject(Text);
+
+			int X = Object.GetValue("X", -1);
+			int Y = Object.GetValue("Y", -1);
+			int W = Object.GetValue("W", -1);
+			int H = Object.GetValue("H", -1);
+
+			if(X == -1 || Y == -1 || W == -1 || H == -1)
+			{
+				return null;
+			}
+			else
+			{
+				return new Rectangle(X, Y, W, H);
+			}
+		}
+
+		static string FormatRectangleValue(Rectangle Value)
+		{
+			ConfigObject Object = new ConfigObject();
+
+			Object.SetValue("X", Value.X);
+			Object.SetValue("Y", Value.Y);
+			Object.SetValue("W", Value.Width);
+			Object.SetValue("H", Value.Height);
+
+			return Object.ToString();
+		}
+
+		public static DirectoryReference GetConfigDir(DirectoryReference WorkspaceDir)
+		{
+			DirectoryReference ConfigDir = DirectoryReference.Combine(WorkspaceDir, ".ugs");
+
+			DirectoryInfo ConfigDirInfo = ConfigDir.ToDirectoryInfo();
+			if (!ConfigDirInfo.Exists)
+			{
+				ConfigDirInfo.Create();
+				ConfigDirInfo.Attributes = FileAttributes.Directory | FileAttributes.Hidden;
+			}
+
+			return ConfigDir;
+		}
+
+		public UserWorkspaceState FindOrAddWorkspaceState(UserWorkspaceSettings Settings)
+		{
+			UserWorkspaceState? State;
+			if (!WorkspaceDirToState.TryGetValue(Settings.RootDir, out State))
+			{
+				if (!UserWorkspaceState.TryLoad(Settings.RootDir, out State))
+				{
+					State = ImportWorkspaceState(Settings.RootDir, Settings.ClientName, Settings.BranchPath);
+					State.Save();
+				}
+				WorkspaceDirToState.Add(Settings.RootDir, State);
+			}
+			if (!String.Equals(Settings.ClientProjectPath, State.CurrentProjectIdentifier))
+			{
+				State = new UserWorkspaceState();
+				State.RootDir = Settings.RootDir;
+				State.CurrentProjectIdentifier = Settings.ClientProjectPath;
+				WorkspaceDirToState[Settings.RootDir] = State;
+			}
+			return State;
+		}
+
+		public UserWorkspaceSettings FindOrAddWorkspaceSettings(DirectoryReference RootDir, string? ServerAndPort, string? UserName, string ClientName, string BranchPath, string ProjectPath)
+		{
+			ProjectInfo.ValidateBranchPath(BranchPath);
+			ProjectInfo.ValidateProjectPath(ProjectPath);
+
+			UserWorkspaceSettings? Settings;
+			if (!WorkspaceDirToSettings.TryGetValue(RootDir, out Settings))
+			{
+				if (!UserWorkspaceSettings.TryLoad(RootDir, out Settings))
+				{
+					Settings = ImportWorkspaceSettings(RootDir, ClientName, BranchPath);
+				}
+				WorkspaceDirToSettings[RootDir] = Settings;
+			}
+
+			Settings.ServerAndPort = ServerAndPort;
+			Settings.UserName = UserName;
+			Settings.ClientName = ClientName;
+			Settings.BranchPath = BranchPath;
+			Settings.ProjectPath = ProjectPath;
+			Settings.Save();
+
+			return Settings;
+		}
+
+		UserWorkspaceState ImportWorkspaceState(DirectoryReference RootDir, string ClientName, string BranchPath)
+		{
+			UserWorkspaceState CurrentWorkspace = new UserWorkspaceState();
+			CurrentWorkspace.RootDir = RootDir;
+
+			// Read the workspace settings
+			ConfigSection WorkspaceSection = ConfigFile.FindSection(ClientName + BranchPath);
+			if(WorkspaceSection == null)
+			{
+				string LegacyBranchAndClientKey = ClientName + BranchPath;
+
+				int SlashIdx = LegacyBranchAndClientKey.IndexOf('/');
+				if(SlashIdx != -1)
+				{
+					LegacyBranchAndClientKey = LegacyBranchAndClientKey.Substring(0, SlashIdx) + "$" + LegacyBranchAndClientKey.Substring(SlashIdx + 1);
+				}
+
+				string? CurrentSync = ConfigFile.GetValue("Clients." + LegacyBranchAndClientKey, null);
+				if(CurrentSync != null)
+				{
+					int AtIdx = CurrentSync.LastIndexOf('@');
+					if(AtIdx != -1)
+					{
+						int ChangeNumber;
+						if(int.TryParse(CurrentSync.Substring(AtIdx + 1), out ChangeNumber))
+						{
+							CurrentWorkspace.CurrentProjectIdentifier = CurrentSync.Substring(0, AtIdx);
+							CurrentWorkspace.CurrentChangeNumber = ChangeNumber;
+						}
+					}
+				}
+
+				string? LastUpdateResultText = ConfigFile.GetValue("Clients." + LegacyBranchAndClientKey + "$LastUpdate", null);
+				if(LastUpdateResultText != null)
+				{
+					int ColonIdx = LastUpdateResultText.LastIndexOf(':');
+					if(ColonIdx != -1)
+					{
+						int ChangeNumber;
+						if(int.TryParse(LastUpdateResultText.Substring(0, ColonIdx), out ChangeNumber))
+						{
+							WorkspaceUpdateResult Result;
+							if(Enum.TryParse(LastUpdateResultText.Substring(ColonIdx + 1), out Result))
+							{
+								CurrentWorkspace.LastSyncChangeNumber = ChangeNumber;
+								CurrentWorkspace.LastSyncResult = Result;
+							}
+						}
+					}
+				}
+			}
+			else
+			{
+				CurrentWorkspace.CurrentProjectIdentifier = WorkspaceSection.GetValue("CurrentProjectPath");
+				CurrentWorkspace.CurrentChangeNumber = WorkspaceSection.GetValue("CurrentChangeNumber", -1);
+				CurrentWorkspace.CurrentSyncFilterHash = WorkspaceSection.GetValue("CurrentSyncFilterHash", null);
+				foreach(string AdditionalChangeNumberString in WorkspaceSection.GetValues("AdditionalChangeNumbers", new string[0]))
+				{
+					int AdditionalChangeNumber;
+					if(int.TryParse(AdditionalChangeNumberString, out AdditionalChangeNumber))
+					{
+						CurrentWorkspace.AdditionalChangeNumbers.Add(AdditionalChangeNumber);
+					}
+				}
+
+				WorkspaceUpdateResult LastSyncResult;
+				Enum.TryParse(WorkspaceSection.GetValue("LastSyncResult", ""), out LastSyncResult);
+				CurrentWorkspace.LastSyncResult = LastSyncResult;
+
+				CurrentWorkspace.LastSyncResultMessage = UnescapeText(WorkspaceSection.GetValue("LastSyncResultMessage"));
+				CurrentWorkspace.LastSyncChangeNumber = WorkspaceSection.GetValue("LastSyncChangeNumber", -1);
+
+				DateTime LastSyncTime;
+				if(DateTime.TryParse(WorkspaceSection.GetValue("LastSyncTime", ""), out LastSyncTime))
+				{
+					CurrentWorkspace.LastSyncTime = LastSyncTime;
+				}
+
+				CurrentWorkspace.LastSyncDurationSeconds = WorkspaceSection.GetValue("LastSyncDuration", 0);
+				CurrentWorkspace.LastBuiltChangeNumber = WorkspaceSection.GetValue("LastBuiltChangeNumber", 0);
+				CurrentWorkspace.ExpandedArchiveTypes = WorkspaceSection.GetValues("ExpandedArchiveName", new string[0]);
+
+				string[] BisectEntries = WorkspaceSection.GetValues("Bisect", new string[0]);
+				foreach(string BisectEntry in BisectEntries)
+				{
+					ConfigObject BisectEntryObject = new ConfigObject(BisectEntry);
+
+					int ChangeNumber = BisectEntryObject.GetValue("Change", -1);
+					if(ChangeNumber != -1)
+					{
+						BisectState State;
+						if(Enum.TryParse(BisectEntryObject.GetValue("State", ""), out State))
+						{
+							BisectEntry Entry = new BisectEntry();
+							Entry.Change = ChangeNumber;
+							Entry.State = State;
+							CurrentWorkspace.BisectChanges.Add(Entry);
+						}
+					}
+				}
+			}
+
+			return CurrentWorkspace;
+		}
+
+		UserWorkspaceSettings ImportWorkspaceSettings(DirectoryReference RootDir, string ClientName, string BranchPath)
+		{
+			UserWorkspaceSettings CurrentWorkspace = new UserWorkspaceSettings();
+			CurrentWorkspace.RootDir = RootDir;
+
+			ConfigSection WorkspaceSection = ConfigFile.FindSection(ClientName + BranchPath);
+			if (WorkspaceSection == null)
+			{
+				CurrentWorkspace.SyncView = new string[0];
+				CurrentWorkspace.SyncCategories = new List<SyncCategory>();
+				CurrentWorkspace.bSyncAllProjects = null;
+				CurrentWorkspace.bIncludeAllProjectsInSolution = null;
+			}
+			else
+			{
+				CurrentWorkspace.SyncView = WorkspaceSection.GetValues("SyncFilter", new string[0]);
+				CurrentWorkspace.SyncCategories = GetCategorySettings(WorkspaceSection, "SyncIncludedCategories", "SyncExcludedCategories").Select(x => new SyncCategory { Id = x.Key, Enable = x.Value }).ToList();
+
+				int SyncAllProjects = WorkspaceSection.GetValue("SyncAllProjects", -1);
+				CurrentWorkspace.bSyncAllProjects = (SyncAllProjects == 0) ? (bool?)false : (SyncAllProjects == 1) ? (bool?)true : (bool?)null;
+
+				int IncludeAllProjectsInSolution = WorkspaceSection.GetValue("IncludeAllProjectsInSolution", -1);
+				CurrentWorkspace.bIncludeAllProjectsInSolution = (IncludeAllProjectsInSolution == 0) ? (bool?)false : (IncludeAllProjectsInSolution == 1) ? (bool?)true : (bool?)null;
+			}
+
+			return CurrentWorkspace;
+		}
+
+		public UserProjectSettings FindOrAddProjectSettings(DirectoryReference LocalRoot, string ClientProjectFileName)
+		{
+			FileReference ConfigFile = FileReference.Combine(GetConfigDir(LocalRoot), $"project_{ClientProjectFileName.Trim('/').Replace('/', '+')}.json");
+			if (!ProjectKeyToSettings.TryGetValue(ConfigFile, out UserProjectSettings? Settings))
+			{
+				if (!Utility.TryLoadJson(ConfigFile, out Settings))
+				{
+					Settings = ImportProjectSettings(ClientProjectFileName);
+					Utility.SaveJson(ConfigFile, Settings);
+				}
+
+				Settings.File = ConfigFile;
+				ProjectKeyToSettings.Add(ConfigFile, Settings);
+			}
+			return Settings;
+		}
+
+		public UserProjectSettings ImportProjectSettings(string ClientProjectFileName)
+		{
+			// Read the project settings
+			UserProjectSettings CurrentProject = new UserProjectSettings();
+	
+			ConfigSection ProjectSection = ConfigFile.FindOrAddSection(ClientProjectFileName);
+			CurrentProject.BuildSteps.AddRange(ProjectSection.GetValues("BuildStep", new string[0]).Select(x => new ConfigObject(x)));
+
+			FilterType FilterType;
+			if (!Enum.TryParse(ProjectSection.GetValue("FilterType", ""), true, out FilterType))
+			{
+				FilterType = FilterType.None;
+			}
+
+			CurrentProject.FilterType = FilterType;
+			CurrentProject.FilterBadges.UnionWith(ProjectSection.GetValues("FilterBadges", new string[0]));
+
+			return CurrentProject;
+		}
+
+		public void Save()
+		{
+			// General settings
+			ConfigSection GeneralSection = ConfigFile.FindOrAddSection("General");
+			GeneralSection.Clear();
+			GeneralSection.SetValue("Version", (int)Version);
+			GeneralSection.SetValue("BuildAfterSync", bBuildAfterSync);
+			GeneralSection.SetValue("RunAfterSync", bRunAfterSync);
+			GeneralSection.SetValue("OpenSolutionAfterSync", bOpenSolutionAfterSync);
+			GeneralSection.SetValue("ShowLogWindow", bShowLogWindow);
+			GeneralSection.SetValue("AutoResolveConflicts", bAutoResolveConflicts);
+			GeneralSection.SetValue("ShowUnreviewed", bShowUnreviewedChanges);
+			GeneralSection.SetValue("ShowAutomated", bShowAutomatedChanges);
+			GeneralSection.SetValue("RobomergeFilter", ShowRobomerge.ToString());
+			GeneralSection.SetValue("AnnotateRobomerge", bAnnotateRobmergeChanges);
+			GeneralSection.SetValue("ShowLocalTimes", bShowLocalTimes);
+			if(LastProject != null)
+			{
+				GeneralSection.SetValue("LastProject", LastProject.ToConfigEntry());
+			}
+			GeneralSection.SetValues("OpenProjects", OpenProjects.Select(x => x.ToConfigEntry()).ToArray());
+			GeneralSection.SetValue("KeepInTray", bKeepInTray);
+			GeneralSection.SetValues("EnabledTools", EnabledTools);
+			GeneralSection.SetValue("ShowNetCoreInfo", bShowNetCoreInfo);
+			GeneralSection.SetValue("FilterIndex", FilterIndex);
+			GeneralSection.SetValues("RecentProjects", RecentProjects.Select(x => x.ToConfigEntry()).ToArray());
+			GeneralSection.SetValues("SyncFilter", SyncView);
+			SetCategorySettings(GeneralSection, "SyncIncludedCategories", "SyncExcludedCategories", SyncCategories);
+			GeneralSection.SetValue("SyncAllProjects", bSyncAllProjects);
+			GeneralSection.SetValue("IncludeAllProjectsInSolution", bIncludeAllProjectsInSolution);
+			GeneralSection.SetValue("SyncType", SyncType.ToString());
+
+			// Build configuration
+			GeneralSection.SetValue("BuildConfig", CompiledEditorBuildConfig.ToString());
+
+			// Tab labels
+			GeneralSection.SetValue("TabLabels", TabLabels.ToString());
+
+			// Editor arguments
+			List<string> EditorArgumentList = new List<string>();
+			foreach(Tuple<string, bool> EditorArgument in EditorArguments)
+			{
+				EditorArgumentList.Add(String.Format("{0}:{1}", EditorArgument.Item2? 1 : 0, EditorArgument.Item1));
+			}
+			GeneralSection.SetValues("EditorArguments", EditorArgumentList.ToArray());
+			GeneralSection.SetValue("EditorArgumentsPrompt", bEditorArgumentsPrompt);
+
+			// Schedule settings
+			ConfigSection ScheduleSection = ConfigFile.FindOrAddSection("Schedule");
+			ScheduleSection.Clear();
+			ScheduleSection.SetValue("Enabled", bScheduleEnabled);
+			ScheduleSection.SetValue("Time", ScheduleTime.ToString());
+			ScheduleSection.SetValue("Change", ScheduleChange.ToString());
+			ScheduleSection.SetValue("AnyOpenProject", ScheduleAnyOpenProject);
+			ScheduleSection.SetValues("Projects", ScheduleProjects.Select(x => x.ToConfigEntry()).ToArray());
+
+			// Precompiled binaries
+			ConfigSection ArchivesSection = ConfigFile.FindOrAddSection("PrecompiledBinaries");
+			ArchivesSection.SetValues("Archives", Archives.Select(x => x.ToConfigEntry()).ToArray());
+
+			// OIDC Settings
+			ConfigSection OIDCSection = ConfigFile.FindOrAddSection("OIDCProviders");
+			List<ConfigObject> TokenObjects = new List<ConfigObject>();
+			foreach (KeyValuePair<string, string> Pair in ProviderToRefreshTokens)
+			{
+				ConfigObject TokenEntryObject = new ConfigObject();
+				TokenEntryObject.SetValue("Provider", Pair.Key);
+				TokenEntryObject.SetValue("Token", Pair.Value);
+				TokenObjects.Add(TokenEntryObject);
+			}
+			OIDCSection.SetValues("Tokens", TokenObjects.Select(x => x.ToString()).ToArray());
+
+
+			// Window settings
+			ConfigSection WindowSection = ConfigFile.FindOrAddSection("Window");
+			WindowSection.Clear();
+			WindowSection.SetValue("Visible", bWindowVisible);
+			WindowSection.SetValue("State", WindowState);
+			if(WindowBounds != null)
+			{
+				WindowSection.SetValue("Bounds", FormatRectangleValue(WindowBounds.Value));
+			}
+
+			// Notification settings
+			ConfigSection NotificationSection = ConfigFile.FindOrAddSection("Notifications");
+			NotificationSection.Clear();
+			if (NotifyProjects.Count > 0)
+			{
+				NotificationSection.SetValues("NotifyProjects", NotifyProjects.ToArray());
+			}
+			if (NotifyUnassignedMinutes != -1)
+			{
+				NotificationSection.SetValue("NotifyUnassignedMinutes", NotifyUnassignedMinutes);
+			}
+			if (NotifyUnacknowledgedMinutes != -1)
+			{
+				NotificationSection.SetValue("NotifyUnacknowledgedMinutes", NotifyUnacknowledgedMinutes);
+			}
+			if (NotifyUnresolvedMinutes != -1)
+			{
+				NotificationSection.SetValue("NotifyUnresolvedMinutes", NotifyUnresolvedMinutes);
+			}
+
+			// Perforce settings
+			ConfigSection PerforceSection = ConfigFile.FindOrAddSection("Perforce");
+			PerforceSection.Clear();
+			if(SyncOptions.NumRetries > 0 && SyncOptions.NumRetries != PerforceSyncOptions.DefaultNumRetries)
+			{
+				PerforceSection.SetValue("NumRetries", SyncOptions.NumRetries);
+			}
+			if(SyncOptions.NumThreads > 0 && SyncOptions.NumThreads != PerforceSyncOptions.DefaultNumThreads)
+			{
+				PerforceSection.SetValue("NumThreads", SyncOptions.NumThreads);
+			}
+			if(SyncOptions.TcpBufferSize > 0 && SyncOptions.TcpBufferSize != PerforceSyncOptions.DefaultTcpBufferSize)
+			{
+				PerforceSection.SetValue("TcpBufferSize", SyncOptions.TcpBufferSize);
+			}
+			if (SyncOptions.FileBufferSize > 0 && SyncOptions.FileBufferSize != PerforceSyncOptions.DefaultFileBufferSize)
+			{
+				PerforceSection.SetValue("FileBufferSize", SyncOptions.FileBufferSize);
+			}
+			if (SyncOptions.MaxCommandsPerBatch > 0 && SyncOptions.MaxCommandsPerBatch != PerforceSyncOptions.DefaultMaxCommandsPerBatch)
+			{
+				PerforceSection.SetValue("MaxCommandsPerBatch", SyncOptions.MaxCommandsPerBatch);
+			}
+			if (SyncOptions.MaxSizePerBatch > 0 && SyncOptions.MaxSizePerBatch != PerforceSyncOptions.DefaultMaxSizePerBatch)
+			{
+				PerforceSection.SetValue("MaxSizePerBatch", SyncOptions.MaxSizePerBatch);
+			}
+
+			// Save the file
+			ConfigFile.Save(FileName);
+		}
+
+		public static string[] GetCombinedSyncFilter(Dictionary<Guid, WorkspaceSyncCategory> UniqueIdToFilter, string[] GlobalView, Dictionary<Guid, bool> GlobalCategoryIdToSetting, string[] WorkspaceView, Dictionary<Guid, bool> WorkspaceCategoryIdToSetting)
+		{
+			List<string> Lines = new List<string>();
+			foreach(string ViewLine in Enumerable.Concat(GlobalView, WorkspaceView).Select(x => x.Trim()).Where(x => x.Length > 0 && !x.StartsWith(";")))
+			{
+				Lines.Add(ViewLine);
+			}
+
+			HashSet<Guid> Enabled = new HashSet<Guid>();
+			foreach (WorkspaceSyncCategory Filter in UniqueIdToFilter.Values)
+			{
+				bool bEnable = Filter.bEnable;
+
+				bool bGlobalEnable;
+				if (GlobalCategoryIdToSetting.TryGetValue(Filter.UniqueId, out bGlobalEnable))
+				{
+					bEnable = bGlobalEnable;
+				}
+
+				bool bWorkspaceEnable;
+				if (WorkspaceCategoryIdToSetting.TryGetValue(Filter.UniqueId, out bWorkspaceEnable))
+				{
+					bEnable = bWorkspaceEnable;
+				}
+
+				if(bEnable)
+				{
+					EnableFilter(Filter.UniqueId, Enabled, UniqueIdToFilter);
+				}
+			}
+
+			foreach (WorkspaceSyncCategory Filter in UniqueIdToFilter.Values.OrderBy(x => x.Name))
+			{
+				if (!Enabled.Contains(Filter.UniqueId))
+				{
+					Lines.AddRange(Filter.Paths.Select(x => "-" + x.Trim()));
+				}
+			}
+
+			return Lines.ToArray();
+		}
+
+		static void EnableFilter(Guid UniqueId, HashSet<Guid> Enabled, Dictionary<Guid, WorkspaceSyncCategory> UniqueIdToFilter)
+		{
+			if(Enabled.Add(UniqueId))
+			{
+				WorkspaceSyncCategory? Category;
+				if(UniqueIdToFilter.TryGetValue(UniqueId, out Category))
+				{
+					foreach(Guid RequiresUniqueId in Category.Requires)
+					{
+						EnableFilter(RequiresUniqueId, Enabled, UniqueIdToFilter);
+					}
+				}
+			}
+		}
+
+		[return: NotNullIfNotNull("Text")]
+		static string? EscapeText(string? Text)
+		{
+			if(Text == null)
+			{
+				return null;
+			}
+
+			StringBuilder Result = new StringBuilder();
+			for(int Idx = 0; Idx < Text.Length; Idx++)
+			{
+				switch(Text[Idx])
+				{
+					case '\\':
+						Result.Append("\\\\");
+						break;
+					case '\t':
+						Result.Append("\\t");
+						break;
+					case '\r':
+						Result.Append("\\r");
+						break;
+					case '\n':
+						Result.Append("\\n");
+						break;
+					case '\'':
+						Result.Append("\\\'");
+						break;
+					case '\"':
+						Result.Append("\\\"");
+						break;
+					default:
+						Result.Append(Text[Idx]);
+						break;
+				}
+			}
+			return Result.ToString();
+		}
+
+		[return: NotNullIfNotNull("Text")]
+		static string? UnescapeText(string? Text)
+		{
+			if(Text == null)
+			{
+				return null;
+			}
+
+			StringBuilder Result = new StringBuilder();
+			for(int Idx = 0; Idx < Text.Length; Idx++)
+			{
+				if(Text[Idx] == '\\' && Idx + 1 < Text.Length)
+				{
+					switch(Text[++Idx])
+					{
+						case 't':
+							Result.Append('\t');
+							break;
+						case 'r':
+							Result.Append('\r');
+							break;
+						case 'n':
+							Result.Append('\n');
+							break;
+						case '\'':
+							Result.Append('\'');
+							break;
+						case '\"':
+							Result.Append('\"');
+							break;
+						default:
+							Result.Append(Text[Idx]);
+							break;
+					}
+				}
+				else
+				{
+					Result.Append(Text[Idx]);
+				}
+			}
+			return Result.ToString();
+		}
+
+		public IEnumerable<FileReference> GetCachedFilePaths()
+		{
+			List<FileReference> Files = new List<FileReference>();
+			Files.AddRange(WorkspaceDirToSettings.Values.Select(x => x.ConfigFile));
+			Files.AddRange(WorkspaceDirToState.Values.Select(x => x.ConfigFile));
+			Files.AddRange(ProjectKeyToSettings.Keys);
+			return Files;
+		}
+	}
+}
diff --git a/Engine/Source/Programs/UnrealGameSync/UnrealGameSyncShared/WorkspaceUpdate.cs b/Engine/Source/Programs/UnrealGameSync/UnrealGameSyncShared/WorkspaceUpdate.cs
new file mode 100644
--- /dev/null
+++ b/Engine/Source/Programs/UnrealGameSync/UnrealGameSyncShared/WorkspaceUpdate.cs
@@ -0,0 +1,1846 @@
+// Copyright Epic Games, Inc. All Rights Reserved.
+
+using EpicGames.Core;
+using EpicGames.Perforce;
+using Microsoft.Extensions.DependencyInjection;
+using Microsoft.Extensions.Logging;
+using System;
+using System.Collections.Concurrent;
+using System.Collections.Generic;
+using System.Diagnostics;
+using System.Diagnostics.CodeAnalysis;
+using System.Globalization;
+using System.IO;
+using System.IO.Compression;
+using System.Linq;
+using System.Security.Cryptography;
+using System.Text;
+using System.Text.Encodings.Web;
+using System.Text.Json;
+using System.Threading;
+using System.Threading.Tasks;
+
+namespace UnrealGameSync
+{
+	[Flags]
+	public enum WorkspaceUpdateOptions
+	{
+		Sync = 0x01,
+		SyncSingleChange = 0x02,
+		AutoResolveChanges = 0x04,
+		GenerateProjectFiles = 0x08,
+		SyncArchives = 0x10,
+		Build = 0x20,
+		Clean = 0x40,
+		ScheduledBuild = 0x80,
+		RunAfterSync = 0x100,
+		OpenSolutionAfterSync = 0x200,
+		ContentOnly = 0x400,
+		UpdateFilter = 0x800,
+		SyncAllProjects = 0x1000,
+		IncludeAllProjectsInSolution = 0x2000,
+	}
+
+	public enum WorkspaceUpdateResult
+	{
+		Canceled,
+		FailedToSync,
+		FailedToSyncLoginExpired,
+		FilesToDelete,
+		FilesToResolve,
+		FilesToClobber,
+		FailedToCompile,
+		FailedToCompileWithCleanWorkspace,
+		Success,
+	}
+
+	public class PerforceSyncOptions
+	{
+		public const int DefaultNumRetries = 0;
+		public const int DefaultNumThreads = 2;
+		public const int DefaultTcpBufferSize = 0;
+		public const int DefaultFileBufferSize = 0;
+
+		public const int DefaultMaxCommandsPerBatch = 200;
+		public const int DefaultMaxSizePerBatch = 128 * 1024 * 1024;
+
+		public int NumRetries = DefaultNumRetries;
+		public int NumThreads = DefaultNumThreads;
+		public int TcpBufferSize = DefaultTcpBufferSize;
+		public int FileBufferSize = DefaultFileBufferSize;
+
+		public int MaxCommandsPerBatch = DefaultMaxCommandsPerBatch;
+		public int MaxSizePerBatch = DefaultMaxSizePerBatch;
+
+
+		public PerforceSyncOptions Clone()
+		{
+			return (PerforceSyncOptions)MemberwiseClone();
+		}
+	}
+
+	public interface IArchiveInfo
+	{
+		string Name { get; }
+		string Type { get; }
+		string BasePath { get; }
+		string? Target { get; }
+		bool Exists();
+		bool TryGetArchiveKeyForChangeNumber(int ChangeNumber, [NotNullWhen(true)] out string? ArchiveKey);
+		Task<bool> DownloadArchive(IPerforceConnection Perforce, string ArchiveKey, DirectoryReference LocalRootPath, FileReference ManifestFileName, ILogger Logger, ProgressValue Progress, CancellationToken CancellationToken);
+	}
+
+	public class WorkspaceUpdateContext
+	{
+		public DateTime StartTime = DateTime.UtcNow;
+		public int ChangeNumber;
+		public WorkspaceUpdateOptions Options;
+		public string[]? SyncFilter;
+		public Dictionary<string, Tuple<IArchiveInfo, string>?> ArchiveTypeToArchive = new Dictionary<string, Tuple<IArchiveInfo, string>?>();
+		public Dictionary<string, bool> DeleteFiles = new Dictionary<string,bool>();
+		public Dictionary<string, bool> ClobberFiles = new Dictionary<string,bool>();
+		public Dictionary<Guid,ConfigObject> DefaultBuildSteps;
+		public List<ConfigObject> UserBuildStepObjects;
+		public HashSet<Guid>? CustomBuildSteps;
+		public Dictionary<string, string> Variables;
+		public PerforceSyncOptions? PerforceSyncOptions;
+		public List<HaveRecord>? HaveFiles; // Cached when sync filter has changed
+
+		// May be updated during sync
+		public ConfigFile? ProjectConfigFile;
+		public IReadOnlyList<string>? ProjectStreamFilter;
+
+		public WorkspaceUpdateContext(int InChangeNumber, WorkspaceUpdateOptions InOptions, string[]? InSyncFilter, Dictionary<Guid, ConfigObject> InDefaultBuildSteps, List<ConfigObject> InUserBuildSteps, HashSet<Guid>? InCustomBuildSteps, Dictionary<string, string> InVariables)
+		{
+			ChangeNumber = InChangeNumber;
+			Options = InOptions;
+			SyncFilter = InSyncFilter;
+			DefaultBuildSteps = InDefaultBuildSteps;
+			UserBuildStepObjects = InUserBuildSteps;
+			CustomBuildSteps = InCustomBuildSteps;
+			Variables = InVariables;
+		}
+	}
+
+	public class WorkspaceSyncCategory
+	{
+		public Guid UniqueId;
+		public bool bEnable;
+		public string Name;
+		public string[] Paths;
+		public bool bHidden;
+		public Guid[] Requires;
+
+		public WorkspaceSyncCategory(Guid UniqueId) : this(UniqueId, "Unnamed")
+		{
+		}
+
+		public WorkspaceSyncCategory(Guid UniqueId, string Name, params string[] Paths)
+		{
+			this.UniqueId = UniqueId;
+			this.bEnable = true;
+			this.Name = Name;
+			this.Paths = Paths;
+			this.Requires = new Guid[0];
+		}
+
+		public static Dictionary<Guid, bool> GetDefault(IEnumerable<WorkspaceSyncCategory> Categories)
+		{
+			return Categories.ToDictionary(x => x.UniqueId, x => x.bEnable);
+		}
+
+		public static Dictionary<Guid, bool> GetDelta(Dictionary<Guid, bool> Source, Dictionary<Guid, bool> Target)
+		{
+			Dictionary<Guid, bool> Changes = new Dictionary<Guid, bool>();
+			foreach (KeyValuePair<Guid, bool> Pair in Target)
+			{
+				bool bValue;
+				if (!Source.TryGetValue(Pair.Key, out bValue) || bValue != Pair.Value)
+				{
+					Changes[Pair.Key] = Pair.Value;
+				}
+			}
+			return Changes;
+		}
+
+		public static void ApplyDelta(Dictionary<Guid, bool> Categories, Dictionary<Guid, bool> Delta)
+		{
+			foreach(KeyValuePair<Guid, bool> Pair in Delta)
+			{
+				Categories[Pair.Key] = Pair.Value;
+			}
+		}
+
+		public override string ToString()
+		{
+			return Name;
+		}
+	}
+
+	public class ProjectInfo
+	{
+		public DirectoryReference LocalRootPath { get; } // ie. local mapping of clientname + branchpath
+		public string ClientName { get; }
+		public string BranchPath { get; } // starts with a slash if non-empty. does not end with a slash.
+		public string ProjectPath { get; } // starts with a slash, uses forward slashes
+		public string ProjectIdentifier { get; } // stream path to project
+		public string TelemetryProjectIdentifier { get; }
+		public bool bIsEnterpriseProject { get; }
+
+		public FileReference LocalFileName => new FileReference(LocalRootPath.FullName + ProjectPath);
+		public string ClientRootPath => $"//{ClientName}{BranchPath}";
+		public string ClientFileName => $"//{ClientName}{BranchPath}{ProjectPath}";
+
+		public ProjectInfo(DirectoryReference LocalRootPath, string ClientName, string BranchPath, string ProjectPath, string ProjectIdentifier, bool bIsEnterpriseProject)
+		{
+			ValidateBranchPath(BranchPath);
+			ValidateProjectPath(ProjectPath);
+
+			this.LocalRootPath = LocalRootPath;
+			this.ClientName = ClientName;
+			this.BranchPath = BranchPath;
+			this.ProjectPath = ProjectPath;
+			this.ProjectIdentifier = ProjectIdentifier;
+			this.TelemetryProjectIdentifier = PerforceUtils.GetClientOrDepotDirectoryName(ProjectIdentifier);
+			this.bIsEnterpriseProject = bIsEnterpriseProject;
+		}
+
+		public static void ValidateBranchPath(string BranchPath)
+		{
+			if (BranchPath.Length > 0 && (!BranchPath.StartsWith("/") || BranchPath.EndsWith("/")))
+			{
+				throw new ArgumentException("Branch path must start with a slash, and not end with a slash", nameof(BranchPath));
+			}
+		}
+
+		public static void ValidateProjectPath(string ProjectPath)
+		{
+			if (!ProjectPath.StartsWith("/"))
+			{
+				throw new ArgumentException("Project path must start with a slash", nameof(ProjectPath));
+			}
+			if (!ProjectPath.EndsWith(".uproject", StringComparison.OrdinalIgnoreCase) && !ProjectPath.EndsWith(".uprojectdirs", StringComparison.OrdinalIgnoreCase))
+			{
+				throw new ArgumentException("Project path must be to a .uproject or .uprojectdirs file", nameof(ProjectPath));
+			}
+		}
+	}
+
+	public class WorkspaceUpdate : IAsyncDisposable
+	{
+		const string BuildVersionFileName = "/Engine/Build/Build.version";
+		const string VersionHeaderFileName = "/Engine/Source/Runtime/Launch/Resources/Version.h";
+		const string ObjectVersionFileName = "/Engine/Source/Runtime/Core/Private/UObject/ObjectVersion.cpp";
+
+		static readonly string LocalVersionHeaderFileName = VersionHeaderFileName.Replace('/', '\\');
+		static readonly string LocalObjectVersionFileName = ObjectVersionFileName.Replace('/', '\\');
+
+		static SemaphoreSlim UpdateSemaphore = new SemaphoreSlim(1);
+
+		class RecordCounter : IDisposable
+		{
+			ProgressValue Progress;
+			string Message;
+			int Count;
+			Stopwatch Timer = Stopwatch.StartNew();
+
+			public RecordCounter(ProgressValue Progress, string Message)
+			{
+				this.Progress = Progress;
+				this.Message = Message;
+
+				Progress.Set(Message);
+			}
+
+			public void Dispose()
+			{
+				UpdateMessage();
+			}
+
+			public void Increment()
+			{
+				Count++;
+				if (Timer.ElapsedMilliseconds > 250)
+				{
+					UpdateMessage();
+				}
+			}
+
+			public void UpdateMessage()
+			{
+				Progress.Set(String.Format("{0} ({1:N0})", Message, Count));
+				Timer.Restart();
+			}
+		}
+
+		class SyncBatchBuilder
+		{
+			public int MaxCommandsPerList { get; }
+			public long MaxSizePerList { get; }
+			public Queue<List<string>> Batches { get; }
+
+			List<string>? Commands;
+			List<string>? DeleteCommands;
+			long Size;
+
+			public SyncBatchBuilder(int? MaxCommandsPerList, long? MaxSizePerList)
+			{
+				this.MaxCommandsPerList = MaxCommandsPerList ?? PerforceSyncOptions.DefaultMaxCommandsPerBatch;
+				this.MaxSizePerList = MaxSizePerList ?? PerforceSyncOptions.DefaultMaxSizePerBatch;
+				this.Batches = new Queue<List<string>>();
+			}
+
+			public void Add(string NewCommand, long NewSize)
+			{
+				if (NewSize == 0)
+				{
+					if (DeleteCommands == null || DeleteCommands.Count >= MaxCommandsPerList)
+					{
+						DeleteCommands = new List<string>();
+						Batches.Enqueue(DeleteCommands);
+					}
+
+					DeleteCommands.Add(NewCommand);
+				}
+				else
+				{
+					if (Commands == null || Commands.Count >= MaxCommandsPerList || Size + NewSize >= MaxSizePerList)
+					{
+						Commands = new List<string>();
+						Batches.Enqueue(Commands);
+						Size = 0;
+					}
+
+					Commands.Add(NewCommand);
+					Size += NewSize;
+				}
+			}
+		}
+
+		class SyncTree
+		{
+			public bool bCanUseWildcard;
+			public int TotalIncludedFiles;
+			public long TotalSize;
+			public int TotalExcludedFiles;
+			public Dictionary<string, long> IncludedFiles = new Dictionary<string, long>();
+			public Dictionary<string, SyncTree> NameToSubTree = new Dictionary<string, SyncTree>(StringComparer.OrdinalIgnoreCase);
+
+			public SyncTree(bool bCanUseWildcard)
+			{
+				this.bCanUseWildcard = bCanUseWildcard;
+			}
+
+			public SyncTree FindOrAddSubTree(string Name)
+			{
+				SyncTree? Result;
+				if (!NameToSubTree.TryGetValue(Name, out Result))
+				{
+					Result = new SyncTree(bCanUseWildcard);
+					NameToSubTree.Add(Name, Result);
+				}
+				return Result;
+			}
+
+			public void IncludeFile(string Path, long Size)
+			{
+				int Idx = Path.IndexOf('/');
+				if (Idx == -1)
+				{
+					IncludedFiles.Add(Path, Size);
+				}
+				else
+				{
+					SyncTree SubTree = FindOrAddSubTree(Path.Substring(0, Idx));
+					SubTree.IncludeFile(Path.Substring(Idx + 1), Size);
+				}
+				TotalIncludedFiles++;
+				TotalSize += Size;
+			}
+
+			public void ExcludeFile(string Path)
+			{
+				int Idx = Path.IndexOf('/');
+				if (Idx != -1)
+				{
+					SyncTree SubTree = FindOrAddSubTree(Path.Substring(0, Idx));
+					SubTree.ExcludeFile(Path.Substring(Idx + 1));
+				}
+				TotalExcludedFiles++;
+			}
+
+			public void GetOptimizedSyncCommands(string Prefix, int ChangeNumber, SyncBatchBuilder Builder)
+			{
+				if (bCanUseWildcard && TotalExcludedFiles == 0 && TotalSize < Builder.MaxSizePerList)
+				{
+					Builder.Add(String.Format("{0}/...@{1}", Prefix, ChangeNumber), TotalSize);
+				}
+				else
+				{
+					foreach (KeyValuePair<string, long> File in IncludedFiles)
+					{
+						Builder.Add(String.Format("{0}/{1}@{2}", Prefix, File.Key, ChangeNumber), File.Value);
+					}
+					foreach (KeyValuePair<string, SyncTree> Pair in NameToSubTree)
+					{
+						Pair.Value.GetOptimizedSyncCommands(String.Format("{0}/{1}", Prefix, PerforceUtils.EscapePath(Pair.Key)), ChangeNumber, Builder);
+					}
+				}
+			}
+		}
+
+		CancellationTokenSource CancellationSource;
+		public WorkspaceUpdateContext Context { get; }
+		public ProgressValue Progress { get; } = new ProgressValue();
+		public Task UpdateTask { get; }
+
+		public WorkspaceUpdate(WorkspaceUpdateContext Context, Func<WorkspaceUpdate, CancellationToken, Task> TaskFunc)
+		{
+			this.CancellationSource = new CancellationTokenSource();
+			this.Context = Context;
+			this.UpdateTask = Task.Run(() => TaskFunc(this, CancellationSource.Token));
+		}
+
+		public void Cancel()
+		{
+			CancellationSource.Cancel();
+		}
+
+		public ValueTask DisposeAsync()
+		{
+			Task DisposeTask = Task.CompletedTask;
+			if (UpdateTask != null)
+			{
+				CancellationSource.Cancel();
+				DisposeTask = UpdateTask.ContinueWith(Task => CancellationSource.Dispose());
+			}
+			return new ValueTask(DisposeTask);
+		}
+
+		class SyncFile
+		{
+			public string DepotFile;
+			public string RelativePath;
+			public long Size;
+
+			public SyncFile(string DepotFile, string RelativePath, long Size)
+			{
+				this.DepotFile = DepotFile;
+				this.RelativePath = RelativePath;
+				this.Size = Size;
+			}
+		};
+
+		class SemaphoreScope : IDisposable
+		{
+			SemaphoreSlim Semaphore;
+			public bool HasLock { get; private set; }
+
+			public SemaphoreScope(SemaphoreSlim Semaphore)
+			{
+				this.Semaphore = Semaphore;
+			}
+
+			public bool TryAcquire()
+			{
+				HasLock = HasLock || Semaphore.Wait(0);
+				return HasLock;
+			}
+
+			public async Task AcquireAsync(CancellationToken CancellationToken)
+			{
+				await Semaphore.WaitAsync(CancellationToken);
+				HasLock = true;
+			}
+
+			public void Release()
+			{
+				if (HasLock)
+				{
+					Semaphore.Release();
+					HasLock = false;
+				}
+			}
+
+			public void Dispose() => Release();
+		}
+
+		public async Task<(WorkspaceUpdateResult, string)> ExecuteAsync(IPerforceSettings PerforceSettings, ProjectInfo Project, UserWorkspaceState State, ILogger Logger, CancellationToken CancellationToken)
+		{
+			using IPerforceConnection Perforce = await PerforceConnection.CreateAsync(PerforceSettings, Logger);
+
+			string CmdExe = Path.Combine(Environment.GetFolderPath(Environment.SpecialFolder.System), "cmd.exe");
+			if (!File.Exists(CmdExe))
+			{
+				return (WorkspaceUpdateResult.FailedToSync, $"Missing {CmdExe}.");
+			}
+
+			List<Tuple<string, TimeSpan>> Times = new List<Tuple<string, TimeSpan>>();
+
+			int NumFilesSynced = 0;
+			if (Context.Options.HasFlag(WorkspaceUpdateOptions.Sync) || Context.Options.HasFlag(WorkspaceUpdateOptions.SyncSingleChange))
+			{
+				using (TelemetryStopwatch SyncTelemetryStopwatch = new TelemetryStopwatch("Workspace_Sync", Project.TelemetryProjectIdentifier))
+				{
+					Logger.LogInformation("Syncing to {Change}...", Context.ChangeNumber);
+
+					// Make sure we're logged in
+					PerforceResponse<LoginRecord> LoginResponse = await Perforce.TryGetLoginStateAsync(CancellationToken);
+					if (!LoginResponse.Succeeded)
+					{
+						return (WorkspaceUpdateResult.FailedToSyncLoginExpired, "User is not logged in.");
+					}
+
+					// Figure out which paths to sync
+					List<string> RelativeSyncPaths = GetRelativeSyncPaths(Project, (Context.Options & WorkspaceUpdateOptions.SyncAllProjects) != 0, Context.SyncFilter);
+					List<string> SyncPaths = new List<string>(RelativeSyncPaths.Select(x => Project.ClientRootPath + x));
+
+					// Get the user's sync filter
+					FileFilter UserFilter = new FileFilter(FileFilterType.Include);
+					if (Context.SyncFilter != null)
+					{
+						UserFilter.AddRules(Context.SyncFilter.Select(x => x.Trim()).Where(x => x.Length > 0 && !x.StartsWith(";") && !x.StartsWith("#")));
+					}
+
+					// Check if the new sync filter matches the previous one. If not, we'll enumerate all files in the workspace and make sure there's nothing extra there.
+					string? NextSyncFilterHash = null;
+					using (SHA1Managed SHA = new SHA1Managed())
+					{
+						StringBuilder CombinedFilter = new StringBuilder();
+						foreach (string RelativeSyncPath in RelativeSyncPaths)
+						{
+							CombinedFilter.AppendFormat("{0}\n", RelativeSyncPath);
+						}
+						if (Context.SyncFilter != null)
+						{
+							CombinedFilter.Append("--FROM--\n");
+							CombinedFilter.Append(String.Join("\n", Context.SyncFilter));
+						}
+						NextSyncFilterHash = BitConverter.ToString(SHA.ComputeHash(Encoding.UTF8.GetBytes(CombinedFilter.ToString()))).Replace("-", "");
+					}
+
+					// If the hash differs, enumerate everything in the workspace to find what needs to be removed
+					if (NextSyncFilterHash != State.CurrentSyncFilterHash)
+					{
+						using (TelemetryStopwatch FilterStopwatch = new TelemetryStopwatch("Workspace_Sync_FilterChanged", Project.TelemetryProjectIdentifier))
+						{
+							Logger.LogInformation("Filter has changed ({PrevHash} -> {NextHash}); finding files in workspace that need to be removed.", (String.IsNullOrEmpty(State.CurrentSyncFilterHash)) ? "None" : State.CurrentSyncFilterHash, NextSyncFilterHash);
+
+							// Find all the files that are in this workspace
+							List<HaveRecord>? HaveFiles = Context.HaveFiles;
+							if (HaveFiles == null)
+							{
+								HaveFiles = new List<HaveRecord>();
+								using (RecordCounter HaveCounter = new RecordCounter(Progress, "Sync filter changed; checking workspace..."))
+								{
+									await foreach (PerforceResponse<HaveRecord> Record in Perforce.TryHaveAsync(FileSpecList.Any, CancellationToken))
+									{
+										if (Record.Succeeded)
+										{
+											HaveFiles.Add(Record.Data);
+											HaveCounter.Increment();
+										}
+										else
+										{
+											return (WorkspaceUpdateResult.FailedToSync, $"Unable to query files ({Record}).");
+										}
+									}
+								}
+								Context.HaveFiles = HaveFiles;
+							}
+
+							// Build a filter for the current sync paths
+							FileFilter SyncPathsFilter = new FileFilter(FileFilterType.Exclude);
+							foreach (string RelativeSyncPath in RelativeSyncPaths)
+							{
+								SyncPathsFilter.Include(RelativeSyncPath);
+							}
+
+							// Remove all the files that are not included by the filter
+							List<string> RemoveDepotPaths = new List<string>();
+							foreach (HaveRecord HaveFile in HaveFiles)
+							{
+								try
+								{
+									FileReference FullPath = new FileReference(HaveFile.Path);
+									if (MatchFilter(Project, FullPath, SyncPathsFilter) && !MatchFilter(Project, FullPath, UserFilter))
+									{
+										Logger.LogInformation("  {DepotFile}", HaveFile.DepotFile);
+										RemoveDepotPaths.Add(HaveFile.DepotFile);
+									}
+								}
+								catch (PathTooLongException)
+								{
+									// We don't actually care about this when looking for files to remove. Perforce may think that it's synced the path, and silently failed. Just ignore it.
+								}
+							}
+
+							// Check if there are any paths outside the regular sync paths
+							if (RemoveDepotPaths.Count > 0)
+							{
+								bool bDeleteListMatches = true;
+
+								Dictionary<string, bool> NewDeleteFiles = new Dictionary<string, bool>(StringComparer.OrdinalIgnoreCase);
+								foreach (string RemoveDepotPath in RemoveDepotPaths)
+								{
+									bool bDelete;
+									if (!Context.DeleteFiles.TryGetValue(RemoveDepotPath, out bDelete))
+									{
+										bDeleteListMatches = false;
+										bDelete = true;
+									}
+									NewDeleteFiles[RemoveDepotPath] = bDelete;
+								}
+								Context.DeleteFiles = NewDeleteFiles;
+
+								if (!bDeleteListMatches)
+								{
+									return (WorkspaceUpdateResult.FilesToDelete, $"Cancelled after finding {NewDeleteFiles.Count} files excluded by filter");
+								}
+
+								RemoveDepotPaths.RemoveAll(x => !Context.DeleteFiles[x]);
+							}
+
+							// Actually delete any files that we don't want
+							if (RemoveDepotPaths.Count > 0)
+							{
+								// Clear the current sync filter hash. If the sync is canceled, we'll be in an indeterminate state, and we should always clean next time round.
+								State.CurrentSyncFilterHash = "INVALID";
+								State.Save();
+
+								// Find all the depot paths that will be synced
+								HashSet<string> RemainingDepotPathsToRemove = new HashSet<string>(StringComparer.OrdinalIgnoreCase);
+								RemainingDepotPathsToRemove.UnionWith(RemoveDepotPaths);
+
+								// Build the list of revisions to sync
+								List<string> RevisionsToRemove = new List<string>();
+								RevisionsToRemove.AddRange(RemoveDepotPaths.Select(x => String.Format("{0}#0", x)));
+
+								(WorkspaceUpdateResult, string) RemoveResult = await SyncFileRevisions(Perforce, "Removing files...", Context, RevisionsToRemove, RemainingDepotPathsToRemove, Progress, Logger, CancellationToken);
+								if (RemoveResult.Item1 != WorkspaceUpdateResult.Success)
+								{
+									return RemoveResult;
+								}
+							}
+
+							// Update the sync filter hash. We've removed any files we need to at this point.
+							State.CurrentSyncFilterHash = NextSyncFilterHash;
+							State.Save();
+						}
+					}
+
+					// Create a filter for all the files we don't want
+					FileFilter Filter = new FileFilter(UserFilter);
+					Filter.Exclude(BuildVersionFileName);
+					if (Context.Options.HasFlag(WorkspaceUpdateOptions.ContentOnly))
+					{
+						Filter.Exclude("*.usf");
+						Filter.Exclude("*.ush");
+					}
+
+					// Create a tree to store the sync path
+					SyncTree SyncTree = new SyncTree(false);
+					if (!Context.Options.HasFlag(WorkspaceUpdateOptions.SyncSingleChange))
+					{
+						foreach (string RelativeSyncPath in RelativeSyncPaths)
+						{
+							const string WildcardSuffix = "/...";
+							if (RelativeSyncPath.EndsWith(WildcardSuffix, StringComparison.Ordinal))
+							{
+								SyncTree Leaf = SyncTree;
+
+								string[] Fragments = RelativeSyncPath.Split('/');
+								for (int Idx = 1; Idx < Fragments.Length - 1; Idx++)
+								{
+									Leaf = Leaf.FindOrAddSubTree(Fragments[Idx]);
+								}
+
+								Leaf.bCanUseWildcard = true;
+							}
+						}
+					}
+
+					// Find all the server changes, and anything that's opened for edit locally. We need to sync files we have open to schedule a resolve.
+					SyncBatchBuilder BatchBuilder = new SyncBatchBuilder(Context.PerforceSyncOptions?.MaxCommandsPerBatch, Context.PerforceSyncOptions?.MaxSizePerBatch);
+					List<string> SyncDepotPaths = new List<string>();
+					using (RecordCounter Counter = new RecordCounter(Progress, "Filtering files..."))
+					{
+						// Track the total new bytes that will be required on disk when syncing. Add an extra 100MB for padding.
+						long RequiredFreeSpace = 100 * 1024 * 1024;
+
+						foreach (string SyncPath in SyncPaths)
+						{
+							string SyncFilter = Context.Options.HasFlag(WorkspaceUpdateOptions.Sync) ? $"{Context.ChangeNumber}" : $"={Context.ChangeNumber}";
+
+							List<SyncFile> SyncFiles = new List<SyncFile>();
+							await foreach (PerforceResponse<SyncRecord> Response in Perforce.TrySyncAsync(SyncOptions.PreviewOnly, -1, $"{SyncPath}@{SyncFilter}", CancellationToken))
+							{
+								if (!Response.Succeeded)
+								{
+									return (WorkspaceUpdateResult.FailedToSync, $"Couldn't enumerate changes matching {SyncPath}.");
+								}
+
+								SyncRecord Record = Response.Data;
+
+								// Get the full local path
+								string RelativePath;
+								try
+								{
+									FileReference SyncFile = new FileReference(Record.Path);
+									RelativePath = PerforceUtils.GetClientRelativePath(Project.LocalRootPath, SyncFile);
+								}
+								catch (PathTooLongException)
+								{
+									Logger.LogInformation("The local path for {Path} exceeds the maximum allowed by Windows. Re-sync your workspace to a directory with a shorter name, or delete the file from the server.", Record.Path);
+									return (WorkspaceUpdateResult.FailedToSync, "File exceeds maximum path length allowed by Windows.");
+								}
+
+								// Create the sync record
+								long SyncSize = (Record.Action == SyncAction.Deleted) ? 0 : Record.FileSize;
+								SyncFiles.Add(new SyncFile(Record.DepotFile, RelativePath, SyncSize));
+								Counter.Increment();
+							}
+
+							// Also sync the currently open files
+							await foreach (PerforceResponse<OpenedRecord> Response in Perforce.TryOpenedAsync(OpenedOptions.None, FileSpecList.Any, CancellationToken))
+							{
+								if (!Response.Succeeded)
+								{
+									return (WorkspaceUpdateResult.FailedToSync, $"Couldn't enumerate changes matching {SyncPath}.");
+								}
+
+								OpenedRecord Record = Response.Data;
+								if (Record.Action != FileAction.Add || Record.Action != FileAction.Branch || Record.Action != FileAction.MoveAdd)
+								{
+									string RelativePath = PerforceUtils.GetClientRelativePath(Record.ClientFile);
+									SyncFiles.Add(new SyncFile(Record.DepotFile, RelativePath, 0));
+								}
+							}
+
+							// Enumerate all the files to be synced. NOTE: depotPath is escaped, whereas clientPath is not.
+							foreach (SyncFile SyncRecord in SyncFiles)
+							{
+								if (Filter.Matches(SyncRecord.RelativePath))
+								{
+									SyncTree.IncludeFile(PerforceUtils.EscapePath(SyncRecord.RelativePath), SyncRecord.Size);
+									SyncDepotPaths.Add(SyncRecord.DepotFile);
+									RequiredFreeSpace += SyncRecord.Size;
+
+									// If the file exists the required free space can be reduced as those bytes will be replaced.
+									FileInfo LocalFileInfo = FileReference.Combine(Project.LocalRootPath, SyncRecord.RelativePath).ToFileInfo();
+									if (LocalFileInfo.Exists)
+									{
+										RequiredFreeSpace -= LocalFileInfo.Length;
+									}
+								}
+								else
+								{
+									SyncTree.ExcludeFile(PerforceUtils.EscapePath(SyncRecord.RelativePath));
+								}
+							}
+						}
+
+						try
+						{
+							DirectoryInfo LocalRootInfo = Project.LocalRootPath.ToDirectoryInfo();
+							DriveInfo Drive = new DriveInfo(LocalRootInfo.Root.FullName);
+
+							if (Drive.AvailableFreeSpace < RequiredFreeSpace)
+							{
+								Logger.LogInformation("Syncing requires {RequiredSpace} which exceeds the {AvailableSpace} available free space on {Drive}.", StringUtils.FormatBytesString(RequiredFreeSpace), StringUtils.FormatBytesString(Drive.AvailableFreeSpace), Drive.Name);
+								return (WorkspaceUpdateResult.FailedToSync, "Not enough available free space.");
+							}
+						}
+						catch (SystemException)
+						{
+							Logger.LogInformation("Unable to check available free space for {RootPath}.", Project.LocalRootPath);
+						}
+					}
+					SyncTree.GetOptimizedSyncCommands(Project.ClientRootPath, Context.ChangeNumber, BatchBuilder);
+
+					// Clear the current sync changelist, in case we cancel
+					if (!Context.Options.HasFlag(WorkspaceUpdateOptions.SyncSingleChange))
+					{
+						State.CurrentChangeNumber = -1;
+						State.AdditionalChangeNumbers.Clear();
+						State.Save();
+					}
+
+					// Find all the depot paths that will be synced
+					HashSet<string> RemainingDepotPaths = new HashSet<string>(StringComparer.OrdinalIgnoreCase);
+					RemainingDepotPaths.UnionWith(SyncDepotPaths);
+
+					using (TelemetryStopwatch TransferStopwatch = new TelemetryStopwatch("Workspace_Sync_TransferFiles", Project.TelemetryProjectIdentifier))
+					{
+						TransferStopwatch.AddData(new { MachineName = Environment.MachineName, DomainName = Environment.UserDomainName, ServerAndPort = Perforce.Settings.ServerAndPort, UserName = Perforce.Settings.UserName, IncludedFiles = SyncTree.TotalIncludedFiles, ExcludedFiles = SyncTree.TotalExcludedFiles, Size = SyncTree.TotalSize, NumThreads = Context.PerforceSyncOptions?.NumThreads ?? PerforceSyncOptions.DefaultNumThreads });
+
+						(WorkspaceUpdateResult, string) SyncResult = await SyncFileRevisions(Perforce, "Syncing files...", Context, BatchBuilder.Batches, RemainingDepotPaths, Progress, Logger, CancellationToken);
+						if (SyncResult.Item1 != WorkspaceUpdateResult.Success)
+						{
+							TransferStopwatch.AddData(new { SyncResult = SyncResult.Item2.ToString(), CompletedFilesFiles = SyncDepotPaths.Count - RemainingDepotPaths.Count });
+							return SyncResult;
+						}
+
+						TransferStopwatch.Stop("Ok");
+						TransferStopwatch.AddData(new { TransferRate = SyncTree.TotalSize / Math.Max(TransferStopwatch.Elapsed.TotalSeconds, 0.0001f) });
+					}
+
+					int VersionChangeNumber = -1;
+					if (Context.Options.HasFlag(WorkspaceUpdateOptions.Sync) && !Context.Options.HasFlag(WorkspaceUpdateOptions.UpdateFilter))
+					{
+						// Read the new config file
+						Context.ProjectConfigFile = await ReadProjectConfigFile(Project.LocalRootPath, Project.LocalFileName, Logger);
+						Context.ProjectStreamFilter = await ReadProjectStreamFilter(Perforce, Context.ProjectConfigFile, Logger, CancellationToken);
+
+						// Get the branch name
+						string? BranchOrStreamName = await Perforce.GetCurrentStreamAsync(CancellationToken);
+						if (BranchOrStreamName != null)
+						{
+							// If it's a virtual stream, take the concrete parent stream instead
+							for (; ; )
+							{
+								StreamRecord StreamSpec = await Perforce.GetStreamAsync(BranchOrStreamName, false, CancellationToken);
+								if (StreamSpec.Type != "virtual" || StreamSpec.Parent == "none" || StreamSpec.Parent == null)
+								{
+									break;
+								}
+								BranchOrStreamName = StreamSpec.Parent;
+							}
+						}
+						else
+						{
+							// Otherwise use the depot path for GenerateProjectFiles.bat in the root of the workspace
+							List<WhereRecord> Files = await Perforce.WhereAsync(Project.ClientRootPath + "/GenerateProjectFiles.bat").Where(x => !x.Unmap).ToListAsync(CancellationToken);
+							if (Files.Count != 1)
+							{
+								return (WorkspaceUpdateResult.FailedToSync, $"Couldn't determine branch name for {Project.ClientFileName}.");
+							}
+							BranchOrStreamName = PerforceUtils.GetClientOrDepotDirectoryName(Files[0].DepotFile);
+						}
+
+						// Find the last code change before this changelist. For consistency in versioning between local builds and precompiled binaries, we need to use the last submitted code changelist as our version number.
+						string[] CodeFilter = new string[] { ".cs", ".h", ".cpp", ".usf", ".ush", ".uproject", ".uplugin" }.SelectMany(x => SyncPaths.Select(y => String.Format("{0}{1}@<={2}", y, x, Context.ChangeNumber))).ToArray();
+
+						PerforceResponseList<ChangesRecord> CodeChanges = await Perforce.TryGetChangesAsync(ChangesOptions.None, 1, ChangeStatus.Submitted, CodeFilter, CancellationToken);
+						if (!CodeChanges.Succeeded)
+						{
+							return (WorkspaceUpdateResult.FailedToSync, $"Couldn't determine last code changelist before CL {Context.ChangeNumber}.");
+						}
+						if (CodeChanges.Count == 0)
+						{
+							return (WorkspaceUpdateResult.FailedToSync, $"Could not find any code changes before CL {Context.ChangeNumber}.");
+						}
+
+						// Get the last code change
+						if (Context.ProjectConfigFile.GetValue("Options.VersionToLastCodeChange", true))
+						{
+							VersionChangeNumber = CodeChanges.Max(x => x.Data.Number);
+						}
+						else
+						{
+							VersionChangeNumber = Context.ChangeNumber;
+						}
+
+						// Update the version files
+						if (Context.ProjectConfigFile.GetValue("Options.UseFastModularVersioningV2", false))
+						{
+							bool bIsLicenseeVersion = await IsLicenseeVersion(Perforce, Project, Logger, CancellationToken);
+							if (!await UpdateVersionFile(Perforce, Project.ClientRootPath + BuildVersionFileName, Context.ChangeNumber, Text => UpdateBuildVersion(Text, Context.ChangeNumber, VersionChangeNumber, BranchOrStreamName, bIsLicenseeVersion), Logger, CancellationToken))
+							{
+								return (WorkspaceUpdateResult.FailedToSync, $"Failed to update {BuildVersionFileName}.");
+							}
+						}
+						else if (Context.ProjectConfigFile.GetValue("Options.UseFastModularVersioning", false))
+						{
+							bool bIsLicenseeVersion = await IsLicenseeVersion(Perforce, Project, Logger, CancellationToken);
+							if (!await UpdateVersionFile(Perforce, Project.ClientRootPath + BuildVersionFileName, Context.ChangeNumber, Text => UpdateBuildVersion(Text, Context.ChangeNumber, VersionChangeNumber, BranchOrStreamName, bIsLicenseeVersion), Logger, CancellationToken))
+							{
+								return (WorkspaceUpdateResult.FailedToSync, $"Failed to update {BuildVersionFileName}");
+							}
+
+							Dictionary<string, string> VersionHeaderStrings = new Dictionary<string, string>();
+							VersionHeaderStrings["#define ENGINE_IS_PROMOTED_BUILD"] = " (0)";
+							VersionHeaderStrings["#define BUILT_FROM_CHANGELIST"] = " 0";
+							VersionHeaderStrings["#define BRANCH_NAME"] = " \"" + BranchOrStreamName.Replace('/', '+') + "\"";
+							if (!await UpdateVersionFile(Perforce, Project.ClientRootPath + VersionHeaderFileName, VersionHeaderStrings, Context.ChangeNumber, Logger, CancellationToken))
+							{
+								return (WorkspaceUpdateResult.FailedToSync, $"Failed to update {VersionHeaderFileName}.");
+							}
+							if (!await UpdateVersionFile(Perforce, Project.ClientRootPath + ObjectVersionFileName, new Dictionary<string, string>(), Context.ChangeNumber, Logger, CancellationToken))
+							{
+								return (WorkspaceUpdateResult.FailedToSync, $"Failed to update {ObjectVersionFileName}.");
+							}
+						}
+						else
+						{
+							if (!await UpdateVersionFile(Perforce, Project.ClientRootPath + BuildVersionFileName, new Dictionary<string, string>(), Context.ChangeNumber, Logger, CancellationToken))
+							{
+								return (WorkspaceUpdateResult.FailedToSync, $"Failed to update {BuildVersionFileName}");
+							}
+
+							Dictionary<string, string> VersionStrings = new Dictionary<string, string>();
+							VersionStrings["#define ENGINE_VERSION"] = " " + VersionChangeNumber.ToString();
+							VersionStrings["#define ENGINE_IS_PROMOTED_BUILD"] = " (0)";
+							VersionStrings["#define BUILT_FROM_CHANGELIST"] = " " + VersionChangeNumber.ToString();
+							VersionStrings["#define BRANCH_NAME"] = " \"" + BranchOrStreamName.Replace('/', '+') + "\"";
+							if (!await UpdateVersionFile(Perforce, Project.ClientRootPath + VersionHeaderFileName, VersionStrings, Context.ChangeNumber, Logger, CancellationToken))
+							{
+								return (WorkspaceUpdateResult.FailedToSync, $"Failed to update {VersionHeaderFileName}");
+							}
+							if (!await UpdateVersionFile(Perforce, Project.ClientRootPath + ObjectVersionFileName, VersionStrings, Context.ChangeNumber, Logger, CancellationToken))
+							{
+								return (WorkspaceUpdateResult.FailedToSync, $"Failed to update {ObjectVersionFileName}");
+							}
+						}
+					}
+
+					// Check if there are any files which need resolving
+					List<FStatRecord> UnresolvedFiles = await FindUnresolvedFiles(Perforce, SyncPaths, CancellationToken);
+					if (UnresolvedFiles.Count > 0 && Context.Options.HasFlag(WorkspaceUpdateOptions.AutoResolveChanges))
+					{
+						foreach (FStatRecord UnresolvedFile in UnresolvedFiles)
+						{
+							if (UnresolvedFile.DepotFile != null)
+							{
+								await Perforce.ResolveAsync(-1, ResolveOptions.Automatic, UnresolvedFile.DepotFile, CancellationToken);
+							}
+						}
+						UnresolvedFiles = await FindUnresolvedFiles(Perforce, SyncPaths, CancellationToken);
+					}
+					if (UnresolvedFiles.Count > 0)
+					{
+						Logger.LogInformation("{NumFiles} files need resolving:", UnresolvedFiles.Count);
+						foreach (FStatRecord UnresolvedFile in UnresolvedFiles)
+						{
+							Logger.LogInformation("  {ClientFile}", UnresolvedFile.ClientFile);
+						}
+						return (WorkspaceUpdateResult.FilesToResolve, "Files need resolving.");
+					}
+
+					// Continue processing sync-only actions
+					if (Context.Options.HasFlag(WorkspaceUpdateOptions.Sync) && !Context.Options.HasFlag(WorkspaceUpdateOptions.UpdateFilter))
+					{
+						Context.ProjectConfigFile ??= await ReadProjectConfigFile(Project.LocalRootPath, Project.LocalFileName, Logger);
+
+						// Execute any project specific post-sync steps
+						string[]? PostSyncSteps = Context.ProjectConfigFile.GetValues("Sync.Step", null);
+						if (PostSyncSteps != null)
+						{
+							Logger.LogInformation("");
+							Logger.LogInformation("Executing post-sync steps...");
+
+							Dictionary<string, string> PostSyncVariables = new Dictionary<string, string>(Context.Variables);
+							PostSyncVariables["Change"] = Context.ChangeNumber.ToString();
+							PostSyncVariables["CodeChange"] = VersionChangeNumber.ToString();
+
+							foreach (string PostSyncStep in PostSyncSteps.Select(x => x.Trim()))
+							{
+								ConfigObject PostSyncStepObject = new ConfigObject(PostSyncStep);
+
+								string ToolFileName = Utility.ExpandVariables(PostSyncStepObject.GetValue("FileName", ""), PostSyncVariables);
+								if (ToolFileName != null)
+								{
+									string ToolArguments = Utility.ExpandVariables(PostSyncStepObject.GetValue("Arguments", ""), PostSyncVariables);
+
+									Logger.LogInformation("post-sync> Running {FileName} {Arguments}", ToolFileName, ToolArguments);
+
+									int ResultFromTool = await Utility.ExecuteProcessAsync(ToolFileName, null, ToolArguments, Line => ProcessOutput(Line, "post-sync> ", Progress, Logger), CancellationToken);
+									if (ResultFromTool != 0)
+									{
+										return (WorkspaceUpdateResult.FailedToSync, $"Post-sync step terminated with exit code {ResultFromTool}.");
+									}
+								}
+							}
+						}
+					}
+
+					// Update the current state
+					if (Context.Options.HasFlag(WorkspaceUpdateOptions.SyncSingleChange))
+					{
+						State.AdditionalChangeNumbers.Add(Context.ChangeNumber);
+					}
+					else
+					{
+						State.CurrentChangeNumber = Context.ChangeNumber;
+					}
+					State.Save();
+
+					// Update the timing info
+					Times.Add(new Tuple<string, TimeSpan>("Sync", SyncTelemetryStopwatch.Stop("Success")));
+
+					// Save the number of files synced
+					NumFilesSynced = SyncDepotPaths.Count;
+					Logger.LogInformation("");
+				}
+			}
+
+			// Extract an archive from the depot path
+			if (Context.Options.HasFlag(WorkspaceUpdateOptions.SyncArchives))
+			{
+				using (TelemetryStopwatch Stopwatch = new TelemetryStopwatch("Workspace_SyncArchives", Project.TelemetryProjectIdentifier))
+				{
+					// Create the directory for extracted archive manifests
+					DirectoryReference ManifestDirectoryName;
+					if (Project.LocalFileName.HasExtension(".uproject"))
+					{
+						ManifestDirectoryName = DirectoryReference.Combine(Project.LocalFileName.Directory, "Saved", "UnrealGameSync");
+					}
+					else
+					{
+						ManifestDirectoryName = DirectoryReference.Combine(Project.LocalFileName.Directory, "Engine", "Saved", "UnrealGameSync");
+					}
+					DirectoryReference.CreateDirectory(ManifestDirectoryName);
+
+					// Sync and extract (or just remove) the given archives
+					foreach (KeyValuePair<string, Tuple<IArchiveInfo, string>?> ArchiveTypeAndArchive in Context.ArchiveTypeToArchive)
+					{
+						string ArchiveType = ArchiveTypeAndArchive.Key;
+
+						// Remove any existing binaries
+						FileReference ManifestFileName = FileReference.Combine(ManifestDirectoryName, String.Format("{0}.zipmanifest", ArchiveType));
+						if (FileReference.Exists(ManifestFileName))
+						{
+							Logger.LogInformation("Removing {ArchiveType} binaries...", ArchiveType);
+							Progress.Set(String.Format("Removing {0} binaries...", ArchiveType), 0.0f);
+							ArchiveUtils.RemoveExtractedFiles(Project.LocalRootPath, ManifestFileName, Progress, Logger);
+							FileReference.Delete(ManifestFileName);
+							Logger.LogInformation("");
+						}
+
+						// If we have a new depot path, sync it down and extract it
+						if (ArchiveTypeAndArchive.Value != null)
+						{
+							IArchiveInfo ArchiveInfo = ArchiveTypeAndArchive.Value.Item1;
+							string ArchiveKey = ArchiveTypeAndArchive.Value.Item2;
+
+							Logger.LogInformation("Syncing {ArchiveType} binaries...", ArchiveType.ToLowerInvariant());
+							Progress.Set(String.Format("Syncing {0} binaries...", ArchiveType.ToLowerInvariant()), 0.0f);
+							if (!ArchiveInfo.DownloadArchive(Perforce, ArchiveKey, Project.LocalRootPath, ManifestFileName, Logger, Progress, CancellationToken.None).Result)
+							{
+								return (WorkspaceUpdateResult.FailedToSync, $"Couldn't read {ArchiveKey}");
+							}
+						}
+					}
+
+					// Update the state
+					State.ExpandedArchiveTypes = Context.ArchiveTypeToArchive.Where(x => x.Value != null).Select(x => x.Key).ToArray();
+					State.Save();
+
+					// Add the finish time
+					Times.Add(new Tuple<string, TimeSpan>("Archive", Stopwatch.Stop("Success")));
+				}
+			}
+
+			// Take the lock before doing anything else. Building and generating project files can only be done on one workspace at a time.
+			using SemaphoreScope Scope = new SemaphoreScope(UpdateSemaphore);
+			if (Context.Options.HasFlag(WorkspaceUpdateOptions.GenerateProjectFiles) || Context.Options.HasFlag(WorkspaceUpdateOptions.Build))
+			{
+				if (!Scope.TryAcquire())
+				{
+					Logger.LogInformation("Waiting for other workspaces to finish...");
+					await Scope.AcquireAsync(CancellationToken);
+				}
+			}
+
+			// Generate project files in the workspace
+			if (Context.Options.HasFlag(WorkspaceUpdateOptions.GenerateProjectFiles))
+			{
+				using (TelemetryStopwatch Stopwatch = new TelemetryStopwatch("Workspace_GenerateProjectFiles", Project.TelemetryProjectIdentifier))
+				{
+					Progress.Set("Generating project files...", 0.0f);
+
+					StringBuilder CommandLine = new StringBuilder();
+					CommandLine.AppendFormat("/C \"\"{0}\"", FileReference.Combine(Project.LocalRootPath, "GenerateProjectFiles.bat"));
+					if ((Context.Options & WorkspaceUpdateOptions.SyncAllProjects) == 0 && (Context.Options & WorkspaceUpdateOptions.IncludeAllProjectsInSolution) == 0)
+					{
+						if (Project.LocalFileName.HasExtension(".uproject"))
+						{
+							CommandLine.AppendFormat(" \"{0}\"", Project.LocalFileName);
+						}
+					}
+					CommandLine.Append(" -progress\"");
+
+					Logger.LogInformation("Generating project files...");
+					Logger.LogInformation("gpf> Running {FileName} {Arguments}", CmdExe, CommandLine);
+
+					int GenerateProjectFilesResult = await Utility.ExecuteProcessAsync(CmdExe, null, CommandLine.ToString(), Line => ProcessOutput(Line, "gpf> ", Progress, Logger), CancellationToken);
+					if (GenerateProjectFilesResult != 0)
+					{
+						return (WorkspaceUpdateResult.FailedToCompile, $"Failed to generate project files (exit code {GenerateProjectFilesResult}).");
+					}
+
+					Logger.LogInformation("");
+					Times.Add(new Tuple<string, TimeSpan>("Prj gen", Stopwatch.Stop("Success")));
+				}
+			}
+
+			// Build everything using MegaXGE
+			if (Context.Options.HasFlag(WorkspaceUpdateOptions.Build))
+			{
+				Context.ProjectConfigFile ??= await ReadProjectConfigFile(Project.LocalRootPath, Project.LocalFileName, Logger);
+
+				// Compile all the build steps together
+				Dictionary<Guid, ConfigObject> BuildStepObjects = Context.DefaultBuildSteps.ToDictionary(x => x.Key, x => new ConfigObject(x.Value));
+				BuildStep.MergeBuildStepObjects(BuildStepObjects, Context.ProjectConfigFile.GetValues("Build.Step", new string[0]).Select(x => new ConfigObject(x)));
+				BuildStep.MergeBuildStepObjects(BuildStepObjects, Context.UserBuildStepObjects);
+
+				// Construct build steps from them
+				List<BuildStep> BuildSteps = BuildStepObjects.Values.Select(x => new BuildStep(x)).OrderBy(x => (x.OrderIndex == -1) ? 10000 : x.OrderIndex).ToList();
+				if (Context.CustomBuildSteps != null && Context.CustomBuildSteps.Count > 0)
+				{
+					BuildSteps.RemoveAll(x => !Context.CustomBuildSteps.Contains(x.UniqueId));
+				}
+				else if (Context.Options.HasFlag(WorkspaceUpdateOptions.ScheduledBuild))
+				{
+					BuildSteps.RemoveAll(x => !x.bScheduledSync);
+				}
+				else
+				{
+					BuildSteps.RemoveAll(x => !x.bNormalSync);
+				}
+
+				// Check if the last successful build was before a change that we need to force a clean for
+				bool bForceClean = false;
+				if (State.LastBuiltChangeNumber != 0)
+				{
+					foreach (string CleanBuildChange in Context.ProjectConfigFile.GetValues("ForceClean.Changelist", new string[0]))
+					{
+						int ChangeNumber;
+						if (int.TryParse(CleanBuildChange, out ChangeNumber))
+						{
+							if ((State.LastBuiltChangeNumber >= ChangeNumber) != (State.CurrentChangeNumber >= ChangeNumber))
+							{
+								Logger.LogInformation("Forcing clean build due to changelist {Change}.", ChangeNumber);
+								Logger.LogInformation("");
+								bForceClean = true;
+								break;
+							}
+						}
+					}
+				}
+
+				// Execute them all
+				using (TelemetryStopwatch Stopwatch = new TelemetryStopwatch("Workspace_Build", Project.TelemetryProjectIdentifier))
+				{
+					Progress.Set("Starting build...", 0.0f);
+
+					// Execute all the steps
+					float MaxProgressFraction = 0.0f;
+					foreach (BuildStep Step in BuildSteps)
+					{
+						MaxProgressFraction += (float)Step.EstimatedDuration / (float)Math.Max(BuildSteps.Sum(x => x.EstimatedDuration), 1);
+
+						Progress.Set(Step.StatusText ?? "Executing build step");
+						Progress.Push(MaxProgressFraction);
+
+						Logger.LogInformation("{Status}", Step.StatusText);
+
+						if (Step.IsValid())
+						{
+							switch (Step.Type)
+							{
+								case BuildStepType.Compile:
+									using (TelemetryStopwatch StepStopwatch = new TelemetryStopwatch("Workspace_Execute_Compile", Project.TelemetryProjectIdentifier))
+									{
+										StepStopwatch.AddData(new { Target = Step.Target });
+
+										FileReference BuildBat = FileReference.Combine(Project.LocalRootPath, "Engine", "Build", "BatchFiles", "Build.bat");
+										string CommandLine = String.Format("{0} {1} {2} {3} -NoHotReloadFromIDE", Step.Target, Step.Platform, Step.Configuration, Utility.ExpandVariables(Step.Arguments ?? "", Context.Variables));
+										if (Context.Options.HasFlag(WorkspaceUpdateOptions.Clean) || bForceClean)
+										{
+											Logger.LogInformation("ubt> Running {FileName} {Arguments}", BuildBat, CommandLine + " -clean");
+											await Utility.ExecuteProcessAsync(CmdExe, null, "/C \"\"" + BuildBat + "\" " + CommandLine + " -clean\"", Line => ProcessOutput(Line, "ubt> ", Progress, Logger), CancellationToken);
+										}
+
+										Logger.LogInformation("ubt> Running {FileName} {Arguments}", BuildBat, CommandLine + " -progress");
+
+										int ResultFromBuild = await Utility.ExecuteProcessAsync(CmdExe, null, "/C \"\"" + BuildBat + "\" " + CommandLine + " -progress\"", Line => ProcessOutput(Line, "ubt> ", Progress, Logger), CancellationToken);
+										if (ResultFromBuild != 0)
+										{
+											StepStopwatch.Stop("Failed");
+
+											WorkspaceUpdateResult Result;
+											if (await HasModifiedSourceFiles(Perforce, Project, CancellationToken) || Context.UserBuildStepObjects.Count > 0)
+											{
+												Result = WorkspaceUpdateResult.FailedToCompile;
+											}
+											else
+											{
+												Result = WorkspaceUpdateResult.FailedToCompileWithCleanWorkspace;
+											}
+
+											return (Result, $"Failed to compile {Step.Target}");
+										}
+
+										StepStopwatch.Stop("Success");
+									}
+									break;
+								case BuildStepType.Cook:
+									using (TelemetryStopwatch StepStopwatch = new TelemetryStopwatch("Workspace_Execute_Cook", Project.TelemetryProjectIdentifier))
+									{
+										StepStopwatch.AddData(new { Project = Path.GetFileNameWithoutExtension(Step.FileName) });
+
+										FileReference LocalRunUAT = FileReference.Combine(Project.LocalRootPath, "Engine", "Build", "BatchFiles", "RunUAT.bat");
+										string Arguments = String.Format("/C \"\"{0}\" -profile=\"{1}\"\"", LocalRunUAT, FileReference.Combine(Project.LocalRootPath, Step.FileName ?? "unknown"));
+										Logger.LogInformation("uat> Running {FileName} {Argument}", LocalRunUAT, Arguments);
+
+										int ResultFromUAT = await Utility.ExecuteProcessAsync(CmdExe, null, Arguments, Line => ProcessOutput(Line, "uat> ", Progress, Logger), CancellationToken);
+										if (ResultFromUAT != 0)
+										{
+											StepStopwatch.Stop("Failed");
+											return (WorkspaceUpdateResult.FailedToCompile, $"Cook failed. ({ResultFromUAT})");
+										}
+
+										StepStopwatch.Stop("Success");
+									}
+									break;
+								case BuildStepType.Other:
+									using (TelemetryStopwatch StepStopwatch = new TelemetryStopwatch("Workspace_Execute_Custom", Project.TelemetryProjectIdentifier))
+									{
+										StepStopwatch.AddData(new { FileName = Path.GetFileNameWithoutExtension(Step.FileName) });
+
+										FileReference ToolFileName = FileReference.Combine(Project.LocalRootPath, Utility.ExpandVariables(Step.FileName ?? "unknown", Context.Variables));
+										string ToolWorkingDir = String.IsNullOrWhiteSpace(Step.WorkingDir) ? ToolFileName.Directory.FullName : Utility.ExpandVariables(Step.WorkingDir, Context.Variables);
+										string ToolArguments = Utility.ExpandVariables(Step.Arguments ?? "", Context.Variables);
+										Logger.LogInformation("tool> Running {0} {1}", ToolFileName, ToolArguments);
+
+										if (Step.bUseLogWindow)
+										{
+											int ResultFromTool = await Utility.ExecuteProcessAsync(ToolFileName.FullName, ToolWorkingDir, ToolArguments, Line => ProcessOutput(Line, "tool> ", Progress, Logger), CancellationToken);
+											if (ResultFromTool != 0)
+											{
+												StepStopwatch.Stop("Failed");
+												return (WorkspaceUpdateResult.FailedToCompile, $"Tool terminated with exit code {ResultFromTool}.");
+											}
+										}
+										else
+										{
+											ProcessStartInfo StartInfo = new ProcessStartInfo(ToolFileName.FullName, ToolArguments);
+											StartInfo.WorkingDirectory = ToolWorkingDir;
+											using (Process.Start(StartInfo))
+											{
+											}
+										}
+
+										StepStopwatch.Stop("Success");
+									}
+									break;
+							}
+						}
+
+						Logger.LogInformation("");
+						Progress.Pop();
+					}
+
+					Times.Add(new Tuple<string, TimeSpan>("Build", Stopwatch.Stop("Success")));
+				}
+
+				// Update the last successful build change number
+				if (Context.CustomBuildSteps == null || Context.CustomBuildSteps.Count == 0)
+				{
+					State.LastBuiltChangeNumber = State.CurrentChangeNumber;
+					State.Save();
+				}
+			}
+
+			// Write out all the timing information
+			Logger.LogInformation("Total time: {TotalTime}", FormatTime(Times.Sum(x => (long)(x.Item2.TotalMilliseconds / 1000))));
+			foreach (Tuple<string, TimeSpan> Time in Times)
+			{
+				Logger.LogInformation("   {Name,-8}: {TimeSeconds}", Time.Item1, FormatTime((long)(Time.Item2.TotalMilliseconds / 1000)));
+			}
+			if (NumFilesSynced > 0)
+			{
+				Logger.LogInformation("{NumFiles} files synced.", NumFilesSynced);
+			}
+
+			DateTime FinishTime = DateTime.Now;
+			Logger.LogInformation("");
+			Logger.LogInformation("UPDATE SUCCEEDED ({FinishDate} {FinishTime})", FinishTime.ToShortDateString(), FinishTime.ToShortTimeString());
+
+			return (WorkspaceUpdateResult.Success, "Update succeeded");
+		}
+
+		static void ProcessOutput(string Line, string Prefix, ProgressValue Progress, ILogger Logger)
+		{
+			string? ParsedLine = ProgressTextWriter.ParseLine(Line, Progress);
+			if (ParsedLine != null)
+			{
+				Logger.LogInformation("{Prefix}{Message}", Prefix, Line);
+			}
+		}
+
+		static async Task<bool> IsLicenseeVersion(IPerforceConnection Perforce, ProjectInfo Project, ILogger Logger, CancellationToken CancellationToken)
+		{
+			string[] Files =
+			{
+				Project.ClientRootPath + "/Engine/Build/NotForLicensees/EpicInternal.txt",
+				Project.ClientRootPath + "/Engine/Restricted/NotForLicensees/Build/EpicInternal.txt"
+			};
+
+			List<FStatRecord> Records = await Perforce.FStatAsync(Files, CancellationToken).ToListAsync(CancellationToken);
+			return !Records.Any(x => x.IsMapped);
+		}
+
+		public static List<string> GetSyncPaths(ProjectInfo Project, bool bSyncAllProjects, string[] SyncFilter)
+		{
+			List<string> SyncPaths = GetRelativeSyncPaths(Project, bSyncAllProjects, SyncFilter);
+			return SyncPaths.Select(x => Project.ClientRootPath + x).ToList();
+		}
+
+		public static List<string> GetRelativeSyncPaths(ProjectInfo Project, bool bSyncAllProjects, string[]? SyncFilter)
+		{
+			List<string> SyncPaths = new List<string>();
+
+			// Check the client path is formatted correctly
+			if (!Project.ClientFileName.StartsWith(Project.ClientRootPath + "/"))
+			{
+				throw new Exception(String.Format("Expected '{0}' to start with '{1}'", Project.ClientFileName, Project.ClientRootPath));
+			}
+
+			// Add the default project paths
+			int LastSlashIdx = Project.ClientFileName.LastIndexOf('/');
+			if (bSyncAllProjects || !Project.ClientFileName.EndsWith(".uproject", StringComparison.OrdinalIgnoreCase) || LastSlashIdx <= Project.ClientRootPath.Length)
+			{
+				SyncPaths.Add("/...");
+			}
+			else
+			{
+				SyncPaths.Add("/*");
+				SyncPaths.Add("/Engine/...");
+				if (Project.bIsEnterpriseProject)
+				{
+					SyncPaths.Add("/Enterprise/...");
+				}
+				SyncPaths.Add(Project.ClientFileName.Substring(Project.ClientRootPath.Length, LastSlashIdx - Project.ClientRootPath.Length) + "/...");
+			}
+
+			// Apply the sync filter to that list. We only want inclusive rules in the output list, but we can manually apply exclusions to previous entries.
+			if (SyncFilter != null)
+			{
+				foreach (string SyncPath in SyncFilter)
+				{
+					string TrimSyncPath = SyncPath.Trim();
+					if (TrimSyncPath.StartsWith("/"))
+					{
+						SyncPaths.Add(TrimSyncPath);
+					}
+					else if (TrimSyncPath.StartsWith("-/") && TrimSyncPath.EndsWith("..."))
+					{
+						SyncPaths.RemoveAll(x => x.StartsWith(TrimSyncPath.Substring(1, TrimSyncPath.Length - 4)));
+					}
+				}
+			}
+
+			// Sort the remaining paths by length, and remove any paths which are included twice
+			SyncPaths = SyncPaths.OrderBy(x => x.Length).ToList();
+			for (int Idx = 0; Idx < SyncPaths.Count; Idx++)
+			{
+				string SyncPath = SyncPaths[Idx];
+				if (SyncPath.EndsWith("..."))
+				{
+					string SyncPathPrefix = SyncPath.Substring(0, SyncPath.Length - 3);
+					for (int OtherIdx = SyncPaths.Count - 1; OtherIdx > Idx; OtherIdx--)
+					{
+						if (SyncPaths[OtherIdx].StartsWith(SyncPathPrefix))
+						{
+							SyncPaths.RemoveAt(OtherIdx);
+						}
+					}
+				}
+			}
+
+			return SyncPaths;
+		}
+
+		public static bool MatchFilter(ProjectInfo Project, FileReference FileName, FileFilter Filter)
+		{
+			bool bMatch = true;
+			if (FileName.IsUnderDirectory(Project.LocalRootPath))
+			{
+				string RelativePath = FileName.MakeRelativeTo(Project.LocalRootPath);
+				if (!Filter.Matches(RelativePath))
+				{
+					bMatch = false;
+				}
+			}
+			return bMatch;
+		}
+
+		class SyncState
+		{
+			public int TotalDepotPaths;
+			public HashSet<string> RemainingDepotPaths;
+			public Queue<List<string>> SyncCommandLists;
+			public string StatusMessage;
+			public WorkspaceUpdateResult Result = WorkspaceUpdateResult.Success;
+
+			public SyncState(HashSet<string> RemainingDepotPaths, Queue<List<string>> SyncCommandLists)
+			{
+				this.TotalDepotPaths = RemainingDepotPaths.Count;
+				this.RemainingDepotPaths = RemainingDepotPaths;
+				this.SyncCommandLists = SyncCommandLists;
+				this.StatusMessage = "Succeeded.";
+			}
+		}
+
+		static Task<(WorkspaceUpdateResult, string)> SyncFileRevisions(IPerforceConnection Perforce, string Prefix, WorkspaceUpdateContext Context, List<string> SyncCommands, HashSet<string> RemainingDepotPaths, ProgressValue Progress, ILogger Logger, CancellationToken CancellationToken)
+		{
+			Queue<List<string>> SyncCommandLists = new Queue<List<string>>();
+			SyncCommandLists.Enqueue(SyncCommands);
+			return SyncFileRevisions(Perforce, Prefix, Context, SyncCommandLists, RemainingDepotPaths, Progress, Logger, CancellationToken);
+		}
+
+		static async Task<(WorkspaceUpdateResult, string)> SyncFileRevisions(IPerforceConnection Perforce, string Prefix, WorkspaceUpdateContext Context, Queue<List<string>> SyncCommandLists, HashSet<string> RemainingDepotPaths, ProgressValue Progress, ILogger Logger, CancellationToken CancellationToken)
+		{
+			// Figure out the number of additional background threads we want to run with. We can run worker on the current thread.
+			int NumThreads = Context.PerforceSyncOptions?.NumThreads ?? PerforceSyncOptions.DefaultNumThreads;
+			int NumExtraThreads = Math.Max(Math.Min(SyncCommandLists.Count, NumThreads) - 1, 0);
+
+			List<IPerforceConnection> ChildConnections = new List<IPerforceConnection>();
+			List<Task> ChildTasks = new List<Task>(NumExtraThreads);
+			try
+			{
+				// Create the state object shared by all the worker threads
+				SyncState State = new SyncState(RemainingDepotPaths, SyncCommandLists);
+
+				// Wrapper writer around the log class to prevent multiple threads writing to it at once
+				ILogger LogWrapper = Logger;
+
+				// Initialize Sync Progress
+				UpdateSyncState(Prefix, State, Progress);
+
+				// Delegate for updating the sync state after a file has been synced
+				Action<SyncRecord, ILogger> SyncOutput = (Record, LocalLog) => { UpdateSyncState(Prefix, Record, State, Progress, LocalLog); };
+
+				// Create all the child threads
+				for (int ThreadIdx = 0; ThreadIdx < NumExtraThreads; ThreadIdx++)
+				{
+					int ThreadNumber = ThreadIdx + 2;
+
+					// Create connection
+					IPerforceConnection ChildConnection = await PerforceConnection.CreateAsync(Perforce.Settings, Perforce.Logger);
+					ChildConnections.Add(ChildConnection);
+
+					Task ChildTask = Task.Run(() => StaticSyncWorker(ThreadNumber, ChildConnection, Context, State, SyncOutput, LogWrapper, CancellationToken));
+					ChildTasks.Add(ChildTask);
+				}
+
+				// Run one worker on the current thread
+				await StaticSyncWorker(1, Perforce, Context, State, SyncOutput, LogWrapper, CancellationToken);
+
+				// Return the result that was set on the state object
+				return (State.Result, State.StatusMessage);
+			}
+			finally
+			{
+				await Task.WhenAll(ChildTasks);
+
+				foreach (IPerforceConnection ChildConnection in ChildConnections)
+				{
+					ChildConnection.Dispose();
+				}
+
+				foreach (Task ChildTask in ChildTasks)
+				{
+					await ChildTask;
+				}
+			}
+		}
+
+		static void UpdateSyncState(string Prefix, SyncState State, ProgressValue Progress)
+		{
+			lock (State)
+			{
+				string Message = String.Format("{0} ({1}/{2})", Prefix, State.TotalDepotPaths - State.RemainingDepotPaths.Count, State.TotalDepotPaths);
+				float Fraction = Math.Min((float)(State.TotalDepotPaths - State.RemainingDepotPaths.Count) / (float)State.TotalDepotPaths, 1.0f);
+				Progress.Set(Message, Fraction);
+			}
+		}
+
+		static void UpdateSyncState(string Prefix, SyncRecord Record, SyncState State, ProgressValue Progress, ILogger Logger)
+		{
+			lock (State)
+			{
+				State.RemainingDepotPaths.Remove(Record.DepotFile);
+
+				string Message = String.Format("{0} ({1}/{2})", Prefix, State.TotalDepotPaths - State.RemainingDepotPaths.Count, State.TotalDepotPaths);
+				float Fraction = Math.Min((float)(State.TotalDepotPaths - State.RemainingDepotPaths.Count) / (float)State.TotalDepotPaths, 1.0f);
+				Progress.Set(Message, Fraction);
+
+				Logger.LogInformation("p4>   {Action} {Path}", Record.Action, Record.Path);
+			}
+		}
+
+		static async Task StaticSyncWorker(int ThreadNumber, IPerforceConnection Perforce, WorkspaceUpdateContext Context, SyncState State, Action<SyncRecord, ILogger> SyncOutput, ILogger GlobalLog, CancellationToken CancellationToken)
+		{
+			PrefixedTextWriter ThreadLog = new PrefixedTextWriter(String.Format("{0}:", ThreadNumber), GlobalLog);
+			for (; ; )
+			{
+				// Remove the next batch that needs to be synced
+				List<string> SyncCommands;
+				lock (State)
+				{
+					if (State.Result == WorkspaceUpdateResult.Success && State.SyncCommandLists.Count > 0)
+					{
+						SyncCommands = State.SyncCommandLists.Dequeue();
+					}
+					else
+					{
+						break;
+					}
+				}
+
+				// Sync the files
+				(WorkspaceUpdateResult Result, string StatusMessage) = await StaticSyncFileRevisions(Perforce, Context, SyncCommands, Record => SyncOutput(Record, ThreadLog), CancellationToken);
+
+				// If it failed, try to set it on the state if nothing else has failed first
+				if (Result != WorkspaceUpdateResult.Success)
+				{
+					lock (State)
+					{
+						if (State.Result == WorkspaceUpdateResult.Success)
+						{
+							State.Result = Result;
+							State.StatusMessage = StatusMessage;
+						}
+					}
+					break;
+				}
+			}
+		}
+
+		static async Task<(WorkspaceUpdateResult, string)> StaticSyncFileRevisions(IPerforceConnection Perforce, WorkspaceUpdateContext Context, List<string> SyncCommands, Action<SyncRecord> SyncOutput, CancellationToken CancellationToken)
+		{
+			// Sync them all
+			List<PerforceResponse<SyncRecord>> Responses = await Perforce.TrySyncAsync(SyncOptions.None, -1, SyncCommands, CancellationToken).ToListAsync(CancellationToken);
+
+			List<string> TamperedFiles = new List<string>();
+			foreach (PerforceResponse<SyncRecord> Response in Responses)
+			{
+				const string NoClobberPrefix = "Can't clobber writable file ";
+				if (Response.Succeeded)
+				{
+					SyncOutput(Response.Data);
+				}
+				else if (Response.Error != null && Response.Error.Generic == PerforceGenericCode.Client && Response.Error.Data.StartsWith(NoClobberPrefix, StringComparison.OrdinalIgnoreCase))
+				{
+					TamperedFiles.Add(Response.Error.Data.Substring(NoClobberPrefix.Length).Trim());
+				}
+				else
+				{
+					return (WorkspaceUpdateResult.FailedToSync, "Aborted sync due to errors.");
+				}
+			}
+
+			// If any files need to be clobbered, defer to the main thread to figure out which ones
+			if (TamperedFiles.Count > 0)
+			{
+				int NumNewFilesToClobber = 0;
+				foreach (string TamperedFile in TamperedFiles)
+				{
+					if (!Context.ClobberFiles.ContainsKey(TamperedFile))
+					{
+						Context.ClobberFiles[TamperedFile] = true;
+						if (TamperedFile.EndsWith(LocalObjectVersionFileName, StringComparison.OrdinalIgnoreCase) || TamperedFile.EndsWith(LocalVersionHeaderFileName, StringComparison.OrdinalIgnoreCase))
+						{
+							// Hack for UseFastModularVersioningV2; we don't need to update these files any more.
+							continue;
+						}
+						NumNewFilesToClobber++;
+					}
+				}
+				if (NumNewFilesToClobber > 0)
+				{
+					return (WorkspaceUpdateResult.FilesToClobber, $"Cancelled sync after checking files to clobber ({NumNewFilesToClobber} new files).");
+				}
+				foreach (string TamperedFile in TamperedFiles)
+				{
+					if (Context.ClobberFiles[TamperedFile])
+					{
+						List<PerforceResponse<SyncRecord>> Response = await Perforce.TrySyncAsync(SyncOptions.Force, -1, TamperedFile, CancellationToken).ToListAsync();
+						if (!Response.Succeeded())
+						{
+							return (WorkspaceUpdateResult.FailedToSync, $"Couldn't sync {TamperedFile}.");
+						}
+					}
+				}
+			}
+
+			// All succeeded
+			return (WorkspaceUpdateResult.Success, "Succeeded.");
+		}
+
+		public static async Task<ConfigFile> ReadProjectConfigFile(DirectoryReference LocalRootPath, FileReference SelectedLocalFileName, ILogger Logger)
+		{
+			// Find the valid config file paths
+			DirectoryInfo EngineDir = DirectoryReference.Combine(LocalRootPath, "Engine").ToDirectoryInfo();
+			List<FileInfo> LocalConfigFiles = Utility.GetLocalConfigPaths(EngineDir, SelectedLocalFileName.ToFileInfo());
+
+			// Read them in
+			ConfigFile ProjectConfig = new ConfigFile();
+			foreach (FileInfo LocalConfigFile in LocalConfigFiles)
+			{
+				try
+				{
+					string[] Lines = await File.ReadAllLinesAsync(LocalConfigFile.FullName);
+					ProjectConfig.Parse(Lines);
+					Logger.LogInformation("Read config file from {FileName}", LocalConfigFile.FullName);
+				}
+				catch (Exception Ex)
+				{
+					Logger.LogInformation(Ex, "Failed to read config file from {FileName}", LocalConfigFile.FullName);
+				}
+			}
+			return ProjectConfig;
+		}
+
+		public static async Task<IReadOnlyList<string>?> ReadProjectStreamFilter(IPerforceConnection Perforce, ConfigFile ProjectConfigFile, ILogger Logger, CancellationToken CancellationToken)
+		{
+			string? StreamListDepotPath = ProjectConfigFile.GetValue("Options.QuickSelectStreamList", null);
+			if (StreamListDepotPath == null)
+			{
+				return null;
+			}
+
+			PerforceResponse<PrintRecord<string[]>> Response = await Perforce.TryPrintLinesAsync(StreamListDepotPath, CancellationToken);
+			if (!Response.Succeeded)
+			{
+				return null;
+			}
+
+			return Response.Data.Contents.Select(x => x.Trim()).Where(x => x.Length > 0).ToList().AsReadOnly();
+		}
+
+		static string FormatTime(long Seconds)
+		{
+			if (Seconds >= 60)
+			{
+				return String.Format("{0,3}m {1:00}s", Seconds / 60, Seconds % 60);
+			}
+			else
+			{
+				return String.Format("     {0,2}s", Seconds);
+			}
+		}
+
+		static async Task<bool> HasModifiedSourceFiles(IPerforceConnection Perforce, ProjectInfo Project, CancellationToken CancellationToken)
+		{
+			List<OpenedRecord> OpenFiles = await Perforce.OpenedAsync(OpenedOptions.None, Project.ClientRootPath + "/...", CancellationToken).ToListAsync(CancellationToken);
+			if (OpenFiles.Any(x => x.DepotFile.IndexOf("/Source/", StringComparison.OrdinalIgnoreCase) != -1))
+			{
+				return true;
+			}
+			return false;
+		}
+
+		static async Task<List<FStatRecord>> FindUnresolvedFiles(IPerforceConnection Perforce, IEnumerable<string> SyncPaths, CancellationToken CancellationToken)
+		{
+			List<FStatRecord> UnresolvedFiles = new List<FStatRecord>();
+			foreach (string SyncPath in SyncPaths)
+			{
+				List<FStatRecord> Records = await Perforce.FStatAsync(FStatOptions.OnlyUnresolved, SyncPath, CancellationToken).ToListAsync(CancellationToken);
+				UnresolvedFiles.AddRange(Records);
+			}
+			return UnresolvedFiles;
+		}
+
+		static Task<bool> UpdateVersionFile(IPerforceConnection Perforce, string ClientPath, Dictionary<string, string> VersionStrings, int ChangeNumber, ILogger Logger, CancellationToken CancellationToken)
+		{
+			return UpdateVersionFile(Perforce, ClientPath, ChangeNumber, Text => UpdateVersionStrings(Text, VersionStrings), Logger, CancellationToken);
+		}
+
+		static async Task<bool> UpdateVersionFile(IPerforceConnection Perforce, string ClientPath, int ChangeNumber, Func<string, string> Update, ILogger Logger, CancellationToken CancellationToken)
+		{
+			List<PerforceResponse<FStatRecord>> Records = await Perforce.TryFStatAsync(FStatOptions.None, ClientPath, CancellationToken).ToListAsync(CancellationToken);
+			if (!Records.Succeeded())
+			{
+				Logger.LogInformation("Failed to query records for {ClientPath}", ClientPath);
+				return false;
+			}
+			if (Records.Count > 1)
+			{
+				// Attempt to remove any existing file which is synced
+				await Perforce.SyncAsync(SyncOptions.Force, -1, $"{ClientPath}#0", CancellationToken).ToListAsync(CancellationToken);
+
+				// Try to get the mapped files again
+				Records = await Perforce.TryFStatAsync(FStatOptions.None, ClientPath, CancellationToken).ToListAsync(CancellationToken);
+				if (!Records.Succeeded())
+				{
+					Logger.LogInformation("Failed to query records for {ClientPath}", ClientPath);
+					return false;
+				}
+			}
+			if (Records.Count == 0)
+			{
+				Logger.LogInformation("Ignoring {ClientPath}; not found on server.", ClientPath);
+				return true;
+			}
+
+			FStatRecord Record = Records[0].Data;
+			string? LocalPath = Record.ClientFile; // Actually a filesystem path
+			if (LocalPath == null)
+			{
+				Logger.LogInformation("Version file is not mapped to workspace ({ClientFile})", ClientPath);
+				return false;
+			}
+			string? DepotPath = Record.DepotFile;
+			if (DepotPath == null)
+			{
+				Logger.LogInformation("Version file does not exist in depot ({ClientFile})", ClientPath);
+				return false;
+			}
+
+			PerforceResponse<PrintRecord<string[]>> Response = await Perforce.TryPrintLinesAsync($"{DepotPath}@{ChangeNumber}", CancellationToken);
+			if (!Response.Succeeded)
+			{
+				Logger.LogInformation("Couldn't get default contents of {DepotPath}", DepotPath);
+				return false;
+			}
+
+			string[]? Contents = Response.Data.Contents;
+			if (Contents == null)
+			{
+				Logger.LogInformation("No data returned for {DepotPath}", DepotPath);
+				return false;
+			}
+
+			string Text = String.Join("\n", Contents);
+			Text = Update(Text);
+			return await WriteVersionFile(Perforce, LocalPath, DepotPath, Text, Logger, CancellationToken);
+		}
+
+		static string UpdateVersionStrings(string Text, Dictionary<string, string> VersionStrings)
+		{
+			StringWriter Writer = new StringWriter();
+			foreach (string Line in Text.Split('\n'))
+			{
+				string NewLine = Line;
+				foreach (KeyValuePair<string, string> VersionString in VersionStrings)
+				{
+					if (UpdateVersionLine(ref NewLine, VersionString.Key, VersionString.Value))
+					{
+						break;
+					}
+				}
+				Writer.WriteLine(NewLine);
+			}
+			return Writer.ToString();
+		}
+
+		static string UpdateBuildVersion(string Text, int Changelist, int CodeChangelist, string BranchOrStreamName, bool bIsLicenseeVersion)
+		{
+			Dictionary<string, object> Object = JsonSerializer.Deserialize<Dictionary<string, object>>(Text, Utility.DefaultJsonSerializerOptions);
+
+			int PrevCompatibleChangelist = 0;
+			if (Object.TryGetValue("CompatibleChangelist", out object? PrevCompatibleChangelistObj))
+			{
+				int.TryParse(PrevCompatibleChangelistObj?.ToString(), out PrevCompatibleChangelist);
+			}
+
+			int PrevIsLicenseeVersion = 0;
+			if (Object.TryGetValue("IsLicenseeVersion", out object? PrevIsLicenseeVersionObj))
+			{
+				int.TryParse(PrevIsLicenseeVersionObj?.ToString(), out PrevIsLicenseeVersion);
+			}
+
+			Object["Changelist"] = Changelist;
+			if (PrevCompatibleChangelist == 0 || (PrevIsLicenseeVersion != 0) != bIsLicenseeVersion)
+			{
+				// Don't overwrite the compatible changelist if we're in a hotfix release
+				Object["CompatibleChangelist"] = CodeChangelist;
+			}
+			Object["BranchName"] = BranchOrStreamName.Replace('/', '+');
+			Object["IsPromotedBuild"] = 0;
+			Object["IsLicenseeVersion"] = bIsLicenseeVersion ? 1 : 0;
+
+			return JsonSerializer.Serialize(Object, new JsonSerializerOptions
+			{
+				WriteIndented = true,
+				// do not escape +
+				Encoder = JavaScriptEncoder.UnsafeRelaxedJsonEscaping
+			});
+		}
+
+		static async Task<bool> WriteVersionFile(IPerforceConnection Perforce, string LocalPath, string DepotPath, string NewText, ILogger Logger, CancellationToken CancellationToken)
+		{
+			try
+			{
+				if (File.Exists(LocalPath) && File.ReadAllText(LocalPath) == NewText)
+				{
+					Logger.LogInformation("Ignored {FileName}; contents haven't changed", LocalPath);
+				}
+				else
+				{
+					Directory.CreateDirectory(Path.GetDirectoryName(LocalPath));
+					Utility.ForceDeleteFile(LocalPath);
+					if (DepotPath != null)
+					{
+						await Perforce.SyncAsync(DepotPath + "#0", CancellationToken).ToListAsync(CancellationToken);
+					}
+					File.WriteAllText(LocalPath, NewText);
+					Logger.LogInformation("Written {FileName}", LocalPath);
+				}
+				return true;
+			}
+			catch (Exception Ex)
+			{
+				Logger.LogError(Ex, "Failed to write to {FileName}.", LocalPath);
+				return false;
+			}
+		}
+
+		static bool UpdateVersionLine(ref string Line, string Prefix, string Suffix)
+		{
+			int LineIdx = 0;
+			int PrefixIdx = 0;
+			for (; ; )
+			{
+				string? PrefixToken = ReadToken(Prefix, ref PrefixIdx);
+				if (PrefixToken == null)
+				{
+					break;
+				}
+
+				string? LineToken = ReadToken(Line, ref LineIdx);
+				if (LineToken == null || LineToken != PrefixToken)
+				{
+					return false;
+				}
+			}
+			Line = Line.Substring(0, LineIdx) + Suffix;
+			return true;
+		}
+
+		static string? ReadToken(string Line, ref int LineIdx)
+		{
+			for (; ; LineIdx++)
+			{
+				if (LineIdx == Line.Length)
+				{
+					return null;
+				}
+				else if (!Char.IsWhiteSpace(Line[LineIdx]))
+				{
+					break;
+				}
+			}
+
+			int StartIdx = LineIdx++;
+			if (Char.IsLetterOrDigit(Line[StartIdx]) || Line[StartIdx] == '_')
+			{
+				while (LineIdx < Line.Length && (Char.IsLetterOrDigit(Line[LineIdx]) || Line[LineIdx] == '_'))
+				{
+					LineIdx++;
+				}
+			}
+
+			return Line.Substring(StartIdx, LineIdx - StartIdx);
+		}
+
+		public Tuple<string, float> CurrentProgress
+		{
+			get { return Progress.Current; }
+		}
+	}
+}
