UGS: Move WorkspaceSettings into the shared library.

diff --git a/Engine/Source/Programs/UnrealGameSync/UnrealGameSync/DetectProjectSettingsTask.cs b/Engine/Source/Programs/UnrealGameSync/UnrealGameSync/DetectProjectSettingsTask.cs
deleted file mode 100644
diff --git a/Engine/Source/Programs/UnrealGameSync/UnrealGameSyncShared/WorkspaceSettings.cs b/Engine/Source/Programs/UnrealGameSync/UnrealGameSyncShared/WorkspaceSettings.cs
new file mode 100644
--- /dev/null
+++ b/Engine/Source/Programs/UnrealGameSync/UnrealGameSyncShared/WorkspaceSettings.cs
@@ -0,0 +1,392 @@
+// Copyright Epic Games, Inc. All Rights Reserved.
+
+using EpicGames.Core;
+using EpicGames.Perforce;
+using Microsoft.Extensions.Logging;
+using System;
+using System.Collections.Generic;
+using System.Drawing;
+using System.IO;
+using System.Linq;
+using System.Text.RegularExpressions;
+using System.Threading;
+using System.Threading.Tasks;
+
+namespace UnrealGameSync
+{
+	public class WorkspaceSettings
+	{
+		public IPerforceSettings PerforceSettings;
+		public UserSelectedProjectSettings SelectedProject { get; }
+		public ProjectInfo ProjectInfo { get; private set; }
+		public string NewSelectedProjectIdentifier;
+		public string NewProjectEditorTarget;
+		public string StreamName;
+		public DirectoryReference DataFolder => GetDataFolder(ProjectInfo.LocalRootPath);
+		public DirectoryReference CacheFolder => GetCacheFolder(ProjectInfo.LocalRootPath);
+		public ConfigFile LatestProjectConfigFile;
+		public ConfigFile WorkspaceProjectConfigFile;
+		public IReadOnlyList<string>? WorkspaceProjectStreamFilter;
+		public List<KeyValuePair<FileReference, DateTime>> LocalConfigFiles;
+
+		public WorkspaceSettings(IPerforceSettings PerforceSettings, UserSelectedProjectSettings ProjectSettings, ProjectInfo ProjectInfo, string NewSelectedProjectIdentifier, string NewProjectEditorTarget, string StreamName, ConfigFile LatestProjectConfigFile, ConfigFile WorkspaceProjectConfigFile, IReadOnlyList<string>? WorkspaceProjectStreamFilter, List<KeyValuePair<FileReference, DateTime>> LocalConfigFiles)
+		{
+			this.PerforceSettings = PerforceSettings;
+			this.SelectedProject = ProjectSettings;
+			this.ProjectInfo = ProjectInfo;
+			this.NewSelectedProjectIdentifier = NewSelectedProjectIdentifier;
+			this.NewProjectEditorTarget = NewProjectEditorTarget;
+			this.StreamName = StreamName;
+			this.LatestProjectConfigFile = LatestProjectConfigFile;
+			this.WorkspaceProjectConfigFile = WorkspaceProjectConfigFile;
+			this.WorkspaceProjectStreamFilter = WorkspaceProjectStreamFilter;
+			this.LocalConfigFiles = LocalConfigFiles;
+		}
+
+		public static DirectoryReference GetDataFolder(DirectoryReference WorkspaceDir) => DirectoryReference.Combine(WorkspaceDir, ".ugs");
+		public static DirectoryReference GetCacheFolder(DirectoryReference WorkspaceDir) => DirectoryReference.Combine(WorkspaceDir, ".ugs", "cache");
+
+		public static async Task<WorkspaceSettings> CreateAsync(IPerforceSettings DefaultPerforceSettings, UserSelectedProjectSettings SelectedProject, ILogger<WorkspaceSettings> Logger, CancellationToken CancellationToken)
+		{
+			PerforceSettings PerforceSettings = Utility.OverridePerforceSettings(DefaultPerforceSettings, SelectedProject.ServerAndPort, SelectedProject.UserName);
+			using IPerforceConnection Perforce = await PerforceConnection.CreateAsync(PerforceSettings, Logger);
+
+			// Make sure we're logged in
+			PerforceResponse<LoginRecord> LoginState = await Perforce.TryGetLoginStateAsync(CancellationToken);
+			if (!LoginState.Succeeded)
+			{
+				throw new UserErrorException("User is not logged in to Perforce.");
+			}
+
+			// Execute like a regular task
+			return await CreateAsync(Perforce, SelectedProject, Logger, CancellationToken);
+		}
+
+		public static async Task<WorkspaceSettings> CreateAsync(IPerforceConnection DefaultConnection, UserSelectedProjectSettings SelectedProject, ILogger<WorkspaceSettings> Logger, CancellationToken CancellationToken)
+		{
+			using IDisposable LoggerScope = Logger.BeginScope("Project {SelectedProject}", SelectedProject.ToString());
+			Logger.LogInformation("Detecting settings for {Project}", SelectedProject);
+
+			// Use the cached client path to the file if it's available; it's much quicker than trying to find the correct workspace.
+			IPerforceConnection? PerforceClient = null;
+			try
+			{
+				IPerforceSettings PerforceSettings;
+
+				FileReference NewSelectedFileName;
+				string NewSelectedClientFileName;
+				if (!String.IsNullOrEmpty(SelectedProject.ClientPath))
+				{
+					// Get the client path
+					NewSelectedClientFileName = SelectedProject.ClientPath;
+
+					// Get the client name
+					string? ClientName;
+					if (!PerforceUtils.TryGetClientName(NewSelectedClientFileName, out ClientName))
+					{
+						throw new UserErrorException($"Couldn't get client name from {NewSelectedClientFileName}");
+					}
+
+					// Create the client
+					PerforceSettings = new PerforceSettings(DefaultConnection.Settings) { ClientName = ClientName };
+					PerforceClient = await PerforceConnection.CreateAsync(PerforceSettings, Logger);
+
+					// Figure out the path on the client. Use the cached location if it's valid.
+					string? LocalPath = SelectedProject.LocalPath;
+					if (LocalPath == null || !File.Exists(LocalPath))
+					{
+						List<WhereRecord> Records = await PerforceClient.WhereAsync(NewSelectedClientFileName, CancellationToken).Where(x => !x.Unmap).ToListAsync(CancellationToken);
+						if (Records.Count != 1)
+						{
+							throw new UserErrorException($"Couldn't get client path for {NewSelectedClientFileName}");
+						}
+						LocalPath = Path.GetFullPath(Records[0].Path);
+					}
+					NewSelectedFileName = new FileReference(LocalPath);
+				}
+				else
+				{
+					// Get the perforce server settings
+					InfoRecord PerforceInfo = await DefaultConnection.GetInfoAsync(InfoOptions.ShortOutput, CancellationToken);
+
+					// Use the path as the selected filename
+					NewSelectedFileName = new FileReference(SelectedProject.LocalPath!);
+
+					// Make sure the project exists
+					if (!FileReference.Exists(NewSelectedFileName))
+					{
+						throw new UserErrorException($"{SelectedProject.LocalPath} does not exist.");
+					}
+
+					// Find all the clients for this user
+					Logger.LogInformation("Enumerating clients for {UserName}...", PerforceInfo.UserName);
+
+					List<ClientsRecord> Clients = await DefaultConnection.GetClientsAsync(ClientsOptions.None, DefaultConnection.Settings.UserName, CancellationToken);
+
+					List<IPerforceSettings> CandidateClients = await FilterClients(Clients, NewSelectedFileName, DefaultConnection.Settings, PerforceInfo.ClientHost, Logger, CancellationToken);
+					if (CandidateClients.Count == 0)
+					{
+						// Search through all workspaces. We may find a suitable workspace which is for any user.
+						Logger.LogInformation("Enumerating shared clients...");
+						Clients = await DefaultConnection.GetClientsAsync(ClientsOptions.None, "", CancellationToken);
+
+						// Filter this list of clients
+						CandidateClients = await FilterClients(Clients, NewSelectedFileName, DefaultConnection.Settings, PerforceInfo.ClientHost, Logger, CancellationToken);
+
+						// If we still couldn't find any, fail.
+						if (CandidateClients.Count == 0)
+						{
+							throw new UserErrorException($"Couldn't find any Perforce workspace containing {NewSelectedFileName}. Check your connection settings.");
+						}
+					}
+
+					// Check there's only one client
+					if (CandidateClients.Count > 1)
+					{
+						throw new UserErrorException(String.Format("Found multiple workspaces containing {0}:\n\n{1}\n\nCannot determine which to use.", Path.GetFileName(NewSelectedFileName.GetFileName()), String.Join("\n", CandidateClients.Select(x => x.ClientName))));
+					}
+
+					// Take the client we've chosen
+					PerforceSettings = CandidateClients[0];
+					PerforceClient = await PerforceConnection.CreateAsync(PerforceSettings, Logger);
+
+					// Get the client path for the project file
+					List<WhereRecord> Records = await PerforceClient.WhereAsync(NewSelectedFileName.FullName, CancellationToken).Where(x => !x.Unmap).ToListAsync(CancellationToken);
+					if (Records.Count == 0)
+					{
+						throw new UserErrorException("File is not mapped to any client");
+					}
+					else if (Records.Count > 1)
+					{
+						throw new UserErrorException($"File is mapped to {Records.Count} locations: {String.Join(", ", Records.Select(x => x.Path))}");
+					}
+
+					NewSelectedClientFileName = Records[0].ClientFile;
+				}
+
+				// Make sure the path case is correct. This can cause UBT intermediates to be out of date if the case mismatches.
+				NewSelectedFileName = FileReference.FindCorrectCase(NewSelectedFileName);
+
+				// Update the selected project with all the data we've found
+				SelectedProject = new UserSelectedProjectSettings(SelectedProject.ServerAndPort, SelectedProject.UserName, SelectedProject.Type, NewSelectedClientFileName, NewSelectedFileName.FullName);
+
+				// Get the local branch root
+				string? BranchClientPath = null;
+				DirectoryReference? BranchDirectoryName = null;
+
+				// Figure out where the engine is in relation to it
+				int EndIdx = NewSelectedClientFileName.Length - 1;
+				if (EndIdx != -1 && NewSelectedClientFileName.EndsWith(".uproject", StringComparison.InvariantCultureIgnoreCase))
+				{
+					EndIdx = NewSelectedClientFileName.LastIndexOf('/') - 1;
+				}
+				for (; EndIdx >= 2; EndIdx--)
+				{
+					if (NewSelectedClientFileName[EndIdx] == '/')
+					{
+						List<PerforceResponse<FStatRecord>> FileRecords = await PerforceClient.TryFStatAsync(FStatOptions.None, NewSelectedClientFileName.Substring(0, EndIdx) + "/Engine/Build/Build.version", CancellationToken).ToListAsync();
+						if (FileRecords.Succeeded())
+						{
+							FStatRecord FileRecord = FileRecords[0].Data;
+							if (FileRecord.ClientFile == null)
+							{
+								throw new UserErrorException($"Missing client path for {FileRecord.DepotFile}");
+							}
+
+							BranchClientPath = NewSelectedClientFileName.Substring(0, EndIdx);
+							BranchDirectoryName = new FileReference(FileRecord.ClientFile).Directory.ParentDirectory?.ParentDirectory;
+							break;
+						}
+					}
+				}
+				if(BranchClientPath == null || BranchDirectoryName == null)
+				{
+					throw new UserErrorException($"Could not find engine in Perforce relative to project path ({NewSelectedClientFileName})");
+				}
+
+				Logger.LogInformation("Found branch root at {RootPath}", BranchClientPath);
+
+				// Find the editor target for this project
+				string? NewProjectEditorTarget = null;
+				if (NewSelectedFileName.HasExtension(".uproject"))
+				{
+					List<FStatRecord> Records = await PerforceClient.FStatAsync(PerforceUtils.GetClientOrDepotDirectoryName(NewSelectedClientFileName) + "/Source/*Editor.Target.cs", CancellationToken).ToListAsync(CancellationToken);
+
+					FStatRecord Record = Records.FirstOrDefault(x => x.Action != FileAction.Delete || x.Action != FileAction.MoveDelete);
+					if (Record != null)
+					{
+						string DepotPath = Record.DepotFile!;
+						NewProjectEditorTarget = Path.GetFileNameWithoutExtension(Path.GetFileNameWithoutExtension(DepotPath.Substring(DepotPath.LastIndexOf('/') + 1)));
+						Logger.LogInformation("Using {TargetName} as editor target name (from {DepotPath})", NewProjectEditorTarget, DepotPath);
+					}
+
+					if (NewProjectEditorTarget == null)
+					{
+						Logger.LogInformation("Couldn't find any editor targets for this project.");
+					}
+				}
+
+				string NewSelectedProjectIdentifier;
+				// Get a unique name for the project that's selected. For regular branches, this can be the depot path. For streams, we want to include the stream name to encode imports.
+				string? StreamName = await PerforceClient.GetCurrentStreamAsync(CancellationToken);
+				if (StreamName != null)
+				{
+					string ExpectedPrefix = String.Format("//{0}/", PerforceClient.Settings.ClientName);
+					if (!NewSelectedClientFileName.StartsWith(ExpectedPrefix, StringComparison.InvariantCultureIgnoreCase))
+					{
+						throw new UserErrorException($"Unexpected client path; expected '{NewSelectedClientFileName}' to begin with '{ExpectedPrefix}'");
+					}
+					string? StreamPrefix = await TryGetStreamPrefixAsync(PerforceClient, StreamName, CancellationToken);
+					if (StreamPrefix == null)
+					{
+						throw new UserErrorException("Unable to get stream prefix");
+					}
+					NewSelectedProjectIdentifier = String.Format("{0}/{1}", StreamPrefix, NewSelectedClientFileName.Substring(ExpectedPrefix.Length));
+				}
+				else
+				{
+					List<PerforceResponse<WhereRecord>> Records = await PerforceClient.TryWhereAsync(NewSelectedClientFileName, CancellationToken).Where(x => !x.Succeeded || !x.Data.Unmap).ToListAsync(CancellationToken);
+					if (!Records.Succeeded() || Records.Count != 1)
+					{
+						throw new UserErrorException($"Couldn't get depot path for {NewSelectedFileName}");
+					}
+
+					NewSelectedProjectIdentifier = Records[0].Data.DepotFile;
+
+					Match Match = Regex.Match(NewSelectedProjectIdentifier, "//([^/]+)/");
+					if (Match.Success)
+					{
+						DepotRecord Depot = await PerforceClient.GetDepotAsync(Match.Groups[1].Value, CancellationToken);
+						if (Depot.Type == "stream")
+						{
+							throw new UserErrorException($"Cannot use a legacy client ({PerforceClient.Settings.ClientName}) with a stream depot ({Depot.Depot}).");
+						}
+					}
+				}
+
+				// Figure out if it's an enterprise project. Use the synced version if we have it.
+				bool bIsEnterpriseProject = false;
+				if (NewSelectedClientFileName.EndsWith(".uproject", StringComparison.InvariantCultureIgnoreCase))
+				{
+					string Text;
+					if (FileReference.Exists(NewSelectedFileName))
+					{
+						Text = FileReference.ReadAllText(NewSelectedFileName);
+					}
+					else
+					{
+						PerforceResponse<PrintRecord<string[]>> ProjectLines = await PerforceClient.TryPrintLinesAsync(NewSelectedClientFileName, CancellationToken);
+						if (!ProjectLines.Succeeded)
+						{
+							throw new UserErrorException($"Unable to get contents of {NewSelectedClientFileName}");
+						}
+						Text = String.Join("\n", ProjectLines.Data.Contents!);
+					}
+					bIsEnterpriseProject = Utility.IsEnterpriseProjectFromText(Text);
+				}
+
+				// Make sure the drive containing the project exists, to prevent other errors down the line
+				string PathRoot = Path.GetPathRoot(NewSelectedFileName.FullName)!;
+				if (!Directory.Exists(PathRoot))
+				{
+					throw new UserErrorException($"Path '{NewSelectedFileName}' is invalid");
+				}
+
+				// Read the initial config file
+				List<KeyValuePair<FileReference, DateTime>> LocalConfigFiles = new List<KeyValuePair<FileReference, DateTime>>();
+				ConfigFile LatestProjectConfigFile = await ConfigUtils.ReadProjectConfigFileAsync(PerforceClient, BranchClientPath, NewSelectedClientFileName, GetCacheFolder(BranchDirectoryName), LocalConfigFiles, Logger, CancellationToken);
+
+				// Get the local config file and stream filter
+				ConfigFile WorkspaceProjectConfigFile = await WorkspaceUpdate.ReadProjectConfigFile(BranchDirectoryName, NewSelectedFileName, Logger);
+				IReadOnlyList<string>? WorkspaceProjectStreamFilter = await WorkspaceUpdate.ReadProjectStreamFilter(PerforceClient, WorkspaceProjectConfigFile, Logger, CancellationToken);
+
+				int BranchIdx = BranchClientPath.IndexOf('/', 2);
+				string BranchPath = (BranchIdx == -1) ? String.Empty : BranchClientPath.Substring(BranchIdx);
+				string ProjectPath = NewSelectedClientFileName.Substring(BranchClientPath.Length);
+
+				ProjectInfo ProjectInfo = new ProjectInfo(BranchDirectoryName, PerforceSettings.ClientName!, BranchPath, ProjectPath, NewSelectedProjectIdentifier, bIsEnterpriseProject);
+
+				WorkspaceSettings WorkspaceSettings = new WorkspaceSettings(PerforceSettings, SelectedProject, ProjectInfo, NewSelectedProjectIdentifier, NewProjectEditorTarget!, StreamName!, LatestProjectConfigFile, WorkspaceProjectConfigFile, WorkspaceProjectStreamFilter, LocalConfigFiles);
+				DirectoryReference.CreateDirectory(WorkspaceSettings.DataFolder);
+				DirectoryReference.CreateDirectory(WorkspaceSettings.CacheFolder);
+
+				return WorkspaceSettings;
+			}
+			finally
+			{
+				PerforceClient?.Dispose();
+			}
+		}
+
+		static async Task<List<IPerforceSettings>> FilterClients(List<ClientsRecord> Clients, FileReference NewSelectedFileName, IPerforceSettings DefaultPerforceSettings, string? HostName, ILogger Logger, CancellationToken CancellationToken)
+		{
+			List<IPerforceSettings> CandidateClients = new List<IPerforceSettings>();
+			foreach(ClientsRecord Client in Clients)
+			{
+				// Make sure the client is well formed
+				if(!String.IsNullOrEmpty(Client.Name) && (!String.IsNullOrEmpty(Client.Host) || !String.IsNullOrEmpty(Client.Owner)) && !String.IsNullOrEmpty(Client.Root))
+				{
+					// Require either a username or host name match
+					if((String.IsNullOrEmpty(Client.Host) || String.Compare(Client.Host, HostName, StringComparison.OrdinalIgnoreCase) == 0) && (String.IsNullOrEmpty(Client.Owner) || String.Compare(Client.Owner, DefaultPerforceSettings.UserName, StringComparison.OrdinalIgnoreCase) == 0))
+					{
+						if(!Utility.SafeIsFileUnderDirectory(NewSelectedFileName.FullName, Client.Root))
+						{
+							Logger.LogInformation("Rejecting {ClientName} due to root mismatch ({RootPath})", Client.Name, Client.Root);
+							continue;
+						}
+
+						PerforceSettings CandidateSettings = new PerforceSettings(DefaultPerforceSettings) { ClientName = Client.Name };
+						using IPerforceConnection CandidateClient = await PerforceConnection.CreateAsync(CandidateSettings, Logger);
+
+						List<PerforceResponse<WhereRecord>> WhereRecords = await CandidateClient.TryWhereAsync(NewSelectedFileName.FullName, CancellationToken).Where(x => x.Failed || !x.Data.Unmap).ToListAsync(CancellationToken);
+						if(!WhereRecords.Succeeded() || WhereRecords.Count != 1)
+						{
+							Logger.LogInformation("Rejecting {ClientName} due to file not existing in workspace", Client.Name);
+							continue;
+						}
+
+						List<PerforceResponse<FStatRecord>> Records = await CandidateClient.TryFStatAsync(FStatOptions.None, NewSelectedFileName.FullName, CancellationToken).ToListAsync(CancellationToken);
+						if (!Records.Succeeded())
+						{
+							Logger.LogInformation("Rejecting {ClientName} due to {FileName} not in depot", Client.Name, NewSelectedFileName);
+							continue;
+						}
+
+						Records.RemoveAll(x => !x.Data.IsMapped);
+						if(Records.Count == 0)
+						{
+							Logger.LogInformation("Rejecting {ClientName} due to {NumRecords} matching records", Client.Name, Records.Count);
+							continue;
+						}
+
+						Logger.LogInformation("Found valid client {ClientName}", Client.Name);
+						CandidateClients.Add(CandidateSettings);
+					}
+				}
+			}
+			return CandidateClients;
+		}
+
+		static async Task<string?> TryGetStreamPrefixAsync(IPerforceConnection Perforce, string StreamName, CancellationToken CancellationToken)
+		{ 
+			string? CurrentStreamName = StreamName;
+			while(!String.IsNullOrEmpty(CurrentStreamName))
+			{
+				PerforceResponse<StreamRecord> Response = await Perforce.TryGetStreamAsync(CurrentStreamName, false, CancellationToken);
+				if (!Response.Succeeded)
+				{
+					return null;
+				}
+
+				StreamRecord StreamSpec = Response.Data;
+				if(StreamSpec.Type != "virtual")
+				{
+					return CurrentStreamName;
+				}
+
+				CurrentStreamName = StreamSpec.Parent;
+			}
+			return null;
+		}
+	}
+}
