UGS: Compute build steps and environment variables for custom commands as part of the sync process. This ensures that we always use the latest values, and moves a lot of business logic out of the WinForms implementation.

diff --git a/Engine/Source/Programs/UnrealGameSync/UnrealGameSync/Controls/WorkspaceControl.cs b/Engine/Source/Programs/UnrealGameSync/UnrealGameSync/Controls/WorkspaceControl.cs
--- a/Engine/Source/Programs/UnrealGameSync/UnrealGameSync/Controls/WorkspaceControl.cs
+++ b/Engine/Source/Programs/UnrealGameSync/UnrealGameSync/Controls/WorkspaceControl.cs
@@ -200,8 +200,6 @@ namespace UnrealGameSync
 
 		const int BuildListExpandCount = 250;
 
-		const string EditorArchiveType = "Editor";
-
 		IWorkspaceControlOwner Owner;
 		string? ApiUrl;
 		DirectoryReference DataFolder;
@@ -968,7 +966,7 @@ namespace UnrealGameSync
 				CombinedSyncFilter = AdditionalPaths.Union(CombinedSyncFilter).ToArray();
 			}
 
-			WorkspaceUpdateContext Context = new WorkspaceUpdateContext(ChangeNumber, Options, CombinedSyncFilter, GetDefaultBuildStepObjects(), ProjectSettings.BuildSteps, null, GetWorkspaceVariables(ChangeNumber));
+			WorkspaceUpdateContext Context = new WorkspaceUpdateContext(ChangeNumber, Options, GetEditorBuildConfig(), CombinedSyncFilter, ProjectSettings.BuildSteps, null);
 			if (Options.HasFlag(WorkspaceUpdateOptions.SyncArchives))
 			{
 				IReadOnlyList<IArchiveInfo> Archives = GetArchives();
@@ -987,7 +985,7 @@ namespace UnrealGameSync
 						return;
 					}
 
-					if (Archive.Type == EditorArchiveType)
+					if (Archive.Type == IArchiveInfo.EditorArchiveType)
 					{
 						Context.Options &= ~(WorkspaceUpdateOptions.Build | WorkspaceUpdateOptions.GenerateProjectFiles | WorkspaceUpdateOptions.OpenSolutionAfterSync);
 					}
@@ -1038,15 +1036,13 @@ namespace UnrealGameSync
 			{
 				if (!Context.Options.HasFlag(WorkspaceUpdateOptions.ContentOnly) && (Context.CustomBuildSteps == null || Context.CustomBuildSteps.Count == 0))
 				{
+					FileReference TargetFile = ConfigUtils.GetEditorTargetFile(Workspace.Project, Workspace.ProjectConfigFile);
 					foreach (BuildConfig Config in Enum.GetValues(typeof(BuildConfig)).OfType<BuildConfig>())
 					{
-						List<FileReference> EditorReceiptPaths = GetEditorReceiptPaths(Config);
-						foreach (FileReference EditorReceiptPath in EditorReceiptPaths)
+						FileReference ReceiptFile = ConfigUtils.GetReceiptFile(Workspace.Project, TargetFile, Config.ToString());
+						if (FileReference.Exists(ReceiptFile))
 						{
-							if (FileReference.Exists(EditorReceiptPath))
-							{
-								try { FileReference.Delete(EditorReceiptPath); } catch (Exception) { }
-							}
+							try { FileReference.Delete(ReceiptFile); } catch (Exception) { }
 						}
 					}
 				}
@@ -2673,42 +2669,7 @@ namespace UnrealGameSync
 
 		private bool TryGetProjectSetting(ConfigFile ProjectConfigFile, string Name, [NotNullWhen(true)] out string? Value)
 		{
-			string Path = SelectedProjectIdentifier;
-			for(; ;)
-			{
-				ConfigSection ProjectSection = ProjectConfigFile.FindSection(Path);
-				if (ProjectSection != null)
-				{
-					string? NewValue = ProjectSection.GetValue(Name, null);
-					if (NewValue != null)
-					{
-						Value = NewValue;
-						return true;
-					}
-				}
-
-				int LastSlash = Path.LastIndexOf('/');
-				if(LastSlash < 2)
-				{
-					break;
-				}
-
-				Path = Path.Substring(0, LastSlash);
-			}
-
-			ConfigSection DefaultSection = ProjectConfigFile.FindSection("Default");
-			if (DefaultSection != null)
-			{
-				string? NewValue = DefaultSection.GetValue(Name, null);
-				if (NewValue != null)
-				{
-					Value = NewValue;
-					return true;
-				}
-			}
-
-			Value = null;
-			return false;
+			return ConfigUtils.TryGetProjectSetting(ProjectConfigFile, SelectedProjectIdentifier, Name, out Value);
 		}
 
 		private bool TryGetProjectSetting(ConfigFile ProjectConfigFile, string Name, [NotNullWhen(true)] out string[]? Values)
@@ -3035,10 +2996,8 @@ namespace UnrealGameSync
 			{
 				BuildConfig EditorBuildConfig = GetEditorBuildConfig();
 
-				List<FileReference> ReceiptPaths = GetEditorReceiptPaths(EditorBuildConfig);
-
-				FileReference EditorExe = GetEditorExePath(EditorBuildConfig);
-				if (ReceiptPaths.Any(x => FileReference.Exists(x)) && FileReference.Exists(EditorExe))
+				FileReference ReceiptFile = ConfigUtils.GetEditorReceiptFile(Workspace.Project, Workspace.ProjectConfigFile, EditorBuildConfig);
+				if (ConfigUtils.TryReadEditorReceipt(Workspace.Project, ReceiptFile, out TargetReceipt? Receipt) && Receipt.Launch != null && File.Exists(Receipt.Launch))
 				{
 					if (Settings.bEditorArgumentsPrompt && !ModifyEditorArguments())
 					{
@@ -3062,9 +3021,9 @@ namespace UnrealGameSync
 						LaunchArguments.Append(" -debug");
 					}
 
-					if (!Utility.SpawnProcess(EditorExe.FullName, LaunchArguments.ToString()))
+					if (!Utility.SpawnProcess(Receipt.Launch, LaunchArguments.ToString()))
 					{
-						ShowErrorDialog("Unable to spawn {0} {1}", EditorExe, LaunchArguments.ToString());
+						ShowErrorDialog("Unable to spawn {0} {1}", Receipt.Launch, LaunchArguments.ToString());
 					}
 				}
 				else
@@ -3074,7 +3033,7 @@ namespace UnrealGameSync
 						Owner.ShowAndActivate();
 
 						WorkspaceUpdateOptions Options = WorkspaceUpdateOptions.Build | WorkspaceUpdateOptions.RunAfterSync;
-						WorkspaceUpdateContext Context = new WorkspaceUpdateContext(Workspace.CurrentChangeNumber, Options, null, GetDefaultBuildStepObjects(), ProjectSettings.BuildSteps, null, GetWorkspaceVariables(Workspace.CurrentChangeNumber));
+						WorkspaceUpdateContext Context = new WorkspaceUpdateContext(Workspace.CurrentChangeNumber, Options, Settings.CompiledEditorBuildConfig, null, ProjectSettings.BuildSteps, null);
 						StartWorkspaceUpdate(Context, null);
 					}
 				}
@@ -3089,7 +3048,7 @@ namespace UnrealGameSync
 				List<IArchiveInfo> Archives = GetSelectedArchives(PerforceMonitor.AvailableArchives);
 				foreach(IArchiveInfo Archive in Archives)
 				{
-					if(Archive.Type == EditorArchiveType)
+					if(Archive.Type == IArchiveInfo.EditorArchiveType)
 					{
 						if(Archive.Target != null)
 						{
@@ -3101,53 +3060,6 @@ namespace UnrealGameSync
 			return EditorTargetName;
 		}
 
-		private class ReceiptJsonObject
-		{
-			public string? Launch { get; set; }
-		}
-
-		private FileReference GetEditorExePath(BuildConfig Config)
-		{
-			// Try to read the executable path from the target receipt
-			List<FileReference> ReceiptFileNames = GetEditorReceiptPaths(Config);
-			foreach (FileReference ReceiptFileName in ReceiptFileNames)
-			{
-				if (FileReference.Exists(ReceiptFileName))
-				{
-					Logger.LogInformation("Reading {FileName}", ReceiptFileName);
-					try
-					{
-						string Text = FileReference.ReadAllText(ReceiptFileName);
-						ReceiptJsonObject Receipt = JsonSerializer.Deserialize<ReceiptJsonObject>(Text, Utility.DefaultJsonSerializerOptions);
-
-						string? LaunchFileName = Receipt.Launch;
-						if (LaunchFileName != null)
-						{
-							LaunchFileName = LaunchFileName.Replace("$(EngineDir)", DirectoryReference.Combine(BranchDirectoryName, "Engine").FullName);
-							LaunchFileName = LaunchFileName.Replace("$(ProjectDir)", SelectedFileName.Directory.FullName);
-							return new FileReference(LaunchFileName);
-						}
-					}
-					catch (Exception Ex)
-					{
-						Logger.LogError(Ex, "Exception while parsing receipt");
-					}
-					break;
-				}
-			}
-
-			// Get the default editor target name
-			string TargetName = GetDefaultEditorTargetName();
-
-			// Otherwise use the standard editor path
-			string ExeFileName = String.Format("{0}.exe", TargetName);
-			if ((Config != BuildConfig.DebugGame || PerforceMonitor.LatestProjectConfigFile.GetValue("Options.DebugGameHasSeparateExecutable", false)) && Config != BuildConfig.Development)
-			{
-				ExeFileName = String.Format("{0}-Win64-{1}.exe", TargetName, Config.ToString());
-			}
-			return FileReference.Combine(BranchDirectoryName, "Engine", "Binaries", "Win64", ExeFileName);
-		}
-
 		private bool WaitForProgramsToFinish()
 		{
 			FileReference[] ProcessFileNames = GetProcessesRunningInWorkspace();
@@ -3223,35 +3135,6 @@ namespace UnrealGameSync
 			return ProcessFileNames.ToArray();
 		}
 
-		private List<FileReference> GetEditorReceiptPaths(BuildConfig Config)
-		{
-			string ConfigSuffix = (Config == BuildConfig.Development) ? "" : String.Format("-Win64-{0}", Config.ToString());
-
-			List<FileReference> PossiblePaths = new List<FileReference>();
-			string ActualEditorName = TryGetEditorNameFromArchive();
-			if (ActualEditorName != null)
-			{
-				PossiblePaths.Add(FileReference.Combine(SelectedFileName.Directory, "Binaries", "Win64", String.Format("{0}{1}.target", ActualEditorName, ConfigSuffix)));
-			}
-			else if (SelectedFileName.HasExtension(".uproject"))
-			{
-				if (bIsEnterpriseProject)
-				{
-					PossiblePaths.Add(FileReference.Combine(SelectedFileName.Directory, "Binaries", "Win64", String.Format("StudioEditor{0}.target", ConfigSuffix)));
-				}
-				else
-				{
-					PossiblePaths.Add(FileReference.Combine(SelectedFileName.Directory, "Binaries", "Win64", String.Format("{0}{1}.target", GetDefaultEditorTargetName(), ConfigSuffix)));
-					PossiblePaths.Add(FileReference.Combine(Workspace.Project.LocalRootPath, "Engine", "Binaries", "Win64", String.Format("{0}{1}.target", GetDefaultEditorTargetName(), ConfigSuffix)));
-				}
-			}
-			else
-			{
-				PossiblePaths.Add(FileReference.Combine(SelectedFileName.Directory, "Engine", "Binaries", "Win64", String.Format("{0}{1}.target", GetDefaultEditorTargetName(), ConfigSuffix)));
-			}
-			return PossiblePaths;
-		}
-
 		private void TimerCallback(object? Sender, EventArgs Args)
 		{
 			Tuple<string, float> Progress = Workspace.CurrentProgress;
@@ -3348,7 +3231,7 @@ namespace UnrealGameSync
 			DirectoryReference? ToolDir = Owner.ToolUpdateMonitor.GetToolPath(Tool.Name);
 			if (ToolDir != null)
 			{
-				Dictionary<string, string> Variables = GetWorkspaceVariables(Workspace.CurrentChangeNumber);
+				Dictionary<string, string> Variables = Workspace.GetVariables(GetEditorBuildConfig());
 				Variables["ToolDir"] = ToolDir.FullName;
 
 				string FileName = Utility.ExpandVariables(Link.FileName, Variables);
@@ -3763,7 +3646,7 @@ namespace UnrealGameSync
 			string[]? SdkInfoEntries;
 			if (TryGetProjectSetting(PerforceMonitor.LatestProjectConfigFile, "SdkInfo", out SdkInfoEntries))
 			{
-				Dictionary<string, string> Variables = GetWorkspaceVariables(-1);
+				Dictionary<string, string> Variables = Workspace.GetVariables(GetEditorBuildConfig());
 				SdkInfoWindow Window = new SdkInfoWindow(SdkInfoEntries, Variables, BadgeFont!);
 				Window.ShowDialog();
 			}
@@ -4229,7 +4112,7 @@ namespace UnrealGameSync
 						ConfigFile? ProjectConfigFile = PerforceMonitor.LatestProjectConfigFile;
 						if (ProjectConfigFile != null)
 						{
-							Dictionary<string, string> Variables = GetWorkspaceVariables(ContextMenuChange.Number);
+							Dictionary<string, string> Variables = Workspace.GetVariables(GetEditorBuildConfig(), ContextMenuChange.Number, -1);
 
 							string[] ChangeContextMenuEntries = ProjectConfigFile.GetValues("Options.ContextMenu", new string[0]);
 							foreach (string ChangeContextMenuEntry in ChangeContextMenuEntries)
@@ -4446,7 +4329,7 @@ namespace UnrealGameSync
 				OptionsContextMenu_SyncPrecompiledBinaries.ToolTipText = null;
 				OptionsContextMenu_SyncPrecompiledBinaries.Checked = SelectedArchives.Count > 0;
 
-				if (Archives.Count > 1 || Archives[0].Type != EditorArchiveType)
+				if (Archives.Count > 1 || Archives[0].Type != IArchiveInfo.EditorArchiveType)
 				{
 					ToolStripMenuItem DisableItem = new ToolStripMenuItem("Disable (compile locally)");
 					DisableItem.Checked = (SelectedArchives.Count == 0);
@@ -4580,7 +4463,7 @@ namespace UnrealGameSync
 				List<ConfigObject> UserBuildSteps = new List<ConfigObject>();
 				UserBuildSteps.Add(CustomBuildStep.ToConfigObject());
 
-				WorkspaceUpdateContext Context = new WorkspaceUpdateContext(Workspace.CurrentChangeNumber, WorkspaceUpdateOptions.Build, null, GetDefaultBuildStepObjects(), UserBuildSteps, new HashSet<Guid> { CustomBuildStep.UniqueId }, GetWorkspaceVariables(Workspace.CurrentChangeNumber));
+				WorkspaceUpdateContext Context = new WorkspaceUpdateContext(Workspace.CurrentChangeNumber, WorkspaceUpdateOptions.Build, GetEditorBuildConfig(), null, UserBuildSteps, new HashSet<Guid> { CustomBuildStep.UniqueId });
 				StartWorkspaceUpdate(Context, null);
 			}
 		}
@@ -5291,7 +5174,8 @@ namespace UnrealGameSync
 						}
 					}
 
-					Dictionary<string, string> Variables = GetWorkspaceVariables(Workspace.CurrentChangeNumber);
+					WorkspaceUpdateContext Context = new WorkspaceUpdateContext(Workspace.CurrentChangeNumber, WorkspaceUpdateOptions.Build, GetEditorBuildConfig(), null, ProjectSettings.BuildSteps, StepSet);
+
 					if (Step.ToolId != Guid.Empty)
 					{
 						string? ToolName = Owner.ToolUpdateMonitor.GetToolName(Step.ToolId);
@@ -5300,43 +5184,16 @@ namespace UnrealGameSync
 							DirectoryReference? ToolPath = Owner.ToolUpdateMonitor.GetToolPath(ToolName);
 							if (ToolPath != null)
 							{
-								Variables["ToolDir"] = ToolPath.FullName;
+								Context.AdditionalVariables["ToolDir"] = ToolPath.FullName;
 							}
 						}
 					}
 
-					WorkspaceUpdateContext Context = new WorkspaceUpdateContext(Workspace.CurrentChangeNumber, WorkspaceUpdateOptions.Build, null, GetDefaultBuildStepObjects(), ProjectSettings.BuildSteps, StepSet, Variables);
 					StartWorkspaceUpdate(Context, null);
 				}
 			}
 		}
 
-		private Dictionary<string, string> GetWorkspaceVariables(int ChangeNumber)
-		{
-			BuildConfig EditorBuildConfig = GetEditorBuildConfig();
-
-			string? SdkInstallerDir;
-			TryGetProjectSetting(PerforceMonitor.LatestProjectConfigFile, "SdkInstallerDir", out SdkInstallerDir);
-
-			Dictionary<string, string> Variables = new Dictionary<string, string>();
-			Variables.Add("Stream", StreamName);
-			Variables.Add("Change", ChangeNumber.ToString());
-			Variables.Add("ClientName", ClientName);
-			Variables.Add("BranchDir", BranchDirectoryName.FullName);
-			Variables.Add("ProjectDir", SelectedFileName.Directory.FullName);
-			Variables.Add("ProjectFile", SelectedFileName.FullName);
-			Variables.Add("UE4EditorExe", GetEditorExePath(EditorBuildConfig).FullName);
-			Variables.Add("UE4EditorCmdExe", GetEditorExePath(EditorBuildConfig).FullName.Replace(".exe", "-Cmd.exe"));
-			Variables.Add("UE4EditorConfig", EditorBuildConfig.ToString());
-			Variables.Add("UE4EditorDebugArg", (EditorBuildConfig == BuildConfig.Debug || EditorBuildConfig == BuildConfig.DebugGame) ? " -debug" : "");
-			Variables.Add("UseIncrementalBuilds", "1");
-			if (!String.IsNullOrEmpty(SdkInstallerDir))
-			{
-				Variables.Add("SdkInstallerDir", SdkInstallerDir);
-			}
-			return Variables;
-		}
-
 		private void OptionsContextMenu_EditBuildSteps_Click(object sender, EventArgs e)
 		{
 			ConfigFile ProjectConfigFile = Workspace.ProjectConfigFile;
@@ -5361,7 +5218,8 @@ namespace UnrealGameSync
 				List<BuildStep> UserSteps = GetUserBuildSteps(ProjectBuildStepObjects);
 
 				// Show the dialog
-				ModifyBuildStepsWindow EditStepsWindow = new ModifyBuildStepsWindow(TargetNames, UserSteps, new HashSet<Guid>(ProjectBuildStepObjects.Keys), BranchDirectoryName, GetWorkspaceVariables(Workspace.CurrentChangeNumber));
+				Dictionary<string, string> Variables = Workspace.GetVariables(GetEditorBuildConfig());
+				ModifyBuildStepsWindow EditStepsWindow = new ModifyBuildStepsWindow(TargetNames, UserSteps, new HashSet<Guid>(ProjectBuildStepObjects.Keys), BranchDirectoryName, Variables);
 				EditStepsWindow.ShowDialog();
 
 				// Update the user settings
@@ -5412,7 +5270,7 @@ namespace UnrealGameSync
 
 		private bool ShouldSyncPrecompiledEditor
 		{
-			get { return Settings.Archives.Any(x => x.bEnabled && x.Type == EditorArchiveType) && GetArchives().Any(x => x.Type == "Editor"); }
+			get { return Settings.Archives.Any(x => x.bEnabled && x.Type == IArchiveInfo.EditorArchiveType) && GetArchives().Any(x => x.Type == "Editor"); }
 		}
 
 		public BuildConfig GetEditorBuildConfig()
@@ -5420,54 +5278,12 @@ namespace UnrealGameSync
 			return ShouldSyncPrecompiledEditor ? BuildConfig.Development : Settings.CompiledEditorBuildConfig;
 		}
 
-		private string GetDefaultEditorTargetName()
-		{
-			string? EditorTarget;
-			if (!TryGetProjectSetting(PerforceMonitor.LatestProjectConfigFile, "EditorTarget", out EditorTarget))
-			{
-				EditorTarget = "UE4Editor";
-			}
-			return EditorTarget;
-		}
-
-		private Dictionary<Guid, ConfigObject> GetDefaultBuildStepObjects()
-		{
-			string ProjectArgument = "";
-			if (SelectedFileName.HasExtension(".uproject"))
-			{
-				ProjectArgument = String.Format("\"{0}\"", SelectedFileName);
-			}
-
-			string ActualEditorTargetName;
-			if (EditorTargetName != null)
-			{
-				ActualEditorTargetName = EditorTargetName;
-			}
-			else if (SelectedFileName.HasExtension(".uproject") && bIsEnterpriseProject)
-			{
-				ActualEditorTargetName = "StudioEditor";
-			}
-			else
-			{
-				ActualEditorTargetName = GetDefaultEditorTargetName();
-			}
-
-			bool bUseCrashReportClientEditor = PerforceMonitor.LatestProjectConfigFile.GetValue("Options.UseCrashReportClientEditor", false);
-
-			List<BuildStep> DefaultBuildSteps = new List<BuildStep>();
-			DefaultBuildSteps.Add(new BuildStep(new Guid("{01F66060-73FA-4CC8-9CB3-E217FBBA954E}"), 0, "Compile UnrealHeaderTool", "Compiling UnrealHeaderTool...", 1, "UnrealHeaderTool", "Win64", "Development", "", !ShouldSyncPrecompiledEditor));
-			DefaultBuildSteps.Add(new BuildStep(new Guid("{F097FF61-C916-4058-8391-35B46C3173D5}"), 1, String.Format("Compile {0}", ActualEditorTargetName), String.Format("Compiling {0}...", ActualEditorTargetName), 10, ActualEditorTargetName, "Win64", Settings.CompiledEditorBuildConfig.ToString(), ProjectArgument, !ShouldSyncPrecompiledEditor));
-			DefaultBuildSteps.Add(new BuildStep(new Guid("{C6E633A1-956F-4AD3-BC95-6D06D131E7B4}"), 2, "Compile ShaderCompileWorker", "Compiling ShaderCompileWorker...", 1, "ShaderCompileWorker", "Win64", "Development", "", !ShouldSyncPrecompiledEditor));
-			DefaultBuildSteps.Add(new BuildStep(new Guid("{24FFD88C-7901-4899-9696-AE1066B4B6E8}"), 3, "Compile UnrealLightmass", "Compiling UnrealLightmass...", 1, "UnrealLightmass", "Win64", "Development", "", !ShouldSyncPrecompiledEditor));
-			DefaultBuildSteps.Add(new BuildStep(new Guid("{FFF20379-06BF-4205-8A3E-C53427736688}"), 4, "Compile CrashReportClient", "Compiling CrashReportClient...", 1, "CrashReportClient", "Win64", "Shipping", "", !ShouldSyncPrecompiledEditor && !bUseCrashReportClientEditor));
-			DefaultBuildSteps.Add(new BuildStep(new Guid("{7143D861-58D3-4F83-BADC-BC5DCB2079F6}"), 5, "Compile CrashReportClientEditor", "Compiling CrashReportClientEditor...", 1, "CrashReportClientEditor", "Win64", "Shipping", "", !ShouldSyncPrecompiledEditor && bUseCrashReportClientEditor));
-
-			return DefaultBuildSteps.ToDictionary(x => x.UniqueId, x => x.ToConfigObject());
-		}
-
 		private Dictionary<Guid, ConfigObject> GetProjectBuildStepObjects(ConfigFile ProjectConfigFile)
 		{
-			Dictionary<Guid, ConfigObject> ProjectBuildSteps = GetDefaultBuildStepObjects();
+			FileReference EditorTargetFile = ConfigUtils.GetEditorTargetFile(ProjectInfo, ProjectConfigFile);
+			string EditorTargetName = EditorTargetFile.GetFileNameWithoutAnyExtensions();
+
+			Dictionary<Guid, ConfigObject> ProjectBuildSteps = ConfigUtils.GetDefaultBuildStepObjects(ProjectInfo, EditorTargetName, Settings.CompiledEditorBuildConfig, ProjectConfigFile, ShouldSyncPrecompiledEditor);
 			foreach (string Line in ProjectConfigFile.GetValues("Build.Step", new string[0]))
 			{
 				AddOrUpdateBuildStep(ProjectBuildSteps, new ConfigObject(Line));
@@ -5492,7 +5308,7 @@ namespace UnrealGameSync
 		{
 			if (OptionsContextMenu_SyncPrecompiledBinaries.DropDownItems.Count == 0)
 			{
-				IArchiveInfo EditorArchive = GetArchives().FirstOrDefault(x => x.Type == EditorArchiveType);
+				IArchiveInfo EditorArchive = GetArchives().FirstOrDefault(x => x.Type == IArchiveInfo.EditorArchiveType);
 				if (EditorArchive != null)
 				{
 					SetSelectedArchive(EditorArchive, !OptionsContextMenu_SyncPrecompiledBinaries.Checked);
diff --git a/Engine/Source/Programs/UnrealGameSync/UnrealGameSync/Forms/VariablesWindow.cs b/Engine/Source/Programs/UnrealGameSync/UnrealGameSync/Forms/VariablesWindow.cs
--- a/Engine/Source/Programs/UnrealGameSync/UnrealGameSync/Forms/VariablesWindow.cs
+++ b/Engine/Source/Programs/UnrealGameSync/UnrealGameSync/Forms/VariablesWindow.cs
@@ -17,6 +17,15 @@ namespace UnrealGameSync
 {
 	public partial class VariablesWindow : Form
 	{
+		static HashSet<string> LegacyVariables = new HashSet<string>()
+		{
+			"UE4EditorConfig",
+			"UE4EditorDebugArg",
+			"UE4EditorExe",
+			"UE4EditorCmdExe",
+			"UseIncrementalBuilds",
+		};
+
 		public delegate void InsertVariableDelegate(string Name);
 
 		public event InsertVariableDelegate? OnInsertVariable;
@@ -33,10 +42,13 @@ namespace UnrealGameSync
 
 			foreach(KeyValuePair<string, string> Pair in Variables)
 			{
-				ListViewItem Item = new ListViewItem(String.Format("$({0})", Pair.Key));
-				Item.SubItems.Add(Pair.Value);
-				Item.Group = CurrentProjectGroup;
-				MacrosList.Items.Add(Item);
+				if (!LegacyVariables.Contains(Pair.Key))
+				{
+					ListViewItem Item = new ListViewItem(String.Format("$({0})", Pair.Key));
+					Item.SubItems.Add(Pair.Value);
+					Item.Group = CurrentProjectGroup;
+					MacrosList.Items.Add(Item);
+				}
 			}
 
 			foreach(DictionaryEntry Entry in Environment.GetEnvironmentVariables().OfType<DictionaryEntry>())
diff --git a/Engine/Source/Programs/UnrealGameSync/UnrealGameSync/Workspace.cs b/Engine/Source/Programs/UnrealGameSync/UnrealGameSync/Workspace.cs
--- a/Engine/Source/Programs/UnrealGameSync/UnrealGameSync/Workspace.cs
+++ b/Engine/Source/Programs/UnrealGameSync/UnrealGameSync/Workspace.cs
@@ -148,6 +148,20 @@ namespace UnrealGameSync
 			}
 		}
 
+		public Dictionary<string, string> GetVariables(BuildConfig EditorConfig, int? OverrideChange = null, int? OverrideCodeChange = null)
+		{
+			FileReference EditorReceiptFile = ConfigUtils.GetEditorReceiptFile(Project, ProjectConfigFile, EditorConfig);
+
+			TargetReceipt? EditorReceipt;
+			if (!ConfigUtils.TryReadEditorReceipt(Project, EditorReceiptFile, out EditorReceipt))
+			{
+				EditorReceipt = ConfigUtils.CreateDefaultEditorReceipt(Project, ProjectConfigFile, EditorConfig);
+			}
+
+			Dictionary<string, string> Variables = ConfigUtils.GetWorkspaceVariables(Project, OverrideChange ?? State.CurrentChangeNumber, OverrideCodeChange ?? State.CurrentCodeChangeNumber, EditorReceipt, ProjectConfigFile);
+			return Variables;
+		}
+
 		public bool IsBusy()
 		{
 			return bSyncing;
diff --git a/Engine/Source/Programs/UnrealGameSync/UnrealGameSyncShared/ConfigUtils.cs b/Engine/Source/Programs/UnrealGameSync/UnrealGameSyncShared/ConfigUtils.cs
--- a/Engine/Source/Programs/UnrealGameSync/UnrealGameSyncShared/ConfigUtils.cs
+++ b/Engine/Source/Programs/UnrealGameSync/UnrealGameSyncShared/ConfigUtils.cs
@@ -5,16 +5,75 @@ using EpicGames.Perforce;
 using Microsoft.Extensions.Logging;
 using System;
 using System.Collections.Generic;
+using System.Diagnostics.CodeAnalysis;
 using System.IO;
 using System.Linq;
+using System.Runtime.InteropServices;
 using System.Text;
+using System.Text.Json;
 using System.Threading;
 using System.Threading.Tasks;
 
 namespace UnrealGameSync
 {
+	public class TargetReceipt
+	{
+		public string? Configuration { get; set; }
+		public string? Launch { get; set; }
+		public string? LaunchCmd { get; set; }
+
+		public static bool TryRead(FileReference Location, DirectoryReference? EngineDir, DirectoryReference? ProjectDir, [NotNullWhen(true)] out TargetReceipt? Receipt)
+		{
+			if (Utility.TryLoadJson(Location, out Receipt))
+			{
+				Receipt.Launch = ExpandReceiptVariables(Receipt.Launch, EngineDir, ProjectDir);
+				Receipt.LaunchCmd = ExpandReceiptVariables(Receipt.LaunchCmd, EngineDir, ProjectDir);
+				return true;
+			}
+			return false;
+		}
+
+		[return: NotNullIfNotNull("Line")]
+		private static string? ExpandReceiptVariables(string? Line, DirectoryReference? EngineDir, DirectoryReference? ProjectDir)
+		{
+			string? ExpandedLine = Line;
+			if (ExpandedLine != null)
+			{
+				if (EngineDir != null)
+				{
+					ExpandedLine = ExpandedLine.Replace("$(EngineDir)", EngineDir.FullName);
+				}
+				if (ProjectDir != null)
+				{
+					ExpandedLine = ExpandedLine.Replace("$(ProjectDir)", ProjectDir.FullName);
+				}
+			}
+			return ExpandedLine;
+		}
+	}
+
 	public static class ConfigUtils
 	{
+		public static string HostPlatform { get; } = GetHostPlatform();
+
+		public static string HostArchitectureSuffix { get; } = String.Empty;
+
+		static string GetHostPlatform()
+		{
+			if (RuntimeInformation.IsOSPlatform(OSPlatform.Windows))
+			{
+				return "Win64";
+			}
+			else if (RuntimeInformation.IsOSPlatform(OSPlatform.OSX))
+			{
+				return "Mac";
+			}
+			else
+			{
+				return "Linux";
+			}
+		}
+
 		public static async Task<ConfigFile> ReadProjectConfigFileAsync(IPerforceConnection Perforce, string BranchClientPath, string SelectedClientFileName, DirectoryReference CacheFolder, List<KeyValuePair<FileReference, DateTime>> LocalConfigFiles, ILogger Logger, CancellationToken CancellationToken)
 		{
 			List<string> ConfigFilePaths = Utility.GetDepotConfigPaths(BranchClientPath + "/Engine", SelectedClientFileName);
@@ -74,5 +133,251 @@ namespace UnrealGameSync
 			}
 			return ProjectConfig;
 		}
+
+		public static FileReference GetEditorTargetFile(ProjectInfo ProjectInfo, ConfigFile ProjectConfig)
+		{
+			if (ProjectInfo.ProjectPath.EndsWith(".uproject", StringComparison.OrdinalIgnoreCase))
+			{
+				List<FileReference> TargetFiles = FindTargets(ProjectInfo.LocalFileName.Directory);
+
+				FileReference? TargetFile = TargetFiles.OrderBy(x => x.FullName, StringComparer.OrdinalIgnoreCase).FirstOrDefault(x => x.FullName.EndsWith("Editor.target.cs", StringComparison.OrdinalIgnoreCase));
+				if (TargetFile != null)
+				{
+					return TargetFile;
+				}
+			}
+
+			string DefaultEditorTargetName = GetDefaultEditorTargetName(ProjectInfo, ProjectConfig);
+			return FileReference.Combine(ProjectInfo.LocalRootPath, "Engine", "Source", $"{DefaultEditorTargetName}.Target.cs");
+		}
+
+		public static FileReference GetEditorReceiptFile(ProjectInfo ProjectInfo, ConfigFile ProjectConfig, BuildConfig Config)
+		{
+			FileReference TargetFile = GetEditorTargetFile(ProjectInfo, ProjectConfig);
+			return GetReceiptFile(ProjectInfo, TargetFile, Config.ToString());
+		}
+
+		private static List<FileReference> FindTargets(DirectoryReference EngineOrProjectDir)
+		{
+			List<FileReference> Targets = new List<FileReference>();
+
+			DirectoryReference SourceDir = DirectoryReference.Combine(EngineOrProjectDir, "Source");
+			if (DirectoryReference.Exists(SourceDir))
+			{
+				foreach (FileReference TargetFile in DirectoryReference.EnumerateFiles(SourceDir))
+				{
+					const string Extension = ".target.cs";
+					if (TargetFile.FullName.EndsWith(Extension, StringComparison.OrdinalIgnoreCase))
+					{
+						Targets.Add(TargetFile);
+					}
+				}
+			}
+
+			return Targets;
+		}
+
+		public static string GetDefaultEditorTargetName(ProjectInfo ProjectInfo, ConfigFile ProjectConfigFile)
+		{
+			string? EditorTarget;
+			if (!TryGetProjectSetting(ProjectConfigFile, ProjectInfo.ProjectIdentifier, "EditorTarget", out EditorTarget))
+			{
+				if (ProjectInfo.bIsEnterpriseProject)
+				{
+					EditorTarget = "StudioEditor";
+				}
+				else
+				{
+					EditorTarget = "UE4Editor";
+				}
+			}
+			return EditorTarget;
+		}
+
+		public static bool TryReadEditorReceipt(ProjectInfo ProjectInfo, FileReference ReceiptFile, [NotNullWhen(true)] out TargetReceipt? Receipt)
+		{
+			DirectoryReference EngineDir = DirectoryReference.Combine(ProjectInfo.LocalRootPath, "Engine");
+			DirectoryReference ProjectDir = ProjectInfo.LocalFileName.Directory;
+
+			if (ReceiptFile.IsUnderDirectory(ProjectDir))
+			{
+				return TargetReceipt.TryRead(ReceiptFile, EngineDir, ProjectDir, out Receipt);
+			}
+			else
+			{
+				return TargetReceipt.TryRead(ReceiptFile, EngineDir, null, out Receipt);
+			}
+		}
+
+		public static TargetReceipt CreateDefaultEditorReceipt(ProjectInfo ProjectInfo, ConfigFile ProjectConfigFile, BuildConfig Configuration)
+		{
+			string BaseName = GetDefaultEditorTargetName(ProjectInfo, ProjectConfigFile);
+			if (Configuration != BuildConfig.Development || !String.IsNullOrEmpty(HostArchitectureSuffix))
+			{
+				if (Configuration != BuildConfig.DebugGame || ProjectConfigFile.GetValue("Options.DebugGameHasSeparateExecutable", false))
+				{
+					BaseName += $"-{HostPlatform}-{Configuration}{HostArchitectureSuffix}";
+				}
+			}
+
+			string Extension = String.Empty;
+			if (RuntimeInformation.IsOSPlatform(OSPlatform.Windows))
+			{
+				Extension = ".exe";
+			}
+
+			TargetReceipt Receipt = new TargetReceipt();
+			Receipt.Configuration = Configuration.ToString();
+			Receipt.Launch = FileReference.Combine(ProjectInfo.LocalRootPath, "Engine", "Binaries", HostPlatform, $"{BaseName}{Extension}").FullName;
+			Receipt.LaunchCmd = FileReference.Combine(ProjectInfo.LocalRootPath, "Engine", "Binaries", HostPlatform, $"{BaseName}-Cmd{Extension}").FullName;
+			return Receipt;
+		}
+
+		public static FileReference GetReceiptFile(ProjectInfo ProjectInfo, FileReference TargetFile, string Configuration)
+		{
+			string TargetName = TargetFile.GetFileNameWithoutAnyExtensions();
+
+			DirectoryReference? ProjectDir = ProjectInfo.ProjectDir;
+			if (ProjectDir != null && TargetFile.IsUnderDirectory(ProjectDir))
+			{
+				return GetReceiptFile(ProjectDir, TargetName, Configuration);
+			}
+			else
+			{
+				return GetReceiptFile(ProjectInfo.EngineDir, TargetName, Configuration);
+			}
+		}
+
+		public static FileReference GetReceiptFile(DirectoryReference BaseDir, string TargetName, string Configuration)
+		{
+			return GetReceiptFile(BaseDir, TargetName, HostPlatform, Configuration, HostArchitectureSuffix);
+		}
+
+		public static FileReference GetReceiptFile(DirectoryReference BaseDir, string TargetName, string Platform, string Configuration, string ArchitectureSuffix)
+		{
+			if (String.IsNullOrEmpty(ArchitectureSuffix) && Configuration.Equals("Development", StringComparison.OrdinalIgnoreCase))
+			{
+				return FileReference.Combine(BaseDir, "Binaries", Platform, $"{TargetName}.target");
+			}
+			else
+			{
+				return FileReference.Combine(BaseDir, "Binaries", Platform, $"{TargetName}-{Platform}-{Configuration}{ArchitectureSuffix}.target");
+			}
+		}
+
+		public static Dictionary<Guid, ConfigObject> GetDefaultBuildStepObjects(ProjectInfo ProjectInfo, string EditorTarget, BuildConfig EditorConfig, ConfigFile LatestProjectConfigFile, bool ShouldSyncPrecompiledEditor)
+		{
+			string ProjectArgument = "";
+			if (ProjectInfo.LocalFileName.HasExtension(".uproject"))
+			{
+				ProjectArgument = String.Format("\"{0}\"", ProjectInfo.LocalFileName);
+			}
+
+			bool bUseCrashReportClientEditor = LatestProjectConfigFile.GetValue("Options.UseCrashReportClientEditor", false);
+
+			List<BuildStep> DefaultBuildSteps = new List<BuildStep>();
+			DefaultBuildSteps.Add(new BuildStep(new Guid("{01F66060-73FA-4CC8-9CB3-E217FBBA954E}"), 0, "Compile UnrealHeaderTool", "Compiling UnrealHeaderTool...", 1, "UnrealHeaderTool", "Win64", "Development", "", !ShouldSyncPrecompiledEditor));
+			DefaultBuildSteps.Add(new BuildStep(new Guid("{F097FF61-C916-4058-8391-35B46C3173D5}"), 1, $"Compile {EditorTarget}", $"Compiling {EditorTarget}...", 10, EditorTarget, "Win64", EditorConfig.ToString(), ProjectArgument, !ShouldSyncPrecompiledEditor));
+			DefaultBuildSteps.Add(new BuildStep(new Guid("{C6E633A1-956F-4AD3-BC95-6D06D131E7B4}"), 2, "Compile ShaderCompileWorker", "Compiling ShaderCompileWorker...", 1, "ShaderCompileWorker", "Win64", "Development", "", !ShouldSyncPrecompiledEditor));
+			DefaultBuildSteps.Add(new BuildStep(new Guid("{24FFD88C-7901-4899-9696-AE1066B4B6E8}"), 3, "Compile UnrealLightmass", "Compiling UnrealLightmass...", 1, "UnrealLightmass", "Win64", "Development", "", !ShouldSyncPrecompiledEditor));
+			DefaultBuildSteps.Add(new BuildStep(new Guid("{FFF20379-06BF-4205-8A3E-C53427736688}"), 4, "Compile CrashReportClient", "Compiling CrashReportClient...", 1, "CrashReportClient", "Win64", "Shipping", "", !ShouldSyncPrecompiledEditor && !bUseCrashReportClientEditor));
+			DefaultBuildSteps.Add(new BuildStep(new Guid("{7143D861-58D3-4F83-BADC-BC5DCB2079F6}"), 5, "Compile CrashReportClientEditor", "Compiling CrashReportClientEditor...", 1, "CrashReportClientEditor", "Win64", "Shipping", "", !ShouldSyncPrecompiledEditor && bUseCrashReportClientEditor));
+
+			return DefaultBuildSteps.ToDictionary(x => x.UniqueId, x => x.ToConfigObject());
+		}
+
+		public static Dictionary<string, string> GetWorkspaceVariables(ProjectInfo ProjectInfo, int ChangeNumber, int CodeChangeNumber, TargetReceipt? EditorTarget, ConfigFile? ProjectConfigFile)
+		{
+			Dictionary<string, string> Variables = new Dictionary<string, string>(StringComparer.OrdinalIgnoreCase);
+
+			if (ProjectInfo.StreamName != null)
+			{
+				Variables.Add("Stream", ProjectInfo.StreamName);
+			}
+
+			Variables.Add("Change", ChangeNumber.ToString());
+			Variables.Add("CodeChange", CodeChangeNumber.ToString());
+
+			Variables.Add("ClientName", ProjectInfo.ClientName);
+			Variables.Add("BranchDir", ProjectInfo.LocalRootPath.FullName);
+			Variables.Add("ProjectDir", ProjectInfo.LocalFileName.Directory.FullName);
+			Variables.Add("ProjectFile", ProjectInfo.LocalFileName.FullName);
+			Variables.Add("UseIncrementalBuilds", "1");
+
+			string EditorConfig = EditorTarget?.Configuration ?? String.Empty;
+			Variables.Add("EditorConfig", EditorConfig);
+
+			string EditorLaunch = EditorTarget?.Launch ?? String.Empty;
+			Variables.Add("EditorExe", EditorLaunch);
+
+			string EditorLaunchCmd = EditorTarget?.LaunchCmd ?? EditorLaunch.Replace(".exe", "-Cmd.exe");
+			Variables.Add("EditorCmdExe", EditorLaunchCmd);
+
+			// Legacy
+			Variables.Add("UE4EditorConfig", EditorConfig);
+			Variables.Add("UE4EditorDebugArg", (EditorConfig.Equals("Debug", StringComparison.Ordinal) || EditorConfig.Equals("DebugGame", StringComparison.Ordinal)) ? " -debug" : "");
+			Variables.Add("UE4EditorExe", EditorLaunch);
+			Variables.Add("UE4EditorCmdExe", EditorLaunchCmd);
+
+			if (ProjectConfigFile != null)
+			{
+				if (TryGetProjectSetting(ProjectConfigFile, ProjectInfo.ProjectIdentifier, "SdkInstallerDir", out string? SdkInstallerDir))
+				{
+					Variables.Add("SdkInstallerDir", SdkInstallerDir);
+				}
+			}
+
+			return Variables;
+		}
+
+		public static Dictionary<string, string> GetWorkspaceVariables(ProjectInfo ProjectInfo, int ChangeNumber, int CodeChangeNumber, TargetReceipt? EditorTarget, ConfigFile? ProjectConfigFile, IEnumerable<KeyValuePair<string, string>> AdditionalVariables)
+		{
+			Dictionary<string, string> Variables = GetWorkspaceVariables(ProjectInfo, ChangeNumber, CodeChangeNumber, EditorTarget, ProjectConfigFile);
+			foreach ((string Key, string Value) in AdditionalVariables)
+			{
+				Variables[Key] = Value;
+			}
+			return Variables;
+		}
+
+		public static bool TryGetProjectSetting(ConfigFile ProjectConfigFile, string SelectedProjectIdentifier, string Name, [NotNullWhen(true)] out string? Value)
+		{
+			string Path = SelectedProjectIdentifier;
+			for (; ; )
+			{
+				ConfigSection ProjectSection = ProjectConfigFile.FindSection(Path);
+				if (ProjectSection != null)
+				{
+					string? NewValue = ProjectSection.GetValue(Name, null);
+					if (NewValue != null)
+					{
+						Value = NewValue;
+						return true;
+					}
+				}
+
+				int LastSlash = Path.LastIndexOf('/');
+				if (LastSlash < 2)
+				{
+					break;
+				}
+
+				Path = Path.Substring(0, LastSlash);
+			}
+
+			ConfigSection DefaultSection = ProjectConfigFile.FindSection("Default");
+			if (DefaultSection != null)
+			{
+				string? NewValue = DefaultSection.GetValue(Name, null);
+				if (NewValue != null)
+				{
+					Value = NewValue;
+					return true;
+				}
+			}
+
+			Value = null;
+			return false;
+		}
 	}
 }
diff --git a/Engine/Source/Programs/UnrealGameSync/UnrealGameSyncShared/UserSettings.cs b/Engine/Source/Programs/UnrealGameSync/UnrealGameSyncShared/UserSettings.cs
--- a/Engine/Source/Programs/UnrealGameSync/UnrealGameSyncShared/UserSettings.cs
+++ b/Engine/Source/Programs/UnrealGameSync/UnrealGameSyncShared/UserSettings.cs
@@ -306,6 +306,7 @@ namespace UnrealGameSync
 		// Settings for the currently synced project in this workspace. CurrentChangeNumber is only valid for this workspace if CurrentProjectPath is the current project.
 		public string? CurrentProjectIdentifier { get; set; }
 		public int CurrentChangeNumber { get; set; } = -1;
+		public int CurrentCodeChangeNumber { get; set; } = -1;
 		public string? CurrentSyncFilterHash { get; set; }
 		public List<int> AdditionalChangeNumbers { get; set; } = new List<int>();
 
diff --git a/Engine/Source/Programs/UnrealGameSync/UnrealGameSyncShared/WorkspaceSettings.cs b/Engine/Source/Programs/UnrealGameSync/UnrealGameSyncShared/WorkspaceSettings.cs
--- a/Engine/Source/Programs/UnrealGameSync/UnrealGameSyncShared/WorkspaceSettings.cs
+++ b/Engine/Source/Programs/UnrealGameSync/UnrealGameSyncShared/WorkspaceSettings.cs
@@ -304,7 +304,7 @@ namespace UnrealGameSync
 				string BranchPath = (BranchIdx == -1) ? String.Empty : BranchClientPath.Substring(BranchIdx);
 				string ProjectPath = NewSelectedClientFileName.Substring(BranchClientPath.Length);
 
-				ProjectInfo ProjectInfo = new ProjectInfo(BranchDirectoryName, PerforceSettings.ClientName!, BranchPath, ProjectPath, NewSelectedProjectIdentifier, bIsEnterpriseProject);
+				ProjectInfo ProjectInfo = new ProjectInfo(BranchDirectoryName, PerforceSettings.ClientName!, BranchPath, ProjectPath, StreamName, NewSelectedProjectIdentifier, bIsEnterpriseProject);
 
 				WorkspaceSettings WorkspaceSettings = new WorkspaceSettings(PerforceSettings, SelectedProject, ProjectInfo, NewSelectedProjectIdentifier, NewProjectEditorTarget!, StreamName!, LatestProjectConfigFile, WorkspaceProjectConfigFile, WorkspaceProjectStreamFilter, LocalConfigFiles);
 				DirectoryReference.CreateDirectory(WorkspaceSettings.DataFolder);
diff --git a/Engine/Source/Programs/UnrealGameSync/UnrealGameSyncShared/WorkspaceUpdate.cs b/Engine/Source/Programs/UnrealGameSync/UnrealGameSyncShared/WorkspaceUpdate.cs
--- a/Engine/Source/Programs/UnrealGameSync/UnrealGameSyncShared/WorkspaceUpdate.cs
+++ b/Engine/Source/Programs/UnrealGameSync/UnrealGameSyncShared/WorkspaceUpdate.cs
@@ -81,6 +81,8 @@ namespace UnrealGameSync
 
 	public interface IArchiveInfo
 	{
+		public const string EditorArchiveType = "Editor";
+
 		string Name { get; }
 		string Type { get; }
 		string BasePath { get; }
@@ -95,14 +97,14 @@ namespace UnrealGameSync
 		public DateTime StartTime = DateTime.UtcNow;
 		public int ChangeNumber;
 		public WorkspaceUpdateOptions Options;
+		public BuildConfig EditorConfig;
 		public string[]? SyncFilter;
 		public Dictionary<string, Tuple<IArchiveInfo, string>?> ArchiveTypeToArchive = new Dictionary<string, Tuple<IArchiveInfo, string>?>();
 		public Dictionary<string, bool> DeleteFiles = new Dictionary<string,bool>();
 		public Dictionary<string, bool> ClobberFiles = new Dictionary<string,bool>();
-		public Dictionary<Guid,ConfigObject> DefaultBuildSteps;
 		public List<ConfigObject> UserBuildStepObjects;
 		public HashSet<Guid>? CustomBuildSteps;
-		public Dictionary<string, string> Variables;
+		public Dictionary<string, string> AdditionalVariables = new Dictionary<string, string>(StringComparer.OrdinalIgnoreCase);
 		public PerforceSyncOptions? PerforceSyncOptions;
 		public List<HaveRecord>? HaveFiles; // Cached when sync filter has changed
 
@@ -110,15 +112,14 @@ namespace UnrealGameSync
 		public ConfigFile? ProjectConfigFile;
 		public IReadOnlyList<string>? ProjectStreamFilter;
 
-		public WorkspaceUpdateContext(int InChangeNumber, WorkspaceUpdateOptions InOptions, string[]? InSyncFilter, Dictionary<Guid, ConfigObject> InDefaultBuildSteps, List<ConfigObject> InUserBuildSteps, HashSet<Guid>? InCustomBuildSteps, Dictionary<string, string> InVariables)
+		public WorkspaceUpdateContext(int InChangeNumber, WorkspaceUpdateOptions InOptions, BuildConfig InEditorConfig, string[]? InSyncFilter, List<ConfigObject> InUserBuildSteps, HashSet<Guid>? InCustomBuildSteps)
 		{
 			ChangeNumber = InChangeNumber;
 			Options = InOptions;
+			EditorConfig = InEditorConfig;
 			SyncFilter = InSyncFilter;
-			DefaultBuildSteps = InDefaultBuildSteps;
 			UserBuildStepObjects = InUserBuildSteps;
 			CustomBuildSteps = InCustomBuildSteps;
-			Variables = InVariables;
 		}
 	}
 
@@ -180,18 +181,25 @@ namespace UnrealGameSync
 	public class ProjectInfo
 	{
 		public DirectoryReference LocalRootPath { get; } // ie. local mapping of clientname + branchpath
+
 		public string ClientName { get; }
 		public string BranchPath { get; } // starts with a slash if non-empty. does not end with a slash.
 		public string ProjectPath { get; } // starts with a slash, uses forward slashes
+
+		public string? StreamName { get; } // name of the current stream
+
 		public string ProjectIdentifier { get; } // stream path to project
-		public string TelemetryProjectIdentifier { get; }
-		public bool bIsEnterpriseProject { get; }
+		public bool bIsEnterpriseProject { get; } // whether it's an enterprise project
 
+		// derived properties
 		public FileReference LocalFileName => new FileReference(LocalRootPath.FullName + ProjectPath);
 		public string ClientRootPath => $"//{ClientName}{BranchPath}";
 		public string ClientFileName => $"//{ClientName}{BranchPath}{ProjectPath}";
+		public string TelemetryProjectIdentifier => PerforceUtils.GetClientOrDepotDirectoryName(ProjectIdentifier);
+		public DirectoryReference EngineDir => DirectoryReference.Combine(LocalRootPath, "Engine");
+		public DirectoryReference? ProjectDir => ProjectPath.EndsWith(".uproject")? LocalFileName.Directory : null;
 
-		public ProjectInfo(DirectoryReference LocalRootPath, string ClientName, string BranchPath, string ProjectPath, string ProjectIdentifier, bool bIsEnterpriseProject)
+		public ProjectInfo(DirectoryReference LocalRootPath, string ClientName, string BranchPath, string ProjectPath, string? StreamName, string ProjectIdentifier, bool bIsEnterpriseProject)
 		{
 			ValidateBranchPath(BranchPath);
 			ValidateProjectPath(ProjectPath);
@@ -200,8 +208,8 @@ namespace UnrealGameSync
 			this.ClientName = ClientName;
 			this.BranchPath = BranchPath;
 			this.ProjectPath = ProjectPath;
+			this.StreamName = StreamName;
 			this.ProjectIdentifier = ProjectIdentifier;
-			this.TelemetryProjectIdentifier = PerforceUtils.GetClientOrDepotDirectoryName(ProjectIdentifier);
 			this.bIsEnterpriseProject = bIsEnterpriseProject;
 		}
 
@@ -929,10 +937,7 @@ namespace UnrealGameSync
 							Logger.LogInformation("");
 							Logger.LogInformation("Executing post-sync steps...");
 
-							Dictionary<string, string> PostSyncVariables = new Dictionary<string, string>(Context.Variables);
-							PostSyncVariables["Change"] = Context.ChangeNumber.ToString();
-							PostSyncVariables["CodeChange"] = VersionChangeNumber.ToString();
-
+							Dictionary<string, string> PostSyncVariables = ConfigUtils.GetWorkspaceVariables(Project, Context.ChangeNumber, VersionChangeNumber, null, Context.ProjectConfigFile);
 							foreach (string PostSyncStep in PostSyncSteps.Select(x => x.Trim()))
 							{
 								ConfigObject PostSyncStepObject = new ConfigObject(PostSyncStep);
@@ -962,6 +967,7 @@ namespace UnrealGameSync
 					else
 					{
 						State.CurrentChangeNumber = Context.ChangeNumber;
+						State.CurrentCodeChangeNumber = VersionChangeNumber;
 					}
 					State.Save();
 
@@ -1079,8 +1085,16 @@ namespace UnrealGameSync
 			{
 				Context.ProjectConfigFile ??= await ReadProjectConfigFile(Project.LocalRootPath, Project.LocalFileName, Logger);
 
-				// Compile all the build steps together
-				Dictionary<Guid, ConfigObject> BuildStepObjects = Context.DefaultBuildSteps.ToDictionary(x => x.Key, x => new ConfigObject(x.Value));
+				// Figure out the new editor target
+				TargetReceipt DefaultEditorReceipt = ConfigUtils.CreateDefaultEditorReceipt(Project, Context.ProjectConfigFile, Context.EditorConfig);
+
+				FileReference EditorTargetFile = ConfigUtils.GetEditorTargetFile(Project, Context.ProjectConfigFile);
+				string EditorTargetName = EditorTargetFile.GetFileNameWithoutAnyExtensions();
+				FileReference EditorReceiptFile = ConfigUtils.GetReceiptFile(Project, EditorTargetFile, Context.EditorConfig.ToString());
+
+				// Get the build steps
+				bool UsingPrecompiledEditor = Context.ArchiveTypeToArchive.TryGetValue(IArchiveInfo.EditorArchiveType, out Tuple<IArchiveInfo, string>? ArchiveInfo) && ArchiveInfo != null;
+				Dictionary<Guid, ConfigObject> BuildStepObjects = ConfigUtils.GetDefaultBuildStepObjects(Project, EditorTargetName, Context.EditorConfig, Context.ProjectConfigFile, UsingPrecompiledEditor);
 				BuildStep.MergeBuildStepObjects(BuildStepObjects, Context.ProjectConfigFile.GetValues("Build.Step", new string[0]).Select(x => new ConfigObject(x)));
 				BuildStep.MergeBuildStepObjects(BuildStepObjects, Context.UserBuildStepObjects);
 
@@ -1137,6 +1151,15 @@ namespace UnrealGameSync
 
 						if (Step.IsValid())
 						{
+							// Get the build variables for this step
+							TargetReceipt? EditorReceipt;
+							if (!ConfigUtils.TryReadEditorReceipt(Project, EditorReceiptFile, out EditorReceipt))
+							{
+								EditorReceipt = DefaultEditorReceipt;
+							}
+							Dictionary<string, string> Variables = ConfigUtils.GetWorkspaceVariables(Project, State.CurrentChangeNumber, State.CurrentCodeChangeNumber, EditorReceipt, Context.ProjectConfigFile, Context.AdditionalVariables);
+
+							// Handle all the different types of steps
 							switch (Step.Type)
 							{
 								case BuildStepType.Compile:
@@ -1145,7 +1168,7 @@ namespace UnrealGameSync
 										StepStopwatch.AddData(new { Target = Step.Target });
 
 										FileReference BuildBat = FileReference.Combine(Project.LocalRootPath, "Engine", "Build", "BatchFiles", "Build.bat");
-										string CommandLine = String.Format("{0} {1} {2} {3} -NoHotReloadFromIDE", Step.Target, Step.Platform, Step.Configuration, Utility.ExpandVariables(Step.Arguments ?? "", Context.Variables));
+										string CommandLine = String.Format("{0} {1} {2} {3} -NoHotReloadFromIDE", Step.Target, Step.Platform, Step.Configuration, Utility.ExpandVariables(Step.Arguments ?? "", Variables));
 										if (Context.Options.HasFlag(WorkspaceUpdateOptions.Clean) || bForceClean)
 										{
 											Logger.LogInformation("ubt> Running {FileName} {Arguments}", BuildBat, CommandLine + " -clean");
@@ -1199,9 +1222,9 @@ namespace UnrealGameSync
 									{
 										StepStopwatch.AddData(new { FileName = Path.GetFileNameWithoutExtension(Step.FileName) });
 
-										FileReference ToolFileName = FileReference.Combine(Project.LocalRootPath, Utility.ExpandVariables(Step.FileName ?? "unknown", Context.Variables));
-										string ToolWorkingDir = String.IsNullOrWhiteSpace(Step.WorkingDir) ? ToolFileName.Directory.FullName : Utility.ExpandVariables(Step.WorkingDir, Context.Variables);
-										string ToolArguments = Utility.ExpandVariables(Step.Arguments ?? "", Context.Variables);
+										FileReference ToolFileName = FileReference.Combine(Project.LocalRootPath, Utility.ExpandVariables(Step.FileName ?? "unknown", Variables));
+										string ToolWorkingDir = String.IsNullOrWhiteSpace(Step.WorkingDir) ? ToolFileName.Directory.FullName : Utility.ExpandVariables(Step.WorkingDir, Variables);
+										string ToolArguments = Utility.ExpandVariables(Step.Arguments ?? "", Variables);
 										Logger.LogInformation("tool> Running {0} {1}", ToolFileName, ToolArguments);
 
 										if (Step.bUseLogWindow)
