UGS: Keep cached project info in the workspace state file, and move background task management into the workspace object in the GUI app.

diff --git a/Engine/Source/Programs/UnrealGameSync/UnrealGameSync/Controls/WorkspaceControl.cs b/Engine/Source/Programs/UnrealGameSync/UnrealGameSync/Controls/WorkspaceControl.cs
--- a/Engine/Source/Programs/UnrealGameSync/UnrealGameSync/Controls/WorkspaceControl.cs
+++ b/Engine/Source/Programs/UnrealGameSync/UnrealGameSync/Controls/WorkspaceControl.cs
@@ -215,44 +215,17 @@ namespace UnrealGameSync
 
 		bool bUserHasOpenIssues = false;
 
-		public UserSelectedProjectSettings SelectedProject
-		{
-			get;
-			private set;
-		}
-
-		public FileReference SelectedFileName
-		{
-			get;
-			private set;
-		}
-
-		string SelectedProjectIdentifier;
-
-		public DirectoryReference BranchDirectoryName
-		{
-			get;
-			private set;
-		}
-
-		public string StreamName
-		{
-			get;
-			private set;
-		}
+		public UserSelectedProjectSettings SelectedProject { get; }
 
-		public string ClientName
-		{
-			get;
-			private set;
-		}
+		public FileReference SelectedFileName => ProjectInfo.LocalFileName;
+		string SelectedProjectIdentifier => ProjectInfo.ProjectIdentifier;
+		public DirectoryReference BranchDirectoryName => ProjectInfo.LocalRootPath;
+		public string? StreamName => ProjectInfo.StreamName;
+		public string ClientName => ProjectInfo.ClientName;
 
 		SynchronizationContext MainThreadSynchronizationContext;
 		bool bIsDisposing;
 
-		bool bUnstable;
-		string EditorTargetName;
-		bool bIsEnterpriseProject;
 		JupiterMonitor JupiterMonitor;
 		PerforceMonitor PerforceMonitor;
 		Workspace Workspace;
@@ -284,8 +257,6 @@ namespace UnrealGameSync
 		Dictionary<string, Size> BadgeLabelToSize = new Dictionary<string, Size>();
 		List<ServiceBadgeInfo> ServiceBadges = new List<ServiceBadgeInfo>();
 
-		string OriginalExecutableFileName;
-
 		NotificationWindow NotificationWindow;
 
 		public Tuple<TaskbarState, float> DesiredTaskbarState
@@ -319,7 +290,7 @@ namespace UnrealGameSync
 		// Placeholder text that is in the control and cleared when the user starts editing.
 		static string AuthorFilterPlaceholderText = "<username>";
 
-		public WorkspaceControl(IWorkspaceControlOwner InOwner, string? InApiUrl, string InOriginalExecutableFileName, bool bInUnstable, WorkspaceSettings WorkspaceSettings, IServiceProvider InServiceProvider, UserSettings InSettings, OIDCTokenManager? InOidcTokenManager)
+		public WorkspaceControl(IWorkspaceControlOwner InOwner, string? InApiUrl, WorkspaceSettings WorkspaceSettings, IServiceProvider InServiceProvider, UserSettings InSettings, OIDCTokenManager? InOidcTokenManager)
 		{
 			InitializeComponent();
 
@@ -328,15 +299,13 @@ namespace UnrealGameSync
 			Owner = InOwner;
 			ApiUrl = InApiUrl;
 			DataFolder = WorkspaceSettings.DataFolder;
-			OriginalExecutableFileName = InOriginalExecutableFileName;
-			bUnstable = bInUnstable;
 			ServiceProvider = InServiceProvider;
 			Logger = InServiceProvider.GetRequiredService<ILogger<WorkspaceControl>>();
 			PerforceSettings = WorkspaceSettings.PerforceSettings;
 			ProjectInfo = WorkspaceSettings.ProjectInfo;
 			Settings = InSettings;
-			this.WorkspaceSettings = InSettings.FindOrAddWorkspaceSettings(ProjectInfo.LocalRootPath, PerforceSettings.ServerAndPort, PerforceSettings.UserName, ProjectInfo.ClientName, ProjectInfo.BranchPath, ProjectInfo.ProjectPath);
-			WorkspaceState = InSettings.FindOrAddWorkspaceState(this.WorkspaceSettings);
+			this.WorkspaceSettings = WorkspaceSettings.UserWorkspaceSettings;
+			this.WorkspaceState = WorkspaceSettings.UserWorkspaceState;
 
 			string ProjectPath = Regex.Replace(WorkspaceSettings.ProjectInfo.ClientFileName, "^//[^/]+/", "");
 			ProjectSettings = InSettings.FindOrAddProjectSettings(WorkspaceSettings.ProjectInfo.LocalRootPath, ProjectPath);
@@ -385,25 +354,7 @@ namespace UnrealGameSync
 
 			// Commit all the new project info
 			IPerforceSettings PerforceClientSettings = WorkspaceSettings.PerforceSettings;
-			ClientName = PerforceClientSettings.ClientName!;
 			SelectedProject = WorkspaceSettings.SelectedProject;
-			SelectedProjectIdentifier = WorkspaceSettings.NewSelectedProjectIdentifier;
-			SelectedFileName = WorkspaceSettings.ProjectInfo.LocalFileName;
-			EditorTargetName = WorkspaceSettings.NewProjectEditorTarget;
-			bIsEnterpriseProject = WorkspaceSettings.ProjectInfo.bIsEnterpriseProject;
-			StreamName = WorkspaceSettings.StreamName;
-
-			// Update the branch directory
-			BranchDirectoryName = WorkspaceSettings.ProjectInfo.LocalRootPath;
-
-			// Check if we've the project we've got open in this workspace is the one we're actually synced to
-			int CurrentChangeNumber = -1;
-			string? CurrentSyncFilterHash = null;
-			if (String.Compare(WorkspaceState.CurrentProjectIdentifier, SelectedProjectIdentifier, true) == 0)
-			{
-				CurrentChangeNumber = WorkspaceState.CurrentChangeNumber;
-				CurrentSyncFilterHash = WorkspaceState.CurrentSyncFilterHash;
-			}
 
 			// Figure out which API server to use
 			string? NewApiUrl;
@@ -3040,26 +2991,6 @@ namespace UnrealGameSync
 			}
 		}
 
-		// returns the default editor target name if not syncing pcbs or Target isn't specified
-		private string TryGetEditorNameFromArchive()
-		{
-			if(ShouldSyncPrecompiledEditor)
-			{
-				List<IArchiveInfo> Archives = GetSelectedArchives(PerforceMonitor.AvailableArchives);
-				foreach(IArchiveInfo Archive in Archives)
-				{
-					if(Archive.Type == IArchiveInfo.EditorArchiveType)
-					{
-						if(Archive.Target != null)
-						{
-							return Archive.Target;
-						}
-					}
-				}
-			}
-			return EditorTargetName;
-		}
-
 		private bool WaitForProgramsToFinish()
 		{
 			FileReference[] ProcessFileNames = GetProcessesRunningInWorkspace();
@@ -3192,7 +3123,7 @@ namespace UnrealGameSync
 						string? Name = Config.GetValue("Name", null);
 						string? Pattern = Config.GetValue("Pattern", null);
 
-						if(Name != null && Pattern != null)
+						if(Name != null && Pattern != null && StreamName != null)
 						{
 							string? DepotName;
 							if (PerforceUtils.TryGetDepotName(StreamName, out DepotName))
diff --git a/Engine/Source/Programs/UnrealGameSync/UnrealGameSync/Forms/IssueBrowserWindow.cs b/Engine/Source/Programs/UnrealGameSync/UnrealGameSync/Forms/IssueBrowserWindow.cs
--- a/Engine/Source/Programs/UnrealGameSync/UnrealGameSync/Forms/IssueBrowserWindow.cs
+++ b/Engine/Source/Programs/UnrealGameSync/UnrealGameSync/Forms/IssueBrowserWindow.cs
@@ -26,7 +26,7 @@ namespace UnrealGameSync
 		TimeSpan? ServerTimeOffset;
 		IServiceProvider ServiceProvider;
 		ILogger Logger;
-		string CurrentStream;
+		string? CurrentStream;
 		int MaxResults = 0;
 		int PendingMaxResults = 0;
 		string? FilterName;
@@ -36,7 +36,7 @@ namespace UnrealGameSync
 		Task? BackgroundTask;
 		bool bDisposed;
 
-		public IssueBrowserWindow(IssueMonitor IssueMonitor, IPerforceSettings PerforceSettings, TimeSpan? ServerTimeOffset, IServiceProvider ServiceProvider, string CurrentStream, Dictionary<string, Func<IssueData, bool>> CustomFilters, string? FilterName)
+		public IssueBrowserWindow(IssueMonitor IssueMonitor, IPerforceSettings PerforceSettings, TimeSpan? ServerTimeOffset, IServiceProvider ServiceProvider, string? CurrentStream, Dictionary<string, Func<IssueData, bool>> CustomFilters, string? FilterName)
 		{
 			this.IssueMonitor = IssueMonitor;
 			this.PerforceSettings = PerforceSettings;
@@ -260,7 +260,7 @@ namespace UnrealGameSync
 
 		static List<IssueBrowserWindow> ExistingWindows = new List<IssueBrowserWindow>();
 
-		public static void Show(Form Owner, IssueMonitor IssueMonitor, IPerforceSettings PerforceSettings, TimeSpan? ServerTimeOffset, IServiceProvider ServiceProvider, string CurrentStream, Dictionary<string, Func<IssueData, bool>> CustomFilters, string? DefaultFilter)
+		public static void Show(Form Owner, IssueMonitor IssueMonitor, IPerforceSettings PerforceSettings, TimeSpan? ServerTimeOffset, IServiceProvider ServiceProvider, string? CurrentStream, Dictionary<string, Func<IssueData, bool>> CustomFilters, string? DefaultFilter)
 		{
 			IssueBrowserWindow Window = ExistingWindows.FirstOrDefault(x => x.IssueMonitor == IssueMonitor);
 			if(Window == null)
diff --git a/Engine/Source/Programs/UnrealGameSync/UnrealGameSync/Forms/IssueDetailsWindow.cs b/Engine/Source/Programs/UnrealGameSync/UnrealGameSync/Forms/IssueDetailsWindow.cs
--- a/Engine/Source/Programs/UnrealGameSync/UnrealGameSync/Forms/IssueDetailsWindow.cs
+++ b/Engine/Source/Programs/UnrealGameSync/UnrealGameSync/Forms/IssueDetailsWindow.cs
@@ -280,7 +280,7 @@ namespace UnrealGameSync
 		System.Windows.Forms.Timer? UpdateTimer;
 		StatusElementResources StatusElementResources;
 
-		IssueDetailsWindow(IssueMonitor IssueMonitor, IssueData Issue, List<IssueBuildData> IssueBuilds, List<IssueDiagnosticData> Diagnostics, IPerforceSettings PerforceSettings, TimeSpan? ServerTimeOffset, IServiceProvider ServiceProvider, string CurrentStream)
+		IssueDetailsWindow(IssueMonitor IssueMonitor, IssueData Issue, List<IssueBuildData> IssueBuilds, List<IssueDiagnosticData> Diagnostics, IPerforceSettings PerforceSettings, TimeSpan? ServerTimeOffset, IServiceProvider ServiceProvider, string? CurrentStream)
 		{
 			this.IssueMonitor = IssueMonitor;
 			this.Issue = Issue;
@@ -1027,7 +1027,7 @@ namespace UnrealGameSync
 			}
 		}
 
-		public static void Show(Form Owner, IssueMonitor IssueMonitor, IPerforceSettings PerforceSettings, TimeSpan? ServerTimeOffset, IssueData Issue, IServiceProvider ServiceProvider, string CurrentStream)
+		public static void Show(Form Owner, IssueMonitor IssueMonitor, IPerforceSettings PerforceSettings, TimeSpan? ServerTimeOffset, IssueData Issue, IServiceProvider ServiceProvider, string? CurrentStream)
 		{
 			Func<CancellationToken, Task<List<IssueBuildData>>> Func = x => RESTApi.GetAsync<List<IssueBuildData>>($"{IssueMonitor.ApiUrl}/api/issues/{Issue.Id}/builds", x);
 			ModalTask<List<IssueBuildData>>? IssueBuilds = ModalTask.Execute(Owner, "Querying issue", "Querying issue, please wait...", Func);
@@ -1037,7 +1037,7 @@ namespace UnrealGameSync
 			}
 		}
 
-		public static void Show(Form Owner, IssueMonitor IssueMonitor, IPerforceSettings PerforceSettings, TimeSpan? ServerTimeOffset, IssueData Issue, List<IssueBuildData> IssueBuilds, IServiceProvider ServiceProvider, string CurrentStream)
+		public static void Show(Form Owner, IssueMonitor IssueMonitor, IPerforceSettings PerforceSettings, TimeSpan? ServerTimeOffset, IssueData Issue, List<IssueBuildData> IssueBuilds, IServiceProvider ServiceProvider, string? CurrentStream)
 		{
 			IssueDetailsWindow Window = ExistingWindows.FirstOrDefault(x => x.IssueMonitor == IssueMonitor && x.Issue.Id == Issue.Id);
 			if(Window == null)
diff --git a/Engine/Source/Programs/UnrealGameSync/UnrealGameSync/Forms/MainWindow.cs b/Engine/Source/Programs/UnrealGameSync/UnrealGameSync/Forms/MainWindow.cs
--- a/Engine/Source/Programs/UnrealGameSync/UnrealGameSync/Forms/MainWindow.cs
+++ b/Engine/Source/Programs/UnrealGameSync/UnrealGameSync/Forms/MainWindow.cs
@@ -1120,7 +1120,7 @@ namespace UnrealGameSync
 				TaskFlags |= ModalTaskFlags.Quiet;
 			}
 
-			ModalTask<WorkspaceSettings>? SettingsTask = PerforceModalTask.Execute(this, "Opening Project", "Opening project, please wait...", DefaultPerforceSettings, (p, c) => OpenProjectWindow.DetectSettingsAsync(p, Project, ProjectLogger, c), ProjectLogger, TaskFlags);
+			ModalTask<WorkspaceSettings>? SettingsTask = PerforceModalTask.Execute(this, "Opening Project", "Opening project, please wait...", DefaultPerforceSettings, (p, c) => OpenProjectWindow.DetectSettingsAsync(p, Project, Settings, ProjectLogger, c), ProjectLogger, TaskFlags);
 			if (SettingsTask == null || SettingsTask.Failed)
 			{
 				if(SettingsTask != null) CreateErrorPanel(ReplaceTabIdx, Project, SettingsTask.Error);
@@ -1184,7 +1184,7 @@ namespace UnrealGameSync
 			}
 
 			// Now that we have the project settings, we can construct the tab
-			WorkspaceControl NewWorkspace = new WorkspaceControl(this, ApiUrl, OriginalExecutableFileName, bUnstable, WorkspaceSettings, ServiceProvider, Settings, OIDCTokenManager);
+			WorkspaceControl NewWorkspace = new WorkspaceControl(this, ApiUrl, WorkspaceSettings, ServiceProvider, Settings, OIDCTokenManager);
 			
 			NewWorkspace.Parent = TabPanel;
 			NewWorkspace.Dock = DockStyle.Fill;
@@ -1260,7 +1260,7 @@ namespace UnrealGameSync
 			switch (Settings.TabLabels)
 			{
 				case TabLabels.Stream:
-					TabName = Workspace.StreamName;
+					TabName = Workspace.StreamName ?? TabName;
 					break;
 				case TabLabels.ProjectFile:
 					TabName = Workspace.SelectedFileName.FullName;
diff --git a/Engine/Source/Programs/UnrealGameSync/UnrealGameSync/Forms/NetCoreWindow.cs b/Engine/Source/Programs/UnrealGameSync/UnrealGameSync/Forms/NetCoreWindow.cs
deleted file mode 100644
diff --git a/Engine/Source/Programs/UnrealGameSync/UnrealGameSync/Forms/OpenProjectWindow.cs b/Engine/Source/Programs/UnrealGameSync/UnrealGameSync/Forms/OpenProjectWindow.cs
--- a/Engine/Source/Programs/UnrealGameSync/UnrealGameSync/Forms/OpenProjectWindow.cs
+++ b/Engine/Source/Programs/UnrealGameSync/UnrealGameSync/Forms/OpenProjectWindow.cs
@@ -286,7 +286,7 @@ namespace UnrealGameSync
 			{
 				ILogger<WorkspaceSettings> Logger = ServiceProvider.GetRequiredService<ILogger<WorkspaceSettings>>();
 
-				ModalTask<WorkspaceSettings>? NewDetectedProjectSettings = PerforceModalTask.Execute(this, "Opening project", "Opening project, please wait...", DefaultSettings, (x, y) => DetectSettingsAsync(x, SelectedProject, Logger, y), Logger);
+				ModalTask<WorkspaceSettings>? NewDetectedProjectSettings = PerforceModalTask.Execute(this, "Opening project", "Opening project, please wait...", DefaultSettings, (x, y) => DetectSettingsAsync(x, SelectedProject, Settings, Logger, y), Logger);
 				if (NewDetectedProjectSettings != null && NewDetectedProjectSettings.Succeeded)
 				{
 					DetectedProjectSettings = NewDetectedProjectSettings.Result;
@@ -296,9 +296,9 @@ namespace UnrealGameSync
 			}
 		}
 
-		public static async Task<WorkspaceSettings> DetectSettingsAsync(IPerforceConnection Perforce, UserSelectedProjectSettings SelectedProject, ILogger<WorkspaceSettings> Logger, CancellationToken CancellationToken)
+		public static async Task<WorkspaceSettings> DetectSettingsAsync(IPerforceConnection Perforce, UserSelectedProjectSettings SelectedProject, UserSettings UserSettings, ILogger<WorkspaceSettings> Logger, CancellationToken CancellationToken)
 		{
-			WorkspaceSettings Settings = await WorkspaceSettings.CreateAsync(Perforce, SelectedProject, Logger, CancellationToken);
+			WorkspaceSettings Settings = await WorkspaceSettings.CreateAsync(Perforce, SelectedProject, UserSettings, Logger, CancellationToken);
 			if (DeploymentSettings.OnDetectProjectSettings != null)
 			{
 				string? Message;
diff --git a/Engine/Source/Programs/UnrealGameSync/UnrealGameSync/Forms/SelectStreamWindow.cs b/Engine/Source/Programs/UnrealGameSync/UnrealGameSync/Forms/SelectStreamWindow.cs
--- a/Engine/Source/Programs/UnrealGameSync/UnrealGameSync/Forms/SelectStreamWindow.cs
+++ b/Engine/Source/Programs/UnrealGameSync/UnrealGameSync/Forms/SelectStreamWindow.cs
@@ -322,7 +322,7 @@ namespace UnrealGameSync
 			return bExpand || (SelectedStream == Stream.Record.Stream);
 		}
 
-		public static bool ShowModal(IWin32Window Owner, IPerforceSettings Perforce, string StreamName, IServiceProvider ServiceProvider, [NotNullWhen(true)] out string? NewStreamName)
+		public static bool ShowModal(IWin32Window Owner, IPerforceSettings Perforce, string? StreamName, IServiceProvider ServiceProvider, [NotNullWhen(true)] out string? NewStreamName)
 		{
 			ILogger Logger = ServiceProvider.GetRequiredService<ILogger<SelectStreamWindow>>();
 
diff --git a/Engine/Source/Programs/UnrealGameSync/UnrealGameSync/ProgramApplicationContext.cs b/Engine/Source/Programs/UnrealGameSync/UnrealGameSync/ProgramApplicationContext.cs
--- a/Engine/Source/Programs/UnrealGameSync/UnrealGameSync/ProgramApplicationContext.cs
+++ b/Engine/Source/Programs/UnrealGameSync/UnrealGameSync/ProgramApplicationContext.cs
@@ -188,7 +188,7 @@ namespace UnrealGameSync
 			foreach (UserSelectedProjectSettings ProjectSettings in Settings.OpenProjects)
 			{
 				ILogger<WorkspaceSettings> Logger = ServiceProvider.GetRequiredService<ILogger<WorkspaceSettings>>();
-				Task<WorkspaceSettings> WorkspaceSettingsTask = Task.Run(() => WorkspaceSettings.CreateAsync(DefaultPerforceSettings, ProjectSettings, Logger, StartupCancellationSource.Token), StartupCancellationSource.Token);
+				Task<WorkspaceSettings> WorkspaceSettingsTask = Task.Run(() => WorkspaceSettings.CreateAsync(DefaultPerforceSettings, ProjectSettings, Settings, Logger, StartupCancellationSource.Token), StartupCancellationSource.Token);
 				StartupTasks.Add((ProjectSettings, new ModalTask<WorkspaceSettings>(WorkspaceSettingsTask)));
 			}
 			StartupTask = Task.WhenAll(StartupTasks.Select(x => x.Item2.Task));
diff --git a/Engine/Source/Programs/UnrealGameSync/UnrealGameSync/Workspace.cs b/Engine/Source/Programs/UnrealGameSync/UnrealGameSync/Workspace.cs
--- a/Engine/Source/Programs/UnrealGameSync/UnrealGameSync/Workspace.cs
+++ b/Engine/Source/Programs/UnrealGameSync/UnrealGameSync/Workspace.cs
@@ -70,6 +70,7 @@ namespace UnrealGameSync
 			get; private set;
 		}
 
+		CancellationTokenSource? PrevCancellationSource;
 		Task PrevUpdateTask = Task.CompletedTask;
 
 		public void Update(WorkspaceUpdateContext Context)
@@ -77,16 +78,28 @@ namespace UnrealGameSync
 			CancelUpdate();
 
 			Task PrevUpdateTaskCopy = PrevUpdateTask;
-			CurrentUpdate = new WorkspaceUpdate(Context, (Update, CancellationToken) => UpdateWorkspaceMini(Update, PrevUpdateTaskCopy, CancellationToken));
+
+			WorkspaceUpdate Update = new WorkspaceUpdate(Context);
+			CurrentUpdate = Update;
+
+			CancellationTokenSource CancellationSource = new CancellationTokenSource();
+			PrevCancellationSource = CancellationSource;
+			PrevUpdateTask = Task.Run(() => UpdateWorkspaceMini(Update, PrevUpdateTaskCopy, CancellationSource.Token));
 		}
 
 		public void CancelUpdate()
 		{
 			// Cancel the current task. We actually terminate the operation asynchronously, but we can signal the cancellation and 
 			// send a cancelled event, then wait for the heavy lifting to finish in the new update task.
-			if (CurrentUpdate != null)
+			if (PrevCancellationSource != null)
+			{
+				CancellationTokenSource PrevCancellationSourceCopy = PrevCancellationSource;
+				PrevCancellationSourceCopy.Cancel();
+				PrevUpdateTask = PrevUpdateTask.ContinueWith(Task => PrevCancellationSourceCopy.Dispose());
+				PrevCancellationSource = null;
+			}
+			if(CurrentUpdate != null)
 			{
-				PrevUpdateTask = CurrentUpdate.DisposeAsync().AsTask();
 				CompleteUpdate(CurrentUpdate, WorkspaceUpdateResult.Canceled, "Cancelled");
 			}
 		}
diff --git a/Engine/Source/Programs/UnrealGameSync/UnrealGameSyncShared/ConfigFile.cs b/Engine/Source/Programs/UnrealGameSync/UnrealGameSyncShared/ConfigFile.cs
--- a/Engine/Source/Programs/UnrealGameSync/UnrealGameSyncShared/ConfigFile.cs
+++ b/Engine/Source/Programs/UnrealGameSync/UnrealGameSyncShared/ConfigFile.cs
@@ -468,7 +468,7 @@ namespace UnrealGameSync
 			FileInfo FileInfo = FileName.ToFileInfo();
 			if (FileInfo.Exists)
 			{
-				Logger.LogInformation("Loading config file from {File} ({Size} bytes)", FileInfo.FullName, FileInfo.Length);
+				Logger.LogDebug("Loading config file from {File} ({Size} bytes)", FileInfo.FullName, FileInfo.Length);
 				Load(FileName);
 				return true;
 			}
@@ -478,12 +478,12 @@ namespace UnrealGameSync
 			FileInfo TempFileInfo = TempFileName.ToFileInfo();
 			if (TempFileInfo.Exists)
 			{
-				Logger.LogInformation("Loading temporary config file from {File} ({Size} bytes)", TempFileInfo.FullName, TempFileInfo.Length);
+				Logger.LogDebug("Loading temporary config file from {File} ({Size} bytes)", TempFileInfo.FullName, TempFileInfo.Length);
 				Load(TempFileName);
 				return true;
 			}
 
-			Logger.LogInformation("No existing config file at {File}", FileName);
+			Logger.LogDebug("No existing config file at {File}", FileName);
 			return false;
 		}
 
diff --git a/Engine/Source/Programs/UnrealGameSync/UnrealGameSyncShared/UserSettings.cs b/Engine/Source/Programs/UnrealGameSync/UnrealGameSyncShared/UserSettings.cs
--- a/Engine/Source/Programs/UnrealGameSync/UnrealGameSyncShared/UserSettings.cs
+++ b/Engine/Source/Programs/UnrealGameSync/UnrealGameSyncShared/UserSettings.cs
@@ -237,6 +237,9 @@ namespace UnrealGameSync
 		[JsonIgnore]
 		public DirectoryReference RootDir { get; set; } = null!;
 
+		[JsonIgnore]
+		public DateTime LastModifiedTimeUtc { get; set; }
+
 		// Connection settings
 		public string? ServerAndPort { get; set; }
 		public string? UserName { get; set; }
@@ -245,7 +248,7 @@ namespace UnrealGameSync
 		// Path to the root of the branch within this client, with a trailing slash if non-empty
 		public string BranchPath { get; set; } = String.Empty;
 
-		// The currently selected project
+		// The currently selected project, relative to the root directory
 		public string ProjectPath { get; set; } = String.Empty;
 
 		// Workspace specific SyncFilters
@@ -267,12 +270,28 @@ namespace UnrealGameSync
 		[JsonIgnore]
 		public string ClientProjectPath => $"//{ClientName}{BranchPath}{ProjectPath}";
 
+		[JsonIgnore]
+		public FileReference LocalProjectPath => new FileReference(RootDir.FullName + ProjectPath);
+
+		public void Init(string? ServerAndPort, string? UserName, string ClientName, string BranchPath, string ProjectPath)
+		{
+			ProjectInfo.ValidateBranchPath(BranchPath);
+			ProjectInfo.ValidateProjectPath(ProjectPath);
+
+			this.ServerAndPort = ServerAndPort;
+			this.UserName = UserName;
+			this.ClientName = ClientName;
+			this.BranchPath = BranchPath;
+			this.ProjectPath = ProjectPath;
+		}
+
 		public static bool TryLoad(DirectoryReference RootDir, [NotNullWhen(true)] out UserWorkspaceSettings? Settings)
 		{
 			FileReference ConfigFile = GetConfigFile(RootDir);
 			if (Utility.TryLoadJson(ConfigFile, out Settings))
 			{
 				Settings.RootDir = RootDir;
+				Settings.LastModifiedTimeUtc = FileReference.GetLastWriteTimeUtc(ConfigFile);
 				return true;
 			}
 			else
@@ -289,6 +308,7 @@ namespace UnrealGameSync
 			lock (SyncRoot)
 			{
 				Utility.SaveJson(ConfigFile, this);
+				LastModifiedTimeUtc = FileReference.GetLastWriteTimeUtc(ConfigFile);
 			}
 		}
 
@@ -303,8 +323,16 @@ namespace UnrealGameSync
 		[JsonIgnore]
 		public DirectoryReference RootDir { get; set; } = null!;
 
+		// Cached state about the project, configured using UserWorkspaceSettings and taken from computed values in ProjectInfo. Assumed valid unless manually updated.
+		public long SettingsTimeUtc { get; set; }
+		public string ClientName { get; set; } = String.Empty;
+		public string BranchPath { get; set; } = String.Empty;
+		public string ProjectPath { get; set; } = String.Empty;
+		public string? StreamName { get; set; }
+		public string ProjectIdentifier { get; set; } = String.Empty; // Should be fully reset if this changes
+		public bool IsEnterpriseProject { get; set; }
+
 		// Settings for the currently synced project in this workspace. CurrentChangeNumber is only valid for this workspace if CurrentProjectPath is the current project.
-		public string? CurrentProjectIdentifier { get; set; }
 		public int CurrentChangeNumber { get; set; } = -1;
 		public int CurrentCodeChangeNumber { get; set; } = -1;
 		public string? CurrentSyncFilterHash { get; set; }
@@ -327,8 +355,34 @@ namespace UnrealGameSync
 		public List<BisectEntry> BisectChanges { get; set; } = new List<BisectEntry>();
 
 		// Path to the config file
+		[JsonIgnore]
 		public FileReference ConfigFile => GetConfigFile(RootDir);
 
+		public static UserWorkspaceState CreateNew(DirectoryReference RootDir)
+		{
+			UserWorkspaceState State = new UserWorkspaceState();
+			State.RootDir = RootDir;
+			return State;
+		}
+
+		public void UpdateCachedProjectInfo(ProjectInfo ProjectInfo, DateTime SettingsTimeUtc)
+		{
+			this.SettingsTimeUtc = SettingsTimeUtc.Ticks;
+
+			RootDir = ProjectInfo.LocalRootPath;
+			ClientName = ProjectInfo.ClientName;
+			BranchPath = ProjectInfo.BranchPath;
+			ProjectPath = ProjectInfo.ProjectPath;
+			StreamName = ProjectInfo.StreamName;
+			ProjectIdentifier = ProjectInfo.ProjectIdentifier;
+			IsEnterpriseProject = ProjectInfo.bIsEnterpriseProject;
+		}
+
+		public bool IsValid(ProjectInfo ProjectInfo)
+		{
+			return ProjectInfo.Equals(ProjectInfo.ProjectIdentifier, StringComparison.Ordinal);
+		}
+
 		public void SetBisectState(int Change, BisectState State)
 		{
 			BisectEntry Entry = BisectChanges.FirstOrDefault(x => x.Change == Change);
@@ -394,7 +448,7 @@ namespace UnrealGameSync
 		};
 
 		FileReference FileName;
-		ConfigFile ConfigFile = new ConfigFile();
+		ConfigFile ConfigFile;
 
 		// General settings
 		public UserSettingsVersion Version = UserSettingsVersion.Latest;
@@ -410,7 +464,6 @@ namespace UnrealGameSync
 		public bool bShowLocalTimes;
 		public bool bKeepInTray;
 		public Guid[] EnabledTools;
-		public bool bShowNetCoreInfo;
 		public int FilterIndex;
 		public UserSelectedProjectSettings? LastProject;
 		public List<UserSelectedProjectSettings> OpenProjects;
@@ -494,9 +547,21 @@ namespace UnrealGameSync
 		}
 
 		public UserSettings(FileReference InFileName, ILogger Logger)
+			: this(InFileName, TryLoad(InFileName, Logger))
+		{
+		}
+
+		static ConfigFile TryLoad(FileReference FileName, ILogger Logger)
 		{
-			FileName = InFileName;
+			ConfigFile ConfigFile = new ConfigFile();
 			ConfigFile.TryLoad(FileName, Logger);
+			return ConfigFile;
+		}
+
+		public UserSettings(FileReference InFileName, ConfigFile InConfigFile)
+		{
+			this.FileName = InFileName;
+			this.ConfigFile = InConfigFile;
 
 			// General settings
 			Version = (UserSettingsVersion)ConfigFile.GetValue("General.Version", (int)UserSettingsVersion.Initial);
@@ -528,7 +593,6 @@ namespace UnrealGameSync
 			}
 			this.EnabledTools = EnabledTools.ToArray();
 
-			bShowNetCoreInfo = ConfigFile.GetValue("General.ShowNetCoreInfo", true);
 			int.TryParse(ConfigFile.GetValue("General.FilterIndex", "0"), out FilterIndex);
 
 			string? LastProjectString = ConfigFile.GetValue("General.LastProject", null);
@@ -745,28 +809,6 @@ namespace UnrealGameSync
 			return ConfigDir;
 		}
 
-		public UserWorkspaceState FindOrAddWorkspaceState(UserWorkspaceSettings Settings)
-		{
-			UserWorkspaceState? State;
-			if (!WorkspaceDirToState.TryGetValue(Settings.RootDir, out State))
-			{
-				if (!UserWorkspaceState.TryLoad(Settings.RootDir, out State))
-				{
-					State = ImportWorkspaceState(Settings.RootDir, Settings.ClientName, Settings.BranchPath);
-					State.Save();
-				}
-				WorkspaceDirToState.Add(Settings.RootDir, State);
-			}
-			if (!String.Equals(Settings.ClientProjectPath, State.CurrentProjectIdentifier))
-			{
-				State = new UserWorkspaceState();
-				State.RootDir = Settings.RootDir;
-				State.CurrentProjectIdentifier = Settings.ClientProjectPath;
-				WorkspaceDirToState[Settings.RootDir] = State;
-			}
-			return State;
-		}
-
 		public UserWorkspaceSettings FindOrAddWorkspaceSettings(DirectoryReference RootDir, string? ServerAndPort, string? UserName, string ClientName, string BranchPath, string ProjectPath)
 		{
 			ProjectInfo.ValidateBranchPath(BranchPath);
@@ -782,16 +824,34 @@ namespace UnrealGameSync
 				WorkspaceDirToSettings[RootDir] = Settings;
 			}
 
-			Settings.ServerAndPort = ServerAndPort;
-			Settings.UserName = UserName;
-			Settings.ClientName = ClientName;
-			Settings.BranchPath = BranchPath;
-			Settings.ProjectPath = ProjectPath;
+			Settings.Init(ServerAndPort, UserName, ClientName, BranchPath, ProjectPath);
 			Settings.Save();
 
 			return Settings;
 		}
 
+		public UserWorkspaceState FindOrAddWorkspaceState(ProjectInfo ProjectInfo, UserWorkspaceSettings Settings)
+		{
+			UserWorkspaceState State = FindOrAddWorkspaceState(ProjectInfo.LocalRootPath, ProjectInfo.ClientName, ProjectInfo.BranchPath);
+			if (!State.IsValid(ProjectInfo))
+			{
+				State = new UserWorkspaceState();
+			}
+			State.UpdateCachedProjectInfo(ProjectInfo, Settings.LastModifiedTimeUtc);
+			return State;
+		}
+
+		public UserWorkspaceState FindOrAddWorkspaceState(DirectoryReference RootDir, string ClientName, string BranchPath)
+		{
+			UserWorkspaceState? State;
+			if (!UserWorkspaceState.TryLoad(RootDir, out State))
+			{
+				State = ImportWorkspaceState(RootDir, ClientName, BranchPath);
+				State.Save();
+			}
+			return State;
+		}
+
 		UserWorkspaceState ImportWorkspaceState(DirectoryReference RootDir, string ClientName, string BranchPath)
 		{
 			UserWorkspaceState CurrentWorkspace = new UserWorkspaceState();
@@ -818,7 +878,7 @@ namespace UnrealGameSync
 						int ChangeNumber;
 						if(int.TryParse(CurrentSync.Substring(AtIdx + 1), out ChangeNumber))
 						{
-							CurrentWorkspace.CurrentProjectIdentifier = CurrentSync.Substring(0, AtIdx);
+							CurrentWorkspace.ProjectIdentifier = CurrentSync.Substring(0, AtIdx);
 							CurrentWorkspace.CurrentChangeNumber = ChangeNumber;
 						}
 					}
@@ -845,7 +905,7 @@ namespace UnrealGameSync
 			}
 			else
 			{
-				CurrentWorkspace.CurrentProjectIdentifier = WorkspaceSection.GetValue("CurrentProjectPath");
+				CurrentWorkspace.ProjectIdentifier = WorkspaceSection.GetValue("CurrentProjectPath", "");
 				CurrentWorkspace.CurrentChangeNumber = WorkspaceSection.GetValue("CurrentChangeNumber", -1);
 				CurrentWorkspace.CurrentSyncFilterHash = WorkspaceSection.GetValue("CurrentSyncFilterHash", null);
 				foreach(string AdditionalChangeNumberString in WorkspaceSection.GetValues("AdditionalChangeNumbers", new string[0]))
@@ -985,7 +1045,6 @@ namespace UnrealGameSync
 			GeneralSection.SetValues("OpenProjects", OpenProjects.Select(x => x.ToConfigEntry()).ToArray());
 			GeneralSection.SetValue("KeepInTray", bKeepInTray);
 			GeneralSection.SetValues("EnabledTools", EnabledTools);
-			GeneralSection.SetValue("ShowNetCoreInfo", bShowNetCoreInfo);
 			GeneralSection.SetValue("FilterIndex", FilterIndex);
 			GeneralSection.SetValues("RecentProjects", RecentProjects.Select(x => x.ToConfigEntry()).ToArray());
 			GeneralSection.SetValues("SyncFilter", SyncView);
diff --git a/Engine/Source/Programs/UnrealGameSync/UnrealGameSyncShared/WorkspaceSettings.cs b/Engine/Source/Programs/UnrealGameSync/UnrealGameSyncShared/WorkspaceSettings.cs
--- a/Engine/Source/Programs/UnrealGameSync/UnrealGameSyncShared/WorkspaceSettings.cs
+++ b/Engine/Source/Programs/UnrealGameSync/UnrealGameSyncShared/WorkspaceSettings.cs
@@ -16,27 +16,28 @@ namespace UnrealGameSync
 {
 	public class WorkspaceSettings
 	{
-		public IPerforceSettings PerforceSettings;
 		public UserSelectedProjectSettings SelectedProject { get; }
-		public ProjectInfo ProjectInfo { get; private set; }
-		public string NewSelectedProjectIdentifier;
-		public string NewProjectEditorTarget;
-		public string StreamName;
+
+		public IPerforceSettings PerforceSettings { get; }
+		public ProjectInfo ProjectInfo { get; }
+		public UserWorkspaceSettings UserWorkspaceSettings { get; }
+		public UserWorkspaceState UserWorkspaceState { get; }
+		public ConfigFile LatestProjectConfigFile { get; }
+		public ConfigFile WorkspaceProjectConfigFile { get; }
+		public IReadOnlyList<string>? WorkspaceProjectStreamFilter { get; }
+		public List<KeyValuePair<FileReference, DateTime>> LocalConfigFiles { get; }
+
 		public DirectoryReference DataFolder => GetDataFolder(ProjectInfo.LocalRootPath);
 		public DirectoryReference CacheFolder => GetCacheFolder(ProjectInfo.LocalRootPath);
-		public ConfigFile LatestProjectConfigFile;
-		public ConfigFile WorkspaceProjectConfigFile;
-		public IReadOnlyList<string>? WorkspaceProjectStreamFilter;
-		public List<KeyValuePair<FileReference, DateTime>> LocalConfigFiles;
 
-		public WorkspaceSettings(IPerforceSettings PerforceSettings, UserSelectedProjectSettings ProjectSettings, ProjectInfo ProjectInfo, string NewSelectedProjectIdentifier, string NewProjectEditorTarget, string StreamName, ConfigFile LatestProjectConfigFile, ConfigFile WorkspaceProjectConfigFile, IReadOnlyList<string>? WorkspaceProjectStreamFilter, List<KeyValuePair<FileReference, DateTime>> LocalConfigFiles)
+		public WorkspaceSettings(UserSelectedProjectSettings SelectedProject, IPerforceSettings PerforceSettings, ProjectInfo ProjectInfo, UserWorkspaceSettings UserWorkspaceSettings, UserWorkspaceState UserWorkspaceState, ConfigFile LatestProjectConfigFile, ConfigFile WorkspaceProjectConfigFile, IReadOnlyList<string>? WorkspaceProjectStreamFilter, List<KeyValuePair<FileReference, DateTime>> LocalConfigFiles)
 		{
+			this.SelectedProject = SelectedProject;
+
 			this.PerforceSettings = PerforceSettings;
-			this.SelectedProject = ProjectSettings;
 			this.ProjectInfo = ProjectInfo;
-			this.NewSelectedProjectIdentifier = NewSelectedProjectIdentifier;
-			this.NewProjectEditorTarget = NewProjectEditorTarget;
-			this.StreamName = StreamName;
+			this.UserWorkspaceSettings = UserWorkspaceSettings;
+			this.UserWorkspaceState = UserWorkspaceState;
 			this.LatestProjectConfigFile = LatestProjectConfigFile;
 			this.WorkspaceProjectConfigFile = WorkspaceProjectConfigFile;
 			this.WorkspaceProjectStreamFilter = WorkspaceProjectStreamFilter;
@@ -46,7 +47,7 @@ namespace UnrealGameSync
 		public static DirectoryReference GetDataFolder(DirectoryReference WorkspaceDir) => DirectoryReference.Combine(WorkspaceDir, ".ugs");
 		public static DirectoryReference GetCacheFolder(DirectoryReference WorkspaceDir) => DirectoryReference.Combine(WorkspaceDir, ".ugs", "cache");
 
-		public static async Task<WorkspaceSettings> CreateAsync(IPerforceSettings DefaultPerforceSettings, UserSelectedProjectSettings SelectedProject, ILogger<WorkspaceSettings> Logger, CancellationToken CancellationToken)
+		public static async Task<WorkspaceSettings> CreateAsync(IPerforceSettings DefaultPerforceSettings, UserSelectedProjectSettings SelectedProject, UserSettings UserSettings, ILogger<WorkspaceSettings> Logger, CancellationToken CancellationToken)
 		{
 			PerforceSettings PerforceSettings = Utility.OverridePerforceSettings(DefaultPerforceSettings, SelectedProject.ServerAndPort, SelectedProject.UserName);
 			using IPerforceConnection Perforce = await PerforceConnection.CreateAsync(PerforceSettings, Logger);
@@ -59,10 +60,10 @@ namespace UnrealGameSync
 			}
 
 			// Execute like a regular task
-			return await CreateAsync(Perforce, SelectedProject, Logger, CancellationToken);
+			return await CreateAsync(Perforce, SelectedProject, UserSettings, Logger, CancellationToken);
 		}
 
-		public static async Task<WorkspaceSettings> CreateAsync(IPerforceConnection DefaultConnection, UserSelectedProjectSettings SelectedProject, ILogger<WorkspaceSettings> Logger, CancellationToken CancellationToken)
+		public static async Task<WorkspaceSettings> CreateAsync(IPerforceConnection DefaultConnection, UserSelectedProjectSettings SelectedProject, UserSettings UserSettings, ILogger<WorkspaceSettings> Logger, CancellationToken CancellationToken)
 		{
 			using IDisposable LoggerScope = Logger.BeginScope("Project {SelectedProject}", SelectedProject.ToString());
 			Logger.LogInformation("Detecting settings for {Project}", SelectedProject);
@@ -164,6 +165,13 @@ namespace UnrealGameSync
 					NewSelectedClientFileName = Records[0].ClientFile;
 				}
 
+				// Make sure the drive containing the project exists, to prevent other errors down the line
+				string PathRoot = Path.GetPathRoot(NewSelectedFileName.FullName)!;
+				if (!Directory.Exists(PathRoot))
+				{
+					throw new UserErrorException($"Path '{NewSelectedFileName}' is invalid");
+				}
+
 				// Make sure the path case is correct. This can cause UBT intermediates to be out of date if the case mismatches.
 				NewSelectedFileName = FileReference.FindCorrectCase(NewSelectedFileName);
 
@@ -206,91 +214,17 @@ namespace UnrealGameSync
 
 				Logger.LogInformation("Found branch root at {RootPath}", BranchClientPath);
 
-				// Find the editor target for this project
-				string? NewProjectEditorTarget = null;
-				if (NewSelectedFileName.HasExtension(".uproject"))
-				{
-					List<FStatRecord> Records = await PerforceClient.FStatAsync(PerforceUtils.GetClientOrDepotDirectoryName(NewSelectedClientFileName) + "/Source/*Editor.Target.cs", CancellationToken).ToListAsync(CancellationToken);
-
-					FStatRecord Record = Records.FirstOrDefault(x => x.Action != FileAction.Delete || x.Action != FileAction.MoveDelete);
-					if (Record != null)
-					{
-						string DepotPath = Record.DepotFile!;
-						NewProjectEditorTarget = Path.GetFileNameWithoutExtension(Path.GetFileNameWithoutExtension(DepotPath.Substring(DepotPath.LastIndexOf('/') + 1)));
-						Logger.LogInformation("Using {TargetName} as editor target name (from {DepotPath})", NewProjectEditorTarget, DepotPath);
-					}
-
-					if (NewProjectEditorTarget == null)
-					{
-						Logger.LogInformation("Couldn't find any editor targets for this project.");
-					}
-				}
-
-				string NewSelectedProjectIdentifier;
-				// Get a unique name for the project that's selected. For regular branches, this can be the depot path. For streams, we want to include the stream name to encode imports.
-				string? StreamName = await PerforceClient.GetCurrentStreamAsync(CancellationToken);
-				if (StreamName != null)
-				{
-					string ExpectedPrefix = String.Format("//{0}/", PerforceClient.Settings.ClientName);
-					if (!NewSelectedClientFileName.StartsWith(ExpectedPrefix, StringComparison.InvariantCultureIgnoreCase))
-					{
-						throw new UserErrorException($"Unexpected client path; expected '{NewSelectedClientFileName}' to begin with '{ExpectedPrefix}'");
-					}
-					string? StreamPrefix = await TryGetStreamPrefixAsync(PerforceClient, StreamName, CancellationToken);
-					if (StreamPrefix == null)
-					{
-						throw new UserErrorException("Unable to get stream prefix");
-					}
-					NewSelectedProjectIdentifier = String.Format("{0}/{1}", StreamPrefix, NewSelectedClientFileName.Substring(ExpectedPrefix.Length));
-				}
-				else
-				{
-					List<PerforceResponse<WhereRecord>> Records = await PerforceClient.TryWhereAsync(NewSelectedClientFileName, CancellationToken).Where(x => !x.Succeeded || !x.Data.Unmap).ToListAsync(CancellationToken);
-					if (!Records.Succeeded() || Records.Count != 1)
-					{
-						throw new UserErrorException($"Couldn't get depot path for {NewSelectedFileName}");
-					}
+				// Read the existing workspace settings from disk, and update them with any info computed here
+				int BranchIdx = BranchClientPath.IndexOf('/', 2);
+				string BranchPath = (BranchIdx == -1) ? String.Empty : BranchClientPath.Substring(BranchIdx);
+				string ProjectPath = NewSelectedClientFileName.Substring(BranchClientPath.Length);
+				UserWorkspaceSettings UserWorkspaceSettings = UserSettings.FindOrAddWorkspaceSettings(BranchDirectoryName, PerforceSettings.ServerAndPort, PerforceSettings.UserName, PerforceSettings.ClientName!, BranchPath, ProjectPath);
 
-					NewSelectedProjectIdentifier = Records[0].Data.DepotFile;
+				// Now compute the updated project info
+				ProjectInfo ProjectInfo = await CreateProjectInfoAsync(PerforceClient, UserWorkspaceSettings, CancellationToken);
 
-					Match Match = Regex.Match(NewSelectedProjectIdentifier, "//([^/]+)/");
-					if (Match.Success)
-					{
-						DepotRecord Depot = await PerforceClient.GetDepotAsync(Match.Groups[1].Value, CancellationToken);
-						if (Depot.Type == "stream")
-						{
-							throw new UserErrorException($"Cannot use a legacy client ({PerforceClient.Settings.ClientName}) with a stream depot ({Depot.Depot}).");
-						}
-					}
-				}
-
-				// Figure out if it's an enterprise project. Use the synced version if we have it.
-				bool bIsEnterpriseProject = false;
-				if (NewSelectedClientFileName.EndsWith(".uproject", StringComparison.InvariantCultureIgnoreCase))
-				{
-					string Text;
-					if (FileReference.Exists(NewSelectedFileName))
-					{
-						Text = FileReference.ReadAllText(NewSelectedFileName);
-					}
-					else
-					{
-						PerforceResponse<PrintRecord<string[]>> ProjectLines = await PerforceClient.TryPrintLinesAsync(NewSelectedClientFileName, CancellationToken);
-						if (!ProjectLines.Succeeded)
-						{
-							throw new UserErrorException($"Unable to get contents of {NewSelectedClientFileName}");
-						}
-						Text = String.Join("\n", ProjectLines.Data.Contents!);
-					}
-					bIsEnterpriseProject = Utility.IsEnterpriseProjectFromText(Text);
-				}
-
-				// Make sure the drive containing the project exists, to prevent other errors down the line
-				string PathRoot = Path.GetPathRoot(NewSelectedFileName.FullName)!;
-				if (!Directory.Exists(PathRoot))
-				{
-					throw new UserErrorException($"Path '{NewSelectedFileName}' is invalid");
-				}
+				// Update the cached workspace state
+				UserWorkspaceState UserWorkspaceState = UserSettings.FindOrAddWorkspaceState(ProjectInfo, UserWorkspaceSettings);
 
 				// Read the initial config file
 				List<KeyValuePair<FileReference, DateTime>> LocalConfigFiles = new List<KeyValuePair<FileReference, DateTime>>();
@@ -300,13 +234,7 @@ namespace UnrealGameSync
 				ConfigFile WorkspaceProjectConfigFile = await WorkspaceUpdate.ReadProjectConfigFile(BranchDirectoryName, NewSelectedFileName, Logger);
 				IReadOnlyList<string>? WorkspaceProjectStreamFilter = await WorkspaceUpdate.ReadProjectStreamFilter(PerforceClient, WorkspaceProjectConfigFile, Logger, CancellationToken);
 
-				int BranchIdx = BranchClientPath.IndexOf('/', 2);
-				string BranchPath = (BranchIdx == -1) ? String.Empty : BranchClientPath.Substring(BranchIdx);
-				string ProjectPath = NewSelectedClientFileName.Substring(BranchClientPath.Length);
-
-				ProjectInfo ProjectInfo = new ProjectInfo(BranchDirectoryName, PerforceSettings.ClientName!, BranchPath, ProjectPath, StreamName, NewSelectedProjectIdentifier, bIsEnterpriseProject);
-
-				WorkspaceSettings WorkspaceSettings = new WorkspaceSettings(PerforceSettings, SelectedProject, ProjectInfo, NewSelectedProjectIdentifier, NewProjectEditorTarget!, StreamName!, LatestProjectConfigFile, WorkspaceProjectConfigFile, WorkspaceProjectStreamFilter, LocalConfigFiles);
+				WorkspaceSettings WorkspaceSettings = new WorkspaceSettings(SelectedProject, PerforceSettings, ProjectInfo, UserWorkspaceSettings, UserWorkspaceState, LatestProjectConfigFile, WorkspaceProjectConfigFile, WorkspaceProjectStreamFilter, LocalConfigFiles);
 				DirectoryReference.CreateDirectory(WorkspaceSettings.DataFolder);
 				DirectoryReference.CreateDirectory(WorkspaceSettings.CacheFolder);
 
@@ -318,6 +246,71 @@ namespace UnrealGameSync
 			}
 		}
 
+		static async Task<ProjectInfo> CreateProjectInfoAsync(IPerforceConnection PerforceClient, UserWorkspaceSettings Settings, CancellationToken CancellationToken)
+		{
+			string? StreamName = await PerforceClient.GetCurrentStreamAsync(CancellationToken);
+
+			// Get a unique name for the project that's selected. For regular branches, this can be the depot path. For streams, we want to include the stream name to encode imports.
+			string NewSelectedProjectIdentifier;
+			if (StreamName != null)
+			{
+				string ExpectedPrefix = String.Format("//{0}/", PerforceClient.Settings.ClientName);
+				if (!Settings.ClientProjectPath.StartsWith(ExpectedPrefix, StringComparison.InvariantCultureIgnoreCase))
+				{
+					throw new UserErrorException($"Unexpected client path; expected '{Settings.ClientProjectPath}' to begin with '{ExpectedPrefix}'");
+				}
+				string? StreamPrefix = await TryGetStreamPrefixAsync(PerforceClient, StreamName, CancellationToken);
+				if (StreamPrefix == null)
+				{
+					throw new UserErrorException("Unable to get stream prefix");
+				}
+				NewSelectedProjectIdentifier = String.Format("{0}/{1}", StreamPrefix, Settings.ClientProjectPath.Substring(ExpectedPrefix.Length));
+			}
+			else
+			{
+				List<PerforceResponse<WhereRecord>> Records = await PerforceClient.TryWhereAsync(Settings.ClientProjectPath, CancellationToken).Where(x => !x.Succeeded || !x.Data.Unmap).ToListAsync(CancellationToken);
+				if (!Records.Succeeded() || Records.Count != 1)
+				{
+					throw new UserErrorException($"Couldn't get depot path for {Settings.ClientProjectPath}");
+				}
+
+				NewSelectedProjectIdentifier = Records[0].Data.DepotFile;
+
+				Match Match = Regex.Match(NewSelectedProjectIdentifier, "//([^/]+)/");
+				if (Match.Success)
+				{
+					DepotRecord Depot = await PerforceClient.GetDepotAsync(Match.Groups[1].Value, CancellationToken);
+					if (Depot.Type == "stream")
+					{
+						throw new UserErrorException($"Cannot use a legacy client ({PerforceClient.Settings.ClientName}) with a stream depot ({Depot.Depot}).");
+					}
+				}
+			}
+
+			// Figure out if it's an enterprise project. Use the synced version if we have it.
+			bool bIsEnterpriseProject = false;
+			if (Settings.ClientProjectPath.EndsWith(".uproject", StringComparison.InvariantCultureIgnoreCase))
+			{
+				string Text;
+				if (FileReference.Exists(Settings.LocalProjectPath))
+				{
+					Text = FileReference.ReadAllText(Settings.LocalProjectPath);
+				}
+				else
+				{
+					PerforceResponse<PrintRecord<string[]>> ProjectLines = await PerforceClient.TryPrintLinesAsync(Settings.ClientProjectPath, CancellationToken);
+					if (!ProjectLines.Succeeded)
+					{
+						throw new UserErrorException($"Unable to get contents of {Settings.ClientProjectPath}");
+					}
+					Text = String.Join("\n", ProjectLines.Data.Contents!);
+				}
+				bIsEnterpriseProject = Utility.IsEnterpriseProjectFromText(Text);
+			}
+
+			return new ProjectInfo(Settings.RootDir, Settings.ClientName, Settings.BranchPath, Settings.ProjectPath, StreamName, NewSelectedProjectIdentifier, bIsEnterpriseProject);
+		}
+
 		static async Task<List<IPerforceSettings>> FilterClients(List<ClientsRecord> Clients, FileReference NewSelectedFileName, IPerforceSettings DefaultPerforceSettings, string? HostName, ILogger Logger, CancellationToken CancellationToken)
 		{
 			List<IPerforceSettings> CandidateClients = new List<IPerforceSettings>();
diff --git a/Engine/Source/Programs/UnrealGameSync/UnrealGameSyncShared/WorkspaceUpdate.cs b/Engine/Source/Programs/UnrealGameSync/UnrealGameSyncShared/WorkspaceUpdate.cs
--- a/Engine/Source/Programs/UnrealGameSync/UnrealGameSyncShared/WorkspaceUpdate.cs
+++ b/Engine/Source/Programs/UnrealGameSync/UnrealGameSyncShared/WorkspaceUpdate.cs
@@ -234,7 +234,7 @@ namespace UnrealGameSync
 		}
 	}
 
-	public class WorkspaceUpdate : IAsyncDisposable
+	public class WorkspaceUpdate
 	{
 		const string BuildVersionFileName = "/Engine/Build/Build.version";
 		const string VersionHeaderFileName = "/Engine/Source/Runtime/Launch/Resources/Version.h";
@@ -397,32 +397,12 @@ namespace UnrealGameSync
 			}
 		}
 
-		CancellationTokenSource CancellationSource;
 		public WorkspaceUpdateContext Context { get; }
 		public ProgressValue Progress { get; } = new ProgressValue();
-		public Task UpdateTask { get; }
 
-		public WorkspaceUpdate(WorkspaceUpdateContext Context, Func<WorkspaceUpdate, CancellationToken, Task> TaskFunc)
+		public WorkspaceUpdate(WorkspaceUpdateContext Context)
 		{
-			this.CancellationSource = new CancellationTokenSource();
 			this.Context = Context;
-			this.UpdateTask = Task.Run(() => TaskFunc(this, CancellationSource.Token));
-		}
-
-		public void Cancel()
-		{
-			CancellationSource.Cancel();
-		}
-
-		public ValueTask DisposeAsync()
-		{
-			Task DisposeTask = Task.CompletedTask;
-			if (UpdateTask != null)
-			{
-				CancellationSource.Cancel();
-				DisposeTask = UpdateTask.ContinueWith(Task => CancellationSource.Dispose());
-			}
-			return new ValueTask(DisposeTask);
 		}
 
 		class SyncFile
