UGS: Initial version of command-line UGS client.

Currently only tested on Windows; Mac/Linux ports to follow. Workspace state and core implementation is shared with the WinForms application, allowing switching between command line and GUI as needed.

Currently implemented commands are:

  ugs init [stream-path] [-client=..] [-server=..] [-user=..] [-branch=..] [-project=..]
    Create a client for the given stream, or initializes an existing client for use by UGS.

  ugs switch [project name|project path|stream]
    Changes the active project to the one in the workspace with the given name, or switches to a new stream.

  ugs config
    Updates the configuration for the current workspace.

  ugs sync [change|'latest'] [-build] [-only]
    Syncs the current workspace to the given changelist, optionally removing all local state.

  ugs build [id] [-list]
    Runs the default build steps for the current project, or a particular step referenced by id.

  ugs status [-update]
    Shows the status of the currently synced branch.

diff --git a/Engine/Source/Programs/UnrealGameSync/UnrealGameSync/Controls/WorkspaceControl.cs b/Engine/Source/Programs/UnrealGameSync/UnrealGameSync/Controls/WorkspaceControl.cs
--- a/Engine/Source/Programs/UnrealGameSync/UnrealGameSync/Controls/WorkspaceControl.cs
+++ b/Engine/Source/Programs/UnrealGameSync/UnrealGameSync/Controls/WorkspaceControl.cs
@@ -290,7 +290,7 @@ namespace UnrealGameSync
 		// Placeholder text that is in the control and cleared when the user starts editing.
 		static string AuthorFilterPlaceholderText = "<username>";
 
-		public WorkspaceControl(IWorkspaceControlOwner InOwner, string? InApiUrl, WorkspaceSettings WorkspaceSettings, IServiceProvider InServiceProvider, UserSettings InSettings, OIDCTokenManager? InOidcTokenManager)
+		public WorkspaceControl(IWorkspaceControlOwner InOwner, string? InApiUrl, OpenProjectInfo OpenProjectInfo, IServiceProvider InServiceProvider, UserSettings InSettings, OIDCTokenManager? InOidcTokenManager)
 		{
 			InitializeComponent();
 
@@ -298,17 +298,15 @@ namespace UnrealGameSync
 
 			Owner = InOwner;
 			ApiUrl = InApiUrl;
-			DataFolder = WorkspaceSettings.DataFolder;
+			DataFolder = OpenProjectInfo.ProjectInfo.DataFolder;
 			ServiceProvider = InServiceProvider;
 			Logger = InServiceProvider.GetRequiredService<ILogger<WorkspaceControl>>();
-			PerforceSettings = WorkspaceSettings.PerforceSettings;
-			ProjectInfo = WorkspaceSettings.ProjectInfo;
+			PerforceSettings = OpenProjectInfo.PerforceSettings;
+			ProjectInfo = OpenProjectInfo.ProjectInfo;
 			Settings = InSettings;
-			this.WorkspaceSettings = WorkspaceSettings.UserWorkspaceSettings;
-			this.WorkspaceState = WorkspaceSettings.UserWorkspaceState;
-
-			string ProjectPath = Regex.Replace(WorkspaceSettings.ProjectInfo.ClientFileName, "^//[^/]+/", "");
-			ProjectSettings = InSettings.FindOrAddProjectSettings(WorkspaceSettings.ProjectInfo.LocalRootPath, ProjectPath);
+			this.WorkspaceSettings = OpenProjectInfo.WorkspaceSettings;
+			this.WorkspaceState = OpenProjectInfo.WorkspaceState;
+			ProjectSettings = InSettings.FindOrAddProjectSettings(OpenProjectInfo.ProjectInfo);
 
 			DesiredTaskbarState = Tuple.Create(TaskbarState.NoProgress, 0.0f);
 
@@ -353,36 +351,36 @@ namespace UnrealGameSync
 			}
 
 			// Commit all the new project info
-			IPerforceSettings PerforceClientSettings = WorkspaceSettings.PerforceSettings;
-			SelectedProject = WorkspaceSettings.SelectedProject;
+			IPerforceSettings PerforceClientSettings = OpenProjectInfo.PerforceSettings;
+			SelectedProject = OpenProjectInfo.SelectedProject;
 
 			// Figure out which API server to use
 			string? NewApiUrl;
-			if (TryGetProjectSetting(WorkspaceSettings.LatestProjectConfigFile, "ApiUrl", out NewApiUrl))
+			if (TryGetProjectSetting(OpenProjectInfo.LatestProjectConfigFile, "ApiUrl", out NewApiUrl))
 			{
 				ApiUrl = NewApiUrl;
 			}
 
-			ProjectInfo Project = WorkspaceSettings.ProjectInfo;
+			ProjectInfo Project = OpenProjectInfo.ProjectInfo;
 
-			Workspace = new Workspace(PerforceClientSettings, Project, WorkspaceState, WorkspaceSettings.WorkspaceProjectConfigFile, WorkspaceSettings.WorkspaceProjectStreamFilter, new LogControlTextWriter(SyncLog), ServiceProvider);
+			Workspace = new Workspace(PerforceClientSettings, Project, WorkspaceState, OpenProjectInfo.WorkspaceProjectConfigFile, OpenProjectInfo.WorkspaceProjectStreamFilter, new LogControlTextWriter(SyncLog), ServiceProvider);
 			Workspace.OnUpdateComplete += UpdateCompleteCallback;
 
 			FileReference ProjectLogBaseName = FileReference.Combine(DataFolder, "sync.log");
 
 			ILogger PerforceLogger = ServiceProvider.GetRequiredService<ILogger<PerforceMonitor>>();
-			PerforceMonitor = new PerforceMonitor(PerforceClientSettings, WorkspaceSettings.ProjectInfo, WorkspaceSettings.LatestProjectConfigFile, WorkspaceSettings.CacheFolder, WorkspaceSettings.LocalConfigFiles, ServiceProvider);
+			PerforceMonitor = new PerforceMonitor(PerforceClientSettings, OpenProjectInfo.ProjectInfo, OpenProjectInfo.LatestProjectConfigFile, OpenProjectInfo.ProjectInfo.CacheFolder, OpenProjectInfo.LocalConfigFiles, ServiceProvider);
 			PerforceMonitor.OnUpdate += UpdateBuildListCallback;
 			PerforceMonitor.OnUpdateMetadata += UpdateBuildMetadataCallback;
 			PerforceMonitor.OnStreamChange += StreamChanged;
 			PerforceMonitor.OnLoginExpired += LoginExpired;
 
 			ILogger EventLogger = ServiceProvider.GetRequiredService < ILogger<EventMonitor>>();
-			EventMonitor = new EventMonitor(ApiUrl, PerforceUtils.GetClientOrDepotDirectoryName(SelectedProjectIdentifier), WorkspaceSettings.PerforceSettings.UserName, ServiceProvider);
+			EventMonitor = new EventMonitor(ApiUrl, PerforceUtils.GetClientOrDepotDirectoryName(SelectedProjectIdentifier), OpenProjectInfo.PerforceSettings.UserName, ServiceProvider);
 			EventMonitor.OnUpdatesReady += UpdateReviewsCallback;
 
 			ILogger<JupiterMonitor> JupiterLogger = ServiceProvider.GetRequiredService<ILogger<JupiterMonitor>>();
-			JupiterMonitor = JupiterMonitor.CreateFromConfigFile(InOidcTokenManager, JupiterLogger, WorkspaceSettings.LatestProjectConfigFile, SelectedProjectIdentifier);
+			JupiterMonitor = JupiterMonitor.CreateFromConfigFile(InOidcTokenManager, JupiterLogger, OpenProjectInfo.LatestProjectConfigFile, SelectedProjectIdentifier);
 
 			UpdateColumnSettings(true);
 
@@ -404,7 +402,7 @@ namespace UnrealGameSync
 			StartupCallbacks = new List<WorkspaceStartupCallback>();
 
 			string? IssuesApiUrl = GetIssuesApiUrl();
-			IssueMonitor = InOwner.CreateIssueMonitor(IssuesApiUrl, WorkspaceSettings.PerforceSettings.UserName);
+			IssueMonitor = InOwner.CreateIssueMonitor(IssuesApiUrl, OpenProjectInfo.PerforceSettings.UserName);
 			IssueMonitor.OnIssuesChanged += IssueMonitor_OnIssuesChangedAsync;
 
 			if (SelectedFileName.HasExtension(".uproject"))
@@ -1650,54 +1648,7 @@ namespace UnrealGameSync
 
 		Dictionary<Guid, WorkspaceSyncCategory> GetSyncCategories()
 		{
-			Dictionary<Guid, WorkspaceSyncCategory> UniqueIdToCategory = new Dictionary<Guid, WorkspaceSyncCategory>();
-
-			// Add the custom filters
-			ConfigFile ProjectConfigFile = PerforceMonitor.LatestProjectConfigFile;
-			if (ProjectConfigFile != null)
-			{
-				string[] CategoryLines = ProjectConfigFile.GetValues("Options.SyncCategory", new string[0]);
-				foreach (string CategoryLine in CategoryLines)
-				{
-					ConfigObject Object = new ConfigObject(CategoryLine);
-
-					Guid UniqueId;
-					if (Guid.TryParse(Object.GetValue("UniqueId", ""), out UniqueId))
-					{
-						WorkspaceSyncCategory? Category;
-						if (!UniqueIdToCategory.TryGetValue(UniqueId, out Category))
-						{
-							Category = new WorkspaceSyncCategory(UniqueId);
-							UniqueIdToCategory.Add(UniqueId, Category);
-						}
-
-						if (Object.GetValue("Clear", false))
-						{
-							Category.Paths = new string[0];
-							Category.Requires = new Guid[0];
-						}
-
-						Category.Name = Object.GetValue("Name", Category.Name);
-						Category.bEnable = Object.GetValue("Enable", Category.bEnable);
-						Category.Paths = Enumerable.Concat(Category.Paths, Object.GetValue("Paths", "").Split(';').Select(x => x.Trim())).Where(x => x.Length > 0).Distinct().OrderBy(x => x).ToArray();
-						Category.bHidden = Object.GetValue("Hidden", Category.bHidden);
-						Category.Requires = Enumerable.Concat(Category.Requires, ParseGuids(Object.GetValue("Requires", "").Split(';'))).Distinct().OrderBy(x => x).ToArray();
-					}
-				}
-			}
-			return UniqueIdToCategory;
-		}
-
-		static IEnumerable<Guid> ParseGuids(IEnumerable<string> Values)
-		{
-			foreach(string Value in Values)
-			{
-				Guid Guid;
-				if(Guid.TryParse(Value, out Guid))
-				{
-					yield return Guid;
-				}
-			}
+			return ConfigUtils.GetSyncCategories(PerforceMonitor.LatestProjectConfigFile);
 		}
 
 		List<string> GetProjectRoots(string InUProjectDirsPath)
@@ -5172,7 +5123,7 @@ namespace UnrealGameSync
 
 				// Save the settings
 				ProjectSettings.BuildSteps = ModifiedBuildSteps;
-				ProjectSettings.Save();
+				Settings.Save();
 
 				// Update the custom tools menu, because we might have changed it
 				UpdateBuildSteps();
@@ -5769,7 +5720,7 @@ namespace UnrealGameSync
 			{
 				ProjectSettings.FilterBadges.Add(BadgeName);
 			}
-			ProjectSettings.Save();
+			Settings.Save();
 
 			UpdateBuildListFilter();
 		}
@@ -5778,7 +5729,7 @@ namespace UnrealGameSync
 		{
 			ProjectSettings.FilterBadges.Clear();
 			ProjectSettings.FilterType = FilterType.None;
-			ProjectSettings.Save();
+			Settings.Save();
 
 			Settings.bShowAutomatedChanges = false;
 			Settings.Save();
@@ -5870,7 +5821,7 @@ namespace UnrealGameSync
 		private void FilterContextMenu_Type_ShowAll_Click(object sender, EventArgs e)
 		{
 			ProjectSettings.FilterType = FilterType.None;
-			ProjectSettings.Save();
+			Settings.Save();
 
 			UpdateBuildListFilter();
 		}
@@ -5878,7 +5829,7 @@ namespace UnrealGameSync
 		private void FilterContextMenu_Type_Code_Click(object sender, EventArgs e)
 		{
 			ProjectSettings.FilterType = FilterType.Code;
-			ProjectSettings.Save();
+			Settings.Save();
 
 			UpdateBuildListFilter();
 		}
@@ -5886,7 +5837,7 @@ namespace UnrealGameSync
 		private void FilterContextMenu_Type_Content_Click(object sender, EventArgs e)
 		{
 			ProjectSettings.FilterType = FilterType.Content;
-			ProjectSettings.Save();
+			Settings.Save();
 
 			UpdateBuildListFilter();
 		}
diff --git a/Engine/Source/Programs/UnrealGameSync/UnrealGameSync/DeploymentSettings.cs b/Engine/Source/Programs/UnrealGameSync/UnrealGameSync/DeploymentSettings.cs
--- a/Engine/Source/Programs/UnrealGameSync/UnrealGameSync/DeploymentSettings.cs
+++ b/Engine/Source/Programs/UnrealGameSync/UnrealGameSync/DeploymentSettings.cs
@@ -69,7 +69,7 @@ namespace UnrealGameSync
 		/// <param name="Logger">The logger</param>
 		/// <param name="Error">Receives an error on failure</param>
 		/// <returns></returns>
-		public delegate bool DetectProjectSettingsEvent(WorkspaceSettings Workspace, ILogger Logger, [NotNullWhen(false)] out string? Error);
+		public delegate bool DetectProjectSettingsEvent(OpenProjectInfo OpenProjectInfo, ILogger Logger, [NotNullWhen(false)] out string? Error);
 
 		/// <summary>
 		/// Called to validate the project settings
diff --git a/Engine/Source/Programs/UnrealGameSync/UnrealGameSync/Forms/MainWindow.cs b/Engine/Source/Programs/UnrealGameSync/UnrealGameSync/Forms/MainWindow.cs
--- a/Engine/Source/Programs/UnrealGameSync/UnrealGameSync/Forms/MainWindow.cs
+++ b/Engine/Source/Programs/UnrealGameSync/UnrealGameSync/Forms/MainWindow.cs
@@ -119,7 +119,7 @@ namespace UnrealGameSync
 
 		public ToolUpdateMonitor ToolUpdateMonitor { get; private set; }
 
-		public MainWindow(UpdateMonitor InUpdateMonitor, string? InApiUrl, DirectoryReference InDataFolder, DirectoryReference InCacheFolder, bool bInRestoreStateOnLoad, string InOriginalExecutableFileName, bool bInUnstable, List<(UserSelectedProjectSettings, ModalTask<WorkspaceSettings>)> StartupTasks, IPerforceSettings InDefaultPerforceSettings, IServiceProvider InServiceProvider, UserSettings InSettings, string? InUri, OIDCTokenManager? InOidcTokenManager)
+		public MainWindow(UpdateMonitor InUpdateMonitor, string? InApiUrl, DirectoryReference InDataFolder, DirectoryReference InCacheFolder, bool bInRestoreStateOnLoad, string InOriginalExecutableFileName, bool bInUnstable, List<(UserSelectedProjectSettings, ModalTask<OpenProjectInfo>)> StartupTasks, IPerforceSettings InDefaultPerforceSettings, IServiceProvider InServiceProvider, UserSettings InSettings, string? InUri, OIDCTokenManager? InOidcTokenManager)
 		{
 			ServiceProvider = InServiceProvider;
 			Logger = ServiceProvider.GetRequiredService<ILogger<MainWindow>>();
@@ -157,7 +157,7 @@ namespace UnrealGameSync
 			SetupDefaultControl();
 
 			int SelectTabIdx = -1;
-			foreach((UserSelectedProjectSettings Project, ModalTask<WorkspaceSettings> StartupTask) in StartupTasks)
+			foreach((UserSelectedProjectSettings Project, ModalTask<OpenProjectInfo> StartupTask) in StartupTasks)
 			{
 				int TabIdx = -1;
 				if (StartupTask.Succeeded)
@@ -1035,15 +1035,15 @@ namespace UnrealGameSync
 
 		public void OpenNewProject()
 		{
-			WorkspaceSettings? WorkspaceSettings = OpenProjectWindow.ShowModal(this, null, Settings, DataFolder, CacheFolder, DefaultPerforceSettings, ServiceProvider);
-			if(WorkspaceSettings != null)
+			OpenProjectInfo? OpenProjectInfo = OpenProjectWindow.ShowModal(this, null, Settings, DataFolder, CacheFolder, DefaultPerforceSettings, ServiceProvider);
+			if(OpenProjectInfo != null)
 			{
-				int NewTabIdx = TryOpenProject(WorkspaceSettings, -1, OpenProjectOptions.None);
+				int NewTabIdx = TryOpenProject(OpenProjectInfo, -1, OpenProjectOptions.None);
 				if(NewTabIdx != -1)
 				{
 					TabControl.SelectTab(NewTabIdx);
 					SaveTabSettings();
-					UpdateRecentProjectsList(WorkspaceSettings.SelectedProject);
+					UpdateRecentProjectsList(OpenProjectInfo.SelectedProject);
 				}
 			}
 		}
@@ -1097,22 +1097,22 @@ namespace UnrealGameSync
 
 		public void EditSelectedProject(int TabIdx, UserSelectedProjectSettings SelectedProject)
 		{
-			WorkspaceSettings? WorkspaceSettings = OpenProjectWindow.ShowModal(this, SelectedProject, Settings, DataFolder, CacheFolder, DefaultPerforceSettings, ServiceProvider);
-			if(WorkspaceSettings != null)
+			OpenProjectInfo? OpenProjectInfo = OpenProjectWindow.ShowModal(this, SelectedProject, Settings, DataFolder, CacheFolder, DefaultPerforceSettings, ServiceProvider);
+			if(OpenProjectInfo != null)
 			{
-				int NewTabIdx = TryOpenProject(WorkspaceSettings, TabIdx, OpenProjectOptions.None);
+				int NewTabIdx = TryOpenProject(OpenProjectInfo, TabIdx, OpenProjectOptions.None);
 				if(NewTabIdx != -1)
 				{
 					TabControl.SelectTab(NewTabIdx);
 					SaveTabSettings();
-					UpdateRecentProjectsList(WorkspaceSettings.SelectedProject);
+					UpdateRecentProjectsList(OpenProjectInfo.SelectedProject);
 				}
 			}
 		}
 
 		int TryOpenProject(UserSelectedProjectSettings Project, int ReplaceTabIdx, OpenProjectOptions Options = OpenProjectOptions.None)
 		{
-			ILogger<WorkspaceSettings> ProjectLogger = ServiceProvider.GetRequiredService<ILogger<WorkspaceSettings>>();
+			ILogger<OpenProjectInfo> ProjectLogger = ServiceProvider.GetRequiredService<ILogger<OpenProjectInfo>>();
 
 			ModalTaskFlags TaskFlags = ModalTaskFlags.None;
 			if((Options & OpenProjectOptions.Quiet) != 0)
@@ -1120,7 +1120,7 @@ namespace UnrealGameSync
 				TaskFlags |= ModalTaskFlags.Quiet;
 			}
 
-			ModalTask<WorkspaceSettings>? SettingsTask = PerforceModalTask.Execute(this, "Opening Project", "Opening project, please wait...", DefaultPerforceSettings, (p, c) => OpenProjectWindow.DetectSettingsAsync(p, Project, Settings, ProjectLogger, c), ProjectLogger, TaskFlags);
+			ModalTask<OpenProjectInfo>? SettingsTask = PerforceModalTask.Execute(this, "Opening Project", "Opening project, please wait...", DefaultPerforceSettings, (p, c) => OpenProjectWindow.DetectSettingsAsync(p, Project, Settings, ProjectLogger, c), ProjectLogger, TaskFlags);
 			if (SettingsTask == null || SettingsTask.Failed)
 			{
 				if(SettingsTask != null) CreateErrorPanel(ReplaceTabIdx, Project, SettingsTask.Error);
@@ -1130,7 +1130,7 @@ namespace UnrealGameSync
 			return TryOpenProject(SettingsTask.Result, ReplaceTabIdx, Options);
 		}
 
-		int TryOpenProject(WorkspaceSettings WorkspaceSettings, int ReplaceTabIdx, OpenProjectOptions Options)
+		int TryOpenProject(OpenProjectInfo OpenProjectInfo, int ReplaceTabIdx, OpenProjectOptions Options)
 		{
 			Logger.LogInformation("Trying to open project {Project}", WorkspaceSettings.ProjectInfo.ClientFileName);
 
@@ -1142,7 +1142,7 @@ namespace UnrealGameSync
 					WorkspaceControl? Workspace = TabControl.GetTabData(TabIdx) as WorkspaceControl;
 					if(Workspace != null)
 					{
-						if(Workspace.SelectedFileName == WorkspaceSettings.ProjectInfo.LocalFileName)
+						if(Workspace.SelectedFileName == OpenProjectInfo.ProjectInfo.LocalFileName)
 						{
 							Logger.LogInformation("  Already open in tab {TabIdx}", TabIdx);
 							if((Options & OpenProjectOptions.Quiet) == 0)
diff --git a/Engine/Source/Programs/UnrealGameSync/UnrealGameSync/Forms/OpenProjectWindow.cs b/Engine/Source/Programs/UnrealGameSync/UnrealGameSync/Forms/OpenProjectWindow.cs
--- a/Engine/Source/Programs/UnrealGameSync/UnrealGameSync/Forms/OpenProjectWindow.cs
+++ b/Engine/Source/Programs/UnrealGameSync/UnrealGameSync/Forms/OpenProjectWindow.cs
@@ -25,7 +25,7 @@ namespace UnrealGameSync
 	{
 		string? ServerAndPortOverride;
 		string? UserNameOverride;
-		WorkspaceSettings? DetectedProjectSettings;
+		OpenProjectInfo? OpenProjectInfo;
 		IPerforceSettings DefaultSettings;
 		IServiceProvider ServiceProvider;
 		UserSettings Settings;
@@ -90,7 +90,7 @@ namespace UnrealGameSync
 			get => Utility.OverridePerforceSettings(DefaultSettings, ServerAndPortOverride, UserNameOverride);
 		}
 
-		public static WorkspaceSettings? ShowModal(IWin32Window Owner, UserSelectedProjectSettings? Project, UserSettings Settings, DirectoryReference DataFolder, DirectoryReference CacheFolder, IPerforceSettings DefaultPerforceSettings, IServiceProvider ServiceProvider)
+		public static OpenProjectInfo? ShowModal(IWin32Window Owner, UserSelectedProjectSettings? Project, UserSettings Settings, DirectoryReference DataFolder, DirectoryReference CacheFolder, IPerforceSettings DefaultPerforceSettings, IServiceProvider ServiceProvider)
 		{
 			OpenProjectWindow Window = new OpenProjectWindow(Project, Settings, DefaultPerforceSettings, ServiceProvider);
 			if(Window.ShowDialog(Owner) == DialogResult.OK)
@@ -284,21 +284,21 @@ namespace UnrealGameSync
 			UserSelectedProjectSettings? SelectedProject;
 			if(TryGetSelectedProject(out SelectedProject))
 			{
-				ILogger<WorkspaceSettings> Logger = ServiceProvider.GetRequiredService<ILogger<WorkspaceSettings>>();
+				ILogger<OpenProjectInfo> Logger = ServiceProvider.GetRequiredService<ILogger<OpenProjectInfo>>();
 
-				ModalTask<WorkspaceSettings>? NewDetectedProjectSettings = PerforceModalTask.Execute(this, "Opening project", "Opening project, please wait...", DefaultSettings, (x, y) => DetectSettingsAsync(x, SelectedProject, Settings, Logger, y), Logger);
-				if (NewDetectedProjectSettings != null && NewDetectedProjectSettings.Succeeded)
+				ModalTask<OpenProjectInfo>? NewOpenProjectInfo = PerforceModalTask.Execute(this, "Opening project", "Opening project, please wait...", DefaultSettings, (x, y) => DetectSettingsAsync(x, SelectedProject, Settings, Logger, y), Logger);
+				if (NewOpenProjectInfo != null && NewOpenProjectInfo.Succeeded)
 				{
-					DetectedProjectSettings = NewDetectedProjectSettings.Result;
+					DetectedProjectSettings = NewOpenProjectInfo.Result;
 					DialogResult = DialogResult.OK;
 					Close();
 				}
 			}
 		}
 
-		public static async Task<WorkspaceSettings> DetectSettingsAsync(IPerforceConnection Perforce, UserSelectedProjectSettings SelectedProject, UserSettings UserSettings, ILogger<WorkspaceSettings> Logger, CancellationToken CancellationToken)
+		public static async Task<OpenProjectInfo> DetectSettingsAsync(IPerforceConnection Perforce, UserSelectedProjectSettings SelectedProject, UserSettings UserSettings, ILogger<OpenProjectInfo> Logger, CancellationToken CancellationToken)
 		{
-			WorkspaceSettings Settings = await WorkspaceSettings.CreateAsync(Perforce, SelectedProject, UserSettings, Logger, CancellationToken);
+			OpenProjectInfo Settings = await OpenProjectInfo.CreateAsync(Perforce, SelectedProject, UserSettings, Logger, CancellationToken);
 			if (DeploymentSettings.OnDetectProjectSettings != null)
 			{
 				string? Message;
diff --git a/Engine/Source/Programs/UnrealGameSync/UnrealGameSync/OIDCTokenManager.cs b/Engine/Source/Programs/UnrealGameSync/UnrealGameSync/OIDCTokenManager.cs
--- a/Engine/Source/Programs/UnrealGameSync/UnrealGameSync/OIDCTokenManager.cs
+++ b/Engine/Source/Programs/UnrealGameSync/UnrealGameSync/OIDCTokenManager.cs
@@ -67,11 +67,11 @@ namespace UnrealGameSync
 			return TokenClients.Any(Pair => Pair.Value.GetStatusForProvider() == OIDCStatus.NotLoggedIn);
 		}
 
-		internal static OIDCTokenManager CreateFromConfigFile(UserSettings Settings, List<WorkspaceSettings> ConfigFiles)
+		internal static OIDCTokenManager CreateFromConfigFile(UserSettings Settings, List<OpenProjectInfo> ConfigFiles)
 		{
 			// join the provider configuration from all projects
 			Dictionary<string, ProviderInfo> Providers = new Dictionary<string, ProviderInfo>();
-			foreach (WorkspaceSettings DetectProjectSettingsTask in ConfigFiles)
+			foreach (OpenProjectInfo DetectProjectSettingsTask in ConfigFiles)
 			{
 				if(DetectProjectSettingsTask == null)
 				{
diff --git a/Engine/Source/Programs/UnrealGameSync/UnrealGameSync/ProgramApplicationContext.cs b/Engine/Source/Programs/UnrealGameSync/UnrealGameSync/ProgramApplicationContext.cs
--- a/Engine/Source/Programs/UnrealGameSync/UnrealGameSync/ProgramApplicationContext.cs
+++ b/Engine/Source/Programs/UnrealGameSync/UnrealGameSync/ProgramApplicationContext.cs
@@ -184,12 +184,12 @@ namespace UnrealGameSync
 			NotifyIcon.MouseDown += new MouseEventHandler(NotifyIcon_MouseDown);
 
 			// Create the startup tasks
-			List<(UserSelectedProjectSettings, ModalTask<WorkspaceSettings>)> StartupTasks = new List<(UserSelectedProjectSettings, ModalTask<WorkspaceSettings>)>();
+			List<(UserSelectedProjectSettings, ModalTask<OpenProjectInfo>)> StartupTasks = new List<(UserSelectedProjectSettings, ModalTask<OpenProjectInfo>)>();
 			foreach (UserSelectedProjectSettings ProjectSettings in Settings.OpenProjects)
 			{
-				ILogger<WorkspaceSettings> Logger = ServiceProvider.GetRequiredService<ILogger<WorkspaceSettings>>();
-				Task<WorkspaceSettings> WorkspaceSettingsTask = Task.Run(() => WorkspaceSettings.CreateAsync(DefaultPerforceSettings, ProjectSettings, Settings, Logger, StartupCancellationSource.Token), StartupCancellationSource.Token);
-				StartupTasks.Add((ProjectSettings, new ModalTask<WorkspaceSettings>(WorkspaceSettingsTask)));
+				ILogger<OpenProjectInfo> Logger = ServiceProvider.GetRequiredService<ILogger<OpenProjectInfo>>();
+				Task<OpenProjectInfo> StartupTask = Task.Run(() => OpenProjectInfo.CreateAsync(DefaultPerforceSettings, ProjectSettings, Settings, Logger, StartupCancellationSource.Token), StartupCancellationSource.Token);
+				StartupTasks.Add((ProjectSettings, new ModalTask<OpenProjectInfo>(StartupTask)));
 			}
 			StartupTask = Task.WhenAll(StartupTasks.Select(x => x.Item2.Task));
 
@@ -223,7 +223,7 @@ namespace UnrealGameSync
 			return Project;
 		}
 
-		private void OnStartupComplete(List<(UserSelectedProjectSettings, ModalTask<WorkspaceSettings>)> StartupTasks)
+		private void OnStartupComplete(List<(UserSelectedProjectSettings, ModalTask<OpenProjectInfo>)> StartupTasks)
 		{
 			// Close the startup window
 			bool bVisible = StartupWindow!.Visible;
diff --git a/Engine/Source/Programs/UnrealGameSync/UnrealGameSyncCmd/Program.cs b/Engine/Source/Programs/UnrealGameSync/UnrealGameSyncCmd/Program.cs
new file mode 100644
--- /dev/null
+++ b/Engine/Source/Programs/UnrealGameSync/UnrealGameSyncCmd/Program.cs
@@ -0,0 +1,730 @@
+// Copyright Epic Games, Inc. All Rights Reserved.
+
+using EpicGames.Core;
+using EpicGames.Perforce;
+using Microsoft.Extensions.Logging;
+using Serilog;
+using System;
+using System.Collections.Generic;
+using System.Diagnostics;
+using System.Globalization;
+using System.Linq;
+using System.Net;
+using System.Runtime.InteropServices;
+using System.Text.RegularExpressions;
+using UnrealGameSync;
+
+namespace UnrealGameSyncCmd
+{
+	using ILogger = Microsoft.Extensions.Logging.ILogger;
+
+	sealed class UserErrorException : Exception
+	{
+		public LogEvent Event { get; }
+		public int Code { get; }
+
+		public UserErrorException(LogEvent Event)
+			: base(Event.ToString())
+		{
+			this.Event = Event;
+			this.Code = 1;
+		}
+
+		public UserErrorException(string Message, params object[] Args)
+			: this(LogEvent.Create(LogLevel.Error, Message, Args))
+		{
+		}
+	}
+
+	public class Program
+	{
+		class CommandInfo
+		{
+			public string Name { get; }
+			public Type Type { get; }
+			public string Usage { get; }
+			public string Brief { get; }
+
+			public CommandInfo(string Name, Type Type, string Usage, string Brief)
+			{
+				this.Name = Name;
+				this.Type = Type;
+				this.Usage = Usage;
+				this.Brief = Brief;
+			}
+		}
+
+		static CommandInfo[] Commands =
+		{
+			new CommandInfo("init", typeof(InitCommand),
+				"ugs init [stream-path] [-client=..] [-server=..] [-user=..] [-branch=..] [-project=..]",
+				"Create a client for the given stream, or initializes an existing client for use by UGS."
+			),
+			new CommandInfo("switch", typeof(SwitchCommand),
+				"ugs switch [project name|project path|stream]",
+				"Changes the active project to the one in the workspace with the given name, or switches to a new stream."
+			),
+			new CommandInfo("config", typeof(ConfigCommand),
+				"ugs config",
+				"Updates the configuration for the current workspace."
+			),
+			new CommandInfo("sync", typeof(SyncCommand),
+				"ugs sync [change|'latest'] [-build] [-only]",
+				"Syncs the current workspace to the given changelist, optionally removing all local state."
+			),
+			new CommandInfo("build", typeof(BuildCommand),
+				"ugs build [id] [-list]",
+				"Runs the default build steps for the current project, or a particular step referenced by id."
+			),
+			new CommandInfo("status", typeof(StatusCommand),
+				"ugs status [-update]",
+				"Shows the status of the currently synced branch."
+			)
+		};
+
+		class CommandContext
+		{
+			public CommandLineArguments Arguments { get; }
+			public ILogger Logger { get; }
+			public ILoggerFactory LoggerFactory { get; }
+			public UserSettings UserSettings { get; }
+
+			public CommandContext(CommandLineArguments Arguments, ILogger Logger, ILoggerFactory LoggerFactory, UserSettings UserSettings)
+			{
+				this.Arguments = Arguments;
+				this.Logger = Logger;
+				this.LoggerFactory = LoggerFactory;
+				this.UserSettings = UserSettings;
+			}
+		}
+
+		class ProjectConfigOptions
+		{
+			[CommandLine("-Server=")]
+			public string? ServerAndPort { get; set; }
+
+			[CommandLine("-User=")]
+			public string? UserName { get; set; }
+
+			public void ApplyTo(UserWorkspaceSettings Settings)
+			{
+				if (ServerAndPort != null)
+				{
+					Settings.ServerAndPort = (ServerAndPort.Length == 0) ? null : ServerAndPort;
+				}
+				if (UserName != null)
+				{
+					Settings.UserName = (UserName.Length == 0) ? null : UserName;
+				}
+			}
+		}
+
+		class ProjectInitOptions : ProjectConfigOptions
+		{
+			[CommandLine("-Client=")]
+			public string? ClientName { get; set; }
+
+			[CommandLine("-Branch=")]
+			public string? BranchPath { get; set; }
+
+			[CommandLine("-Project=")]
+			public string? ProjectName { get; set; }
+		}
+
+		public static async Task<int> Main(string[] RawArgs)
+		{
+			DirectoryReference GlobalConfigFolder = DirectoryReference.Combine(DirectoryReference.GetSpecialFolder(Environment.SpecialFolder.LocalApplicationData)!, "UnrealGameSync");
+			DirectoryReference LogFolder = GlobalConfigFolder;
+
+			Serilog.ILogger SerilogLogger = new LoggerConfiguration()
+				.Enrich.FromLogContext()
+				.WriteTo.Console(Serilog.Events.LogEventLevel.Information, outputTemplate: "{Message:lj}{NewLine}")
+				.WriteTo.File(FileReference.Combine(LogFolder, "ugs-cmd.txt").FullName, Serilog.Events.LogEventLevel.Debug, rollingInterval: RollingInterval.Day, rollOnFileSizeLimit: true, fileSizeLimitBytes: 20 * 1024 * 1024, retainedFileCountLimit: 10)
+				.CreateLogger();
+
+			using ILoggerFactory LoggerFactory = new Serilog.Extensions.Logging.SerilogLoggerFactory(SerilogLogger, true);
+			ILogger Logger = LoggerFactory.CreateLogger("Main");
+			try
+			{
+				ConfigFile ConfigFile = new ConfigFile();
+				if (RuntimeInformation.IsOSPlatform(OSPlatform.Windows))
+				{
+					ConfigFile.TryLoad(FileReference.Combine(GlobalConfigFolder, "UnrealGameSync.ini"), Logger);
+				}
+
+				UserSettings UserSettings = new UserSettings(null!, ConfigFile);
+
+				CommandLineArguments Args = new CommandLineArguments(RawArgs);
+
+				string? CommandName;
+				if (!Args.TryGetPositionalArgument(out CommandName))
+				{
+					PrintHelp();
+					return 0;
+				}
+
+				CommandInfo? Command = Commands.FirstOrDefault(x => x.Name.Equals(CommandName, StringComparison.OrdinalIgnoreCase));
+				if (Command == null)
+				{
+					Logger.LogError($"unknown command '{CommandName}'");
+					Console.WriteLine();
+					PrintHelp();
+					return 1;
+				}
+
+				Command Instance = (Command)Activator.CreateInstance(Command.Type)!;
+				await Instance.ExecuteAsync(new CommandContext(Args, Logger, LoggerFactory, UserSettings));
+				return 0;
+			}
+			catch (UserErrorException Ex)
+			{
+				Logger.Log(Ex.Event.Level, "{Message}", Ex.Event.ToString());
+				return Ex.Code;
+			}
+			catch (PerforceException Ex)
+			{
+				Logger.LogError(Ex, "{Message}", Ex.Message);
+				return 1;
+			}
+			catch (Exception Ex)
+			{
+				Logger.LogError(Ex, "Unhandled exception.\n{Str}", Ex);
+				return 1;
+			}
+		}
+
+		static void PrintHelp()
+		{
+			Console.WriteLine("Usage:");
+			foreach (CommandInfo Command in Commands)
+			{
+				Console.WriteLine();
+				ConsoleUtils.WriteLineWithWordWrap(Command.Usage, 2, 8);
+				ConsoleUtils.WriteLineWithWordWrap(Command.Brief, 4, 4);
+			}
+		}
+
+		public static UserWorkspaceSettings? ReadOptionalUserWorkspaceSettings()
+		{
+			DirectoryReference? Dir = DirectoryReference.GetCurrentDirectory();
+			for (; Dir != null; Dir = Dir.ParentDirectory)
+			{
+				UserWorkspaceSettings? Settings;
+				if (UserWorkspaceSettings.TryLoad(Dir, out Settings))
+				{
+					return Settings;
+				}
+			}
+			return null;
+		}
+
+		public static UserWorkspaceSettings ReadRequiredUserWorkspaceSettings()
+		{
+			UserWorkspaceSettings? Settings = ReadOptionalUserWorkspaceSettings();
+			if (Settings == null)
+			{
+				throw new UserErrorException("Unable to find UGS workspace in current directory.");
+			}
+			return Settings;
+		}
+
+		public static async Task<UserWorkspaceState> ReadWorkspaceState(IPerforceConnection PerforceClient, UserWorkspaceSettings Settings, UserSettings UserSettings, ILogger Logger)
+		{
+			UserWorkspaceState State = UserSettings.FindOrAddWorkspaceState(Settings);
+			if (State.SettingsTimeUtc != Settings.LastModifiedTimeUtc)
+			{
+				Logger.LogDebug("Updating state due to modified settings timestamp");
+				ProjectInfo Info = await ProjectInfo.CreateAsync(PerforceClient, Settings, CancellationToken.None);
+				State.UpdateCachedProjectInfo(Info, Settings.LastModifiedTimeUtc);
+				State.Save();
+			}
+			return State;
+		}
+
+		public static Task<IPerforceConnection> ConnectAsync(string? ServerAndPort, string? UserName, string? ClientName, ILoggerFactory LoggerFactory)
+		{
+			PerforceSettings Settings = new PerforceSettings(PerforceSettings.Default);
+			Settings.ClientName = ClientName;
+			Settings.PreferNativeClient = true;
+			if (!String.IsNullOrEmpty(ServerAndPort))
+			{
+				Settings.ServerAndPort = ServerAndPort;
+			}
+			if (!String.IsNullOrEmpty(UserName))
+			{
+				Settings.UserName = UserName;
+			}
+
+			return PerforceConnection.CreateAsync(Settings, LoggerFactory.CreateLogger("Perforce"));
+		}
+
+		public static Task<IPerforceConnection> ConnectAsync(UserWorkspaceSettings Settings, ILoggerFactory LoggerFactory)
+		{
+			return ConnectAsync(Settings.ServerAndPort, Settings.UserName, Settings.ClientName, LoggerFactory);
+		}
+
+		static string[] ReadSyncFilter(UserWorkspaceSettings WorkspaceSettings, UserSettings UserSettings, ConfigFile ProjectConfig)
+		{
+			Dictionary<Guid, WorkspaceSyncCategory> SyncCategories = ConfigUtils.GetSyncCategories(ProjectConfig);
+			string[] CombinedSyncFilter = UserSettings.GetCombinedSyncFilter(SyncCategories, UserSettings.SyncView, UserSettings.SyncCategories, WorkspaceSettings.SyncView, WorkspaceSettings.SyncCategoriesDict);
+
+			ConfigSection PerforceSection = ProjectConfig.FindSection("Perforce");
+			if (PerforceSection != null)
+			{
+				IEnumerable<string> AdditionalPaths = PerforceSection.GetValues("AdditionalPathsToSync", new string[0]);
+				CombinedSyncFilter = AdditionalPaths.Union(CombinedSyncFilter).ToArray();
+			}
+
+			return CombinedSyncFilter;
+		}
+
+		static async Task<string> FindProjectPathAsync(IPerforceConnection Perforce, string ClientName, string BranchPath, string? ProjectName)
+		{
+			using IPerforceConnection PerforceClient = await PerforceConnection.CreateAsync(new PerforceSettings(Perforce.Settings) { ClientName = ClientName }, Perforce.Logger);
+
+			// Find or validate the selected project
+			string SearchPath;
+			if (ProjectName == null)
+			{
+				SearchPath = $"//{ClientName}{BranchPath}/*.uprojectdirs";
+			}
+			else if (ProjectName.Contains('.'))
+			{
+				SearchPath = $"//{ClientName}{BranchPath}/{ProjectName.TrimStart('/')}";
+			}
+			else
+			{
+				SearchPath = $"//{ClientName}{BranchPath}/.../{ProjectName}.uproject";
+			}
+
+			List<FStatRecord> ProjectFileRecords = await PerforceClient.FStatAsync(FStatOptions.ClientFileInPerforceSyntax, SearchPath).ToListAsync();
+			ProjectFileRecords.RemoveAll(x => !x.IsMapped);
+
+			List<string> Paths = ProjectFileRecords.Select(x => PerforceUtils.GetClientRelativePath(x.ClientFile!)).Distinct(StringComparer.Ordinal).ToList();
+			if (Paths.Count == 0)
+			{
+				throw new UserErrorException("No project file found matching {SearchPath}", SearchPath);
+			}
+			if (Paths.Count > 1)
+			{
+				throw new UserErrorException("Multiple projects found matching {SearchPath}: {Paths}", SearchPath, String.Join(", ", Paths));
+			}
+
+			return "/" + Paths[0];
+		}
+
+		abstract class Command
+		{
+			public abstract Task ExecuteAsync(CommandContext Context);
+		}
+
+		class InitCommand : Command
+		{
+			public override async Task ExecuteAsync(CommandContext Context)
+			{
+				ILogger Logger = Context.Logger;
+
+				// Get the positional argument indicating the file to look for
+				string? InitName;
+				Context.Arguments.TryGetPositionalArgument(out InitName);
+
+				// Get the config settings from the command line
+				ProjectInitOptions Options = new ProjectInitOptions();
+				Context.Arguments.ApplyTo(Options);
+				Context.Arguments.CheckAllArgumentsUsed();
+
+				// Get the host name
+				using IPerforceConnection Perforce = await ConnectAsync(Options.ServerAndPort, Options.UserName, null, Context.LoggerFactory);
+				InfoRecord PerforceInfo = await Perforce.GetInfoAsync(InfoOptions.ShortOutput);
+				string HostName = PerforceInfo.ClientHost ?? Dns.GetHostName();
+
+				// Create the perforce connection
+				if (InitName != null)
+				{
+					await InitNewClientAsync(Perforce, Context, InitName, HostName, Options, Logger);
+				}
+				else
+				{
+					await InitExistingClientAsync(Perforce, Context, HostName, Options, Logger);
+				}
+			}
+
+			async Task InitNewClientAsync(IPerforceConnection Perforce, CommandContext Context, string StreamName, string HostName, ProjectInitOptions Options, ILogger Logger)
+			{
+				Logger.LogInformation("Checking stream...");
+
+				// Get the given stream
+				PerforceResponse<StreamRecord> StreamResponse = await Perforce.TryGetStreamAsync(StreamName, true);
+				if (!StreamResponse.Succeeded)
+				{
+					throw new UserErrorException($"Unable to find stream '{StreamName}'");
+				}
+				StreamRecord Stream = StreamResponse.Data;
+
+				// Get the new directory for the client
+				DirectoryReference ClientDir = DirectoryReference.Combine(DirectoryReference.GetCurrentDirectory(), Stream.Stream.Replace('/', '+'));
+				DirectoryReference.CreateDirectory(ClientDir);
+
+				// Make up a new client name 
+				string ClientName = Options.ClientName ?? Regex.Replace($"{Perforce.Settings.UserName}_{HostName}_{Stream.Stream.Trim('/')}", "[^0-9a-zA-Z_.-]", "+");
+
+				// Check there are no existing clients under the current path
+				List<ClientsRecord> Clients = await FindExistingClients(Perforce, HostName, ClientDir);
+				if (Clients.Count > 0)
+				{
+					if (Clients.Count == 1 && ClientName.Equals(Clients[0].Name, StringComparison.OrdinalIgnoreCase) && ClientDir == TryParseRoot(Clients[0].Root))
+					{
+						Logger.LogInformation("Reusing existing client for {ClientDir} ({ClientName})", ClientDir, Options.ClientName);
+					}
+					else
+					{
+						throw new UserErrorException("Current directory is already within a Perforce workspace ({ClientName})", Clients[0].Name);
+					}
+				}
+
+				// Create the new client
+				ClientRecord Client = new ClientRecord(ClientName, Perforce.Settings.UserName, ClientDir.FullName);
+				Client.Host = HostName;
+				Client.Stream = Stream.Stream;
+				Client.Options = ClientOptions.Rmdir;
+				await Perforce.CreateClientAsync(Client);
+
+				// Branch root is currently hard-coded at the root
+				string BranchPath = Options.BranchPath ?? String.Empty;
+				string ProjectPath = await FindProjectPathAsync(Perforce, ClientName, BranchPath, Options.ProjectName);
+
+				// Create the settings object
+				UserWorkspaceSettings Settings = new UserWorkspaceSettings();
+				Settings.RootDir = ClientDir;
+				Settings.Init(Perforce.Settings.ServerAndPort, Perforce.Settings.UserName, ClientName, BranchPath, ProjectPath);
+				Options.ApplyTo(Settings);
+				Settings.Save();
+
+				Logger.LogInformation("Initialized {ClientName} with root at {RootDir}", ClientName, ClientDir);
+			}
+
+			static DirectoryReference? TryParseRoot(string Root)
+			{
+				try
+				{
+					return new DirectoryReference(Root);
+				}
+				catch
+				{
+					return null;
+				}
+			}
+
+			async Task InitExistingClientAsync(IPerforceConnection Perforce, CommandContext Context, string HostName, ProjectInitOptions Options, ILogger Logger)
+			{
+				DirectoryReference CurrentDir = DirectoryReference.GetCurrentDirectory();
+
+				// Make sure the client name is set
+				string? ClientName = Options.ClientName;
+				if (ClientName == null)
+				{
+					List<ClientsRecord> Clients = await FindExistingClients(Perforce, HostName, CurrentDir);
+					if (Clients.Count == 0)
+					{
+						throw new UserErrorException("Unable to find client for {HostName} under {ClientDir}", HostName, CurrentDir);
+					}
+					if (Clients.Count > 1)
+					{
+						throw new UserErrorException("Multiple clients found for {HostName} under {ClientDir}: {ClientList}", HostName, CurrentDir, String.Join(", ", Clients.Select(x => x.Name)));
+					}
+
+					ClientName = Clients[0].Name;
+					Logger.LogInformation("Found client {ClientName}", ClientName);
+				}
+
+				// Get the client info
+				ClientRecord Client = await Perforce.GetClientAsync(Options.ClientName);
+				DirectoryReference ClientDir = new DirectoryReference(Client.Root);
+
+				// If a project path was specified in local syntax, try to convert it to client-relative syntax
+				string? ProjectName = Options.ProjectName;
+				if (Options.ProjectName != null && Options.ProjectName.Contains('.'))
+				{
+					Options.ProjectName = FileReference.Combine(CurrentDir, Options.ProjectName).MakeRelativeTo(ClientDir).Replace('\\', '/');
+				}
+
+				// Branch root is currently hard-coded at the root
+				string BranchPath = Options.BranchPath ?? String.Empty;
+				string ProjectPath = await FindProjectPathAsync(Perforce, ClientName, BranchPath, ProjectName);
+
+				// Create the settings object
+				UserWorkspaceSettings Settings = new UserWorkspaceSettings();
+				Settings.RootDir = ClientDir;
+				Settings.Init(Perforce.Settings.ServerAndPort, Perforce.Settings.UserName, ClientName, BranchPath, ProjectPath);
+				Options.ApplyTo(Settings);
+				Settings.Save();
+
+				Logger.LogInformation("Initialized workspace at {RootDir} for {ClientProject}", ClientDir, Settings.ClientProjectPath);
+			}
+
+			static async Task<List<ClientsRecord>> FindExistingClients(IPerforceConnection Perforce, string HostName, DirectoryReference ClientDir)
+			{
+				List<ClientsRecord> MatchingClients = new List<ClientsRecord>();
+
+				List<ClientsRecord> Clients = await Perforce.GetClientsAsync(ClientsOptions.None, Perforce.Settings.UserName);
+				foreach (ClientsRecord Client in Clients)
+				{
+					if (!String.IsNullOrEmpty(Client.Root) && !String.IsNullOrEmpty(Client.Host) && String.Compare(HostName, Client.Host, StringComparison.OrdinalIgnoreCase) == 0)
+					{
+						DirectoryReference? RootDir;
+						try
+						{
+							RootDir = new DirectoryReference(Client.Root);
+						}
+						catch
+						{
+							RootDir = null;
+						}
+
+						if (RootDir != null && ClientDir.IsUnderDirectory(RootDir))
+						{
+							MatchingClients.Add(Client);
+						}
+					}
+				}
+
+				return MatchingClients;
+			}
+		}
+
+		class SyncCommand : Command
+		{
+			class SyncOptions
+			{
+				[CommandLine("-Only")]
+				public bool SingleChange { get; set; }
+
+				[CommandLine("-Build")]
+				public bool Build { get; set; }
+			}
+
+			public override async Task ExecuteAsync(CommandContext Context)
+			{
+				ILogger Logger = Context.Logger;
+				Context.Arguments.TryGetPositionalArgument(out string? ChangeString);
+
+				SyncOptions SyncOptions = new SyncOptions();
+				Context.Arguments.ApplyTo(SyncOptions);
+
+				Context.Arguments.CheckAllArgumentsUsed();
+
+				UserWorkspaceSettings Settings = ReadRequiredUserWorkspaceSettings();
+				using IPerforceConnection PerforceClient = await ConnectAsync(Settings, Context.LoggerFactory);
+				UserWorkspaceState State = await ReadWorkspaceState(PerforceClient, Settings, Context.UserSettings, Logger);
+
+				ChangeString ??= "latest";
+
+				int Change;
+				if (!int.TryParse(ChangeString, out Change))
+				{
+					if (String.Equals(ChangeString, "latest", StringComparison.OrdinalIgnoreCase))
+					{
+						List<ChangesRecord> Changes = await PerforceClient.GetChangesAsync(ChangesOptions.None, 1, ChangeStatus.Submitted, $"//{Settings.ClientName}/...");
+						Change = Changes[0].Number;
+					}
+					else
+					{
+						throw new UserErrorException("Unknown change type for sync '{Change}'", ChangeString);
+					}
+				}
+
+				WorkspaceUpdateOptions Options = SyncOptions.SingleChange? WorkspaceUpdateOptions.SyncSingleChange : WorkspaceUpdateOptions.Sync;
+				if (SyncOptions.Build)
+				{
+					Options |= WorkspaceUpdateOptions.Build;
+				}
+
+				ProjectInfo ProjectInfo = State.CreateProjectInfo();
+				UserProjectSettings ProjectSettings = Context.UserSettings.FindOrAddProjectSettings(ProjectInfo);
+
+				ConfigFile ProjectConfig = await ConfigUtils.ReadProjectConfigFileAsync(PerforceClient, ProjectInfo, Logger, CancellationToken.None);
+				string[] SyncFilter = ReadSyncFilter(Settings, Context.UserSettings, ProjectConfig);
+
+				WorkspaceUpdateContext UpdateContext = new WorkspaceUpdateContext(Change, Options, BuildConfig.Development, SyncFilter, ProjectSettings.BuildSteps, null);
+
+				WorkspaceUpdate Update = new WorkspaceUpdate(UpdateContext);
+				await Update.ExecuteAsync(PerforceClient.Settings, ProjectInfo, State, Context.Logger, CancellationToken.None);
+			}
+		}
+
+		class ConfigCommand : Command
+		{
+			public override async Task ExecuteAsync(CommandContext Context)
+			{
+				ILogger Logger = Context.Logger;
+
+				UserWorkspaceSettings Settings = ReadRequiredUserWorkspaceSettings();
+				if (!Context.Arguments.GetUnusedArguments().Any())
+				{
+					ProcessStartInfo StartInfo = new ProcessStartInfo();
+					StartInfo.FileName = Settings.ConfigFile.FullName;
+					StartInfo.UseShellExecute = true;
+					using (Process? Editor = Process.Start(StartInfo))
+					{
+						if (Editor != null)
+						{
+							await Editor.WaitForExitAsync();
+						}
+					}
+				}
+				else
+				{
+					ProjectConfigOptions Options = new ProjectConfigOptions();
+					Context.Arguments.ApplyTo(Options);
+					Context.Arguments.CheckAllArgumentsUsed(Context.Logger);
+
+					Options.ApplyTo(Settings);
+					Settings.Save();
+
+					Logger.LogInformation("Updated {ConfigFile}", Settings.ConfigFile);
+				}
+			}
+		}
+
+		class BuildCommand : Command
+		{
+			public override async Task ExecuteAsync(CommandContext Context)
+			{
+				ILogger Logger = Context.Logger;
+				Context.Arguments.TryGetPositionalArgument(out string? Target);
+				bool ListOnly = Context.Arguments.HasOption("-List");
+				Context.Arguments.CheckAllArgumentsUsed();
+
+				UserWorkspaceSettings Settings = ReadRequiredUserWorkspaceSettings();
+				using IPerforceConnection PerforceClient = await ConnectAsync(Settings, Context.LoggerFactory);
+				UserWorkspaceState State = await ReadWorkspaceState(PerforceClient, Settings, Context.UserSettings, Logger);
+
+				ProjectInfo ProjectInfo = State.CreateProjectInfo();
+
+				if (ListOnly)
+				{
+					ConfigFile ProjectConfig = await ConfigUtils.ReadProjectConfigFileAsync(PerforceClient, ProjectInfo, Logger, CancellationToken.None);
+
+					FileReference EditorTarget = ConfigUtils.GetEditorTargetFile(ProjectInfo, ProjectConfig);
+
+					Dictionary<Guid, ConfigObject> BuildStepObjects = ConfigUtils.GetDefaultBuildStepObjects(ProjectInfo, EditorTarget.GetFileNameWithoutAnyExtensions(), BuildConfig.Development, ProjectConfig, false);
+
+					Logger.LogInformation("Available build steps:");
+					Logger.LogInformation("");
+					Logger.LogInformation("  Id                                   | Description                              | Type       | Enabled");
+					Logger.LogInformation("  -------------------------------------|------------------------------------------|------------|-----------------");
+					foreach (BuildStep BuildStep in BuildStepObjects.Values.Select(x => new BuildStep(x)).OrderBy(x => x.OrderIndex))
+					{
+						Logger.LogInformation("  {Id,-36} | {Name,-40} | {Type,-10} | {Enabled,-8}", BuildStep.UniqueId, BuildStep.Description, BuildStep.Type, BuildStep.bNormalSync);
+					}
+					return;
+				}
+
+				HashSet<Guid>? Steps = null;
+				if (Target != null)
+				{
+					Guid Id;
+					if (!Guid.TryParse(Target, out Id))
+					{
+						Logger.LogError("Unable to parse '{Target}' as a GUID. Pass -List to show all available build steps and their identifiers.", Target);
+					}
+					Steps = new HashSet<Guid> { Id };
+				}
+
+				WorkspaceUpdateContext UpdateContext = new WorkspaceUpdateContext(State.CurrentChangeNumber, WorkspaceUpdateOptions.Build, BuildConfig.Development, null, new List<ConfigObject>(), Steps);
+
+				WorkspaceUpdate Update = new WorkspaceUpdate(UpdateContext);
+				await Update.ExecuteAsync(PerforceClient.Settings, ProjectInfo, State, Context.Logger, CancellationToken.None);
+			}
+		}
+
+		class StatusCommand : Command
+		{
+			public override async Task ExecuteAsync(CommandContext Context)
+			{
+				ILogger Logger = Context.Logger;
+				bool bUpdate = Context.Arguments.HasOption("-Update");
+				Context.Arguments.CheckAllArgumentsUsed();
+
+				UserWorkspaceSettings Settings = ReadRequiredUserWorkspaceSettings();
+				Logger.LogInformation("User: {UserName}", Settings.UserName);
+				Logger.LogInformation("Server: {ServerAndPort}", Settings.ServerAndPort);
+				Logger.LogInformation("Project: {ClientProjectPath}", Settings.ClientProjectPath);
+
+				using IPerforceConnection PerforceClient = await ConnectAsync(Settings, Context.LoggerFactory);
+
+				UserWorkspaceState State = await ReadWorkspaceState(PerforceClient, Settings, Context.UserSettings, Logger);
+				if (bUpdate)
+				{
+					ProjectInfo NewProjectInfo = await ProjectInfo.CreateAsync(PerforceClient, Settings, CancellationToken.None);
+					State.UpdateCachedProjectInfo(NewProjectInfo, Settings.LastModifiedTimeUtc);
+				}
+
+				string StreamOrBranchName = State.StreamName ?? Settings.BranchPath.TrimStart('/');
+				if (State.LastSyncResultMessage == null)
+				{
+					Logger.LogInformation("Not currently synced to {Stream}", StreamOrBranchName);
+				}
+				else if (State.LastSyncResult == WorkspaceUpdateResult.Success)
+				{
+					Logger.LogInformation("Synced to {Stream} CL {Change}", StreamOrBranchName, State.LastSyncChangeNumber);
+				}
+				else
+				{
+					Logger.LogWarning("Last sync to {Stream} CL {Change} failed: {Result}", StreamOrBranchName, State.LastSyncChangeNumber, State.LastSyncResultMessage);
+				}
+			}
+		}
+
+		class SwitchCommand : Command
+		{
+			public override async Task ExecuteAsync(CommandContext Context)
+			{
+				// Get the positional argument indicating the file to look for
+				string? TargetName;
+				if (!Context.Arguments.TryGetPositionalArgument(out TargetName))
+				{
+					throw new UserErrorException("Missing stream or project name to switch to.");
+				}
+
+				bool Force = TargetName.StartsWith("//", StringComparison.Ordinal) && Context.Arguments.HasOption("-Force");
+
+				// Finish argument parsing
+				Context.Arguments.CheckAllArgumentsUsed();
+
+				// Get a connection to the client for this workspace
+				UserWorkspaceSettings Settings = ReadRequiredUserWorkspaceSettings();
+				using IPerforceConnection PerforceClient = await ConnectAsync(Settings, Context.LoggerFactory);
+
+				// Check whether we're switching stream or project
+				if (TargetName.StartsWith("//", StringComparison.Ordinal))
+				{
+					await SwitchStreamAsync(PerforceClient, TargetName, Force, Context.Logger);
+				}
+				else
+				{
+					await SwitchProjectAsync(PerforceClient, Settings, TargetName, Context.Logger);
+				}
+			}
+
+			public async Task SwitchStreamAsync(IPerforceConnection PerforceClient, string StreamName, bool Force, ILogger Logger)
+			{
+				if (!Force && await PerforceClient.OpenedAsync(OpenedOptions.None, FileSpecList.Any).AnyAsync())
+				{
+					throw new UserErrorException("Client {ClientName} has files opened. Use -Force to switch anyway.", PerforceClient.Settings.ClientName!);
+				}
+
+				await PerforceClient.SwitchClientToStreamAsync(PerforceClient.Settings.ClientName!, StreamName, SwitchClientOptions.IgnoreOpenFiles);
+
+				Logger.LogInformation("Switched to stream {StreamName}", StreamName);
+			}
+
+			public async Task SwitchProjectAsync(IPerforceConnection PerforceClient, UserWorkspaceSettings Settings, string ProjectName, ILogger Logger)
+			{
+				Settings.ProjectPath = await FindProjectPathAsync(PerforceClient, Settings.ClientName, Settings.BranchPath, ProjectName);
+				Settings.Save();
+				Logger.LogInformation("Switched to project {ProjectPath}", Settings.ClientProjectPath);
+			}
+		}
+	}
+}
diff --git a/Engine/Source/Programs/UnrealGameSync/UnrealGameSyncCmd/UnrealGameSyncCmd.csproj b/Engine/Source/Programs/UnrealGameSync/UnrealGameSyncCmd/UnrealGameSyncCmd.csproj
new file mode 100644
--- /dev/null
+++ b/Engine/Source/Programs/UnrealGameSync/UnrealGameSyncCmd/UnrealGameSyncCmd.csproj
@@ -0,0 +1,23 @@
+<Project Sdk="Microsoft.NET.Sdk">
+
+  <PropertyGroup>
+    <OutputType>Exe</OutputType>
+    <TargetFramework>net6.0-windows</TargetFramework>
+    <ImplicitUsings>enable</ImplicitUsings>
+    <Nullable>enable</Nullable>
+    <AssemblyName>ugs</AssemblyName>
+  </PropertyGroup>
+
+  <ItemGroup>
+    <PackageReference Include="Serilog" Version="2.10.0" />
+    <PackageReference Include="Serilog.Sinks.Console" Version="4.0.1" />
+    <PackageReference Include="Serilog.Sinks.File" Version="5.0.0" />
+    <PackageReference Include="Serilog.Extensions.Logging" Version="3.1.0" />
+  </ItemGroup>
+
+  <ItemGroup>
+    <ProjectReference Include="..\..\Shared\EpicGames.Core\EpicGames.Core.csproj" />
+    <ProjectReference Include="..\UnrealGameSyncShared\UnrealGameSyncShared.csproj" />
+  </ItemGroup>
+
+</Project>
diff --git a/Engine/Source/Programs/UnrealGameSync/UnrealGameSyncShared/ConfigUtils.cs b/Engine/Source/Programs/UnrealGameSync/UnrealGameSyncShared/ConfigUtils.cs
--- a/Engine/Source/Programs/UnrealGameSync/UnrealGameSyncShared/ConfigUtils.cs
+++ b/Engine/Source/Programs/UnrealGameSync/UnrealGameSyncShared/ConfigUtils.cs
@@ -74,6 +74,16 @@ namespace UnrealGameSync
 			}
 		}
 
+		public static Task<ConfigFile> ReadProjectConfigFileAsync(IPerforceConnection Perforce, ProjectInfo ProjectInfo, ILogger Logger, CancellationToken CancellationToken)
+		{
+			return ReadProjectConfigFileAsync(Perforce, ProjectInfo, new List<KeyValuePair<FileReference, DateTime>>(), Logger, CancellationToken);
+		}
+
+		public static Task<ConfigFile> ReadProjectConfigFileAsync(IPerforceConnection Perforce, ProjectInfo ProjectInfo, List<KeyValuePair<FileReference, DateTime>> LocalConfigFiles, ILogger Logger, CancellationToken CancellationToken)
+		{
+			return ReadProjectConfigFileAsync(Perforce, ProjectInfo.ClientRootPath, ProjectInfo.ClientFileName, ProjectInfo.CacheFolder, LocalConfigFiles, Logger, CancellationToken);
+		}
+
 		public static async Task<ConfigFile> ReadProjectConfigFileAsync(IPerforceConnection Perforce, string BranchClientPath, string SelectedClientFileName, DirectoryReference CacheFolder, List<KeyValuePair<FileReference, DateTime>> LocalConfigFiles, ILogger Logger, CancellationToken CancellationToken)
 		{
 			List<string> ConfigFilePaths = Utility.GetDepotConfigPaths(BranchClientPath + "/Engine", SelectedClientFileName);
@@ -122,7 +132,7 @@ namespace UnrealGameSync
 						try
 						{
 							ProjectConfig.Parse(Lines.ToArray());
-							Logger.LogInformation("Read config file from {DepotFile}", FileRecord.DepotFile);
+							Logger.LogDebug("Read config file from {DepotFile}", FileRecord.DepotFile);
 						}
 						catch (Exception Ex)
 						{
@@ -379,5 +389,54 @@ namespace UnrealGameSync
 			Value = null;
 			return false;
 		}
+
+		public static Dictionary<Guid, WorkspaceSyncCategory> GetSyncCategories(ConfigFile ProjectConfigFile)
+		{
+			Dictionary<Guid, WorkspaceSyncCategory> UniqueIdToCategory = new Dictionary<Guid, WorkspaceSyncCategory>();
+			if (ProjectConfigFile != null)
+			{
+				string[] CategoryLines = ProjectConfigFile.GetValues("Options.SyncCategory", new string[0]);
+				foreach (string CategoryLine in CategoryLines)
+				{
+					ConfigObject Object = new ConfigObject(CategoryLine);
+
+					Guid UniqueId;
+					if (Guid.TryParse(Object.GetValue("UniqueId", ""), out UniqueId))
+					{
+						WorkspaceSyncCategory? Category;
+						if (!UniqueIdToCategory.TryGetValue(UniqueId, out Category))
+						{
+							Category = new WorkspaceSyncCategory(UniqueId);
+							UniqueIdToCategory.Add(UniqueId, Category);
+						}
+
+						if (Object.GetValue("Clear", false))
+						{
+							Category.Paths = new string[0];
+							Category.Requires = new Guid[0];
+						}
+
+						Category.Name = Object.GetValue("Name", Category.Name);
+						Category.bEnable = Object.GetValue("Enable", Category.bEnable);
+						Category.Paths = Enumerable.Concat(Category.Paths, Object.GetValue("Paths", "").Split(';').Select(x => x.Trim())).Where(x => x.Length > 0).Distinct().OrderBy(x => x).ToArray();
+						Category.bHidden = Object.GetValue("Hidden", Category.bHidden);
+						Category.Requires = Enumerable.Concat(Category.Requires, ParseGuids(Object.GetValue("Requires", "").Split(';'))).Distinct().OrderBy(x => x).ToArray();
+					}
+				}
+			}
+			return UniqueIdToCategory;
+		}
+
+		static IEnumerable<Guid> ParseGuids(IEnumerable<string> Values)
+		{
+			foreach (string Value in Values)
+			{
+				Guid Guid;
+				if (Guid.TryParse(Value, out Guid))
+				{
+					yield return Guid;
+				}
+			}
+		}
 	}
 }
diff --git a/Engine/Source/Programs/UnrealGameSync/UnrealGameSyncShared/OpenProjectInfo.cs b/Engine/Source/Programs/UnrealGameSync/UnrealGameSyncShared/OpenProjectInfo.cs
new file mode 100644
--- /dev/null
+++ b/Engine/Source/Programs/UnrealGameSync/UnrealGameSyncShared/OpenProjectInfo.cs
@@ -0,0 +1,290 @@
+// Copyright Epic Games, Inc. All Rights Reserved.
+
+using EpicGames.Core;
+using EpicGames.Perforce;
+using Microsoft.Extensions.Logging;
+using System;
+using System.Collections.Generic;
+using System.Drawing;
+using System.IO;
+using System.Linq;
+using System.Text.RegularExpressions;
+using System.Threading;
+using System.Threading.Tasks;
+
+namespace UnrealGameSync
+{
+	public class OpenProjectInfo
+	{
+		public UserSelectedProjectSettings SelectedProject { get; }
+
+		public IPerforceSettings PerforceSettings { get; }
+		public ProjectInfo ProjectInfo { get; }
+		public UserWorkspaceSettings WorkspaceSettings { get; }
+		public UserWorkspaceState WorkspaceState { get; }
+		public ConfigFile LatestProjectConfigFile { get; }
+		public ConfigFile WorkspaceProjectConfigFile { get; }
+		public IReadOnlyList<string>? WorkspaceProjectStreamFilter { get; }
+		public List<KeyValuePair<FileReference, DateTime>> LocalConfigFiles { get; }
+
+		public OpenProjectInfo(UserSelectedProjectSettings SelectedProject, IPerforceSettings PerforceSettings, ProjectInfo ProjectInfo, UserWorkspaceSettings WorkspaceSettings, UserWorkspaceState WorkspaceState, ConfigFile LatestProjectConfigFile, ConfigFile WorkspaceProjectConfigFile, IReadOnlyList<string>? WorkspaceProjectStreamFilter, List<KeyValuePair<FileReference, DateTime>> LocalConfigFiles)
+		{
+			this.SelectedProject = SelectedProject;
+
+			this.PerforceSettings = PerforceSettings;
+			this.ProjectInfo = ProjectInfo;
+			this.WorkspaceSettings = WorkspaceSettings;
+			this.WorkspaceState = WorkspaceState;
+			this.LatestProjectConfigFile = LatestProjectConfigFile;
+			this.WorkspaceProjectConfigFile = WorkspaceProjectConfigFile;
+			this.WorkspaceProjectStreamFilter = WorkspaceProjectStreamFilter;
+			this.LocalConfigFiles = LocalConfigFiles;
+		}
+
+		public static async Task<OpenProjectInfo> CreateAsync(IPerforceSettings DefaultPerforceSettings, UserSelectedProjectSettings SelectedProject, UserSettings UserSettings, ILogger<OpenProjectInfo> Logger, CancellationToken CancellationToken)
+		{
+			PerforceSettings PerforceSettings = Utility.OverridePerforceSettings(DefaultPerforceSettings, SelectedProject.ServerAndPort, SelectedProject.UserName);
+			using IPerforceConnection Perforce = await PerforceConnection.CreateAsync(PerforceSettings, Logger);
+
+			// Make sure we're logged in
+			PerforceResponse<LoginRecord> LoginState = await Perforce.TryGetLoginStateAsync(CancellationToken);
+			if (!LoginState.Succeeded)
+			{
+				throw new UserErrorException("User is not logged in to Perforce.");
+			}
+
+			// Execute like a regular task
+			return await CreateAsync(Perforce, SelectedProject, UserSettings, Logger, CancellationToken);
+		}
+
+		public static async Task<OpenProjectInfo> CreateAsync(IPerforceConnection DefaultConnection, UserSelectedProjectSettings SelectedProject, UserSettings UserSettings, ILogger<OpenProjectInfo> Logger, CancellationToken CancellationToken)
+		{
+			using IDisposable LoggerScope = Logger.BeginScope("Project {SelectedProject}", SelectedProject.ToString());
+			Logger.LogInformation("Detecting settings for {Project}", SelectedProject);
+
+			// Use the cached client path to the file if it's available; it's much quicker than trying to find the correct workspace.
+			IPerforceConnection? PerforceClient = null;
+			try
+			{
+				IPerforceSettings PerforceSettings;
+
+				FileReference NewSelectedFileName;
+				string NewSelectedClientFileName;
+				if (!String.IsNullOrEmpty(SelectedProject.ClientPath))
+				{
+					// Get the client path
+					NewSelectedClientFileName = SelectedProject.ClientPath;
+
+					// Get the client name
+					string? ClientName;
+					if (!PerforceUtils.TryGetClientName(NewSelectedClientFileName, out ClientName))
+					{
+						throw new UserErrorException($"Couldn't get client name from {NewSelectedClientFileName}");
+					}
+
+					// Create the client
+					PerforceSettings = new PerforceSettings(DefaultConnection.Settings) { ClientName = ClientName };
+					PerforceClient = await PerforceConnection.CreateAsync(PerforceSettings, Logger);
+
+					// Figure out the path on the client. Use the cached location if it's valid.
+					string? LocalPath = SelectedProject.LocalPath;
+					if (LocalPath == null || !File.Exists(LocalPath))
+					{
+						List<WhereRecord> Records = await PerforceClient.WhereAsync(NewSelectedClientFileName, CancellationToken).Where(x => !x.Unmap).ToListAsync(CancellationToken);
+						if (Records.Count != 1)
+						{
+							throw new UserErrorException($"Couldn't get client path for {NewSelectedClientFileName}");
+						}
+						LocalPath = Path.GetFullPath(Records[0].Path);
+					}
+					NewSelectedFileName = new FileReference(LocalPath);
+				}
+				else
+				{
+					// Get the perforce server settings
+					InfoRecord PerforceInfo = await DefaultConnection.GetInfoAsync(InfoOptions.ShortOutput, CancellationToken);
+
+					// Use the path as the selected filename
+					NewSelectedFileName = new FileReference(SelectedProject.LocalPath!);
+
+					// Make sure the project exists
+					if (!FileReference.Exists(NewSelectedFileName))
+					{
+						throw new UserErrorException($"{SelectedProject.LocalPath} does not exist.");
+					}
+
+					// Find all the clients for this user
+					Logger.LogInformation("Enumerating clients for {UserName}...", PerforceInfo.UserName);
+
+					List<ClientsRecord> Clients = await DefaultConnection.GetClientsAsync(ClientsOptions.None, DefaultConnection.Settings.UserName, CancellationToken);
+
+					List<IPerforceSettings> CandidateClients = await FilterClients(Clients, NewSelectedFileName, DefaultConnection.Settings, PerforceInfo.ClientHost, Logger, CancellationToken);
+					if (CandidateClients.Count == 0)
+					{
+						// Search through all workspaces. We may find a suitable workspace which is for any user.
+						Logger.LogInformation("Enumerating shared clients...");
+						Clients = await DefaultConnection.GetClientsAsync(ClientsOptions.None, "", CancellationToken);
+
+						// Filter this list of clients
+						CandidateClients = await FilterClients(Clients, NewSelectedFileName, DefaultConnection.Settings, PerforceInfo.ClientHost, Logger, CancellationToken);
+
+						// If we still couldn't find any, fail.
+						if (CandidateClients.Count == 0)
+						{
+							throw new UserErrorException($"Couldn't find any Perforce workspace containing {NewSelectedFileName}. Check your connection settings.");
+						}
+					}
+
+					// Check there's only one client
+					if (CandidateClients.Count > 1)
+					{
+						throw new UserErrorException(String.Format("Found multiple workspaces containing {0}:\n\n{1}\n\nCannot determine which to use.", Path.GetFileName(NewSelectedFileName.GetFileName()), String.Join("\n", CandidateClients.Select(x => x.ClientName))));
+					}
+
+					// Take the client we've chosen
+					PerforceSettings = CandidateClients[0];
+					PerforceClient = await PerforceConnection.CreateAsync(PerforceSettings, Logger);
+
+					// Get the client path for the project file
+					List<WhereRecord> Records = await PerforceClient.WhereAsync(NewSelectedFileName.FullName, CancellationToken).Where(x => !x.Unmap).ToListAsync(CancellationToken);
+					if (Records.Count == 0)
+					{
+						throw new UserErrorException("File is not mapped to any client");
+					}
+					else if (Records.Count > 1)
+					{
+						throw new UserErrorException($"File is mapped to {Records.Count} locations: {String.Join(", ", Records.Select(x => x.Path))}");
+					}
+
+					NewSelectedClientFileName = Records[0].ClientFile;
+				}
+
+				// Make sure the drive containing the project exists, to prevent other errors down the line
+				string PathRoot = Path.GetPathRoot(NewSelectedFileName.FullName)!;
+				if (!Directory.Exists(PathRoot))
+				{
+					throw new UserErrorException($"Path '{NewSelectedFileName}' is invalid");
+				}
+
+				// Make sure the path case is correct. This can cause UBT intermediates to be out of date if the case mismatches.
+				NewSelectedFileName = FileReference.FindCorrectCase(NewSelectedFileName);
+
+				// Update the selected project with all the data we've found
+				SelectedProject = new UserSelectedProjectSettings(SelectedProject.ServerAndPort, SelectedProject.UserName, SelectedProject.Type, NewSelectedClientFileName, NewSelectedFileName.FullName);
+
+				// Get the local branch root
+				string? BranchClientPath = null;
+				DirectoryReference? BranchDirectoryName = null;
+
+				// Figure out where the engine is in relation to it
+				int EndIdx = NewSelectedClientFileName.Length - 1;
+				if (EndIdx != -1 && NewSelectedClientFileName.EndsWith(".uproject", StringComparison.InvariantCultureIgnoreCase))
+				{
+					EndIdx = NewSelectedClientFileName.LastIndexOf('/') - 1;
+				}
+				for (; EndIdx >= 2; EndIdx--)
+				{
+					if (NewSelectedClientFileName[EndIdx] == '/')
+					{
+						List<PerforceResponse<FStatRecord>> FileRecords = await PerforceClient.TryFStatAsync(FStatOptions.None, NewSelectedClientFileName.Substring(0, EndIdx) + "/Engine/Build/Build.version", CancellationToken).ToListAsync();
+						if (FileRecords.Succeeded())
+						{
+							FStatRecord FileRecord = FileRecords[0].Data;
+							if (FileRecord.ClientFile == null)
+							{
+								throw new UserErrorException($"Missing client path for {FileRecord.DepotFile}");
+							}
+
+							BranchClientPath = NewSelectedClientFileName.Substring(0, EndIdx);
+							BranchDirectoryName = new FileReference(FileRecord.ClientFile).Directory.ParentDirectory?.ParentDirectory;
+							break;
+						}
+					}
+				}
+				if(BranchClientPath == null || BranchDirectoryName == null)
+				{
+					throw new UserErrorException($"Could not find engine in Perforce relative to project path ({NewSelectedClientFileName})");
+				}
+
+				Logger.LogInformation("Found branch root at {RootPath}", BranchClientPath);
+
+				// Read the existing workspace settings from disk, and update them with any info computed here
+				int BranchIdx = BranchClientPath.IndexOf('/', 2);
+				string BranchPath = (BranchIdx == -1) ? String.Empty : BranchClientPath.Substring(BranchIdx);
+				string ProjectPath = NewSelectedClientFileName.Substring(BranchClientPath.Length);
+				UserWorkspaceSettings UserWorkspaceSettings = UserSettings.FindOrAddWorkspaceSettings(BranchDirectoryName, PerforceSettings.ServerAndPort, PerforceSettings.UserName, PerforceSettings.ClientName!, BranchPath, ProjectPath);
+
+				// Now compute the updated project info
+				ProjectInfo ProjectInfo = await ProjectInfo.CreateAsync(PerforceClient, UserWorkspaceSettings, CancellationToken);
+
+				// Update the cached workspace state
+				UserWorkspaceState UserWorkspaceState = UserSettings.FindOrAddWorkspaceState(ProjectInfo, UserWorkspaceSettings);
+
+				// Read the initial config file
+				List<KeyValuePair<FileReference, DateTime>> LocalConfigFiles = new List<KeyValuePair<FileReference, DateTime>>();
+				ConfigFile LatestProjectConfigFile = await ConfigUtils.ReadProjectConfigFileAsync(PerforceClient, ProjectInfo, LocalConfigFiles, Logger, CancellationToken);
+
+				// Get the local config file and stream filter
+				ConfigFile WorkspaceProjectConfigFile = await WorkspaceUpdate.ReadProjectConfigFile(BranchDirectoryName, NewSelectedFileName, Logger);
+				IReadOnlyList<string>? WorkspaceProjectStreamFilter = await WorkspaceUpdate.ReadProjectStreamFilter(PerforceClient, WorkspaceProjectConfigFile, Logger, CancellationToken);
+
+				OpenProjectInfo WorkspaceSettings = new OpenProjectInfo(SelectedProject, PerforceSettings, ProjectInfo, UserWorkspaceSettings, UserWorkspaceState, LatestProjectConfigFile, WorkspaceProjectConfigFile, WorkspaceProjectStreamFilter, LocalConfigFiles);
+
+				return WorkspaceSettings;
+			}
+			finally
+			{
+				PerforceClient?.Dispose();
+			}
+		}
+
+		static async Task<List<IPerforceSettings>> FilterClients(List<ClientsRecord> Clients, FileReference NewSelectedFileName, IPerforceSettings DefaultPerforceSettings, string? HostName, ILogger Logger, CancellationToken CancellationToken)
+		{
+			List<IPerforceSettings> CandidateClients = new List<IPerforceSettings>();
+			foreach(ClientsRecord Client in Clients)
+			{
+				// Make sure the client is well formed
+				if(!String.IsNullOrEmpty(Client.Name) && (!String.IsNullOrEmpty(Client.Host) || !String.IsNullOrEmpty(Client.Owner)) && !String.IsNullOrEmpty(Client.Root))
+				{
+					// Require either a username or host name match
+					if((String.IsNullOrEmpty(Client.Host) || String.Compare(Client.Host, HostName, StringComparison.OrdinalIgnoreCase) == 0) && (String.IsNullOrEmpty(Client.Owner) || String.Compare(Client.Owner, DefaultPerforceSettings.UserName, StringComparison.OrdinalIgnoreCase) == 0))
+					{
+						if(!Utility.SafeIsFileUnderDirectory(NewSelectedFileName.FullName, Client.Root))
+						{
+							Logger.LogInformation("Rejecting {ClientName} due to root mismatch ({RootPath})", Client.Name, Client.Root);
+							continue;
+						}
+
+						PerforceSettings CandidateSettings = new PerforceSettings(DefaultPerforceSettings) { ClientName = Client.Name };
+						using IPerforceConnection CandidateClient = await PerforceConnection.CreateAsync(CandidateSettings, Logger);
+
+						List<PerforceResponse<WhereRecord>> WhereRecords = await CandidateClient.TryWhereAsync(NewSelectedFileName.FullName, CancellationToken).Where(x => x.Failed || !x.Data.Unmap).ToListAsync(CancellationToken);
+						if(!WhereRecords.Succeeded() || WhereRecords.Count != 1)
+						{
+							Logger.LogInformation("Rejecting {ClientName} due to file not existing in workspace", Client.Name);
+							continue;
+						}
+
+						List<PerforceResponse<FStatRecord>> Records = await CandidateClient.TryFStatAsync(FStatOptions.None, NewSelectedFileName.FullName, CancellationToken).ToListAsync(CancellationToken);
+						if (!Records.Succeeded())
+						{
+							Logger.LogInformation("Rejecting {ClientName} due to {FileName} not in depot", Client.Name, NewSelectedFileName);
+							continue;
+						}
+
+						Records.RemoveAll(x => !x.Data.IsMapped);
+						if(Records.Count == 0)
+						{
+							Logger.LogInformation("Rejecting {ClientName} due to {NumRecords} matching records", Client.Name, Records.Count);
+							continue;
+						}
+
+						Logger.LogInformation("Found valid client {ClientName}", Client.Name);
+						CandidateClients.Add(CandidateSettings);
+					}
+				}
+			}
+			return CandidateClients;
+		}
+	}
+}
diff --git a/Engine/Source/Programs/UnrealGameSync/UnrealGameSyncShared/UserSettings.cs b/Engine/Source/Programs/UnrealGameSync/UnrealGameSyncShared/UserSettings.cs
--- a/Engine/Source/Programs/UnrealGameSync/UnrealGameSyncShared/UserSettings.cs
+++ b/Engine/Source/Programs/UnrealGameSync/UnrealGameSyncShared/UserSettings.cs
@@ -121,9 +121,6 @@ namespace UnrealGameSync
 
 	public class UserSelectedProjectSettings
 	{
-		[JsonIgnore]
-		FileReference? File { get; set; }
-
 		public string? ServerAndPort { get; set; }
 		public string? UserName { get; set; }
 		public UserSelectedProjectType Type { get; set; }
@@ -238,7 +235,7 @@ namespace UnrealGameSync
 		public DirectoryReference RootDir { get; set; } = null!;
 
 		[JsonIgnore]
-		public DateTime LastModifiedTimeUtc { get; set; }
+		public long LastModifiedTimeUtc { get; set; }
 
 		// Connection settings
 		public string? ServerAndPort { get; set; }
@@ -291,7 +288,7 @@ namespace UnrealGameSync
 			if (Utility.TryLoadJson(ConfigFile, out Settings))
 			{
 				Settings.RootDir = RootDir;
-				Settings.LastModifiedTimeUtc = FileReference.GetLastWriteTimeUtc(ConfigFile);
+				Settings.LastModifiedTimeUtc = FileReference.GetLastWriteTimeUtc(ConfigFile).Ticks;
 				return true;
 			}
 			else
@@ -308,7 +305,7 @@ namespace UnrealGameSync
 			lock (SyncRoot)
 			{
 				Utility.SaveJson(ConfigFile, this);
-				LastModifiedTimeUtc = FileReference.GetLastWriteTimeUtc(ConfigFile);
+				LastModifiedTimeUtc = FileReference.GetLastWriteTimeUtc(ConfigFile).Ticks;
 			}
 		}
 
@@ -365,9 +362,14 @@ namespace UnrealGameSync
 			return State;
 		}
 
-		public void UpdateCachedProjectInfo(ProjectInfo ProjectInfo, DateTime SettingsTimeUtc)
+		public ProjectInfo CreateProjectInfo()
+		{
+			return new ProjectInfo(RootDir, ClientName, BranchPath, ProjectPath, StreamName, ProjectIdentifier, IsEnterpriseProject);
+		}
+
+		public void UpdateCachedProjectInfo(ProjectInfo ProjectInfo, long SettingsTimeUtc)
 		{
-			this.SettingsTimeUtc = SettingsTimeUtc.Ticks;
+			this.SettingsTimeUtc = SettingsTimeUtc;
 
 			RootDir = ProjectInfo.LocalRootPath;
 			ClientName = ProjectInfo.ClientName;
@@ -425,14 +427,9 @@ namespace UnrealGameSync
 
 	public class UserProjectSettings
 	{
-		[JsonIgnore]
-		public FileReference? File { get; set; }
-
 		public List<ConfigObject> BuildSteps { get; set; } = new List<ConfigObject>();
 		public FilterType FilterType { get; set; }
 		public HashSet<string> FilterBadges { get; set; } = new HashSet<string>(StringComparer.OrdinalIgnoreCase);
-
-		public void Save() => Utility.SaveJson(File!, this);
 	}
 
 	public class UserSettings
@@ -841,6 +838,11 @@ namespace UnrealGameSync
 			return State;
 		}
 
+		public UserWorkspaceState FindOrAddWorkspaceState(UserWorkspaceSettings Settings)
+		{
+			return FindOrAddWorkspaceState(Settings.RootDir, Settings.ClientName, Settings.BranchPath);
+		}
+
 		public UserWorkspaceState FindOrAddWorkspaceState(DirectoryReference RootDir, string ClientName, string BranchPath)
 		{
 			UserWorkspaceState? State;
@@ -985,25 +987,10 @@ namespace UnrealGameSync
 			return CurrentWorkspace;
 		}
 
-		public UserProjectSettings FindOrAddProjectSettings(DirectoryReference LocalRoot, string ClientProjectFileName)
+		public UserProjectSettings FindOrAddProjectSettings(ProjectInfo ProjectInfo)
 		{
-			FileReference ConfigFile = FileReference.Combine(GetConfigDir(LocalRoot), $"project_{ClientProjectFileName.Trim('/').Replace('/', '+')}.json");
-			if (!ProjectKeyToSettings.TryGetValue(ConfigFile, out UserProjectSettings? Settings))
-			{
-				if (!Utility.TryLoadJson(ConfigFile, out Settings))
-				{
-					Settings = ImportProjectSettings(ClientProjectFileName);
-					Utility.SaveJson(ConfigFile, Settings);
-				}
+			string ClientProjectFileName = Regex.Replace(ProjectInfo.ClientFileName, "^//[^/]+/", "");
 
-				Settings.File = ConfigFile;
-				ProjectKeyToSettings.Add(ConfigFile, Settings);
-			}
-			return Settings;
-		}
-
-		public UserProjectSettings ImportProjectSettings(string ClientProjectFileName)
-		{
 			// Read the project settings
 			UserProjectSettings CurrentProject = new UserProjectSettings();
 	
diff --git a/Engine/Source/Programs/UnrealGameSync/UnrealGameSyncShared/Utility.cs b/Engine/Source/Programs/UnrealGameSync/UnrealGameSyncShared/Utility.cs
--- a/Engine/Source/Programs/UnrealGameSync/UnrealGameSyncShared/Utility.cs
+++ b/Engine/Source/Programs/UnrealGameSync/UnrealGameSyncShared/Utility.cs
@@ -389,7 +389,7 @@ namespace UnrealGameSync
 			FileReference CacheFile = FileReference.Combine(CacheFolder, Digest);
 			if(FileReference.Exists(CacheFile))
 			{
-				Logger.LogInformation("Reading cached copy of {DepotFile} from {LocalFile}", DepotPath, CacheFile);
+				Logger.LogDebug("Reading cached copy of {DepotFile} from {LocalFile}", DepotPath, CacheFile);
 				string[] Lines = FileReference.ReadAllLines(CacheFile);
 				try
 				{
@@ -403,6 +403,8 @@ namespace UnrealGameSync
 			}
 			else
 			{
+				DirectoryReference.CreateDirectory(CacheFolder);
+
 				FileReference TempFile = new FileReference(String.Format("{0}.{1}.temp", CacheFile.FullName, Guid.NewGuid()));
 				PerforceResponse<PrintRecord> Response = await Perforce.TryPrintAsync(TempFile.FullName, DepotPath, CancellationToken);
 				if (!Response.Succeeded)
diff --git a/Engine/Source/Programs/UnrealGameSync/UnrealGameSyncShared/WorkspaceSettings.cs b/Engine/Source/Programs/UnrealGameSync/UnrealGameSyncShared/WorkspaceSettings.cs
deleted file mode 100644
diff --git a/Engine/Source/Programs/UnrealGameSync/UnrealGameSyncShared/WorkspaceUpdate.cs b/Engine/Source/Programs/UnrealGameSync/UnrealGameSyncShared/WorkspaceUpdate.cs
--- a/Engine/Source/Programs/UnrealGameSync/UnrealGameSyncShared/WorkspaceUpdate.cs
+++ b/Engine/Source/Programs/UnrealGameSync/UnrealGameSyncShared/WorkspaceUpdate.cs
@@ -17,6 +17,7 @@ using System.Security.Cryptography;
 using System.Text;
 using System.Text.Encodings.Web;
 using System.Text.Json;
+using System.Text.RegularExpressions;
 using System.Threading;
 using System.Threading.Tasks;
 
@@ -198,6 +199,11 @@ namespace UnrealGameSync
 		public string TelemetryProjectIdentifier => PerforceUtils.GetClientOrDepotDirectoryName(ProjectIdentifier);
 		public DirectoryReference EngineDir => DirectoryReference.Combine(LocalRootPath, "Engine");
 		public DirectoryReference? ProjectDir => ProjectPath.EndsWith(".uproject")? LocalFileName.Directory : null;
+		public DirectoryReference DataFolder => GetDataFolder(LocalRootPath);
+		public DirectoryReference CacheFolder => GetCacheFolder(LocalRootPath);
+
+		public static DirectoryReference GetDataFolder(DirectoryReference WorkspaceDir) => DirectoryReference.Combine(WorkspaceDir, ".ugs");
+		public static DirectoryReference GetCacheFolder(DirectoryReference WorkspaceDir) => DirectoryReference.Combine(WorkspaceDir, ".ugs", "cache");
 
 		public ProjectInfo(DirectoryReference LocalRootPath, string ClientName, string BranchPath, string ProjectPath, string? StreamName, string ProjectIdentifier, bool bIsEnterpriseProject)
 		{
@@ -213,6 +219,93 @@ namespace UnrealGameSync
 			this.bIsEnterpriseProject = bIsEnterpriseProject;
 		}
 
+		public static async Task<ProjectInfo> CreateAsync(IPerforceConnection PerforceClient, UserWorkspaceSettings Settings, CancellationToken CancellationToken)
+		{
+			string? StreamName = await PerforceClient.GetCurrentStreamAsync(CancellationToken);
+
+			// Get a unique name for the project that's selected. For regular branches, this can be the depot path. For streams, we want to include the stream name to encode imports.
+			string NewSelectedProjectIdentifier;
+			if (StreamName != null)
+			{
+				string ExpectedPrefix = String.Format("//{0}/", PerforceClient.Settings.ClientName);
+				if (!Settings.ClientProjectPath.StartsWith(ExpectedPrefix, StringComparison.InvariantCultureIgnoreCase))
+				{
+					throw new UserErrorException($"Unexpected client path; expected '{Settings.ClientProjectPath}' to begin with '{ExpectedPrefix}'");
+				}
+				string? StreamPrefix = await TryGetStreamPrefixAsync(PerforceClient, StreamName, CancellationToken);
+				if (StreamPrefix == null)
+				{
+					throw new UserErrorException("Unable to get stream prefix");
+				}
+				NewSelectedProjectIdentifier = String.Format("{0}/{1}", StreamPrefix, Settings.ClientProjectPath.Substring(ExpectedPrefix.Length));
+			}
+			else
+			{
+				List<PerforceResponse<WhereRecord>> Records = await PerforceClient.TryWhereAsync(Settings.ClientProjectPath, CancellationToken).Where(x => !x.Succeeded || !x.Data.Unmap).ToListAsync(CancellationToken);
+				if (!Records.Succeeded() || Records.Count != 1)
+				{
+					throw new UserErrorException($"Couldn't get depot path for {Settings.ClientProjectPath}");
+				}
+
+				NewSelectedProjectIdentifier = Records[0].Data.DepotFile;
+
+				Match Match = Regex.Match(NewSelectedProjectIdentifier, "//([^/]+)/");
+				if (Match.Success)
+				{
+					DepotRecord Depot = await PerforceClient.GetDepotAsync(Match.Groups[1].Value, CancellationToken);
+					if (Depot.Type == "stream")
+					{
+						throw new UserErrorException($"Cannot use a legacy client ({PerforceClient.Settings.ClientName}) with a stream depot ({Depot.Depot}).");
+					}
+				}
+			}
+
+			// Figure out if it's an enterprise project. Use the synced version if we have it.
+			bool bIsEnterpriseProject = false;
+			if (Settings.ClientProjectPath.EndsWith(".uproject", StringComparison.InvariantCultureIgnoreCase))
+			{
+				string Text;
+				if (FileReference.Exists(Settings.LocalProjectPath))
+				{
+					Text = FileReference.ReadAllText(Settings.LocalProjectPath);
+				}
+				else
+				{
+					PerforceResponse<PrintRecord<string[]>> ProjectLines = await PerforceClient.TryPrintLinesAsync(Settings.ClientProjectPath, CancellationToken);
+					if (!ProjectLines.Succeeded)
+					{
+						throw new UserErrorException($"Unable to get contents of {Settings.ClientProjectPath}");
+					}
+					Text = String.Join("\n", ProjectLines.Data.Contents!);
+				}
+				bIsEnterpriseProject = Utility.IsEnterpriseProjectFromText(Text);
+			}
+
+			return new ProjectInfo(Settings.RootDir, Settings.ClientName, Settings.BranchPath, Settings.ProjectPath, StreamName, NewSelectedProjectIdentifier, bIsEnterpriseProject);
+		}
+
+		static async Task<string?> TryGetStreamPrefixAsync(IPerforceConnection Perforce, string StreamName, CancellationToken CancellationToken)
+		{
+			string? CurrentStreamName = StreamName;
+			while (!String.IsNullOrEmpty(CurrentStreamName))
+			{
+				PerforceResponse<StreamRecord> Response = await Perforce.TryGetStreamAsync(CurrentStreamName, false, CancellationToken);
+				if (!Response.Succeeded)
+				{
+					return null;
+				}
+
+				StreamRecord StreamSpec = Response.Data;
+				if (StreamSpec.Type != "virtual")
+				{
+					return CurrentStreamName;
+				}
+
+				CurrentStreamName = StreamSpec.Parent;
+			}
+			return null;
+		}
+
 		public static void ValidateBranchPath(string BranchPath)
 		{
 			if (BranchPath.Length > 0 && (!BranchPath.StartsWith("/") || BranchPath.EndsWith("/")))
@@ -1588,11 +1681,11 @@ namespace UnrealGameSync
 				{
 					string[] Lines = await File.ReadAllLinesAsync(LocalConfigFile.FullName);
 					ProjectConfig.Parse(Lines);
-					Logger.LogInformation("Read config file from {FileName}", LocalConfigFile.FullName);
+					Logger.LogDebug("Read config file from {FileName}", LocalConfigFile.FullName);
 				}
 				catch (Exception Ex)
 				{
-					Logger.LogInformation(Ex, "Failed to read config file from {FileName}", LocalConfigFile.FullName);
+					Logger.LogWarning(Ex, "Failed to read config file from {FileName}", LocalConfigFile.FullName);
 				}
 			}
 			return ProjectConfig;
