UGS: Support for running UGS CLI on Mac.

diff --git a/Engine/Source/Programs/Shared/EpicGames.Perforce/EpicGames.Perforce.csproj b/Engine/Source/Programs/Shared/EpicGames.Perforce/EpicGames.Perforce.csproj
--- a/Engine/Source/Programs/Shared/EpicGames.Perforce/EpicGames.Perforce.csproj
+++ b/Engine/Source/Programs/Shared/EpicGames.Perforce/EpicGames.Perforce.csproj
@@ -53,6 +53,9 @@
     <None Include="..\..\Shared\EpicGames.Perforce.Native\bin\linux-x64\Release\EpicGames.Perforce.Native.so" Visible="false">
       <CopyToOutputDirectory>PreserveNewest</CopyToOutputDirectory>
     </None>
+    <None Include="..\..\Shared\EpicGames.Perforce.Native\bin\mac-x64\Release\EpicGames.Perforce.Native.dylib" Visible="false">
+      <CopyToOutputDirectory>PreserveNewest</CopyToOutputDirectory>
+    </None>
   </ItemGroup>
 
   <ItemGroup>
diff --git a/Engine/Source/Programs/Shared/EpicGames.Perforce/PerforceConnection.cs b/Engine/Source/Programs/Shared/EpicGames.Perforce/PerforceConnection.cs
--- a/Engine/Source/Programs/Shared/EpicGames.Perforce/PerforceConnection.cs
+++ b/Engine/Source/Programs/Shared/EpicGames.Perforce/PerforceConnection.cs
@@ -81,7 +81,7 @@ namespace EpicGames.Perforce
 		/// <returns></returns>
 		public static async Task<IPerforceConnection> CreateAsync(IPerforceSettings Settings, ILogger Logger)
 		{
-			if (Settings.PreferNativeClient && !RuntimeInformation.IsOSPlatform(OSPlatform.OSX))
+			if (Settings.PreferNativeClient)
 			{
 				return await NativePerforceConnection.CreateAsync(Settings, Logger);
 			}
diff --git a/Engine/Source/Programs/UnrealGameSync/UnrealGameSyncCmd/Program.cs b/Engine/Source/Programs/UnrealGameSync/UnrealGameSyncCmd/Program.cs
--- a/Engine/Source/Programs/UnrealGameSync/UnrealGameSyncCmd/Program.cs
+++ b/Engine/Source/Programs/UnrealGameSync/UnrealGameSyncCmd/Program.cs
@@ -12,6 +12,8 @@ using System.Linq;
 using System.Net;
 using System.Runtime.InteropServices;
 using System.Text.RegularExpressions;
+using System.Threading;
+using System.Threading.Tasks;
 using UnrealGameSync;
 
 namespace UnrealGameSyncCmd
@@ -554,7 +556,7 @@ namespace UnrealGameSyncCmd
 
 		class ConfigCommand : Command
 		{
-			public override async Task ExecuteAsync(CommandContext Context)
+			public override Task ExecuteAsync(CommandContext Context)
 			{
 				ILogger Logger = Context.Logger;
 
@@ -568,7 +570,7 @@ namespace UnrealGameSyncCmd
 					{
 						if (Editor != null)
 						{
-							await Editor.WaitForExitAsync();
+							Editor.WaitForExit();
 						}
 					}
 				}
@@ -583,6 +585,8 @@ namespace UnrealGameSyncCmd
 
 					Logger.LogInformation("Updated {ConfigFile}", Settings.ConfigFile);
 				}
+				
+				return Task.CompletedTask;
 			}
 		}
 
diff --git a/Engine/Source/Programs/UnrealGameSync/UnrealGameSyncCmd/UnrealGameSyncCmd.csproj b/Engine/Source/Programs/UnrealGameSync/UnrealGameSyncCmd/UnrealGameSyncCmd.csproj
--- a/Engine/Source/Programs/UnrealGameSync/UnrealGameSyncCmd/UnrealGameSyncCmd.csproj
+++ b/Engine/Source/Programs/UnrealGameSync/UnrealGameSyncCmd/UnrealGameSyncCmd.csproj
@@ -2,7 +2,7 @@
 
   <PropertyGroup>
     <OutputType>Exe</OutputType>
-    <TargetFramework>net6.0-windows</TargetFramework>
+    <TargetFramework>netcoreapp3.1</TargetFramework>
     <ImplicitUsings>enable</ImplicitUsings>
     <Nullable>enable</Nullable>
     <AssemblyName>ugs</AssemblyName>
