UGS: Fix exception initializing workspace on Mac. MacOS does not have concept of hidden directories, and the runtime throws an exception when trying to execute Windows code trying to update the attributes.

diff --git a/Engine/Source/Programs/UnrealGameSync/UnrealGameSyncCmd/Program.cs b/Engine/Source/Programs/UnrealGameSync/UnrealGameSyncCmd/Program.cs
--- a/Engine/Source/Programs/UnrealGameSync/UnrealGameSyncCmd/Program.cs
+++ b/Engine/Source/Programs/UnrealGameSync/UnrealGameSyncCmd/Program.cs
@@ -440,7 +440,7 @@ namespace UnrealGameSyncCmd
 				}
 
 				// Get the client info
-				ClientRecord Client = await Perforce.GetClientAsync(Options.ClientName);
+				ClientRecord Client = await Perforce.GetClientAsync(ClientName);
 				DirectoryReference ClientDir = new DirectoryReference(Client.Root);
 
 				// If a project path was specified in local syntax, try to convert it to client-relative syntax
diff --git a/Engine/Source/Programs/UnrealGameSync/UnrealGameSyncShared/UserSettings.cs b/Engine/Source/Programs/UnrealGameSync/UnrealGameSyncShared/UserSettings.cs
--- a/Engine/Source/Programs/UnrealGameSync/UnrealGameSyncShared/UserSettings.cs
+++ b/Engine/Source/Programs/UnrealGameSync/UnrealGameSyncShared/UserSettings.cs
@@ -9,6 +9,7 @@ using System.Diagnostics.CodeAnalysis;
 using System.Drawing;
 using System.IO;
 using System.Linq;
+using System.Runtime.InteropServices;
 using System.Text;
 using System.Text.Json;
 using System.Text.Json.Serialization;
@@ -800,7 +801,10 @@ namespace UnrealGameSync
 			if (!ConfigDirInfo.Exists)
 			{
 				ConfigDirInfo.Create();
-				ConfigDirInfo.Attributes = FileAttributes.Directory | FileAttributes.Hidden;
+				if (RuntimeInformation.IsOSPlatform(OSPlatform.Windows))
+				{
+					ConfigDirInfo.Attributes = FileAttributes.Directory | FileAttributes.Hidden;
+				}
 			}
 
 			return ConfigDir;
