UGS: Disable implicit using option for UGS CLI. Prevents building on Windows.

diff --git a/Engine/Source/Programs/UnrealGameSync/UnrealGameSyncCmd/UnrealGameSyncCmd.csproj b/Engine/Source/Programs/UnrealGameSync/UnrealGameSyncCmd/UnrealGameSyncCmd.csproj
--- a/Engine/Source/Programs/UnrealGameSync/UnrealGameSyncCmd/UnrealGameSyncCmd.csproj
+++ b/Engine/Source/Programs/UnrealGameSync/UnrealGameSyncCmd/UnrealGameSyncCmd.csproj
@@ -3,7 +3,6 @@
   <PropertyGroup>
     <OutputType>Exe</OutputType>
     <TargetFramework>netcoreapp3.1</TargetFramework>
-    <ImplicitUsings>enable</ImplicitUsings>
     <Nullable>enable</Nullable>
     <AssemblyName>ugs</AssemblyName>
   </PropertyGroup>
