UGS: Use standard paths on Mac and Linux for config files and log files.

Logs:

Mac: ~/Library/Logs/Unreal Engine/UnrealGameSync
Linux: ~/

Config files:

Both: ~/.config/UnrealGameSync

diff --git a/Engine/Source/Programs/UnrealGameSync/UnrealGameSyncCmd/Program.cs b/Engine/Source/Programs/UnrealGameSync/UnrealGameSyncCmd/Program.cs
--- a/Engine/Source/Programs/UnrealGameSync/UnrealGameSyncCmd/Program.cs
+++ b/Engine/Source/Programs/UnrealGameSync/UnrealGameSyncCmd/Program.cs
@@ -135,13 +135,39 @@ namespace UnrealGameSyncCmd
 
 		public static async Task<int> Main(string[] RawArgs)
 		{
-			DirectoryReference GlobalConfigFolder = DirectoryReference.Combine(DirectoryReference.GetSpecialFolder(Environment.SpecialFolder.LocalApplicationData)!, "UnrealGameSync");
-			DirectoryReference LogFolder = GlobalConfigFolder;
+			DirectoryReference GlobalConfigFolder;
+			if(RuntimeInformation.IsOSPlatform(OSPlatform.Windows))
+			{
+				GlobalConfigFolder = DirectoryReference.Combine(DirectoryReference.GetSpecialFolder(Environment.SpecialFolder.LocalApplicationData)!, "UnrealGameSync");
+			}
+			else
+			{
+				GlobalConfigFolder = DirectoryReference.Combine(DirectoryReference.GetSpecialFolder(Environment.SpecialFolder.UserProfile)!, ".config", "UnrealGameSync");
+			}
+			DirectoryReference.CreateDirectory(GlobalConfigFolder);
+
+			string LogName;
+			DirectoryReference LogFolder;
+			if(RuntimeInformation.IsOSPlatform(OSPlatform.OSX))
+			{
+				LogFolder = DirectoryReference.Combine(DirectoryReference.GetSpecialFolder(Environment.SpecialFolder.UserProfile)!, "Library", "Logs", "Unreal Engine", "UnrealGameSync");
+				LogName = "UnrealGameSync-.log";
+			}
+			else if(RuntimeInformation.IsOSPlatform(OSPlatform.Linux))
+			{
+				LogFolder = DirectoryReference.GetSpecialFolder(Environment.SpecialFolder.UserProfile)!;
+				LogName = ".ugs-.log";
+			}
+			else
+			{
+				LogFolder = GlobalConfigFolder;
+				LogName = "UnrealGameSyncCmd-.log";
+			}
 
 			Serilog.ILogger SerilogLogger = new LoggerConfiguration()
 				.Enrich.FromLogContext()
 				.WriteTo.Console(Serilog.Events.LogEventLevel.Information, outputTemplate: "{Message:lj}{NewLine}")
-				.WriteTo.File(FileReference.Combine(LogFolder, "ugs-cmd.txt").FullName, Serilog.Events.LogEventLevel.Debug, rollingInterval: RollingInterval.Day, rollOnFileSizeLimit: true, fileSizeLimitBytes: 20 * 1024 * 1024, retainedFileCountLimit: 10)
+				.WriteTo.File(FileReference.Combine(LogFolder, LogName).FullName, Serilog.Events.LogEventLevel.Debug, rollingInterval: RollingInterval.Day, rollOnFileSizeLimit: true, fileSizeLimitBytes: 20 * 1024 * 1024, retainedFileCountLimit: 10)
 				.CreateLogger();
 
 			using ILoggerFactory LoggerFactory = new Serilog.Extensions.Logging.SerilogLoggerFactory(SerilogLogger, true);
