UGS: Add a command for listing submitted changes ("ugs changes -count=123 -lines=3")

diff --git a/Engine/Source/Programs/Shared/EpicGames.Core/StringUtils.cs b/Engine/Source/Programs/Shared/EpicGames.Core/StringUtils.cs
--- a/Engine/Source/Programs/Shared/EpicGames.Core/StringUtils.cs
+++ b/Engine/Source/Programs/Shared/EpicGames.Core/StringUtils.cs
@@ -143,7 +143,7 @@ namespace EpicGames.Core
 				int MaxIdx = GetWordWrapLineEnd(Text, MinIdx, MaxWidthForLine);
 
 				int PrintMaxIdx = MaxIdx;
-				while (PrintMaxIdx > MinIdx && Text[PrintMaxIdx - 1] == ' ')
+				while (PrintMaxIdx > MinIdx && Char.IsWhiteSpace(Text[PrintMaxIdx - 1]))
 				{
 					PrintMaxIdx--;
 				}
@@ -162,13 +162,23 @@ namespace EpicGames.Core
 		/// </summary>
 		static int GetWordWrapLineEnd(string Text, int MinIdx, int MaxWidth)
 		{
-			int MaxIdx = MinIdx + MaxWidth;
-			if (MaxIdx >= Text.Length)
+			MaxWidth = Math.Min(MaxWidth, Text.Length - MinIdx);
+
+			int MaxIdx = Text.IndexOf('\n', MinIdx, MaxWidth);
+			if (MaxIdx == -1)
+			{
+				MaxIdx = MinIdx + MaxWidth;
+			}
+			else
 			{
-				return Text.Length;
+				return MaxIdx + 1;
 			}
 
-			if (Text[MaxIdx] == ' ')
+			if (MaxIdx == Text.Length)
+			{
+				return MaxIdx;
+			}
+			else if (Char.IsWhiteSpace(Text[MaxIdx]))
 			{
 				for (; ; MaxIdx++)
 				{
@@ -559,5 +569,20 @@ namespace EpicGames.Core
 
 			return Line;
 		}
+
+		/// <summary>
+		/// Truncates the given string to the maximum length, appending an elipsis if it is longer than allowed.
+		/// </summary>
+		/// <param name="Text"></param>
+		/// <param name="MaxLength"></param>
+		/// <returns></returns>
+		public static string Truncate(string Text, int MaxLength)
+		{
+			if (Text.Length > MaxLength)
+			{
+				Text = Text.Substring(0, MaxLength - 3) + "...";
+			}
+			return Text;
+		}
 	}
 }
diff --git a/Engine/Source/Programs/UnrealGameSync/UnrealGameSync/PerforceMonitor.cs b/Engine/Source/Programs/UnrealGameSync/UnrealGameSync/PerforceMonitor.cs
--- a/Engine/Source/Programs/UnrealGameSync/UnrealGameSync/PerforceMonitor.cs
+++ b/Engine/Source/Programs/UnrealGameSync/UnrealGameSync/PerforceMonitor.cs
@@ -18,59 +18,6 @@ using System.Threading.Tasks;
 
 namespace UnrealGameSync
 {
-	class PerforceChangeDetails
-	{
-		/// <summary>
-		/// Set of extensions to treat as code
-		/// </summary>
-		static readonly HashSet<string> CodeExtensions = new HashSet<string>
-		{
-			".c",
-			".cc",
-			".cpp",
-			".m",
-			".mm",
-			".rc",
-			".cs",
-			".csproj",
-			".h",
-			".hpp",
-			".inl",
-			".usf",
-			".ush",
-			".uproject",
-			".uplugin",
-			".sln"
-		};
-
-		public string Description;
-		public bool bContainsCode;
-		public bool bContainsContent;
-
-		public PerforceChangeDetails(DescribeRecord DescribeRecord)
-		{
-			Description = DescribeRecord.Description;
-
-			// Check whether the files are code or content
-			foreach(DescribeFileRecord File in DescribeRecord.Files)
-			{
-				if(CodeExtensions.Any(Extension => File.DepotFile.EndsWith(Extension, StringComparison.OrdinalIgnoreCase)))
-				{
-					bContainsCode = true;
-				}
-				else
-				{
-					bContainsContent = true;
-				}
-
-				if (bContainsCode && bContainsContent)
-				{
-					break;
-				}
-			}
-		}
-	}
-
 	interface IArchiveInfoSource
 	{
 		IReadOnlyList<IArchiveInfo> AvailableArchives { get; }
diff --git a/Engine/Source/Programs/UnrealGameSync/UnrealGameSyncCmd/Program.cs b/Engine/Source/Programs/UnrealGameSync/UnrealGameSyncCmd/Program.cs
--- a/Engine/Source/Programs/UnrealGameSync/UnrealGameSyncCmd/Program.cs
+++ b/Engine/Source/Programs/UnrealGameSync/UnrealGameSyncCmd/Program.cs
@@ -66,6 +66,10 @@ namespace UnrealGameSyncCmd
 				"ugs switch [project name|project path|stream]",
 				"Changes the active project to the one in the workspace with the given name, or switches to a new stream."
 			),
+			new CommandInfo("changes", typeof(ChangesCommand),
+				"ugs changes",
+				"List recently submitted changes to the current branch."
+			),
 			new CommandInfo("config", typeof(ConfigCommand),
 				"ugs config",
 				"Updates the configuration for the current workspace."
@@ -580,6 +584,71 @@ namespace UnrealGameSyncCmd
 			}
 		}
 
+		class ChangesCommand : Command
+		{
+			public override async Task ExecuteAsync(CommandContext Context)
+			{
+				ILogger Logger = Context.Logger;
+
+				int Count = Context.Arguments.GetIntegerOrDefault("-Count=", 10);
+				int LineCount = Context.Arguments.GetIntegerOrDefault("-Lines=", 3);
+				Context.Arguments.CheckAllArgumentsUsed(Context.Logger);
+
+				UserWorkspaceSettings Settings = ReadRequiredUserWorkspaceSettings();
+				using IPerforceConnection PerforceClient = await ConnectAsync(Settings, Context.LoggerFactory);
+
+				List<ChangesRecord> Changes = await PerforceClient.GetChangesAsync(ChangesOptions.None, Count, ChangeStatus.Submitted, $"//{Settings.ClientName}/...");
+				foreach(IEnumerable<ChangesRecord> ChangesBatch in Changes.Batch(10))
+				{
+					List<DescribeRecord> DescribeRecords = await PerforceClient.DescribeAsync(ChangesBatch.Select(x => x.Number).ToArray());
+
+					Logger.LogInformation("  Change    Type     Author          Description");
+					foreach (DescribeRecord DescribeRecord in DescribeRecords)
+					{
+						PerforceChangeDetails Details = new PerforceChangeDetails(DescribeRecord);
+
+						string Type;
+						if (Details.bContainsCode)
+						{
+							if (Details.bContainsContent)
+							{
+								Type = "Both";
+							}
+							else
+							{
+								Type = "Code";
+							}
+						}
+						else
+						{
+							if (Details.bContainsContent)
+							{
+								Type = "Content";
+							}
+							else
+							{
+								Type = "None";
+							}
+						}
+
+						string Author = StringUtils.Truncate(DescribeRecord.User, 15);
+
+						List<string> Lines = StringUtils.WordWrap(Details.Description, Math.Max(ConsoleUtils.WindowWidth - 40, 10)).ToList();
+						if (Lines.Count == 0)
+						{
+							Lines.Add(String.Empty);
+						}
+
+						Logger.LogInformation("  {Change,-9} {Type,-8} {Author,-15} {Description}", DescribeRecord.Number, Type, Author, Lines[0]);
+						for (int LineIndex = 1; LineIndex < LineCount; LineIndex++)
+						{
+							Logger.LogInformation("                                     {Description}", Lines[LineIndex]);
+						}
+					}
+				}
+			}
+		}
+
 		class ConfigCommand : Command
 		{
 			public override Task ExecuteAsync(CommandContext Context)
diff --git a/Engine/Source/Programs/UnrealGameSync/UnrealGameSyncShared/Utility.cs b/Engine/Source/Programs/UnrealGameSync/UnrealGameSyncShared/Utility.cs
--- a/Engine/Source/Programs/UnrealGameSync/UnrealGameSyncShared/Utility.cs
+++ b/Engine/Source/Programs/UnrealGameSync/UnrealGameSyncShared/Utility.cs
@@ -30,6 +30,59 @@ namespace UnrealGameSync
 		}
 	}
 
+	public class PerforceChangeDetails
+	{
+		/// <summary>
+		/// Set of extensions to treat as code
+		/// </summary>
+		static readonly HashSet<string> CodeExtensions = new HashSet<string>
+		{
+			".c",
+			".cc",
+			".cpp",
+			".m",
+			".mm",
+			".rc",
+			".cs",
+			".csproj",
+			".h",
+			".hpp",
+			".inl",
+			".usf",
+			".ush",
+			".uproject",
+			".uplugin",
+			".sln"
+		};
+
+		public string Description;
+		public bool bContainsCode;
+		public bool bContainsContent;
+
+		public PerforceChangeDetails(DescribeRecord DescribeRecord)
+		{
+			Description = DescribeRecord.Description;
+
+			// Check whether the files are code or content
+			foreach (DescribeFileRecord File in DescribeRecord.Files)
+			{
+				if (CodeExtensions.Any(Extension => File.DepotFile.EndsWith(Extension, StringComparison.OrdinalIgnoreCase)))
+				{
+					bContainsCode = true;
+				}
+				else
+				{
+					bContainsContent = true;
+				}
+
+				if (bContainsCode && bContainsContent)
+				{
+					break;
+				}
+			}
+		}
+	}
+
 	public static class Utility
 	{
 		static JsonSerializerOptions GetDefaultJsonSerializerOptions()
