UGS: Move global settings common to the CLI/GUI apps into a separate JSON file, making it easier to edit on Mac/Linux.

diff --git a/Engine/Source/Programs/UnrealGameSync/UnrealGameSync/Automation/P4Automation.cs b/Engine/Source/Programs/UnrealGameSync/UnrealGameSync/Automation/P4Automation.cs
--- a/Engine/Source/Programs/UnrealGameSync/UnrealGameSync/Automation/P4Automation.cs
+++ b/Engine/Source/Programs/UnrealGameSync/UnrealGameSync/Automation/P4Automation.cs
@@ -21,7 +21,7 @@ namespace UnrealGameSync
 			string? UserName = null;
 			string? DepotPathSettings = null;
 
-			GlobalSettings.ReadGlobalPerforceSettings(ref ServerAndPort, ref UserName, ref DepotPathSettings);
+			GlobalPerforceSettings.ReadGlobalPerforceSettings(ref ServerAndPort, ref UserName, ref DepotPathSettings);
 
 			return Utility.OverridePerforceSettings(PerforceSettings.Default, ServerAndPort, UserName);
 
diff --git a/Engine/Source/Programs/UnrealGameSync/UnrealGameSync/AutomationServer.cs b/Engine/Source/Programs/UnrealGameSync/UnrealGameSync/AutomationServer.cs
--- a/Engine/Source/Programs/UnrealGameSync/UnrealGameSync/AutomationServer.cs
+++ b/Engine/Source/Programs/UnrealGameSync/UnrealGameSync/AutomationServer.cs
@@ -190,7 +190,7 @@ namespace UnrealGameSync
 		{
 			if(PortNumber <= 0)
 			{
-				GlobalSettings.DeleteRegistryKey(Registry.CurrentUser, "Software\\Epic Games\\UnrealGameSync", "AutomationPort");
+				GlobalPerforceSettings.DeleteRegistryKey(Registry.CurrentUser, "Software\\Epic Games\\UnrealGameSync", "AutomationPort");
 			}
 			else
 			{
diff --git a/Engine/Source/Programs/UnrealGameSync/UnrealGameSync/Controls/WorkspaceControl.cs b/Engine/Source/Programs/UnrealGameSync/UnrealGameSync/Controls/WorkspaceControl.cs
--- a/Engine/Source/Programs/UnrealGameSync/UnrealGameSync/Controls/WorkspaceControl.cs
+++ b/Engine/Source/Programs/UnrealGameSync/UnrealGameSync/Controls/WorkspaceControl.cs
@@ -906,7 +906,7 @@ namespace UnrealGameSync
 				}
 			}
 
-			string[] CombinedSyncFilter = UserSettings.GetCombinedSyncFilter(GetSyncCategories(), Settings.SyncView, Settings.SyncCategories, WorkspaceSettings.SyncView, WorkspaceSettings.SyncCategoriesDict);
+			string[] CombinedSyncFilter = UserSettings.GetCombinedSyncFilter(GetSyncCategories(), Settings.Global.SyncView, Settings.Global.SyncCategories, WorkspaceSettings.SyncView, WorkspaceSettings.SyncCategoriesDict);
 
 			ConfigSection PerforceSection = PerforceMonitor.LatestProjectConfigFile.FindSection("Perforce");
 			if (PerforceSection != null)
@@ -957,11 +957,11 @@ namespace UnrealGameSync
 			{
 				Context.Options |= WorkspaceUpdateOptions.AutoResolveChanges;
 			}
-			if (WorkspaceSettings.bSyncAllProjects ?? Settings.bSyncAllProjects)
+			if (WorkspaceSettings.bSyncAllProjects ?? Settings.Global.bSyncAllProjects)
 			{
 				Context.Options |= WorkspaceUpdateOptions.SyncAllProjects | WorkspaceUpdateOptions.IncludeAllProjectsInSolution;
 			}
-			if (WorkspaceSettings.bIncludeAllProjectsInSolution ?? Settings.bIncludeAllProjectsInSolution)
+			if (WorkspaceSettings.bIncludeAllProjectsInSolution ?? Settings.Global.bIncludeAllProjectsInSolution)
 			{
 				Context.Options |= WorkspaceUpdateOptions.IncludeAllProjectsInSolution;
 			}
@@ -4944,8 +4944,8 @@ namespace UnrealGameSync
 				ExtraSafeToDeleteExtensions = "";
 			}
 
-			string[] CombinedSyncFilter = UserSettings.GetCombinedSyncFilter(GetSyncCategories(), Settings.SyncView, Settings.SyncCategories, WorkspaceSettings.SyncView, WorkspaceSettings.SyncCategoriesDict);
-			List<string> SyncPaths = WorkspaceUpdate.GetSyncPaths(Workspace.Project, WorkspaceSettings.bSyncAllProjects ?? Settings.bSyncAllProjects, CombinedSyncFilter);
+			string[] CombinedSyncFilter = UserSettings.GetCombinedSyncFilter(GetSyncCategories(), Settings.Global.SyncView, Settings.Global.SyncCategories, WorkspaceSettings.SyncView, WorkspaceSettings.SyncCategoriesDict);
+			List<string> SyncPaths = WorkspaceUpdate.GetSyncPaths(Workspace.Project, WorkspaceSettings.bSyncAllProjects ?? Settings.Global.bSyncAllProjects, CombinedSyncFilter);
 
 			CleanWorkspaceWindow.DoClean(ParentForm, PerforceSettings, BranchDirectoryName, Workspace.Project.ClientRootPath, SyncPaths, ExtraSafeToDeleteFolders.Split('\n'), ExtraSafeToDeleteExtensions.Split('\n'), ServiceProvider.GetRequiredService<ILogger<CleanWorkspaceWindow>>());
 		}
@@ -5225,13 +5225,13 @@ namespace UnrealGameSync
 
 		private void OptionsContextMenu_SyncFilter_Click(object sender, EventArgs e)
 		{
-			SyncFilter Filter = new SyncFilter(GetSyncCategories(), Settings.SyncView, Settings.SyncCategories, Settings.bSyncAllProjects, Settings.bIncludeAllProjectsInSolution, WorkspaceSettings.SyncView, WorkspaceSettings.SyncCategoriesDict, WorkspaceSettings.bSyncAllProjects, WorkspaceSettings.bIncludeAllProjectsInSolution);
+			SyncFilter Filter = new SyncFilter(GetSyncCategories(), Settings.Global.SyncView, Settings.Global.SyncCategories, Settings.Global.bSyncAllProjects, Settings.Global.bIncludeAllProjectsInSolution, WorkspaceSettings.SyncView, WorkspaceSettings.SyncCategoriesDict, WorkspaceSettings.bSyncAllProjects, WorkspaceSettings.bIncludeAllProjectsInSolution);
 			if (Filter.ShowDialog() == DialogResult.OK)
 			{
-				Settings.SyncCategories = Filter.GlobalSyncCategories;
-				Settings.SyncView = Filter.GlobalView;
-				Settings.bSyncAllProjects = Filter.bGlobalSyncAllProjects;
-				Settings.bIncludeAllProjectsInSolution = Filter.bGlobalIncludeAllProjectsInSolution;
+				Settings.Global.SyncCategories = Filter.GlobalSyncCategories;
+				Settings.Global.SyncView = Filter.GlobalView;
+				Settings.Global.bSyncAllProjects = Filter.bGlobalSyncAllProjects;
+				Settings.Global.bIncludeAllProjectsInSolution = Filter.bGlobalIncludeAllProjectsInSolution;
 				Settings.Save();
 
 				WorkspaceSettings.SyncCategoriesDict = Filter.WorkspaceSyncCategories;
diff --git a/Engine/Source/Programs/UnrealGameSync/UnrealGameSync/Forms/ApplicationSettingsWindow.cs b/Engine/Source/Programs/UnrealGameSync/UnrealGameSync/Forms/ApplicationSettingsWindow.cs
--- a/Engine/Source/Programs/UnrealGameSync/UnrealGameSync/Forms/ApplicationSettingsWindow.cs
+++ b/Engine/Source/Programs/UnrealGameSync/UnrealGameSync/Forms/ApplicationSettingsWindow.cs
@@ -76,7 +76,7 @@ namespace UnrealGameSync
 			this.ToolUpdateMonitor = ToolUpdateMonitor;
 			this.Logger = Logger;
 
-			GlobalSettings.ReadGlobalPerforceSettings(ref InitialServerAndPort, ref InitialUserName, ref InitialDepotPath);
+			GlobalPerforceSettings.ReadGlobalPerforceSettings(ref InitialServerAndPort, ref InitialUserName, ref InitialDepotPath);
 			bInitialUnstable = bUnstable;
 
 			InitialAutomationPortNumber = AutomationServer.GetPortNumber();
@@ -207,7 +207,7 @@ namespace UnrealGameSync
 				}
 
 				bRestartUnstable = UseUnstableBuildCheckBox.Checked;
-				GlobalSettings.SaveGlobalPerforceSettings(ServerAndPort, UserName, DepotPath);
+				GlobalPerforceSettings.SaveGlobalPerforceSettings(ServerAndPort, UserName, DepotPath);
 				AutomationServer.SetPortNumber(AutomationPortNumber);
 			}
 
diff --git a/Engine/Source/Programs/UnrealGameSync/UnrealGameSync/Forms/MainWindow.cs b/Engine/Source/Programs/UnrealGameSync/UnrealGameSync/Forms/MainWindow.cs
--- a/Engine/Source/Programs/UnrealGameSync/UnrealGameSync/Forms/MainWindow.cs
+++ b/Engine/Source/Programs/UnrealGameSync/UnrealGameSync/Forms/MainWindow.cs
@@ -1132,7 +1132,7 @@ namespace UnrealGameSync
 
 		int TryOpenProject(OpenProjectInfo OpenProjectInfo, int ReplaceTabIdx, OpenProjectOptions Options)
 		{
-			Logger.LogInformation("Trying to open project {Project}", WorkspaceSettings.ProjectInfo.ClientFileName);
+			Logger.LogInformation("Trying to open project {Project}", OpenProjectInfo.ProjectInfo.ClientFileName);
 
 			// Check that none of the other tabs already have it open
 			for(int TabIdx = 0; TabIdx < TabControl.GetTabCount(); TabIdx++)
@@ -1151,9 +1151,9 @@ namespace UnrealGameSync
 							}
 							return TabIdx;
 						}
-						else if(WorkspaceSettings.ProjectInfo.LocalFileName.IsUnderDirectory(Workspace.BranchDirectoryName))
+						else if(OpenProjectInfo.ProjectInfo.LocalFileName.IsUnderDirectory(Workspace.BranchDirectoryName))
 						{
-							if((Options & OpenProjectOptions.Quiet) == 0 && MessageBox.Show(String.Format("{0} is already open under {1}.\n\nWould you like to close it?", Workspace.SelectedFileName.GetFileNameWithoutExtension(), Workspace.BranchDirectoryName, WorkspaceSettings.ProjectInfo.LocalFileName.GetFileNameWithoutExtension()), "Branch already open", MessageBoxButtons.YesNo) == System.Windows.Forms.DialogResult.Yes)
+							if((Options & OpenProjectOptions.Quiet) == 0 && MessageBox.Show(String.Format("{0} is already open under {1}.\n\nWould you like to close it?", Workspace.SelectedFileName.GetFileNameWithoutExtension(), Workspace.BranchDirectoryName, OpenProjectInfo.ProjectInfo.LocalFileName.GetFileNameWithoutExtension()), "Branch already open", MessageBoxButtons.YesNo) == System.Windows.Forms.DialogResult.Yes)
 							{
 								Logger.LogInformation("  Another project already open in this workspace, tab {TabIdx}. Replacing.", TabIdx);
 								TabControl.RemoveTab(TabIdx);
@@ -1178,13 +1178,13 @@ namespace UnrealGameSync
 				if(OldWorkspace != null)
 				{
 					OldWorkspace.Hide();
-					TabControl.SetTabData(ReplaceTabIdx, new ErrorPanel(WorkspaceSettings.SelectedProject));
+					TabControl.SetTabData(ReplaceTabIdx, new ErrorPanel(OpenProjectInfo.SelectedProject));
 					OldWorkspace.Dispose();
 				}
 			}
 
 			// Now that we have the project settings, we can construct the tab
-			WorkspaceControl NewWorkspace = new WorkspaceControl(this, ApiUrl, WorkspaceSettings, ServiceProvider, Settings, OIDCTokenManager);
+			WorkspaceControl NewWorkspace = new WorkspaceControl(this, ApiUrl, OpenProjectInfo, ServiceProvider, Settings, OIDCTokenManager);
 			
 			NewWorkspace.Parent = TabPanel;
 			NewWorkspace.Dock = DockStyle.Fill;
diff --git a/Engine/Source/Programs/UnrealGameSync/UnrealGameSync/Forms/OpenProjectWindow.cs b/Engine/Source/Programs/UnrealGameSync/UnrealGameSync/Forms/OpenProjectWindow.cs
--- a/Engine/Source/Programs/UnrealGameSync/UnrealGameSync/Forms/OpenProjectWindow.cs
+++ b/Engine/Source/Programs/UnrealGameSync/UnrealGameSync/Forms/OpenProjectWindow.cs
@@ -35,7 +35,7 @@ namespace UnrealGameSync
 			InitializeComponent();
 
 			this.Settings = Settings;
-			this.DetectedProjectSettings = null;
+			this.OpenProjectInfo = null;
 			this.DefaultSettings = DefaultSettings;
 			this.ServiceProvider = ServiceProvider;
 
@@ -95,7 +95,7 @@ namespace UnrealGameSync
 			OpenProjectWindow Window = new OpenProjectWindow(Project, Settings, DefaultPerforceSettings, ServiceProvider);
 			if(Window.ShowDialog(Owner) == DialogResult.OK)
 			{
-				return Window.DetectedProjectSettings;
+				return Window.OpenProjectInfo;
 			}
 			else
 			{
@@ -289,7 +289,7 @@ namespace UnrealGameSync
 				ModalTask<OpenProjectInfo>? NewOpenProjectInfo = PerforceModalTask.Execute(this, "Opening project", "Opening project, please wait...", DefaultSettings, (x, y) => DetectSettingsAsync(x, SelectedProject, Settings, Logger, y), Logger);
 				if (NewOpenProjectInfo != null && NewOpenProjectInfo.Succeeded)
 				{
-					DetectedProjectSettings = NewOpenProjectInfo.Result;
+					OpenProjectInfo = NewOpenProjectInfo.Result;
 					DialogResult = DialogResult.OK;
 					Close();
 				}
diff --git a/Engine/Source/Programs/UnrealGameSync/UnrealGameSync/GlobalSettings.cs b/Engine/Source/Programs/UnrealGameSync/UnrealGameSync/GlobalSettings.cs
--- a/Engine/Source/Programs/UnrealGameSync/UnrealGameSync/GlobalSettings.cs
+++ b/Engine/Source/Programs/UnrealGameSync/UnrealGameSync/GlobalSettings.cs
@@ -10,7 +10,7 @@ using System.Windows.Forms;
 
 namespace UnrealGameSync
 {
-	internal class GlobalSettings
+	internal class GlobalPerforceSettings
 	{
 		public static void ReadGlobalPerforceSettings(ref string? ServerAndPort, ref string? UserName, ref string? DepotPath)
 		{
diff --git a/Engine/Source/Programs/UnrealGameSync/UnrealGameSync/Program.cs b/Engine/Source/Programs/UnrealGameSync/UnrealGameSync/Program.cs
--- a/Engine/Source/Programs/UnrealGameSync/UnrealGameSync/Program.cs
+++ b/Engine/Source/Programs/UnrealGameSync/UnrealGameSync/Program.cs
@@ -70,7 +70,7 @@ namespace UnrealGameSync
 			string? ServerAndPort = null;
 			string? UserName = null;
 			string? BaseUpdatePath = null;
-			GlobalSettings.ReadGlobalPerforceSettings(ref ServerAndPort, ref UserName, ref BaseUpdatePath);
+			GlobalPerforceSettings.ReadGlobalPerforceSettings(ref ServerAndPort, ref UserName, ref BaseUpdatePath);
 
 			List<string> RemainingArgs = new List<string>(Args);
 
@@ -143,7 +143,7 @@ namespace UnrealGameSync
 						Logger.LogInformation("Missing server settings; finding defaults.");
 						ServerAndPort ??= PerforceSettings.Default.ServerAndPort;
 						UserName ??= PerforceSettings.Default.UserName;
-						GlobalSettings.SaveGlobalPerforceSettings(ServerAndPort, UserName, BaseUpdatePath);
+						GlobalPerforceSettings.SaveGlobalPerforceSettings(ServerAndPort, UserName, BaseUpdatePath);
 					}
 
 					ILogger TelemetryLogger = LoggerProvider.CreateLogger("Telemetry");
diff --git a/Engine/Source/Programs/UnrealGameSync/UnrealGameSync/ProgramApplicationContext.cs b/Engine/Source/Programs/UnrealGameSync/UnrealGameSync/ProgramApplicationContext.cs
--- a/Engine/Source/Programs/UnrealGameSync/UnrealGameSync/ProgramApplicationContext.cs
+++ b/Engine/Source/Programs/UnrealGameSync/UnrealGameSync/ProgramApplicationContext.cs
@@ -87,7 +87,7 @@ namespace UnrealGameSync
 			MainThreadSynchronizationContext = SynchronizationContext.Current!;
 
 			// Read the user's settings
-			Settings = new UserSettings(FileReference.Combine(DataFolder, "UnrealGameSync.ini"), ServiceProvider.GetRequiredService<ILogger<UserSettings>>());
+			Settings = UserSettings.Create(DataFolder, ServiceProvider.GetRequiredService<ILogger<UserSettings>>());
 			if(!String.IsNullOrEmpty(ProjectFileName))
 			{
 				string FullProjectFileName = Path.GetFullPath(ProjectFileName);
diff --git a/Engine/Source/Programs/UnrealGameSync/UnrealGameSyncCmd/Program.cs b/Engine/Source/Programs/UnrealGameSync/UnrealGameSyncCmd/Program.cs
--- a/Engine/Source/Programs/UnrealGameSync/UnrealGameSyncCmd/Program.cs
+++ b/Engine/Source/Programs/UnrealGameSync/UnrealGameSyncCmd/Program.cs
@@ -93,9 +93,9 @@ namespace UnrealGameSyncCmd
 			public CommandLineArguments Arguments { get; }
 			public ILogger Logger { get; }
 			public ILoggerFactory LoggerFactory { get; }
-			public UserSettings UserSettings { get; }
+			public GlobalSettingsFile UserSettings { get; }
 
-			public CommandContext(CommandLineArguments Arguments, ILogger Logger, ILoggerFactory LoggerFactory, UserSettings UserSettings)
+			public CommandContext(CommandLineArguments Arguments, ILogger Logger, ILoggerFactory LoggerFactory, GlobalSettingsFile UserSettings)
 			{
 				this.Arguments = Arguments;
 				this.Logger = Logger;
@@ -140,7 +140,7 @@ namespace UnrealGameSyncCmd
 		public static async Task<int> Main(string[] RawArgs)
 		{
 			DirectoryReference GlobalConfigFolder;
-			if(RuntimeInformation.IsOSPlatform(OSPlatform.Windows))
+			if (RuntimeInformation.IsOSPlatform(OSPlatform.Windows))
 			{
 				GlobalConfigFolder = DirectoryReference.Combine(DirectoryReference.GetSpecialFolder(Environment.SpecialFolder.LocalApplicationData)!, "UnrealGameSync");
 			}
@@ -152,12 +152,12 @@ namespace UnrealGameSyncCmd
 
 			string LogName;
 			DirectoryReference LogFolder;
-			if(RuntimeInformation.IsOSPlatform(OSPlatform.OSX))
+			if (RuntimeInformation.IsOSPlatform(OSPlatform.OSX))
 			{
 				LogFolder = DirectoryReference.Combine(DirectoryReference.GetSpecialFolder(Environment.SpecialFolder.UserProfile)!, "Library", "Logs", "Unreal Engine", "UnrealGameSync");
 				LogName = "UnrealGameSync-.log";
 			}
-			else if(RuntimeInformation.IsOSPlatform(OSPlatform.Linux))
+			else if (RuntimeInformation.IsOSPlatform(OSPlatform.Linux))
 			{
 				LogFolder = DirectoryReference.GetSpecialFolder(Environment.SpecialFolder.UserProfile)!;
 				LogName = ".ugs-.log";
@@ -178,13 +178,15 @@ namespace UnrealGameSyncCmd
 			ILogger Logger = LoggerFactory.CreateLogger("Main");
 			try
 			{
-				ConfigFile ConfigFile = new ConfigFile();
+				GlobalSettingsFile Settings;
 				if (RuntimeInformation.IsOSPlatform(OSPlatform.Windows))
 				{
-					ConfigFile.TryLoad(FileReference.Combine(GlobalConfigFolder, "UnrealGameSync.ini"), Logger);
+					Settings = UserSettings.Create(FileReference.Combine(GlobalConfigFolder, "UnrealGameSync.ini"), FileReference.Combine(GlobalConfigFolder, "Global.json"), Logger);
+				}
+				else
+				{
+					Settings = GlobalSettingsFile.Create(FileReference.Combine(GlobalConfigFolder, "Global.json"));
 				}
-
-				UserSettings UserSettings = new UserSettings(null!, ConfigFile);
 
 				CommandLineArguments Args = new CommandLineArguments(RawArgs);
 
@@ -205,7 +207,7 @@ namespace UnrealGameSyncCmd
 				}
 
 				Command Instance = (Command)Activator.CreateInstance(Command.Type)!;
-				await Instance.ExecuteAsync(new CommandContext(Args, Logger, LoggerFactory, UserSettings));
+				await Instance.ExecuteAsync(new CommandContext(Args, Logger, LoggerFactory, Settings));
 				return 0;
 			}
 			catch (UserErrorException Ex)
@@ -260,7 +262,7 @@ namespace UnrealGameSyncCmd
 			return Settings;
 		}
 
-		public static async Task<UserWorkspaceState> ReadWorkspaceState(IPerforceConnection PerforceClient, UserWorkspaceSettings Settings, UserSettings UserSettings, ILogger Logger)
+		public static async Task<UserWorkspaceState> ReadWorkspaceState(IPerforceConnection PerforceClient, UserWorkspaceSettings Settings, GlobalSettingsFile UserSettings, ILogger Logger)
 		{
 			UserWorkspaceState State = UserSettings.FindOrAddWorkspaceState(Settings);
 			if (State.SettingsTimeUtc != Settings.LastModifiedTimeUtc)
@@ -295,10 +297,10 @@ namespace UnrealGameSyncCmd
 			return ConnectAsync(Settings.ServerAndPort, Settings.UserName, Settings.ClientName, LoggerFactory);
 		}
 
-		static string[] ReadSyncFilter(UserWorkspaceSettings WorkspaceSettings, UserSettings UserSettings, ConfigFile ProjectConfig)
+		static string[] ReadSyncFilter(UserWorkspaceSettings WorkspaceSettings, GlobalSettingsFile UserSettings, ConfigFile ProjectConfig)
 		{
 			Dictionary<Guid, WorkspaceSyncCategory> SyncCategories = ConfigUtils.GetSyncCategories(ProjectConfig);
-			string[] CombinedSyncFilter = UserSettings.GetCombinedSyncFilter(SyncCategories, UserSettings.SyncView, UserSettings.SyncCategories, WorkspaceSettings.SyncView, WorkspaceSettings.SyncCategoriesDict);
+			string[] CombinedSyncFilter = GlobalSettingsFile.GetCombinedSyncFilter(SyncCategories, UserSettings.Global.SyncView, UserSettings.Global.SyncCategories, WorkspaceSettings.SyncView, WorkspaceSettings.SyncCategoriesDict);
 
 			ConfigSection PerforceSection = ProjectConfig.FindSection("Perforce");
 			if (PerforceSection != null)
diff --git a/Engine/Source/Programs/UnrealGameSync/UnrealGameSyncLauncher/Forms/SettingsWindow.cs b/Engine/Source/Programs/UnrealGameSync/UnrealGameSyncLauncher/Forms/SettingsWindow.cs
--- a/Engine/Source/Programs/UnrealGameSync/UnrealGameSyncLauncher/Forms/SettingsWindow.cs
+++ b/Engine/Source/Programs/UnrealGameSync/UnrealGameSyncLauncher/Forms/SettingsWindow.cs
@@ -85,7 +85,7 @@ namespace UnrealGameSyncLauncher
 				DepotPath = null;
 			}
 
-			GlobalSettings.SaveGlobalPerforceSettings(ServerAndPort, UserName, DepotPath);
+			GlobalPerforceSettings.SaveGlobalPerforceSettings(ServerAndPort, UserName, DepotPath);
 
 			PerforceSettings PerforceSettings = new PerforceSettings(PerforceSettings.Default);
 			if (!String.IsNullOrEmpty(ServerAndPort))
@@ -107,7 +107,7 @@ namespace UnrealGameSyncLauncher
 			{
 				if(Task.Succeeded)
 				{
-					GlobalSettings.SaveGlobalPerforceSettings(ServerAndPort, UserName, DepotPath);
+					GlobalPerforceSettings.SaveGlobalPerforceSettings(ServerAndPort, UserName, DepotPath);
 					DialogResult = DialogResult.OK;
 					Close();
 				}
diff --git a/Engine/Source/Programs/UnrealGameSync/UnrealGameSyncLauncher/Program.cs b/Engine/Source/Programs/UnrealGameSync/UnrealGameSyncLauncher/Program.cs
--- a/Engine/Source/Programs/UnrealGameSync/UnrealGameSyncLauncher/Program.cs
+++ b/Engine/Source/Programs/UnrealGameSync/UnrealGameSyncLauncher/Program.cs
@@ -46,7 +46,7 @@ namespace UnrealGameSyncLauncher
 				string? ServerAndPort = null;
 				string? UserName = null;
 				string? DepotPath = DeploymentSettings.DefaultDepotPath;
-				GlobalSettings.ReadGlobalPerforceSettings(ref ServerAndPort, ref UserName, ref DepotPath);
+				GlobalPerforceSettings.ReadGlobalPerforceSettings(ref ServerAndPort, ref UserName, ref DepotPath);
 
 				// If the shift key is held down, immediately show the settings window
 				SettingsWindow.SyncAndRunDelegate SyncAndRunWrapper = (Perforce, DepotParam, bUnstableParam, LogWriter, CancellationToken) => SyncAndRun(Perforce, DepotParam, bUnstableParam, Args, InstanceMutex, LogWriter, CancellationToken);
diff --git a/Engine/Source/Programs/UnrealGameSync/UnrealGameSyncShared/GlobalSettings.cs b/Engine/Source/Programs/UnrealGameSync/UnrealGameSyncShared/GlobalSettings.cs
new file mode 100644
--- /dev/null
+++ b/Engine/Source/Programs/UnrealGameSync/UnrealGameSyncShared/GlobalSettings.cs
@@ -0,0 +1,189 @@
+// Copyright Epic Games, Inc. All Rights Reserved.
+
+using EpicGames.Core;
+using System;
+using System.Collections.Generic;
+using System.Linq;
+using System.Text;
+using System.Text.Json;
+using System.Text.Json.Serialization;
+using System.Threading.Tasks;
+
+namespace UnrealGameSync
+{
+	public class GlobalSettings
+	{
+		public string[] SyncView { get; set; } = Array.Empty<string>();
+		public List<Guid> IncludeSyncCategories { get; set; } = new List<Guid>();
+		public List<Guid> ExcludeSyncCategories { get; set; } = new List<Guid>();
+		public bool bSyncAllProjects { get; set; } = false;
+		public bool bIncludeAllProjectsInSolution { get; set; } = false;
+
+		[JsonIgnore]
+		public Dictionary<Guid, bool> SyncCategories
+		{
+			get
+			{
+				Dictionary<Guid, bool> Categories = new Dictionary<Guid, bool>();
+				foreach (Guid IncludeCategory in IncludeSyncCategories)
+				{
+					Categories[IncludeCategory] = true;
+				}
+				foreach (Guid ExcludeCategory in ExcludeSyncCategories)
+				{
+					Categories[ExcludeCategory] = false;
+				}
+				return Categories;
+			}
+			set
+			{
+				IncludeSyncCategories = value.Where(x => x.Value).Select(x => x.Key).ToList();
+				ExcludeSyncCategories = value.Where(x => !x.Value).Select(x => x.Key).ToList();
+			}
+		}
+	}
+
+	public class GlobalSettingsFile
+	{
+		public FileReference File { get; }
+		public GlobalSettings Global { get; }
+
+		public virtual UserProjectSettings FindOrAddProjectSettings(ProjectInfo ProjectInfo)
+		{
+			return new UserProjectSettings();
+		}
+
+		protected virtual void ImportWorkspaceSettings(DirectoryReference RootDir, string ClientName, string BranchPath, UserWorkspaceSettings WorkspaceSettings)
+		{
+		}
+
+		protected virtual void ImportWorkspaceState(DirectoryReference RootDir, string ClientName, string BranchPath, UserWorkspaceState WorkspaceState)
+		{
+		}
+
+		public GlobalSettingsFile(FileReference File, GlobalSettings Global)
+		{
+			this.File = File;
+			this.Global = Global;
+		}
+
+		public static GlobalSettingsFile Create(FileReference File)
+		{
+			GlobalSettings? Data;
+			if (!Utility.TryLoadJson(File, out Data))
+			{
+				Data = new GlobalSettings();
+			}
+			return new GlobalSettingsFile(File, Data);
+		}
+
+		public virtual void Save()
+		{
+			Utility.SaveJson(File, Global);
+		}
+
+		public UserWorkspaceState FindOrAddWorkspaceState(UserWorkspaceSettings Settings)
+		{
+			return FindOrAddWorkspaceState(Settings.RootDir, Settings.ClientName, Settings.BranchPath);
+		}
+
+		public UserWorkspaceState FindOrAddWorkspaceState(DirectoryReference RootDir, string ClientName, string BranchPath)
+		{
+			UserWorkspaceState? State;
+			if (!UserWorkspaceState.TryLoad(RootDir, out State))
+			{
+				State = new UserWorkspaceState();
+				ImportWorkspaceState(RootDir, ClientName, BranchPath, State);
+				State.Save();
+			}
+			return State;
+		}
+
+		public UserWorkspaceState FindOrAddWorkspaceState(ProjectInfo ProjectInfo, UserWorkspaceSettings Settings)
+		{
+			UserWorkspaceState State = FindOrAddWorkspaceState(ProjectInfo.LocalRootPath, ProjectInfo.ClientName, ProjectInfo.BranchPath);
+			if (!State.IsValid(ProjectInfo))
+			{
+				State = new UserWorkspaceState();
+			}
+			State.UpdateCachedProjectInfo(ProjectInfo, Settings.LastModifiedTimeUtc);
+			return State;
+		}
+
+		public UserWorkspaceSettings FindOrAddWorkspaceSettings(DirectoryReference RootDir, string? ServerAndPort, string? UserName, string ClientName, string BranchPath, string ProjectPath)
+		{
+			ProjectInfo.ValidateBranchPath(BranchPath);
+			ProjectInfo.ValidateProjectPath(ProjectPath);
+
+			UserWorkspaceSettings? Settings;
+			if (!UserWorkspaceSettings.TryLoad(RootDir, out Settings))
+			{
+				Settings = new UserWorkspaceSettings();
+				Settings.RootDir = RootDir;
+				ImportWorkspaceSettings(RootDir, ClientName, BranchPath, Settings);
+			}
+
+			Settings.Init(ServerAndPort, UserName, ClientName, BranchPath, ProjectPath);
+			Settings.Save();
+
+			return Settings;
+		}
+
+		public static string[] GetCombinedSyncFilter(Dictionary<Guid, WorkspaceSyncCategory> UniqueIdToFilter, string[] GlobalView, Dictionary<Guid, bool> GlobalCategoryIdToSetting, string[] WorkspaceView, Dictionary<Guid, bool> WorkspaceCategoryIdToSetting)
+		{
+			List<string> Lines = new List<string>();
+			foreach (string ViewLine in Enumerable.Concat(GlobalView, WorkspaceView).Select(x => x.Trim()).Where(x => x.Length > 0 && !x.StartsWith(";")))
+			{
+				Lines.Add(ViewLine);
+			}
+
+			HashSet<Guid> Enabled = new HashSet<Guid>();
+			foreach (WorkspaceSyncCategory Filter in UniqueIdToFilter.Values)
+			{
+				bool bEnable = Filter.bEnable;
+
+				bool bGlobalEnable;
+				if (GlobalCategoryIdToSetting.TryGetValue(Filter.UniqueId, out bGlobalEnable))
+				{
+					bEnable = bGlobalEnable;
+				}
+
+				bool bWorkspaceEnable;
+				if (WorkspaceCategoryIdToSetting.TryGetValue(Filter.UniqueId, out bWorkspaceEnable))
+				{
+					bEnable = bWorkspaceEnable;
+				}
+
+				if (bEnable)
+				{
+					EnableFilter(Filter.UniqueId, Enabled, UniqueIdToFilter);
+				}
+			}
+
+			foreach (WorkspaceSyncCategory Filter in UniqueIdToFilter.Values.OrderBy(x => x.Name))
+			{
+				if (!Enabled.Contains(Filter.UniqueId))
+				{
+					Lines.AddRange(Filter.Paths.Select(x => "-" + x.Trim()));
+				}
+			}
+
+			return Lines.ToArray();
+		}
+
+		static void EnableFilter(Guid UniqueId, HashSet<Guid> Enabled, Dictionary<Guid, WorkspaceSyncCategory> UniqueIdToFilter)
+		{
+			if (Enabled.Add(UniqueId))
+			{
+				WorkspaceSyncCategory? Category;
+				if (UniqueIdToFilter.TryGetValue(UniqueId, out Category))
+				{
+					foreach (Guid RequiresUniqueId in Category.Requires)
+					{
+						EnableFilter(RequiresUniqueId, Enabled, UniqueIdToFilter);
+					}
+				}
+			}
+		}
+	}
+}
diff --git a/Engine/Source/Programs/UnrealGameSync/UnrealGameSyncShared/UserSettings.cs b/Engine/Source/Programs/UnrealGameSync/UnrealGameSyncShared/UserSettings.cs
--- a/Engine/Source/Programs/UnrealGameSync/UnrealGameSyncShared/UserSettings.cs
+++ b/Engine/Source/Programs/UnrealGameSync/UnrealGameSyncShared/UserSettings.cs
@@ -1,7 +1,6 @@
 // Copyright Epic Games, Inc. All Rights Reserved.
 
 using EpicGames.Core;
-using EpicGames.Perforce;
 using Microsoft.Extensions.Logging;
 using System;
 using System.Collections.Generic;
@@ -11,10 +10,8 @@ using System.IO;
 using System.Linq;
 using System.Runtime.InteropServices;
 using System.Text;
-using System.Text.Json;
 using System.Text.Json.Serialization;
 using System.Text.RegularExpressions;
-using System.Threading.Tasks;
 
 namespace UnrealGameSync
 {
@@ -433,7 +430,7 @@ namespace UnrealGameSync
 		public HashSet<string> FilterBadges { get; set; } = new HashSet<string>(StringComparer.OrdinalIgnoreCase);
 	}
 
-	public class UserSettings
+	public class UserSettings : GlobalSettingsFile
 	{
 		/// <summary>
 		/// Enum that decribes which robomerge changes to show
@@ -466,10 +463,6 @@ namespace UnrealGameSync
 		public UserSelectedProjectSettings? LastProject;
 		public List<UserSelectedProjectSettings> OpenProjects;
 		public List<UserSelectedProjectSettings> RecentProjects;
-		public string[] SyncView;
-		public Dictionary<Guid, bool> SyncCategories;
-		public bool bSyncAllProjects;
-		public bool bIncludeAllProjectsInSolution;
 		public LatestChangeType SyncType;
 		public BuildConfig CompiledEditorBuildConfig; // NB: This assumes not using precompiled editor. See CurrentBuildConfig.
 		public TabLabels TabLabels;
@@ -544,19 +537,35 @@ namespace UnrealGameSync
 			return Projects;
 		}
 
-		public UserSettings(FileReference InFileName, ILogger Logger)
-			: this(InFileName, TryLoad(InFileName, Logger))
+		public static UserSettings Create(DirectoryReference SettingsDir, ILogger Logger)
 		{
+			return Create(FileReference.Combine(SettingsDir, "UnrealGameSync.ini"), FileReference.Combine(SettingsDir, "Global.json"), Logger);
 		}
 
-		static ConfigFile TryLoad(FileReference FileName, ILogger Logger)
+		public static UserSettings Create(FileReference FileName, FileReference CoreFileName, ILogger Logger)
 		{
 			ConfigFile ConfigFile = new ConfigFile();
 			ConfigFile.TryLoad(FileName, Logger);
-			return ConfigFile;
+
+			GlobalSettings? CoreSettingsData;
+			if(FileReference.Exists(CoreFileName))
+			{
+				CoreSettingsData = Utility.LoadJson<GlobalSettings>(CoreFileName);
+			}
+			else
+			{
+				CoreSettingsData = new GlobalSettings();
+				CoreSettingsData.SyncView = ConfigFile.GetValues("General.SyncFilter", new string[0]);
+				CoreSettingsData.SyncCategories = GetCategorySettings(ConfigFile.FindSection("General"), "SyncIncludedCategories", "SyncExcludedCategories");
+				CoreSettingsData.bSyncAllProjects = ConfigFile.GetValue("General.SyncAllProjects", false);
+				CoreSettingsData.bIncludeAllProjectsInSolution = ConfigFile.GetValue("General.IncludeAllProjectsInSolution", false);
+			}
+
+			return new UserSettings(FileName, ConfigFile, CoreFileName, CoreSettingsData);
 		}
 
-		public UserSettings(FileReference InFileName, ConfigFile InConfigFile)
+		public UserSettings(FileReference InFileName, ConfigFile InConfigFile, FileReference InCoreFileName, GlobalSettings InCoreSettingsData)
+			: base(InCoreFileName, InCoreSettingsData)
 		{
 			this.FileName = InFileName;
 			this.ConfigFile = InConfigFile;
@@ -609,11 +618,6 @@ namespace UnrealGameSync
 
 			OpenProjects = ReadProjectList("General.OpenProjects", "General.OpenProjectFileNames");
 			RecentProjects = ReadProjectList("General.RecentProjects", "General.OtherProjectFileNames");
-			SyncView = ConfigFile.GetValues("General.SyncFilter", new string[0]);
-			SyncCategories = GetCategorySettings(ConfigFile.FindSection("General"), "SyncIncludedCategories", "SyncExcludedCategories");
-
-			bSyncAllProjects = ConfigFile.GetValue("General.SyncAllProjects", false);
-			bIncludeAllProjectsInSolution = ConfigFile.GetValue("General.IncludeAllProjectsInSolution", false);
 			if(!Enum.TryParse(ConfigFile.GetValue("General.SyncType", ""), out SyncType))
 			{
 				SyncType = LatestChangeType.Good;
@@ -810,59 +814,8 @@ namespace UnrealGameSync
 			return ConfigDir;
 		}
 
-		public UserWorkspaceSettings FindOrAddWorkspaceSettings(DirectoryReference RootDir, string? ServerAndPort, string? UserName, string ClientName, string BranchPath, string ProjectPath)
-		{
-			ProjectInfo.ValidateBranchPath(BranchPath);
-			ProjectInfo.ValidateProjectPath(ProjectPath);
-
-			UserWorkspaceSettings? Settings;
-			if (!WorkspaceDirToSettings.TryGetValue(RootDir, out Settings))
-			{
-				if (!UserWorkspaceSettings.TryLoad(RootDir, out Settings))
-				{
-					Settings = ImportWorkspaceSettings(RootDir, ClientName, BranchPath);
-				}
-				WorkspaceDirToSettings[RootDir] = Settings;
-			}
-
-			Settings.Init(ServerAndPort, UserName, ClientName, BranchPath, ProjectPath);
-			Settings.Save();
-
-			return Settings;
-		}
-
-		public UserWorkspaceState FindOrAddWorkspaceState(ProjectInfo ProjectInfo, UserWorkspaceSettings Settings)
-		{
-			UserWorkspaceState State = FindOrAddWorkspaceState(ProjectInfo.LocalRootPath, ProjectInfo.ClientName, ProjectInfo.BranchPath);
-			if (!State.IsValid(ProjectInfo))
-			{
-				State = new UserWorkspaceState();
-			}
-			State.UpdateCachedProjectInfo(ProjectInfo, Settings.LastModifiedTimeUtc);
-			return State;
-		}
-
-		public UserWorkspaceState FindOrAddWorkspaceState(UserWorkspaceSettings Settings)
-		{
-			return FindOrAddWorkspaceState(Settings.RootDir, Settings.ClientName, Settings.BranchPath);
-		}
-
-		public UserWorkspaceState FindOrAddWorkspaceState(DirectoryReference RootDir, string ClientName, string BranchPath)
+		protected override void ImportWorkspaceState(DirectoryReference RootDir, string ClientName, string BranchPath, UserWorkspaceState CurrentWorkspace)
 		{
-			UserWorkspaceState? State;
-			if (!UserWorkspaceState.TryLoad(RootDir, out State))
-			{
-				State = ImportWorkspaceState(RootDir, ClientName, BranchPath);
-				State.Save();
-			}
-			return State;
-		}
-
-		UserWorkspaceState ImportWorkspaceState(DirectoryReference RootDir, string ClientName, string BranchPath)
-		{
-			UserWorkspaceState CurrentWorkspace = new UserWorkspaceState();
-			CurrentWorkspace.RootDir = RootDir;
-
 			// Read the workspace settings
 			ConfigSection WorkspaceSection = ConfigFile.FindSection(ClientName + BranchPath);
 			if(WorkspaceSection == null)
@@ -959,24 +912,12 @@ namespace UnrealGameSync
 					}
 				}
 			}
-
-			return CurrentWorkspace;
 		}
 
-		UserWorkspaceSettings ImportWorkspaceSettings(DirectoryReference RootDir, string ClientName, string BranchPath)
+		protected override void ImportWorkspaceSettings(DirectoryReference RootDir, string ClientName, string BranchPath, UserWorkspaceSettings CurrentWorkspace)
 		{
-			UserWorkspaceSettings CurrentWorkspace = new UserWorkspaceSettings();
-			CurrentWorkspace.RootDir = RootDir;
-
 			ConfigSection WorkspaceSection = ConfigFile.FindSection(ClientName + BranchPath);
-			if (WorkspaceSection == null)
-			{
-				CurrentWorkspace.SyncView = new string[0];
-				CurrentWorkspace.SyncCategories = new List<SyncCategory>();
-				CurrentWorkspace.bSyncAllProjects = null;
-				CurrentWorkspace.bIncludeAllProjectsInSolution = null;
-			}
-			else
+			if (WorkspaceSection != null)
 			{
 				CurrentWorkspace.SyncView = WorkspaceSection.GetValues("SyncFilter", new string[0]);
 				CurrentWorkspace.SyncCategories = GetCategorySettings(WorkspaceSection, "SyncIncludedCategories", "SyncExcludedCategories").Select(x => new SyncCategory { Id = x.Key, Enable = x.Value }).ToList();
@@ -987,11 +928,9 @@ namespace UnrealGameSync
 				int IncludeAllProjectsInSolution = WorkspaceSection.GetValue("IncludeAllProjectsInSolution", -1);
 				CurrentWorkspace.bIncludeAllProjectsInSolution = (IncludeAllProjectsInSolution == 0) ? (bool?)false : (IncludeAllProjectsInSolution == 1) ? (bool?)true : (bool?)null;
 			}
-
-			return CurrentWorkspace;
 		}
 
-		public UserProjectSettings FindOrAddProjectSettings(ProjectInfo ProjectInfo)
+		public override UserProjectSettings FindOrAddProjectSettings(ProjectInfo ProjectInfo)
 		{
 			string ClientProjectFileName = Regex.Replace(ProjectInfo.ClientFileName, "^//[^/]+/", "");
 
@@ -1013,8 +952,10 @@ namespace UnrealGameSync
 			return CurrentProject;
 		}
 
-		public void Save()
+		public override void Save()
 		{
+			base.Save();
+
 			// General settings
 			ConfigSection GeneralSection = ConfigFile.FindOrAddSection("General");
 			GeneralSection.Clear();
@@ -1038,10 +979,6 @@ namespace UnrealGameSync
 			GeneralSection.SetValues("EnabledTools", EnabledTools);
 			GeneralSection.SetValue("FilterIndex", FilterIndex);
 			GeneralSection.SetValues("RecentProjects", RecentProjects.Select(x => x.ToConfigEntry()).ToArray());
-			GeneralSection.SetValues("SyncFilter", SyncView);
-			SetCategorySettings(GeneralSection, "SyncIncludedCategories", "SyncExcludedCategories", SyncCategories);
-			GeneralSection.SetValue("SyncAllProjects", bSyncAllProjects);
-			GeneralSection.SetValue("IncludeAllProjectsInSolution", bIncludeAllProjectsInSolution);
 			GeneralSection.SetValue("SyncType", SyncType.ToString());
 
 			// Build configuration
@@ -1147,63 +1084,6 @@ namespace UnrealGameSync
 			ConfigFile.Save(FileName);
 		}
 
-		public static string[] GetCombinedSyncFilter(Dictionary<Guid, WorkspaceSyncCategory> UniqueIdToFilter, string[] GlobalView, Dictionary<Guid, bool> GlobalCategoryIdToSetting, string[] WorkspaceView, Dictionary<Guid, bool> WorkspaceCategoryIdToSetting)
-		{
-			List<string> Lines = new List<string>();
-			foreach(string ViewLine in Enumerable.Concat(GlobalView, WorkspaceView).Select(x => x.Trim()).Where(x => x.Length > 0 && !x.StartsWith(";")))
-			{
-				Lines.Add(ViewLine);
-			}
-
-			HashSet<Guid> Enabled = new HashSet<Guid>();
-			foreach (WorkspaceSyncCategory Filter in UniqueIdToFilter.Values)
-			{
-				bool bEnable = Filter.bEnable;
-
-				bool bGlobalEnable;
-				if (GlobalCategoryIdToSetting.TryGetValue(Filter.UniqueId, out bGlobalEnable))
-				{
-					bEnable = bGlobalEnable;
-				}
-
-				bool bWorkspaceEnable;
-				if (WorkspaceCategoryIdToSetting.TryGetValue(Filter.UniqueId, out bWorkspaceEnable))
-				{
-					bEnable = bWorkspaceEnable;
-				}
-
-				if(bEnable)
-				{
-					EnableFilter(Filter.UniqueId, Enabled, UniqueIdToFilter);
-				}
-			}
-
-			foreach (WorkspaceSyncCategory Filter in UniqueIdToFilter.Values.OrderBy(x => x.Name))
-			{
-				if (!Enabled.Contains(Filter.UniqueId))
-				{
-					Lines.AddRange(Filter.Paths.Select(x => "-" + x.Trim()));
-				}
-			}
-
-			return Lines.ToArray();
-		}
-
-		static void EnableFilter(Guid UniqueId, HashSet<Guid> Enabled, Dictionary<Guid, WorkspaceSyncCategory> UniqueIdToFilter)
-		{
-			if(Enabled.Add(UniqueId))
-			{
-				WorkspaceSyncCategory? Category;
-				if(UniqueIdToFilter.TryGetValue(UniqueId, out Category))
-				{
-					foreach(Guid RequiresUniqueId in Category.Requires)
-					{
-						EnableFilter(RequiresUniqueId, Enabled, UniqueIdToFilter);
-					}
-				}
-			}
-		}
-
 		[return: NotNullIfNotNull("Text")]
 		static string? EscapeText(string? Text)
 		{
diff --git a/Engine/Source/Programs/UnrealGameSync/UnrealGameSyncShared/Utility.cs b/Engine/Source/Programs/UnrealGameSync/UnrealGameSyncShared/Utility.cs
--- a/Engine/Source/Programs/UnrealGameSync/UnrealGameSyncShared/Utility.cs
+++ b/Engine/Source/Programs/UnrealGameSync/UnrealGameSyncShared/Utility.cs
@@ -88,6 +88,8 @@ namespace UnrealGameSync
 		static JsonSerializerOptions GetDefaultJsonSerializerOptions()
 		{
 			JsonSerializerOptions Options = new JsonSerializerOptions();
+			Options.AllowTrailingCommas = true;
+			Options.ReadCommentHandling = JsonCommentHandling.Skip;
 			Options.PropertyNameCaseInsensitive = true;
 			Options.Converters.Add(new JsonStringEnumConverter());
 			return Options;
@@ -117,11 +119,8 @@ namespace UnrealGameSync
 
 		public static T LoadJson<T>(FileReference File)
 		{
-			JsonSerializerOptions Options = new JsonSerializerOptions { AllowTrailingCommas = true, ReadCommentHandling = JsonCommentHandling.Skip };
-			Options.Converters.Add(new JsonStringEnumConverter());
-
 			byte[] Data = FileReference.ReadAllBytes(File);
-			return JsonSerializer.Deserialize<T>(Data, Options);
+			return JsonSerializer.Deserialize<T>(Data, DefaultJsonSerializerOptions);
 		}
 
 		public static void SaveJson<T>(FileReference File, T Object)
