UGS: Add a filter command to allow adding and remove categories to sync.

diff --git a/Engine/Source/Programs/Shared/EpicGames.Core/CommandLineArguments.cs b/Engine/Source/Programs/Shared/EpicGames.Core/CommandLineArguments.cs
--- a/Engine/Source/Programs/Shared/EpicGames.Core/CommandLineArguments.cs
+++ b/Engine/Source/Programs/Shared/EpicGames.Core/CommandLineArguments.cs
@@ -720,38 +720,11 @@ namespace EpicGames.Core
 		/// Applies these arguments to fields with the [CommandLine] attribute in the given object.
 		/// </summary>
 		/// <param name="TargetObject">The object to configure</param>
-		/// <param name="Logger">Sink for error/warning messages</param>
 		public void ApplyTo(object TargetObject)
 		{
 			ApplyTo(TargetObject, Log.Logger);
 		}
 
-		static List<ArgumentTarget> GetArgumentTargetsForType(Type? TargetType)
-		{
-			List<ArgumentTarget> Targets = new List<ArgumentTarget>();
-			while (TargetType != null && TargetType != typeof(object))
-			{
-				foreach (FieldInfo FieldInfo in TargetType.GetFields(BindingFlags.Instance | BindingFlags.GetField | BindingFlags.Public | BindingFlags.NonPublic | BindingFlags.DeclaredOnly))
-				{
-					IEnumerable<CommandLineAttribute> Attributes = FieldInfo.GetCustomAttributes<CommandLineAttribute>();
-					if (Attributes.Any())
-					{
-						Targets.Add(new ArgumentTarget(FieldInfo, FieldInfo.FieldType, FieldInfo.SetValue, FieldInfo.GetValue, Attributes.ToArray()));
-					}
-				}
-				foreach (PropertyInfo PropertyInfo in TargetType.GetProperties(BindingFlags.Instance | BindingFlags.GetProperty | BindingFlags.Public | BindingFlags.NonPublic | BindingFlags.DeclaredOnly))
-				{
-					IEnumerable<CommandLineAttribute> Attributes = PropertyInfo.GetCustomAttributes<CommandLineAttribute>();
-					if (Attributes.Any())
-					{
-						Targets.Add(new ArgumentTarget(PropertyInfo, PropertyInfo.PropertyType, PropertyInfo.SetValue, PropertyInfo.GetValue, Attributes.ToArray()));
-					}
-				}
-				TargetType = TargetType.BaseType;
-			}
-			return Targets;
-		}
-
 		/// <summary>
 		/// Applies these arguments to fields with the [CommandLine] attribute in the given object.
 		/// </summary>
@@ -853,6 +826,43 @@ namespace EpicGames.Core
 			}
 		}
 
+		/// <summary>
+		/// Applies these arguments to fields with the [CommandLine] attribute in the given object.
+		/// </summary>
+		/// <param name="Logger">Sink for error/warning messages</param>
+		public T ApplyTo<T>(ILogger Logger) where T : new()
+		{
+			T Object = new T();
+			ApplyTo(Object, Logger);
+			return Object;
+		}
+
+		static List<ArgumentTarget> GetArgumentTargetsForType(Type? TargetType)
+		{
+			List<ArgumentTarget> Targets = new List<ArgumentTarget>();
+			while (TargetType != null && TargetType != typeof(object))
+			{
+				foreach (FieldInfo FieldInfo in TargetType.GetFields(BindingFlags.Instance | BindingFlags.GetField | BindingFlags.Public | BindingFlags.NonPublic | BindingFlags.DeclaredOnly))
+				{
+					IEnumerable<CommandLineAttribute> Attributes = FieldInfo.GetCustomAttributes<CommandLineAttribute>();
+					if (Attributes.Any())
+					{
+						Targets.Add(new ArgumentTarget(FieldInfo, FieldInfo.FieldType, FieldInfo.SetValue, FieldInfo.GetValue, Attributes.ToArray()));
+					}
+				}
+				foreach (PropertyInfo PropertyInfo in TargetType.GetProperties(BindingFlags.Instance | BindingFlags.GetProperty | BindingFlags.Public | BindingFlags.NonPublic | BindingFlags.DeclaredOnly))
+				{
+					IEnumerable<CommandLineAttribute> Attributes = PropertyInfo.GetCustomAttributes<CommandLineAttribute>();
+					if (Attributes.Any())
+					{
+						Targets.Add(new ArgumentTarget(PropertyInfo, PropertyInfo.PropertyType, PropertyInfo.SetValue, PropertyInfo.GetValue, Attributes.ToArray()));
+					}
+				}
+				TargetType = TargetType.BaseType;
+			}
+			return Targets;
+		}
+
 		/// <summary>
 		/// Gets help text for the arguments of a given type
 		/// </summary>
diff --git a/Engine/Source/Programs/UnrealGameSync/UnrealGameSyncCmd/Program.cs b/Engine/Source/Programs/UnrealGameSync/UnrealGameSyncCmd/Program.cs
--- a/Engine/Source/Programs/UnrealGameSync/UnrealGameSyncCmd/Program.cs
+++ b/Engine/Source/Programs/UnrealGameSync/UnrealGameSyncCmd/Program.cs
@@ -74,6 +74,10 @@ namespace UnrealGameSyncCmd
 				"ugs config",
 				"Updates the configuration for the current workspace."
 			),
+			new CommandInfo("filter", typeof(FilterCommand),
+				"ugs filter [-reset] [-include=..] [-exclude=..] [-view=..] [-addview=..] [-removeview=..] [-global]",
+				"Displays or updates the workspace or global sync filter"
+			),
 			new CommandInfo("sync", typeof(SyncCommand),
 				"ugs sync [change|'latest'] [-build] [-only]",
 				"Syncs the current workspace to the given changelist, optionally removing all local state."
@@ -300,7 +304,7 @@ namespace UnrealGameSyncCmd
 		static string[] ReadSyncFilter(UserWorkspaceSettings WorkspaceSettings, GlobalSettingsFile UserSettings, ConfigFile ProjectConfig)
 		{
 			Dictionary<Guid, WorkspaceSyncCategory> SyncCategories = ConfigUtils.GetSyncCategories(ProjectConfig);
-			string[] CombinedSyncFilter = GlobalSettingsFile.GetCombinedSyncFilter(SyncCategories, UserSettings.Global.SyncView, UserSettings.Global.SyncCategories, WorkspaceSettings.SyncView, WorkspaceSettings.SyncCategoriesDict);
+			string[] CombinedSyncFilter = GlobalSettingsFile.GetCombinedSyncFilter(SyncCategories, UserSettings.Global.Filter, WorkspaceSettings.Filter);
 
 			ConfigSection PerforceSection = ProjectConfig.FindSection("Perforce");
 			if (PerforceSection != null)
@@ -687,6 +691,183 @@ namespace UnrealGameSyncCmd
 			}
 		}
 
+		class FilterCommand : Command
+		{
+			class FilterCommandOptions
+			{
+				[CommandLine("-Reset")]
+				public bool Reset = false;
+
+				[CommandLine("-Include=")]
+				public List<string> Include { get; set; } = new List<string>();
+
+				[CommandLine("-Exclude=")]
+				public List<string> Exclude { get; set; } = new List<string>();
+
+				[CommandLine("-View=", ListSeparator = ';')]
+				public List<string>? View { get; set; } 
+
+				[CommandLine("-AddView=", ListSeparator = ';')]
+				public List<string> AddView { get; set; } = new List<string>();
+
+				[CommandLine("-RemoveView=", ListSeparator = ';')]
+				public List<string> RemoveView { get; set; } = new List<string>();
+
+				[CommandLine("-AllProjects", Value = "true")]
+				[CommandLine("-OnlyCurrent", Value = "false")]
+				public bool? AllProjects = null;
+
+				[CommandLine("-GpfAllProjects", Value ="true")]
+				[CommandLine("-GpfOnlyCurrent", Value = "false")]
+				public bool? AllProjectsInSln = null;
+
+				[CommandLine("-Global")]
+				public bool Global { get; set; }
+			}
+
+			public override async Task ExecuteAsync(CommandContext Context)
+			{
+				ILogger Logger = Context.Logger;
+
+				UserWorkspaceSettings WorkspaceSettings = ReadRequiredUserWorkspaceSettings();
+				using IPerforceConnection PerforceClient = await ConnectAsync(WorkspaceSettings, Context.LoggerFactory);
+				UserWorkspaceState WorkspaceState = await ReadWorkspaceState(PerforceClient, WorkspaceSettings, Context.UserSettings, Logger);
+				ProjectInfo ProjectInfo = WorkspaceState.CreateProjectInfo();
+
+				ConfigFile ProjectConfig = await ConfigUtils.ReadProjectConfigFileAsync(PerforceClient, ProjectInfo, Logger, CancellationToken.None);
+				Dictionary<Guid, WorkspaceSyncCategory> SyncCategories = ConfigUtils.GetSyncCategories(ProjectConfig);
+
+				FilterSettings GlobalFilter = Context.UserSettings.Global.Filter;
+				FilterSettings WorkspaceFilter = WorkspaceSettings.Filter;
+
+				FilterCommandOptions Options = Context.Arguments.ApplyTo<FilterCommandOptions>(Logger);
+				Context.Arguments.CheckAllArgumentsUsed(Context.Logger);
+
+				if (Options.Global)
+				{
+					ApplyCommandOptions(Context.UserSettings.Global.Filter, Options, SyncCategories.Values, Logger);
+					Context.UserSettings.Save();
+				}
+				else
+				{
+					ApplyCommandOptions(WorkspaceSettings.Filter, Options, SyncCategories.Values, Logger);
+					WorkspaceSettings.Save();
+				}
+
+				Dictionary<Guid, bool> GlobalCategories = GlobalFilter.GetCategories();
+				Dictionary<Guid, bool> WorkspaceCategories = WorkspaceFilter.GetCategories();
+
+				Logger.LogInformation("Categories:");
+				foreach (WorkspaceSyncCategory SyncCategory in SyncCategories.Values)
+				{
+					bool bEnabled;
+
+					string Scope = "(Default)";
+					if (GlobalCategories.TryGetValue(SyncCategory.UniqueId, out bEnabled))
+					{
+						Scope = "(Global)";
+					}
+					else if (WorkspaceCategories.TryGetValue(SyncCategory.UniqueId, out bEnabled))
+					{
+						Scope = "(Workspace)";
+					}
+					else
+					{
+						bEnabled = SyncCategory.bEnable;
+					}
+
+					Logger.LogInformation("  {Id,30} {Enabled,3} {Scope,-9} {Name}", SyncCategory.UniqueId, bEnabled? "Yes" : "No", Scope, SyncCategory.Name);
+				}
+
+				if (GlobalFilter.View.Count > 0)
+				{
+					Logger.LogInformation("");
+					Logger.LogInformation("Global View:");
+					foreach (string Line in GlobalFilter.View)
+					{
+						Logger.LogInformation("  {Line}", Line);
+					}
+				}
+				if (WorkspaceFilter.View.Count > 0)
+				{
+					Logger.LogInformation("");
+					Logger.LogInformation("Workspace View:");
+					foreach (string Line in WorkspaceFilter.View)
+					{
+						Logger.LogInformation("  {Line}", Line);
+					}
+				}
+
+				string[] Filter = ReadSyncFilter(WorkspaceSettings, Context.UserSettings, ProjectConfig);
+
+				Logger.LogInformation("");
+				Logger.LogInformation("Combined view:");
+				foreach (string FilterLine in Filter)
+				{
+					Logger.LogInformation("  {FilterLine}", FilterLine);
+				}
+			}
+
+			static void ApplyCommandOptions(FilterSettings Settings, FilterCommandOptions CommandOptions, IEnumerable<WorkspaceSyncCategory> SyncCategories, ILogger Logger)
+			{
+				if (CommandOptions.Reset)
+				{
+					Logger.LogInformation("Resetting settings...");
+					Settings.Reset();
+				}
+
+				HashSet<Guid> IncludeCategories = new HashSet<Guid>(CommandOptions.Include.Select(x => GetCategoryId(x, SyncCategories)));
+				HashSet<Guid> ExcludeCategories = new HashSet<Guid>(CommandOptions.Exclude.Select(x => GetCategoryId(x, SyncCategories)));
+
+				Guid Id = IncludeCategories.FirstOrDefault(x => ExcludeCategories.Contains(x));
+				if (Id != Guid.Empty)
+				{
+					throw new UserErrorException("Category {Id} cannot be both included and excluded", Id);
+				}
+
+				IncludeCategories.ExceptWith(Settings.IncludeCategories);
+				Settings.IncludeCategories.AddRange(IncludeCategories);
+
+				ExcludeCategories.ExceptWith(Settings.ExcludeCategories);
+				Settings.ExcludeCategories.AddRange(ExcludeCategories);
+
+				if (CommandOptions.View != null)
+				{
+					Settings.View = CommandOptions.View;
+				}
+				if (CommandOptions.RemoveView.Count > 0)
+				{
+					HashSet<string> ViewRemove = new HashSet<string>(CommandOptions.RemoveView, StringComparer.OrdinalIgnoreCase);
+					Settings.View.RemoveAll(x => ViewRemove.Contains(x));
+				}
+				if (CommandOptions.AddView.Count > 0)
+				{
+					HashSet<string> ViewLines = new HashSet<string>(Settings.View, StringComparer.OrdinalIgnoreCase);
+					Settings.View.AddRange(CommandOptions.AddView.Where(x => !ViewLines.Contains(x)));
+				}
+
+				Settings.AllProjects = CommandOptions.AllProjects ?? Settings.AllProjects;
+				Settings.AllProjectsInSln = CommandOptions.AllProjectsInSln ?? Settings.AllProjectsInSln;
+			}
+
+			static Guid GetCategoryId(string Text, IEnumerable<WorkspaceSyncCategory> SyncCategories)
+			{
+				Guid Id;
+				if (Guid.TryParse(Text, out Id))
+				{
+					return Id;
+				}
+
+				WorkspaceSyncCategory? Category = SyncCategories.FirstOrDefault(x => x.Name.Equals(Text, StringComparison.OrdinalIgnoreCase));
+				if (Category != null)
+				{
+					return Category.UniqueId;
+				}
+
+				throw new UserErrorException("Unable to find category '{Category}'", Text);
+			}
+		}
+
 		class BuildCommand : Command
 		{
 			public override async Task ExecuteAsync(CommandContext Context)
diff --git a/Engine/Source/Programs/UnrealGameSync/UnrealGameSyncShared/GlobalSettings.cs b/Engine/Source/Programs/UnrealGameSync/UnrealGameSyncShared/GlobalSettings.cs
--- a/Engine/Source/Programs/UnrealGameSync/UnrealGameSyncShared/GlobalSettings.cs
+++ b/Engine/Source/Programs/UnrealGameSync/UnrealGameSyncShared/GlobalSettings.cs
@@ -11,38 +11,49 @@ using System.Threading.Tasks;
 
 namespace UnrealGameSync
 {
-	public class GlobalSettings
+	public class FilterSettings
 	{
-		public string[] SyncView { get; set; } = Array.Empty<string>();
-		public List<Guid> IncludeSyncCategories { get; set; } = new List<Guid>();
-		public List<Guid> ExcludeSyncCategories { get; set; } = new List<Guid>();
-		public bool bSyncAllProjects { get; set; } = false;
-		public bool bIncludeAllProjectsInSolution { get; set; } = false;
-
-		[JsonIgnore]
-		public Dictionary<Guid, bool> SyncCategories
+		public List<Guid> IncludeCategories { get; set; } = new List<Guid>();
+		public List<Guid> ExcludeCategories { get; set; } = new List<Guid>();
+		public List<string> View { get; set; } = new List<string>();
+		public bool? AllProjects { get; set; }
+		public bool? AllProjectsInSln { get; set; }
+
+		public void Reset()
+		{
+			IncludeCategories.Clear();
+			ExcludeCategories.Clear();
+			View.Clear();
+			AllProjects = null;
+			AllProjectsInSln = null;
+		}
+
+		public void SetCategories(Dictionary<Guid, bool> Categories)
 		{
-			get
+			IncludeCategories = Categories.Where(x => x.Value).Select(x => x.Key).ToList();
+			ExcludeCategories = Categories.Where(x => !x.Value).Select(x => x.Key).ToList();
+		}
+
+		public Dictionary<Guid, bool> GetCategories()
+		{
+			Dictionary<Guid, bool> Categories = new Dictionary<Guid, bool>();
+			foreach (Guid IncludeCategory in IncludeCategories)
 			{
-				Dictionary<Guid, bool> Categories = new Dictionary<Guid, bool>();
-				foreach (Guid IncludeCategory in IncludeSyncCategories)
-				{
-					Categories[IncludeCategory] = true;
-				}
-				foreach (Guid ExcludeCategory in ExcludeSyncCategories)
-				{
-					Categories[ExcludeCategory] = false;
-				}
-				return Categories;
+				Categories[IncludeCategory] = true;
 			}
-			set
+			foreach (Guid ExcludeCategory in ExcludeCategories)
 			{
-				IncludeSyncCategories = value.Where(x => x.Value).Select(x => x.Key).ToList();
-				ExcludeSyncCategories = value.Where(x => !x.Value).Select(x => x.Key).ToList();
+				Categories[ExcludeCategory] = false;
 			}
+			return Categories;
 		}
 	}
 
+	public class GlobalSettings
+	{
+		public FilterSettings Filter { get; set; } = new FilterSettings();
+	}
+
 	public class GlobalSettingsFile
 	{
 		public FileReference File { get; }
@@ -129,14 +140,17 @@ namespace UnrealGameSync
 			return Settings;
 		}
 
-		public static string[] GetCombinedSyncFilter(Dictionary<Guid, WorkspaceSyncCategory> UniqueIdToFilter, string[] GlobalView, Dictionary<Guid, bool> GlobalCategoryIdToSetting, string[] WorkspaceView, Dictionary<Guid, bool> WorkspaceCategoryIdToSetting)
+		public static string[] GetCombinedSyncFilter(Dictionary<Guid, WorkspaceSyncCategory> UniqueIdToFilter, FilterSettings GlobalFilter, FilterSettings WorkspaceFilter)
 		{
 			List<string> Lines = new List<string>();
-			foreach (string ViewLine in Enumerable.Concat(GlobalView, WorkspaceView).Select(x => x.Trim()).Where(x => x.Length > 0 && !x.StartsWith(";")))
+			foreach (string ViewLine in Enumerable.Concat(GlobalFilter.View, WorkspaceFilter.View).Select(x => x.Trim()).Where(x => x.Length > 0 && !x.StartsWith(";")))
 			{
 				Lines.Add(ViewLine);
 			}
 
+			Dictionary<Guid, bool> GlobalCategoryIdToSetting = GlobalFilter.GetCategories();
+			Dictionary<Guid, bool> WorkspaceCategoryIdToSetting = WorkspaceFilter.GetCategories();
+
 			HashSet<Guid> Enabled = new HashSet<Guid>();
 			foreach (WorkspaceSyncCategory Filter in UniqueIdToFilter.Values)
 			{
diff --git a/Engine/Source/Programs/UnrealGameSync/UnrealGameSyncShared/UserSettings.cs b/Engine/Source/Programs/UnrealGameSync/UnrealGameSyncShared/UserSettings.cs
--- a/Engine/Source/Programs/UnrealGameSync/UnrealGameSyncShared/UserSettings.cs
+++ b/Engine/Source/Programs/UnrealGameSync/UnrealGameSyncShared/UserSettings.cs
@@ -247,17 +247,7 @@ namespace UnrealGameSync
 		public string ProjectPath { get; set; } = String.Empty;
 
 		// Workspace specific SyncFilters
-		public string[] SyncView { get; set; } = Array.Empty<string>();
-		public List<SyncCategory> SyncCategories { get; set; } = new List<SyncCategory>();
-		public bool? bSyncAllProjects { get; set; }
-		public bool? bIncludeAllProjectsInSolution { get; set; }
-
-		[JsonIgnore]
-		public Dictionary<Guid, bool> SyncCategoriesDict
-		{
-			get => SyncCategories.ToDictionary(x => x.Id, x => x.Enable);
-			set => SyncCategories = value.Select(x => new SyncCategory { Id = x.Key, Enable = x.Value }).ToList();
-		}
+		public FilterSettings Filter { get; set; } = new FilterSettings();
 
 		[JsonIgnore]
 		public FileReference ConfigFile => GetConfigFile(RootDir);
@@ -555,10 +545,10 @@ namespace UnrealGameSync
 			else
 			{
 				CoreSettingsData = new GlobalSettings();
-				CoreSettingsData.SyncView = ConfigFile.GetValues("General.SyncFilter", new string[0]);
-				CoreSettingsData.SyncCategories = GetCategorySettings(ConfigFile.FindSection("General"), "SyncIncludedCategories", "SyncExcludedCategories");
-				CoreSettingsData.bSyncAllProjects = ConfigFile.GetValue("General.SyncAllProjects", false);
-				CoreSettingsData.bIncludeAllProjectsInSolution = ConfigFile.GetValue("General.IncludeAllProjectsInSolution", false);
+				CoreSettingsData.Filter.View = ConfigFile.GetValues("General.SyncFilter", new string[0]).ToList();
+				CoreSettingsData.Filter.SetCategories(GetCategorySettings(ConfigFile.FindSection("General"), "SyncIncludedCategories", "SyncExcludedCategories"));
+				CoreSettingsData.Filter.AllProjects = ConfigFile.GetValue("General.SyncAllProjects", false);
+				CoreSettingsData.Filter.AllProjectsInSln = ConfigFile.GetValue("General.IncludeAllProjectsInSolution", false);
 			}
 
 			return new UserSettings(FileName, ConfigFile, CoreFileName, CoreSettingsData);
@@ -919,14 +909,14 @@ namespace UnrealGameSync
 			ConfigSection WorkspaceSection = ConfigFile.FindSection(ClientName + BranchPath);
 			if (WorkspaceSection != null)
 			{
-				CurrentWorkspace.SyncView = WorkspaceSection.GetValues("SyncFilter", new string[0]);
-				CurrentWorkspace.SyncCategories = GetCategorySettings(WorkspaceSection, "SyncIncludedCategories", "SyncExcludedCategories").Select(x => new SyncCategory { Id = x.Key, Enable = x.Value }).ToList();
+				CurrentWorkspace.Filter.View = WorkspaceSection.GetValues("SyncFilter", new string[0]).ToList();
+				CurrentWorkspace.Filter.SetCategories(GetCategorySettings(WorkspaceSection, "SyncIncludedCategories", "SyncExcludedCategories"));
 
 				int SyncAllProjects = WorkspaceSection.GetValue("SyncAllProjects", -1);
-				CurrentWorkspace.bSyncAllProjects = (SyncAllProjects == 0) ? (bool?)false : (SyncAllProjects == 1) ? (bool?)true : (bool?)null;
+				CurrentWorkspace.Filter.AllProjects = (SyncAllProjects == 0) ? (bool?)false : (SyncAllProjects == 1) ? (bool?)true : (bool?)null;
 
 				int IncludeAllProjectsInSolution = WorkspaceSection.GetValue("IncludeAllProjectsInSolution", -1);
-				CurrentWorkspace.bIncludeAllProjectsInSolution = (IncludeAllProjectsInSolution == 0) ? (bool?)false : (IncludeAllProjectsInSolution == 1) ? (bool?)true : (bool?)null;
+				CurrentWorkspace.Filter.AllProjectsInSln = (IncludeAllProjectsInSolution == 0) ? (bool?)false : (IncludeAllProjectsInSolution == 1) ? (bool?)true : (bool?)null;
 			}
 		}
 
