UGS: Build fixes for GUI client.

diff --git a/Engine/Source/Programs/UnrealGameSync/UnrealGameSync/Controls/WorkspaceControl.cs b/Engine/Source/Programs/UnrealGameSync/UnrealGameSync/Controls/WorkspaceControl.cs
--- a/Engine/Source/Programs/UnrealGameSync/UnrealGameSync/Controls/WorkspaceControl.cs
+++ b/Engine/Source/Programs/UnrealGameSync/UnrealGameSync/Controls/WorkspaceControl.cs
@@ -906,7 +906,7 @@ namespace UnrealGameSync
 				}
 			}
 
-			string[] CombinedSyncFilter = UserSettings.GetCombinedSyncFilter(GetSyncCategories(), Settings.Global.SyncView, Settings.Global.SyncCategories, WorkspaceSettings.SyncView, WorkspaceSettings.SyncCategoriesDict);
+			string[] CombinedSyncFilter = UserSettings.GetCombinedSyncFilter(GetSyncCategories(), Settings.Global.Filter, WorkspaceSettings.Filter);
 
 			ConfigSection PerforceSection = PerforceMonitor.LatestProjectConfigFile.FindSection("Perforce");
 			if (PerforceSection != null)
@@ -957,11 +957,11 @@ namespace UnrealGameSync
 			{
 				Context.Options |= WorkspaceUpdateOptions.AutoResolveChanges;
 			}
-			if (WorkspaceSettings.bSyncAllProjects ?? Settings.Global.bSyncAllProjects)
+			if (WorkspaceSettings.Filter.AllProjects ?? Settings.Global.Filter.AllProjects ?? false)
 			{
 				Context.Options |= WorkspaceUpdateOptions.SyncAllProjects | WorkspaceUpdateOptions.IncludeAllProjectsInSolution;
 			}
-			if (WorkspaceSettings.bIncludeAllProjectsInSolution ?? Settings.Global.bIncludeAllProjectsInSolution)
+			if (WorkspaceSettings.Filter.AllProjectsInSln ?? Settings.Global.Filter.AllProjectsInSln ?? false)
 			{
 				Context.Options |= WorkspaceUpdateOptions.IncludeAllProjectsInSolution;
 			}
@@ -4944,8 +4944,8 @@ namespace UnrealGameSync
 				ExtraSafeToDeleteExtensions = "";
 			}
 
-			string[] CombinedSyncFilter = UserSettings.GetCombinedSyncFilter(GetSyncCategories(), Settings.Global.SyncView, Settings.Global.SyncCategories, WorkspaceSettings.SyncView, WorkspaceSettings.SyncCategoriesDict);
-			List<string> SyncPaths = WorkspaceUpdate.GetSyncPaths(Workspace.Project, WorkspaceSettings.bSyncAllProjects ?? Settings.Global.bSyncAllProjects, CombinedSyncFilter);
+			string[] CombinedSyncFilter = UserSettings.GetCombinedSyncFilter(GetSyncCategories(), Settings.Global.Filter, WorkspaceSettings.Filter);
+			List<string> SyncPaths = WorkspaceUpdate.GetSyncPaths(Workspace.Project, WorkspaceSettings.Filter.AllProjects ?? Settings.Global.Filter.AllProjects ?? false, CombinedSyncFilter);
 
 			CleanWorkspaceWindow.DoClean(ParentForm, PerforceSettings, BranchDirectoryName, Workspace.Project.ClientRootPath, SyncPaths, ExtraSafeToDeleteFolders.Split('\n'), ExtraSafeToDeleteExtensions.Split('\n'), ServiceProvider.GetRequiredService<ILogger<CleanWorkspaceWindow>>());
 		}
@@ -5225,19 +5225,13 @@ namespace UnrealGameSync
 
 		private void OptionsContextMenu_SyncFilter_Click(object sender, EventArgs e)
 		{
-			SyncFilter Filter = new SyncFilter(GetSyncCategories(), Settings.Global.SyncView, Settings.Global.SyncCategories, Settings.Global.bSyncAllProjects, Settings.Global.bIncludeAllProjectsInSolution, WorkspaceSettings.SyncView, WorkspaceSettings.SyncCategoriesDict, WorkspaceSettings.bSyncAllProjects, WorkspaceSettings.bIncludeAllProjectsInSolution);
+			SyncFilter Filter = new SyncFilter(GetSyncCategories(), Settings.Global.Filter, WorkspaceSettings.Filter);
 			if (Filter.ShowDialog() == DialogResult.OK)
 			{
-				Settings.Global.SyncCategories = Filter.GlobalSyncCategories;
-				Settings.Global.SyncView = Filter.GlobalView;
-				Settings.Global.bSyncAllProjects = Filter.bGlobalSyncAllProjects;
-				Settings.Global.bIncludeAllProjectsInSolution = Filter.bGlobalIncludeAllProjectsInSolution;
+				Settings.Global.Filter = Filter.GlobalFilter;
 				Settings.Save();
 
-				WorkspaceSettings.SyncCategoriesDict = Filter.WorkspaceSyncCategories;
-				WorkspaceSettings.SyncView = Filter.WorkspaceView;
-				WorkspaceSettings.bSyncAllProjects = Filter.bWorkspaceSyncAllProjects;
-				WorkspaceSettings.bIncludeAllProjectsInSolution = Filter.bWorkspaceIncludeAllProjectsInSolution;
+				WorkspaceSettings.Filter = Filter.WorkspaceFilter;
 				WorkspaceSettings.Save();
 			}
 		}
diff --git a/Engine/Source/Programs/UnrealGameSync/UnrealGameSync/Forms/SyncFilter.cs b/Engine/Source/Programs/UnrealGameSync/UnrealGameSync/Forms/SyncFilter.cs
--- a/Engine/Source/Programs/UnrealGameSync/UnrealGameSync/Forms/SyncFilter.cs
+++ b/Engine/Source/Programs/UnrealGameSync/UnrealGameSync/Forms/SyncFilter.cs
@@ -15,42 +15,30 @@ namespace UnrealGameSync
 	partial class SyncFilter : Form
 	{
 		Dictionary<Guid, WorkspaceSyncCategory> UniqueIdToCategory;
-		public string[] GlobalView;
-		public Dictionary<Guid, bool> GlobalSyncCategories;
-		public bool bGlobalSyncAllProjects;
-		public bool bGlobalIncludeAllProjectsInSolution;
-		public string[] WorkspaceView;
-		public Dictionary<Guid, bool> WorkspaceSyncCategories;
-		public bool? bWorkspaceSyncAllProjects;
-		public bool? bWorkspaceIncludeAllProjectsInSolution;
-
-		public SyncFilter(Dictionary<Guid, WorkspaceSyncCategory> InUniqueIdToCategory, string[] InGlobalView, Dictionary<Guid, bool> InGlobalSyncCategories, bool bInGlobalProjectOnly, bool bInGlobalIncludeAllProjectsInSolution, string[] InWorkspaceView, Dictionary<Guid, bool> InWorkspaceSyncCategories, bool? bInWorkspaceProjectOnly, bool? bInWorkspaceIncludeAllProjectsInSolution)
+		public FilterSettings GlobalFilter;
+		public FilterSettings WorkspaceFilter;
+
+		public SyncFilter(Dictionary<Guid, WorkspaceSyncCategory> InUniqueIdToCategory, FilterSettings InGlobalFilter, FilterSettings InWorkspaceFilter)
 		{
 			InitializeComponent();
 
 			UniqueIdToCategory = InUniqueIdToCategory;
-			GlobalSyncCategories = InGlobalSyncCategories;
-			GlobalView = InGlobalView;
-			bGlobalSyncAllProjects = bInGlobalProjectOnly;
-			bGlobalIncludeAllProjectsInSolution = bInGlobalIncludeAllProjectsInSolution;
-			WorkspaceSyncCategories = InWorkspaceSyncCategories;
-			WorkspaceView = InWorkspaceView;
-			bWorkspaceSyncAllProjects = bInWorkspaceProjectOnly;
-			bWorkspaceIncludeAllProjectsInSolution = bInWorkspaceIncludeAllProjectsInSolution;
+			GlobalFilter = InGlobalFilter;
+			WorkspaceFilter = InWorkspaceFilter;
 
 			Dictionary<Guid, bool> SyncCategories = WorkspaceSyncCategory.GetDefault(UniqueIdToCategory.Values);
 
-			WorkspaceSyncCategory.ApplyDelta(SyncCategories, InGlobalSyncCategories);
-			GlobalControl.SetView(GlobalView);
+			WorkspaceSyncCategory.ApplyDelta(SyncCategories, GlobalFilter.GetCategories());
+			GlobalControl.SetView(GlobalFilter.View.ToArray());
 			SetExcludedCategories(GlobalControl.CategoriesCheckList, UniqueIdToCategory, SyncCategories);
-			GlobalControl.SyncAllProjects.Checked = bGlobalSyncAllProjects;
-			GlobalControl.IncludeAllProjectsInSolution.Checked = bGlobalIncludeAllProjectsInSolution;
+			GlobalControl.SyncAllProjects.Checked = GlobalFilter.AllProjects ?? false;
+			GlobalControl.IncludeAllProjectsInSolution.Checked = GlobalFilter.AllProjectsInSln ?? false;
 
-			WorkspaceSyncCategory.ApplyDelta(SyncCategories, InWorkspaceSyncCategories);
-			WorkspaceControl.SetView(WorkspaceView);
+			WorkspaceSyncCategory.ApplyDelta(SyncCategories, WorkspaceFilter.GetCategories());
+			WorkspaceControl.SetView(WorkspaceFilter.View.ToArray());
 			SetExcludedCategories(WorkspaceControl.CategoriesCheckList, UniqueIdToCategory, SyncCategories);
-			WorkspaceControl.SyncAllProjects.Checked = bWorkspaceSyncAllProjects ?? bGlobalSyncAllProjects;
-			WorkspaceControl.IncludeAllProjectsInSolution.Checked = bWorkspaceIncludeAllProjectsInSolution ?? bGlobalIncludeAllProjectsInSolution;
+			WorkspaceControl.SyncAllProjects.Checked = WorkspaceFilter.AllProjects ?? GlobalFilter.AllProjects ?? false;
+			WorkspaceControl.IncludeAllProjectsInSolution.Checked = WorkspaceFilter.AllProjectsInSln ?? GlobalFilter.AllProjectsInSln ?? false;
 
 			GlobalControl.CategoriesCheckList.ItemCheck += GlobalControl_CategoriesCheckList_ItemCheck;
 			GlobalControl.SyncAllProjects.CheckStateChanged += GlobalControl_SyncAllProjects_CheckStateChanged;
@@ -89,15 +77,25 @@ namespace UnrealGameSync
 			}
 		}
 
-		private void GetExcludedCategories(out Dictionary<Guid, bool> NewGlobalSyncCategories, out Dictionary<Guid, bool> NewWorkspaceSyncCategories)
+		private void GetSettings(out FilterSettings NewGlobalFilter, out FilterSettings NewWorkspaceFilter)
 		{
 			Dictionary<Guid, bool> DefaultSyncCategories = WorkspaceSyncCategory.GetDefault(UniqueIdToCategory.Values);
 
-			Dictionary<Guid, bool> GlobalSyncCategories = GetCategorySettings(GlobalControl.CategoriesCheckList, this.GlobalSyncCategories);
-			NewGlobalSyncCategories = WorkspaceSyncCategory.GetDelta(DefaultSyncCategories, GlobalSyncCategories);
+			NewGlobalFilter = new FilterSettings();
+			NewGlobalFilter.View = GlobalControl.GetView().ToList();
+			NewGlobalFilter.AllProjects = GlobalControl.SyncAllProjects.Checked;
+			NewGlobalFilter.AllProjectsInSln = GlobalControl.IncludeAllProjectsInSolution.Checked;
+
+			Dictionary<Guid, bool> GlobalSyncCategories = GetCategorySettings(GlobalControl.CategoriesCheckList, GlobalFilter.GetCategories());
+			NewGlobalFilter.SetCategories(WorkspaceSyncCategory.GetDelta(DefaultSyncCategories, GlobalSyncCategories));
 
-			Dictionary<Guid, bool> WorkspaceSyncCategories = GetCategorySettings(WorkspaceControl.CategoriesCheckList, this.WorkspaceSyncCategories);
-			NewWorkspaceSyncCategories = WorkspaceSyncCategory.GetDelta(GlobalSyncCategories, WorkspaceSyncCategories);
+			NewWorkspaceFilter = new FilterSettings();
+			NewWorkspaceFilter.View = WorkspaceControl.GetView().ToList();
+			NewWorkspaceFilter.AllProjects = (WorkspaceControl.SyncAllProjects.Checked == NewGlobalFilter.AllProjects) ? (bool?)null : WorkspaceControl.SyncAllProjects.Checked;
+			NewWorkspaceFilter.AllProjectsInSln = (WorkspaceControl.IncludeAllProjectsInSolution.Checked == NewGlobalFilter.AllProjectsInSln) ? (bool?)null : WorkspaceControl.IncludeAllProjectsInSolution.Checked;
+
+			Dictionary<Guid, bool> WorkspaceSyncCategories = GetCategorySettings(WorkspaceControl.CategoriesCheckList, WorkspaceFilter.GetCategories());
+			NewWorkspaceFilter.SetCategories(WorkspaceSyncCategory.GetDelta(GlobalSyncCategories, WorkspaceSyncCategories));
 		}
 
 		private Dictionary<Guid, bool> GetCategorySettings(CheckedListBox ListBox, IEnumerable<KeyValuePair<Guid, bool>> OriginalSettings)
@@ -133,26 +131,18 @@ namespace UnrealGameSync
 
 		private void OkButton_Click(object sender, EventArgs e)
 		{
-			string[] NewGlobalView = GlobalControl.GetView();
-			string[] NewWorkspaceView = WorkspaceControl.GetView();
+			GetSettings(out FilterSettings NewGlobalFilter, out FilterSettings NewWorkspaceFilter);
 
-			if(NewGlobalView.Any(x => x.Contains("//")) || NewWorkspaceView.Any(x => x.Contains("//")))
+			if(NewGlobalFilter.View.Any(x => x.Contains("//")) || NewWorkspaceFilter.View.Any(x => x.Contains("//")))
 			{
 				if(MessageBox.Show(this, "Custom views should be relative to the stream root (eg. -/Engine/...).\r\n\r\nFull depot paths (eg. //depot/...) will not match any files.\r\n\r\nAre you sure you want to continue?", "Invalid view", MessageBoxButtons.OKCancel) != System.Windows.Forms.DialogResult.OK)
 				{
 					return;
 				}
 			}
-		
-			GlobalView = NewGlobalView;
-			bGlobalSyncAllProjects = GlobalControl.SyncAllProjects.Checked;
-			bGlobalIncludeAllProjectsInSolution = GlobalControl.IncludeAllProjectsInSolution.Checked;
-
-			WorkspaceView = NewWorkspaceView;
-			bWorkspaceSyncAllProjects = (WorkspaceControl.SyncAllProjects.Checked == bGlobalSyncAllProjects)? (bool?)null : WorkspaceControl.SyncAllProjects.Checked;
-			bWorkspaceIncludeAllProjectsInSolution = (WorkspaceControl.IncludeAllProjectsInSolution.Checked == bGlobalIncludeAllProjectsInSolution)? (bool?)null : WorkspaceControl.IncludeAllProjectsInSolution.Checked;
 
-			GetExcludedCategories(out GlobalSyncCategories, out WorkspaceSyncCategories);
+			GlobalFilter = NewGlobalFilter;
+			WorkspaceFilter = NewWorkspaceFilter;
 
 			DialogResult = DialogResult.OK;
 		}
@@ -164,11 +154,9 @@ namespace UnrealGameSync
 
 		private void ShowCombinedView_Click(object sender, EventArgs e)
 		{
-			Dictionary<Guid, bool> NewGlobalSyncCategories;
-			Dictionary<Guid, bool> NewWorkspaceSyncCategories;
-			GetExcludedCategories(out NewGlobalSyncCategories, out NewWorkspaceSyncCategories);
+			GetSettings(out FilterSettings NewGlobalFilter, out FilterSettings NewWorkspaceFilter);
 
-			string[] Filter = UserSettings.GetCombinedSyncFilter(UniqueIdToCategory, GlobalControl.GetView(), NewGlobalSyncCategories, WorkspaceControl.GetView(), NewWorkspaceSyncCategories);
+			string[] Filter = UserSettings.GetCombinedSyncFilter(UniqueIdToCategory, NewGlobalFilter, NewWorkspaceFilter);
 			if(Filter.Length == 0)
 			{
 				Filter = new string[]{ "All files will be synced." };
