UGS: Fix null reference exception when running filter command on new workspace.

diff --git a/Engine/Source/Programs/UnrealGameSync/UnrealGameSyncShared/GlobalSettings.cs b/Engine/Source/Programs/UnrealGameSync/UnrealGameSyncShared/GlobalSettings.cs
--- a/Engine/Source/Programs/UnrealGameSync/UnrealGameSyncShared/GlobalSettings.cs
+++ b/Engine/Source/Programs/UnrealGameSync/UnrealGameSyncShared/GlobalSettings.cs
@@ -104,6 +104,7 @@ namespace UnrealGameSync
 			if (!UserWorkspaceState.TryLoad(RootDir, out State))
 			{
 				State = new UserWorkspaceState();
+				State.RootDir = RootDir;
 				ImportWorkspaceState(RootDir, ClientName, BranchPath, State);
 				State.Save();
 			}
