UGS: Fix missing error output when syncing.

diff --git a/Engine/Source/Programs/UnrealGameSync/UnrealGameSyncCmd/Program.cs b/Engine/Source/Programs/UnrealGameSync/UnrealGameSyncCmd/Program.cs
--- a/Engine/Source/Programs/UnrealGameSync/UnrealGameSyncCmd/Program.cs
+++ b/Engine/Source/Programs/UnrealGameSync/UnrealGameSyncCmd/Program.cs
@@ -586,7 +586,11 @@ namespace UnrealGameSyncCmd
 				WorkspaceUpdateContext UpdateContext = new WorkspaceUpdateContext(Change, Options, BuildConfig.Development, SyncFilter, ProjectSettings.BuildSteps, null);
 
 				WorkspaceUpdate Update = new WorkspaceUpdate(UpdateContext);
-				await Update.ExecuteAsync(PerforceClient.Settings, ProjectInfo, State, Context.Logger, CancellationToken.None);
+				(WorkspaceUpdateResult Result, string Message) = await Update.ExecuteAsync(PerforceClient.Settings, ProjectInfo, State, Context.Logger, CancellationToken.None);
+				if(Result != WorkspaceUpdateResult.Success)
+				{
+					Logger.LogError("{Message}", Message);
+				}
 			}
 		}
 
