UGS: Fix path to shell on non-Windows platforms.

diff --git a/Engine/Source/Programs/UnrealGameSync/UnrealGameSyncShared/WorkspaceUpdate.cs b/Engine/Source/Programs/UnrealGameSync/UnrealGameSyncShared/WorkspaceUpdate.cs
--- a/Engine/Source/Programs/UnrealGameSync/UnrealGameSyncShared/WorkspaceUpdate.cs
+++ b/Engine/Source/Programs/UnrealGameSync/UnrealGameSyncShared/WorkspaceUpdate.cs
@@ -20,6 +20,7 @@ using System.Text.Json;
 using System.Text.RegularExpressions;
 using System.Threading;
 using System.Threading.Tasks;
+using System.Runtime.InteropServices;
 
 namespace UnrealGameSync
 {
@@ -550,7 +551,15 @@ namespace UnrealGameSync
 		{
 			using IPerforceConnection Perforce = await PerforceConnection.CreateAsync(PerforceSettings, Logger);
 
-			string CmdExe = Path.Combine(Environment.GetFolderPath(Environment.SpecialFolder.System), "cmd.exe");
+			string CmdExe;
+			if(RuntimeInformation.IsOSPlatform(OSPlatform.Windows))
+			{
+				CmdExe = Path.Combine(Environment.GetFolderPath(Environment.SpecialFolder.System), "cmd.exe");
+			}
+			else
+			{
+				CmdExe = "/bin/sh";
+			}
 			if (!File.Exists(CmdExe))
 			{
 				return (WorkspaceUpdateResult.FailedToSync, $"Missing {CmdExe}.");
