UGS: Fix some Windows-isms in cross-platform UGS library.

diff --git a/Engine/Source/Programs/UnrealGameSync/UnrealGameSyncCmd/Program.cs b/Engine/Source/Programs/UnrealGameSync/UnrealGameSyncCmd/Program.cs
--- a/Engine/Source/Programs/UnrealGameSync/UnrealGameSyncCmd/Program.cs
+++ b/Engine/Source/Programs/UnrealGameSync/UnrealGameSyncCmd/Program.cs
@@ -539,6 +539,10 @@ namespace UnrealGameSyncCmd
 
 				[CommandLine("-Build")]
 				public bool Build { get; set; }
+
+				[CommandLine("-NoGPF", Value = "false")]
+				[CommandLine("-NoProjectFiles", Value = "false")]
+				public bool ProjectFiles { get; set; } = true;
 			}
 
 			public override async Task ExecuteAsync(CommandContext Context)
@@ -576,6 +580,10 @@ namespace UnrealGameSyncCmd
 				{
 					Options |= WorkspaceUpdateOptions.Build;
 				}
+				if (SyncOptions.ProjectFiles)
+				{
+					Options |= WorkspaceUpdateOptions.GenerateProjectFiles;
+				}
 
 				ProjectInfo ProjectInfo = State.CreateProjectInfo();
 				UserProjectSettings ProjectSettings = Context.UserSettings.FindOrAddProjectSettings(ProjectInfo);
diff --git a/Engine/Source/Programs/UnrealGameSync/UnrealGameSyncShared/ConfigUtils.cs b/Engine/Source/Programs/UnrealGameSync/UnrealGameSyncShared/ConfigUtils.cs
--- a/Engine/Source/Programs/UnrealGameSync/UnrealGameSyncShared/ConfigUtils.cs
+++ b/Engine/Source/Programs/UnrealGameSync/UnrealGameSyncShared/ConfigUtils.cs
@@ -285,13 +285,27 @@ namespace UnrealGameSync
 
 			bool bUseCrashReportClientEditor = LatestProjectConfigFile.GetValue("Options.UseCrashReportClientEditor", false);
 
+			string HostPlatform;
+			if(RuntimeInformation.IsOSPlatform(OSPlatform.OSX))
+			{
+				HostPlatform = "Mac";
+			}
+			else if(RuntimeInformation.IsOSPlatform(OSPlatform.Linux))
+			{
+				HostPlatform = "Linux";
+			}
+			else
+			{
+				HostPlatform = "Win64";
+			}
+
 			List<BuildStep> DefaultBuildSteps = new List<BuildStep>();
-			DefaultBuildSteps.Add(new BuildStep(new Guid("{01F66060-73FA-4CC8-9CB3-E217FBBA954E}"), 0, "Compile UnrealHeaderTool", "Compiling UnrealHeaderTool...", 1, "UnrealHeaderTool", "Win64", "Development", "", !ShouldSyncPrecompiledEditor));
-			DefaultBuildSteps.Add(new BuildStep(new Guid("{F097FF61-C916-4058-8391-35B46C3173D5}"), 1, $"Compile {EditorTarget}", $"Compiling {EditorTarget}...", 10, EditorTarget, "Win64", EditorConfig.ToString(), ProjectArgument, !ShouldSyncPrecompiledEditor));
-			DefaultBuildSteps.Add(new BuildStep(new Guid("{C6E633A1-956F-4AD3-BC95-6D06D131E7B4}"), 2, "Compile ShaderCompileWorker", "Compiling ShaderCompileWorker...", 1, "ShaderCompileWorker", "Win64", "Development", "", !ShouldSyncPrecompiledEditor));
-			DefaultBuildSteps.Add(new BuildStep(new Guid("{24FFD88C-7901-4899-9696-AE1066B4B6E8}"), 3, "Compile UnrealLightmass", "Compiling UnrealLightmass...", 1, "UnrealLightmass", "Win64", "Development", "", !ShouldSyncPrecompiledEditor));
-			DefaultBuildSteps.Add(new BuildStep(new Guid("{FFF20379-06BF-4205-8A3E-C53427736688}"), 4, "Compile CrashReportClient", "Compiling CrashReportClient...", 1, "CrashReportClient", "Win64", "Shipping", "", !ShouldSyncPrecompiledEditor && !bUseCrashReportClientEditor));
-			DefaultBuildSteps.Add(new BuildStep(new Guid("{7143D861-58D3-4F83-BADC-BC5DCB2079F6}"), 5, "Compile CrashReportClientEditor", "Compiling CrashReportClientEditor...", 1, "CrashReportClientEditor", "Win64", "Shipping", "", !ShouldSyncPrecompiledEditor && bUseCrashReportClientEditor));
+			DefaultBuildSteps.Add(new BuildStep(new Guid("{01F66060-73FA-4CC8-9CB3-E217FBBA954E}"), 0, "Compile UnrealHeaderTool", "Compiling UnrealHeaderTool...", 1, "UnrealHeaderTool", HostPlatform, "Development", "", !ShouldSyncPrecompiledEditor));
+			DefaultBuildSteps.Add(new BuildStep(new Guid("{F097FF61-C916-4058-8391-35B46C3173D5}"), 1, $"Compile {EditorTarget}", $"Compiling {EditorTarget}...", 10, EditorTarget, HostPlatform, EditorConfig.ToString(), ProjectArgument, !ShouldSyncPrecompiledEditor));
+			DefaultBuildSteps.Add(new BuildStep(new Guid("{C6E633A1-956F-4AD3-BC95-6D06D131E7B4}"), 2, "Compile ShaderCompileWorker", "Compiling ShaderCompileWorker...", 1, "ShaderCompileWorker", HostPlatform, "Development", "", !ShouldSyncPrecompiledEditor));
+			DefaultBuildSteps.Add(new BuildStep(new Guid("{24FFD88C-7901-4899-9696-AE1066B4B6E8}"), 3, "Compile UnrealLightmass", "Compiling UnrealLightmass...", 1, "UnrealLightmass", HostPlatform, "Development", "", !ShouldSyncPrecompiledEditor));
+			DefaultBuildSteps.Add(new BuildStep(new Guid("{FFF20379-06BF-4205-8A3E-C53427736688}"), 4, "Compile CrashReportClient", "Compiling CrashReportClient...", 1, "CrashReportClient", HostPlatform, "Shipping", "", !ShouldSyncPrecompiledEditor && !bUseCrashReportClientEditor));
+			DefaultBuildSteps.Add(new BuildStep(new Guid("{7143D861-58D3-4F83-BADC-BC5DCB2079F6}"), 5, "Compile CrashReportClientEditor", "Compiling CrashReportClientEditor...", 1, "CrashReportClientEditor", HostPlatform, "Shipping", "", !ShouldSyncPrecompiledEditor && bUseCrashReportClientEditor));
 
 			return DefaultBuildSteps.ToDictionary(x => x.UniqueId, x => x.ToConfigObject());
 		}
diff --git a/Engine/Source/Programs/UnrealGameSync/UnrealGameSyncShared/WorkspaceUpdate.cs b/Engine/Source/Programs/UnrealGameSync/UnrealGameSyncShared/WorkspaceUpdate.cs
--- a/Engine/Source/Programs/UnrealGameSync/UnrealGameSyncShared/WorkspaceUpdate.cs
+++ b/Engine/Source/Programs/UnrealGameSync/UnrealGameSyncShared/WorkspaceUpdate.cs
@@ -547,23 +547,25 @@ namespace UnrealGameSync
 			public void Dispose() => Release();
 		}
 
-		public async Task<(WorkspaceUpdateResult, string)> ExecuteAsync(IPerforceSettings PerforceSettings, ProjectInfo Project, UserWorkspaceState State, ILogger Logger, CancellationToken CancellationToken)
-		{
-			using IPerforceConnection Perforce = await PerforceConnection.CreateAsync(PerforceSettings, Logger);
+		public static string ShellScriptExt { get; }= RuntimeInformation.IsOSPlatform(OSPlatform.Windows)? "bat" : "sh";
 
-			string CmdExe;
+		public Task<int> ExecuteShellCommandAsync(string CommandLine, string? WorkingDir, Action<string> ProcessOutput, CancellationToken CancellationToken)
+		{
 			if(RuntimeInformation.IsOSPlatform(OSPlatform.Windows))
 			{
-				CmdExe = Path.Combine(Environment.GetFolderPath(Environment.SpecialFolder.System), "cmd.exe");
+				string CmdExe = Path.Combine(Environment.GetFolderPath(Environment.SpecialFolder.System), "cmd.exe");
+				return Utility.ExecuteProcessAsync(CmdExe, null, $"/C \"{CommandLine}\"", ProcessOutput, CancellationToken);
 			}
 			else
 			{
-				CmdExe = "/bin/sh";
-			}
-			if (!File.Exists(CmdExe))
-			{
-				return (WorkspaceUpdateResult.FailedToSync, $"Missing {CmdExe}.");
+				string ShellExe = "/bin/sh";
+				return Utility.ExecuteProcessAsync(ShellExe, null, $"{CommandLine}", ProcessOutput, CancellationToken);
 			}
+		}
+
+		public async Task<(WorkspaceUpdateResult, string)> ExecuteAsync(IPerforceSettings PerforceSettings, ProjectInfo Project, UserWorkspaceState State, ILogger Logger, CancellationToken CancellationToken)
+		{
+			using IPerforceConnection Perforce = await PerforceConnection.CreateAsync(PerforceSettings, Logger);
 
 			List<Tuple<string, TimeSpan>> Times = new List<Tuple<string, TimeSpan>>();
 
@@ -1138,7 +1140,7 @@ namespace UnrealGameSync
 					Progress.Set("Generating project files...", 0.0f);
 
 					StringBuilder CommandLine = new StringBuilder();
-					CommandLine.AppendFormat("/C \"\"{0}\"", FileReference.Combine(Project.LocalRootPath, "GenerateProjectFiles.bat"));
+					CommandLine.AppendFormat("\"{0}\"", FileReference.Combine(Project.LocalRootPath, $"GenerateProjectFiles.{ShellScriptExt}"));
 					if ((Context.Options & WorkspaceUpdateOptions.SyncAllProjects) == 0 && (Context.Options & WorkspaceUpdateOptions.IncludeAllProjectsInSolution) == 0)
 					{
 						if (Project.LocalFileName.HasExtension(".uproject"))
@@ -1146,12 +1148,11 @@ namespace UnrealGameSync
 							CommandLine.AppendFormat(" \"{0}\"", Project.LocalFileName);
 						}
 					}
-					CommandLine.Append(" -progress\"");
-
+					CommandLine.Append(" -progress");
 					Logger.LogInformation("Generating project files...");
-					Logger.LogInformation("gpf> Running {FileName} {Arguments}", CmdExe, CommandLine);
+					Logger.LogInformation("gpf> Running {Arguments}", CommandLine);
 
-					int GenerateProjectFilesResult = await Utility.ExecuteProcessAsync(CmdExe, null, CommandLine.ToString(), Line => ProcessOutput(Line, "gpf> ", Progress, Logger), CancellationToken);
+					int GenerateProjectFilesResult = await ExecuteShellCommandAsync(CommandLine.ToString(), null, Line => ProcessOutput(Line, "gpf> ", Progress, Logger), CancellationToken);
 					if (GenerateProjectFilesResult != 0)
 					{
 						return (WorkspaceUpdateResult.FailedToCompile, $"Failed to generate project files (exit code {GenerateProjectFilesResult}).");
@@ -1231,6 +1232,16 @@ namespace UnrealGameSync
 
 						Logger.LogInformation("{Status}", Step.StatusText);
 
+						DirectoryReference BatchFilesDir = DirectoryReference.Combine(Project.LocalRootPath, "Engine", "Build", "BatchFiles");
+						if(RuntimeInformation.IsOSPlatform(OSPlatform.OSX))
+						{
+							BatchFilesDir = DirectoryReference.Combine(BatchFilesDir, "Mac");
+						}
+						else if(RuntimeInformation.IsOSPlatform(OSPlatform.Linux))
+						{
+							BatchFilesDir = DirectoryReference.Combine(BatchFilesDir, "Linux");
+						}
+
 						if (Step.IsValid())
 						{
 							// Get the build variables for this step
@@ -1249,17 +1260,18 @@ namespace UnrealGameSync
 									{
 										StepStopwatch.AddData(new { Target = Step.Target });
 
-										FileReference BuildBat = FileReference.Combine(Project.LocalRootPath, "Engine", "Build", "BatchFiles", "Build.bat");
-										string CommandLine = String.Format("{0} {1} {2} {3} -NoHotReloadFromIDE", Step.Target, Step.Platform, Step.Configuration, Utility.ExpandVariables(Step.Arguments ?? "", Variables));
+										FileReference BuildBat = FileReference.Combine(BatchFilesDir, $"Build.{ShellScriptExt}");
+
+										string CommandLine = $"\"{BuildBat}\" {Step.Target} {Step.Platform} {Step.Configuration} {Utility.ExpandVariables(Step.Arguments ?? "", Variables)} -NoHotReloadFromIDE";
 										if (Context.Options.HasFlag(WorkspaceUpdateOptions.Clean) || bForceClean)
 										{
-											Logger.LogInformation("ubt> Running {FileName} {Arguments}", BuildBat, CommandLine + " -clean");
-											await Utility.ExecuteProcessAsync(CmdExe, null, "/C \"\"" + BuildBat + "\" " + CommandLine + " -clean\"", Line => ProcessOutput(Line, "ubt> ", Progress, Logger), CancellationToken);
+											Logger.LogInformation("ubt> Running {Arguments}", CommandLine + " -clean");
+											await ExecuteShellCommandAsync(CommandLine + " -clean", null, Line => ProcessOutput(Line, "ubt> ", Progress, Logger), CancellationToken);
 										}
 
 										Logger.LogInformation("ubt> Running {FileName} {Arguments}", BuildBat, CommandLine + " -progress");
 
-										int ResultFromBuild = await Utility.ExecuteProcessAsync(CmdExe, null, "/C \"\"" + BuildBat + "\" " + CommandLine + " -progress\"", Line => ProcessOutput(Line, "ubt> ", Progress, Logger), CancellationToken);
+										int ResultFromBuild = await ExecuteShellCommandAsync(CommandLine + " -progress", null, Line => ProcessOutput(Line, "ubt> ", Progress, Logger), CancellationToken);
 										if (ResultFromBuild != 0)
 										{
 											StepStopwatch.Stop("Failed");
@@ -1285,11 +1297,11 @@ namespace UnrealGameSync
 									{
 										StepStopwatch.AddData(new { Project = Path.GetFileNameWithoutExtension(Step.FileName) });
 
-										FileReference LocalRunUAT = FileReference.Combine(Project.LocalRootPath, "Engine", "Build", "BatchFiles", "RunUAT.bat");
-										string Arguments = String.Format("/C \"\"{0}\" -profile=\"{1}\"\"", LocalRunUAT, FileReference.Combine(Project.LocalRootPath, Step.FileName ?? "unknown"));
+										FileReference LocalRunUAT = FileReference.Combine(BatchFilesDir, $"RunUAT.{ShellScriptExt}");
+										string Arguments = String.Format("\"{0}\" -profile=\"{1}\"", LocalRunUAT, FileReference.Combine(Project.LocalRootPath, Step.FileName ?? "unknown"));
 										Logger.LogInformation("uat> Running {FileName} {Argument}", LocalRunUAT, Arguments);
 
-										int ResultFromUAT = await Utility.ExecuteProcessAsync(CmdExe, null, Arguments, Line => ProcessOutput(Line, "uat> ", Progress, Logger), CancellationToken);
+										int ResultFromUAT = await ExecuteShellCommandAsync(Arguments, null, Line => ProcessOutput(Line, "uat> ", Progress, Logger), CancellationToken);
 										if (ResultFromUAT != 0)
 										{
 											StepStopwatch.Stop("Failed");
