UGS: Add a command to list all clients and their root directories on the current host ("ugs clients") and a command to launch the editor ("ugs run"). Also bypass prompt for removing files from the workspace when running through CLI.

diff --git a/Engine/Source/Programs/UnrealGameSync/UnrealGameSyncCmd/Program.cs b/Engine/Source/Programs/UnrealGameSync/UnrealGameSyncCmd/Program.cs
--- a/Engine/Source/Programs/UnrealGameSync/UnrealGameSyncCmd/Program.cs
+++ b/Engine/Source/Programs/UnrealGameSync/UnrealGameSyncCmd/Program.cs
@@ -8,9 +8,11 @@ using System;
 using System.Collections.Generic;
 using System.Diagnostics;
 using System.Globalization;
+using System.IO;
 using System.Linq;
 using System.Net;
 using System.Runtime.InteropServices;
+using System.Text;
 using System.Text.RegularExpressions;
 using System.Threading;
 using System.Threading.Tasks;
@@ -40,6 +42,8 @@ namespace UnrealGameSyncCmd
 
 	public class Program
 	{
+		static BuildConfig EditorConfig => BuildConfig.Development;
+
 		class CommandInfo
 		{
 			public string Name { get; }
@@ -79,9 +83,17 @@ namespace UnrealGameSyncCmd
 				"Displays or updates the workspace or global sync filter"
 			),
 			new CommandInfo("sync", typeof(SyncCommand),
-				"ugs sync [change|'latest'] [-build] [-only]",
+				"ugs sync [change|'latest'] [-build] [-remove] [-only]",
 				"Syncs the current workspace to the given changelist, optionally removing all local state."
 			),
+			new CommandInfo("clients", typeof(ClientsCommand),
+				"ugs clients",
+				"Lists all clients suitable for use on the current machine."
+			),
+			new CommandInfo("run", typeof(RunCommand),
+				"ugs run",
+				"Runs the editor for the current branch."
+			),
 			new CommandInfo("build", typeof(BuildCommand),
 				"ugs build [id] [-list]",
 				"Runs the default build steps for the current project, or a particular step referenced by id."
@@ -108,14 +120,17 @@ namespace UnrealGameSyncCmd
 			}
 		}
 
-		class ProjectConfigOptions
+		class ServerOptions
 		{
 			[CommandLine("-Server=")]
 			public string? ServerAndPort { get; set; }
 
 			[CommandLine("-User=")]
 			public string? UserName { get; set; }
+		}
 
+		class ProjectConfigOptions : ServerOptions
+		{
 			public void ApplyTo(UserWorkspaceSettings Settings)
 			{
 				if (ServerAndPort != null)
@@ -584,6 +599,7 @@ namespace UnrealGameSyncCmd
 				{
 					Options |= WorkspaceUpdateOptions.GenerateProjectFiles;
 				}
+				Options |= WorkspaceUpdateOptions.RemoveFilteredFiles;
 
 				ProjectInfo ProjectInfo = State.CreateProjectInfo();
 				UserProjectSettings ProjectSettings = Context.UserSettings.FindOrAddProjectSettings(ProjectInfo);
@@ -597,7 +613,86 @@ namespace UnrealGameSyncCmd
 				(WorkspaceUpdateResult Result, string Message) = await Update.ExecuteAsync(PerforceClient.Settings, ProjectInfo, State, Context.Logger, CancellationToken.None);
 				if(Result != WorkspaceUpdateResult.Success)
 				{
-					Logger.LogError("{Message}", Message);
+					Logger.LogError("{Message} (Result: {Result})", Message, Result);
+				}
+			}
+		}
+
+		class ClientsCommand : Command
+		{
+			public class ClientsOptions : ServerOptions
+			{
+			}
+
+			public override async Task ExecuteAsync(CommandContext Context)
+			{
+				ILogger Logger = Context.Logger;
+
+				ClientsOptions Options = Context.Arguments.ApplyTo<ClientsOptions>(Logger);
+				Context.Arguments.CheckAllArgumentsUsed();
+
+				using IPerforceConnection PerforceClient = await ConnectAsync(Options.ServerAndPort, Options.UserName, null, Context.LoggerFactory);
+				InfoRecord Info = await PerforceClient.GetInfoAsync(InfoOptions.ShortOutput);
+
+				List<ClientsRecord> Clients = await PerforceClient.GetClientsAsync(EpicGames.Perforce.ClientsOptions.None, PerforceClient.Settings.UserName);
+				foreach (ClientsRecord Client in Clients)
+				{
+					if (String.Equals(Info.ClientHost, Client.Host, StringComparison.OrdinalIgnoreCase))
+					{
+						Logger.LogInformation("{Client,-50} {Root}", Client.Name, Client.Root);
+					}
+				}
+			}
+		}
+
+		class RunCommand : Command
+		{
+			public override async Task ExecuteAsync(CommandContext Context)
+			{
+				ILogger Logger = Context.Logger;
+
+				UserWorkspaceSettings Settings = ReadRequiredUserWorkspaceSettings();
+				using IPerforceConnection PerforceClient = await ConnectAsync(Settings, Context.LoggerFactory);
+				UserWorkspaceState State = await ReadWorkspaceState(PerforceClient, Settings, Context.UserSettings, Logger);
+
+				ProjectInfo ProjectInfo = State.CreateProjectInfo();
+				ConfigFile ProjectConfig = await ConfigUtils.ReadProjectConfigFileAsync(PerforceClient, ProjectInfo, Logger, CancellationToken.None);
+
+				FileReference ReceiptFile = ConfigUtils.GetEditorReceiptFile(ProjectInfo, ProjectConfig, EditorConfig);
+				Logger.LogDebug("Receipt file: {Receipt}", ReceiptFile);
+
+				if (!ConfigUtils.TryReadEditorReceipt(ProjectInfo, ReceiptFile, out TargetReceipt? Receipt) || String.IsNullOrEmpty(Receipt.Launch))
+				{
+					throw new UserErrorException("The editor needs to be built before you can run it. (Missing {ReceiptFile}).", ReceiptFile);
+				}
+				if (!File.Exists(Receipt.Launch))
+				{
+					throw new UserErrorException("The editor needs to be built before you can run it. (Missing {LaunchFile}).", Receipt.Launch);
+				}
+
+				List<string> LaunchArguments = new List<string>();
+				if (Settings.LocalProjectPath.HasExtension(".uproject"))
+				{
+					LaunchArguments.Add($"\"{Settings.LocalProjectPath}\"");
+				}
+				if (EditorConfig == BuildConfig.Debug || EditorConfig == BuildConfig.DebugGame)
+				{
+					LaunchArguments.Append(" -debug");
+				}
+				for (int Idx = 0; Idx < Context.Arguments.Count; Idx++)
+				{
+					if (!Context.Arguments.HasBeenUsed(Idx))
+					{
+						LaunchArguments.Add(Context.Arguments[Idx]);
+					}
+				}
+
+				string CommandLine = CommandLineArguments.Join(LaunchArguments);
+				Logger.LogInformation("Spawning: {LaunchFile} {CommandLine}", CommandLineArguments.Quote(Receipt.Launch), CommandLine);
+
+				if (!Utility.SpawnProcess(Receipt.Launch, CommandLine))
+				{
+					Logger.LogError("Unable to spawn {0} {1}", Receipt.Launch, LaunchArguments.ToString());
 				}
 			}
 		}
@@ -901,7 +996,7 @@ namespace UnrealGameSyncCmd
 
 					FileReference EditorTarget = ConfigUtils.GetEditorTargetFile(ProjectInfo, ProjectConfig);
 
-					Dictionary<Guid, ConfigObject> BuildStepObjects = ConfigUtils.GetDefaultBuildStepObjects(ProjectInfo, EditorTarget.GetFileNameWithoutAnyExtensions(), BuildConfig.Development, ProjectConfig, false);
+					Dictionary<Guid, ConfigObject> BuildStepObjects = ConfigUtils.GetDefaultBuildStepObjects(ProjectInfo, EditorTarget.GetFileNameWithoutAnyExtensions(), EditorConfig, ProjectConfig, false);
 
 					Logger.LogInformation("Available build steps:");
 					Logger.LogInformation("");
diff --git a/Engine/Source/Programs/UnrealGameSync/UnrealGameSyncShared/WorkspaceUpdate.cs b/Engine/Source/Programs/UnrealGameSync/UnrealGameSyncShared/WorkspaceUpdate.cs
--- a/Engine/Source/Programs/UnrealGameSync/UnrealGameSyncShared/WorkspaceUpdate.cs
+++ b/Engine/Source/Programs/UnrealGameSync/UnrealGameSyncShared/WorkspaceUpdate.cs
@@ -41,6 +41,7 @@ namespace UnrealGameSync
 		UpdateFilter = 0x800,
 		SyncAllProjects = 0x1000,
 		IncludeAllProjectsInSolution = 0x2000,
+		RemoveFilteredFiles = 0x4000,
 	}
 
 	public enum WorkspaceUpdateResult
@@ -668,7 +669,7 @@ namespace UnrealGameSync
 							}
 
 							// Check if there are any paths outside the regular sync paths
-							if (RemoveDepotPaths.Count > 0)
+							if (RemoveDepotPaths.Count > 0 && (Context.Options & WorkspaceUpdateOptions.RemoveFilteredFiles) == 0)
 							{
 								bool bDeleteListMatches = true;
 
