UGS: Add a warning when attempting to add the same file to the sync filter more than once.

diff --git a/Engine/Source/Programs/UnrealGameSync/UnrealGameSyncShared/WorkspaceUpdate.cs b/Engine/Source/Programs/UnrealGameSync/UnrealGameSyncShared/WorkspaceUpdate.cs
--- a/Engine/Source/Programs/UnrealGameSync/UnrealGameSyncShared/WorkspaceUpdate.cs
+++ b/Engine/Source/Programs/UnrealGameSync/UnrealGameSyncShared/WorkspaceUpdate.cs
@@ -335,8 +335,8 @@ namespace UnrealGameSync
 		const string VersionHeaderFileName = "/Engine/Source/Runtime/Launch/Resources/Version.h";
 		const string ObjectVersionFileName = "/Engine/Source/Runtime/Core/Private/UObject/ObjectVersion.cpp";
 
-		static readonly string LocalVersionHeaderFileName = VersionHeaderFileName.Replace('/', '\\');
-		static readonly string LocalObjectVersionFileName = ObjectVersionFileName.Replace('/', '\\');
+		static readonly string LocalVersionHeaderFileName = VersionHeaderFileName.Replace('/', Path.DirectorySeparatorChar);
+		static readonly string LocalObjectVersionFileName = ObjectVersionFileName.Replace('/', Path.DirectorySeparatorChar);
 
 		static SemaphoreSlim UpdateSemaphore = new SemaphoreSlim(1);
 
@@ -445,17 +445,27 @@ namespace UnrealGameSync
 				return Result;
 			}
 
-			public void IncludeFile(string Path, long Size)
+			public void IncludeFile(string Path, long Size, ILogger Logger) => IncludeFile(Path, Path, Size, Logger);
+
+			private void IncludeFile(string FullPath, string Path, long Size, ILogger Logger)
 			{
 				int Idx = Path.IndexOf('/');
 				if (Idx == -1)
 				{
-					IncludedFiles.Add(Path, Size);
+					long OldSize;
+					if (IncludedFiles.TryGetValue(Path, out OldSize))
+					{
+						Logger.LogWarning("File {FullPath} added twice to sync tree (old size {OldSize}, new size {Size})", FullPath, OldSize, Size);
+					}
+					else
+					{
+						IncludedFiles.Add(Path, Size);
+					}
 				}
 				else
 				{
 					SyncTree SubTree = FindOrAddSubTree(Path.Substring(0, Idx));
-					SubTree.IncludeFile(Path.Substring(Idx + 1), Size);
+					SubTree.IncludeFile(FullPath, Path.Substring(Idx + 1), Size, Logger);
 				}
 				TotalIncludedFiles++;
 				TotalSize += Size;
@@ -815,7 +825,7 @@ namespace UnrealGameSync
 							{
 								if (Filter.Matches(SyncRecord.RelativePath))
 								{
-									SyncTree.IncludeFile(PerforceUtils.EscapePath(SyncRecord.RelativePath), SyncRecord.Size);
+									SyncTree.IncludeFile(PerforceUtils.EscapePath(SyncRecord.RelativePath), SyncRecord.Size, Logger);
 									SyncDepotPaths.Add(SyncRecord.DepotFile);
 									RequiredFreeSpace += SyncRecord.Size;
 
