UGS-CLI: Add an option to force the filtered files to be synced with the have table (-refilter).

diff --git a/Engine/Source/Programs/UnrealGameSync/UnrealGameSyncCmd/Program.cs b/Engine/Source/Programs/UnrealGameSync/UnrealGameSyncCmd/Program.cs
--- a/Engine/Source/Programs/UnrealGameSync/UnrealGameSyncCmd/Program.cs
+++ b/Engine/Source/Programs/UnrealGameSync/UnrealGameSyncCmd/Program.cs
@@ -568,6 +568,9 @@ namespace UnrealGameSyncCmd
 
 				[CommandLine("-Clobber")]
 				public bool Clobber { get; set; }
+
+				[CommandLine("-Refilter")]
+				public bool Refilter { get; set; }
 			}
 
 			public override async Task ExecuteAsync(CommandContext Context)
@@ -613,6 +616,10 @@ namespace UnrealGameSyncCmd
 				{
 					Options |= WorkspaceUpdateOptions.Clobber;
 				}
+				if (SyncOptions.Refilter)
+				{
+					Options |= WorkspaceUpdateOptions.Refilter;
+				}
 				Options |= WorkspaceUpdateOptions.RemoveFilteredFiles;
 
 				ProjectInfo ProjectInfo = State.CreateProjectInfo();
diff --git a/Engine/Source/Programs/UnrealGameSync/UnrealGameSyncShared/WorkspaceUpdate.cs b/Engine/Source/Programs/UnrealGameSync/UnrealGameSyncShared/WorkspaceUpdate.cs
--- a/Engine/Source/Programs/UnrealGameSync/UnrealGameSyncShared/WorkspaceUpdate.cs
+++ b/Engine/Source/Programs/UnrealGameSync/UnrealGameSyncShared/WorkspaceUpdate.cs
@@ -43,6 +43,7 @@ namespace UnrealGameSync
 		IncludeAllProjectsInSolution = 0x2000,
 		RemoveFilteredFiles = 0x4000,
 		Clobber = 0x8000,
+		Refilter = 0x10000,
 	}
 
 	public enum WorkspaceUpdateResult
@@ -624,7 +625,7 @@ namespace UnrealGameSync
 					}
 
 					// If the hash differs, enumerate everything in the workspace to find what needs to be removed
-					if (NextSyncFilterHash != State.CurrentSyncFilterHash)
+					if (NextSyncFilterHash != State.CurrentSyncFilterHash || (Context.Options & WorkspaceUpdateOptions.Refilter) != 0)
 					{
 						using (TelemetryStopwatch FilterStopwatch = new TelemetryStopwatch("Workspace_Sync_FilterChanged", Project.TelemetryProjectIdentifier))
 						{
