Allow UGS to have an ini driven change to sync to based on a badge filter.
Users will be able to choose which change is considered good by having that CL match some subset of badges.

diff --git a/Engine/Source/Programs/UnrealGameSync/UnrealGameSync/Controls/WorkspaceControl.cs b/Engine/Source/Programs/UnrealGameSync/UnrealGameSync/Controls/WorkspaceControl.cs
--- a/Engine/Source/Programs/UnrealGameSync/UnrealGameSync/Controls/WorkspaceControl.cs
+++ b/Engine/Source/Programs/UnrealGameSync/UnrealGameSync/Controls/WorkspaceControl.cs
@@ -452,7 +452,7 @@ namespace UnrealGameSync
 			if (StartupTimer != null)
 			{
 				int LatestChangeNumber;
-				if (FindChangeToSync(Settings.SyncType, out LatestChangeNumber))
+				if (FindChangeToSync(Settings.SyncTypeID, out LatestChangeNumber))
 				{
 					StartupTimerElapsed(false);
 				}
@@ -686,7 +686,10 @@ namespace UnrealGameSync
 				components.Dispose();
 			}
 
-			UpdateTimer.Stop();
+			if (UpdateTimer != null)
+			{
+				UpdateTimer.Stop();
+			}
 
 			if (StartupCallbacks != null)
 			{
@@ -2016,6 +2019,73 @@ namespace UnrealGameSync
 			return SelectedArchives.Count == 0 || SelectedArchives.All(x => GetArchiveKeyForChangeNumber(x, ChangeNumber) != null);
 		}
 
+		/// <summary>
+		/// Should read as Do Required Badges Exist and Are They Good.
+		/// Determines if the ini specified INI filter is a subset of the InBadges passed
+		/// BadgeGroupSyncFilter - is used as the ini 
+		/// </summary>
+		/// <param name="InBadges"></param>
+		/// <returns></returns>
+		private bool DoRequiredBadgesExist(List<string> RequiredBadgeList, List<BadgeData> InBadges)
+		{
+			Dictionary<string, BadgeData> InBadgeDictionary = new Dictionary<string, BadgeData>();
+			foreach (BadgeData Badge in InBadges)
+			{
+				InBadgeDictionary.Add(Badge.BadgeName, Badge);
+			}
+
+			// make sure all required badges exist in the InBadgeSet
+			foreach (string BadgeName in RequiredBadgeList)
+			{
+				BadgeData? Badge;
+				if (InBadgeDictionary.TryGetValue(BadgeName, out Badge))
+				{
+					// If any required badge is not successful then the filter isn't matched.
+					if (!Badge.IsSuccess)
+					{
+						return false;
+					}
+				}
+				else
+				{
+					// required badge not found.
+					return false;
+				}
+			}
+
+			// All required badges existed and were successful 
+			return true;
+		}
+
+		private bool CanSyncChangeType(LatestChangeType ChangeType, ChangesRecord Change, EventSummary? Summary)
+		{
+			if (ChangeType.bGood)
+			{
+				if (Summary == null || Summary.Verdict != ReviewVerdict.Good)
+				{
+					return false;
+				}
+			}
+
+			if (ChangeType.bStarred)
+			{
+				if ((Summary == null || Summary.LastStarReview == null || Summary.LastStarReview.Type != EventType.Starred) && !PromotedChangeNumbers.Contains(Change.Number))
+				{
+					return false;
+				}
+			}
+
+			if (ChangeType.RequiredBadges.Count > 0)
+			{
+				if (Summary == null || !DoRequiredBadgesExist(ChangeType.RequiredBadges, Summary.Badges))
+				{
+					return false;
+				}
+			}
+
+			return true;
+		}
+
 		private ChangeLayoutInfo GetChangeLayoutInfo(ChangesRecord Change)
 		{
 			ChangeLayoutInfo? LayoutInfo;
@@ -2465,6 +2535,47 @@ namespace UnrealGameSync
 			return Badges;
 		}
 
+		private void PrintLatestChangeTypeUsage(string ErroringDefinition)
+		{
+			string ErrorMessage = "Error: Name not set for config value in \"Sync.LatestChangeType\"";
+			string UsageMessage = "Expected Format under the [Sync] ini category: +LatestChangeType=(Name=[string], Description=[string], OrderIndex=[int], bGood=[bool], bStarred=[bool], RequiredBadges=\"List of Badges For Changelist to have\")";
+
+			ShowErrorDialog("{0} Erroring String({1}) \n {2}", ErrorMessage, ErroringDefinition, UsageMessage);
+		}
+
+
+		/**
+		 * Array of config specified options for what the Latest Change to Sync should be.
+		 * See PrintLatestChangeTypeUsage for usage information
+		 */
+		public List<LatestChangeType> GetCustomLatestChangeTypes()
+		{
+			List<LatestChangeType> FoundLatestChangeTypes = new List<LatestChangeType>();
+
+			// Add three default change types that can be synced to
+			FoundLatestChangeTypes.Add(LatestChangeType.LatestChange());
+			FoundLatestChangeTypes.Add(LatestChangeType.LatestGoodChange());
+			FoundLatestChangeTypes.Add(LatestChangeType.LatestStarredChange());
+
+			string[] LatestChangeTypeDefinitions = PerforceMonitor.LatestProjectConfigFile.GetValues("Sync.LatestChangeType", new string[0]);
+			foreach (string LatestChangeDefinition in LatestChangeTypeDefinitions.Distinct())
+			{
+				LatestChangeType? NewType;
+				if (LatestChangeType.TryParseConfigEntry(LatestChangeDefinition, out NewType))
+				{
+					FoundLatestChangeTypes.Add(NewType);
+				}
+				else
+				{
+					PrintLatestChangeTypeUsage(LatestChangeDefinition);
+				}
+			}
+
+			FoundLatestChangeTypes = FoundLatestChangeTypes.OrderBy(X => X.OrderIndex).ToList();
+
+			return FoundLatestChangeTypes;
+		}
+
 		private static void SafeProcessStart(ProcessStartInfo StartInfo)
 		{
 			try
@@ -4370,9 +4481,18 @@ namespace UnrealGameSync
 			if (Workspace != null)
 			{
 				int ChangeNumber;
-				if (!FindChangeToSync(Settings.SyncType, out ChangeNumber))
+				if (!FindChangeToSync(Settings.SyncTypeID, out ChangeNumber))
 				{
-					ShowErrorDialog(String.Format("Couldn't find any {0}changelist. Double-click on the change you want to sync manually.", (Settings.SyncType == LatestChangeType.Starred) ? "starred " : (Settings.SyncType == LatestChangeType.Good) ? "good " : ""));
+					string SyncTypeName = "";
+					foreach (LatestChangeType ChangeType in GetCustomLatestChangeTypes())
+					{
+						if (ChangeType.Name == Settings.SyncTypeID)
+						{
+							SyncTypeName = ChangeType.Name;
+							break;
+						}
+					}
+					ShowErrorDialog(String.Format("Couldn't find any {0}changelist. Double-click on the change you want to sync manually.", SyncTypeName));
 				}
 				else if (ChangeNumber < Workspace.CurrentChangeNumber)
 				{
@@ -4797,12 +4917,22 @@ namespace UnrealGameSync
 		{
 			if (Settings.bScheduleEnabled)
 			{
-				Logger.LogInformation("Scheduled sync at {Time} for {Change}.", DateTime.Now, Settings.ScheduleChange);
+				// Check the project config if they have a scheduled setting selected.
+				string? ScheduledSyncTypeID = null;
+				UserSelectedProjectSettings ProjectSetting = Settings.ScheduleProjects.First(x => x.LocalPath != null && x.LocalPath.Equals(SelectedProject.LocalPath, StringComparison.OrdinalIgnoreCase));
+				if (ProjectSetting != null)
+				{
+					ScheduledSyncTypeID = ProjectSetting.ScheduledSyncTypeID;
+				}
+
+				Logger.LogInformation("Scheduled sync at {0} for {1}.",
+					DateTime.Now,
+					ScheduledSyncTypeID != null ? ScheduledSyncTypeID : "Default");
 
 				int ChangeNumber;
-				if (!FindChangeToSync(Settings.ScheduleChange, out ChangeNumber))
+				if (!FindChangeToSync(ScheduledSyncTypeID, out ChangeNumber))
 				{
-					Logger.LogInformation("Couldn't find any matching change");
+					SyncLog.AppendLine("Couldn't find any matching change");
 				}
 				else if (Workspace.CurrentChangeNumber >= ChangeNumber)
 				{
@@ -4820,6 +4950,31 @@ namespace UnrealGameSync
 			}
 		}
 
+		bool FindChangeToSync(string? ChangeTypeID, out int ChangeNumber)
+		{
+			if (ChangeTypeID != null)
+			{
+				LatestChangeType? ChangeTypeToSyncTo = null;
+				foreach (LatestChangeType ChangeType in GetCustomLatestChangeTypes())
+				{
+					if (ChangeType.Name == ChangeTypeID)
+					{
+						ChangeTypeToSyncTo = ChangeType;
+						break;
+					}
+				}
+
+				if (ChangeTypeToSyncTo != null)
+				{
+					return FindChangeToSync(ChangeTypeToSyncTo, out ChangeNumber);
+				}
+			}
+
+			SyncLog.AppendLine("Invalid or not set Scheduled Sync type using (Sync to Latest Good)");
+
+			return FindChangeToSync(LatestChangeType.LatestGoodChange(), out ChangeNumber);
+		}
+
 		bool FindChangeToSync(LatestChangeType ChangeType, out int ChangeNumber)
 		{
 			for (int Idx = 0; Idx < BuildList.Items.Count; Idx++)
@@ -4827,31 +4982,12 @@ namespace UnrealGameSync
 				ChangesRecord Change = (ChangesRecord)BuildList.Items[Idx].Tag;
 				if (Change != null)
 				{
-					if (ChangeType == LatestChangeType.Any)
-					{
-						if (CanSyncChange(Change.Number))
-						{
-							ChangeNumber = FindNewestGoodContentChange(Change.Number);
-							return true;
-						}
-					}
-					else if (ChangeType == LatestChangeType.Good)
+					EventSummary? Summary = EventMonitor.GetSummaryForChange(Change.Number);
+					if (CanSyncChange(Change.Number)
+						&& CanSyncChangeType(ChangeType, Change, Summary))
 					{
-						EventSummary? Summary = EventMonitor.GetSummaryForChange(Change.Number);
-						if (Summary != null && Summary.Verdict == ReviewVerdict.Good && CanSyncChange(Change.Number))
-						{
-							ChangeNumber = FindNewestGoodContentChange(Change.Number);
-							return true;
-						}
-					}
-					else if (ChangeType == LatestChangeType.Starred)
-					{
-						EventSummary? Summary = EventMonitor.GetSummaryForChange(Change.Number);
-						if (((Summary != null && Summary.LastStarReview != null && Summary.LastStarReview.Type == EventType.Starred) || PromotedChangeNumbers.Contains(Change.Number)) && CanSyncChange(Change.Number))
-						{
-							ChangeNumber = FindNewestGoodContentChange(Change.Number);
-							return true;
-						}
+						ChangeNumber = FindNewestGoodContentChange(Change.Number);
+						return true;
 					}
 				}
 			}
@@ -5236,29 +5372,67 @@ namespace UnrealGameSync
 			}
 		}
 
-		private void ShowSyncMenu(Rectangle Bounds)
+		/// <summary>
+		/// Used to load CustomLatestChangeTypes config data into the windows form.
+		/// </summary>
+		private void RefreshSyncMenuItems()
 		{
-			SyncContextMenu_LatestChange.Checked = (Settings.SyncType == LatestChangeType.Any);
-			SyncContextMenu_LatestGoodChange.Checked = (Settings.SyncType == LatestChangeType.Good);
-			SyncContextMenu_LatestStarredChange.Checked = (Settings.SyncType == LatestChangeType.Starred);
-			SyncContextMenu.Show(StatusPanel, new Point(Bounds.Left, Bounds.Bottom), ToolStripDropDownDirection.BelowRight);
-		}
+			SyncContextMenu.Items.Clear();
+			SyncContextMenu.Items.AddRange(new System.Windows.Forms.ToolStripItem[] {
+			this.SyncContexMenu_EnterChangelist,
+			this.toolStripSeparator8});
 
-		private void SyncContextMenu_LatestChange_Click(object sender, EventArgs e)
-		{
-			Settings.SyncType = LatestChangeType.Any;
-			Settings.Save();
+			foreach (LatestChangeType ChangeType in GetCustomLatestChangeTypes())
+			{
+				System.Windows.Forms.ToolStripMenuItem MenuItem = new System.Windows.Forms.ToolStripMenuItem();
+				MenuItem.Name = ChangeType.Name;
+				MenuItem.Text = ChangeType.Description;
+				MenuItem.Size = new System.Drawing.Size(189, 22);
+				MenuItem.Click += (sender, e) => this.SyncContextMenu_LatestChangeType_Click(sender, e, ChangeType.Name);
+
+				SyncContextMenu.Items.Add(MenuItem);
+			}
 		}
 
-		private void SyncContextMenu_LatestGoodChange_Click(object sender, EventArgs e)
+		private void ShowSyncMenu(Rectangle Bounds)
 		{
-			Settings.SyncType = LatestChangeType.Good;
-			Settings.Save();
+			RefreshSyncMenuItems();
+
+			string? NameOfLatestChangeType = null;
+			foreach (LatestChangeType ChangeType in GetCustomLatestChangeTypes())
+			{
+				if (ChangeType.Name == Settings.SyncTypeID)
+				{
+					NameOfLatestChangeType = ChangeType.Name;
+					break;
+				}
+			}
+
+			if (NameOfLatestChangeType != null)
+			{
+				foreach (int value in Enumerable.Range(0, SyncContextMenu.Items.Count))
+				{
+					System.Windows.Forms.ToolStripMenuItem? Item = SyncContextMenu.Items[value] as System.Windows.Forms.ToolStripMenuItem;
+					if (Item != null)
+					{
+						if (Item.Name == NameOfLatestChangeType)
+						{
+							Item.Checked = true;
+						}
+						else
+						{
+							Item.Checked = false;
+						}
+					}
+				}
+			}
+
+			SyncContextMenu.Show(StatusPanel, new Point(Bounds.Left, Bounds.Bottom), ToolStripDropDownDirection.BelowRight);
 		}
 
-		private void SyncContextMenu_LatestStarredChange_Click(object sender, EventArgs e)
+		private void SyncContextMenu_LatestChangeType_Click(object? sender, EventArgs e, string SyncTypeID)
 		{
-			Settings.SyncType = LatestChangeType.Starred;
+			Settings.SyncTypeID = SyncTypeID;
 			Settings.Save();
 		}
 
diff --git a/Engine/Source/Programs/UnrealGameSync/UnrealGameSync/Forms/MainWindow.cs b/Engine/Source/Programs/UnrealGameSync/UnrealGameSync/Forms/MainWindow.cs
--- a/Engine/Source/Programs/UnrealGameSync/UnrealGameSync/Forms/MainWindow.cs
+++ b/Engine/Source/Programs/UnrealGameSync/UnrealGameSync/Forms/MainWindow.cs
@@ -863,10 +863,20 @@ namespace UnrealGameSync
 				OpenProjects.Add(TabPanel.SelectedProject);
 			}
 
-			ScheduleWindow Schedule = new ScheduleWindow(Settings.bScheduleEnabled, Settings.ScheduleChange, Settings.ScheduleTime, Settings.ScheduleAnyOpenProject, Settings.ScheduleProjects, OpenProjects);
+			Dictionary<UserSelectedProjectSettings, List<LatestChangeType>> ProjectToLatestChangeTypes = new Dictionary<UserSelectedProjectSettings, List<LatestChangeType>>();
+			for (int Idx = 0; Idx < TabControl.GetTabCount(); Idx++)
+			{
+				WorkspaceControl? Workspace = TabControl.GetTabData(Idx) as WorkspaceControl;
+				if (Workspace != null)
+				{
+					ProjectToLatestChangeTypes.Add(Workspace.SelectedProject, Workspace.GetCustomLatestChangeTypes());
+				}
+			}
+
+			ScheduleWindow Schedule = new ScheduleWindow(Settings.bScheduleEnabled, Settings.ScheduleTime, Settings.ScheduleAnyOpenProject, Settings.ScheduleProjects, OpenProjects, ProjectToLatestChangeTypes);
 			if(Schedule.ShowDialog() == System.Windows.Forms.DialogResult.OK)
 			{
-				Schedule.CopySettings(out Settings.bScheduleEnabled, out Settings.ScheduleChange, out Settings.ScheduleTime, out Settings.ScheduleAnyOpenProject, out Settings.ScheduleProjects);
+				Schedule.CopySettings(out Settings.bScheduleEnabled, out Settings.ScheduleTime, out Settings.ScheduleAnyOpenProject, out Settings.ScheduleProjects);
 				Settings.Save();
 			}
 
diff --git a/Engine/Source/Programs/UnrealGameSync/UnrealGameSync/Forms/ScheduleWindow.cs b/Engine/Source/Programs/UnrealGameSync/UnrealGameSync/Forms/ScheduleWindow.cs
--- a/Engine/Source/Programs/UnrealGameSync/UnrealGameSync/Forms/ScheduleWindow.cs
+++ b/Engine/Source/Programs/UnrealGameSync/UnrealGameSync/Forms/ScheduleWindow.cs
@@ -16,12 +16,22 @@ namespace UnrealGameSync
 {
 	partial class ScheduleWindow : Form
 	{
-		public ScheduleWindow(bool InEnabled, LatestChangeType InChange, TimeSpan InTime, bool AnyOpenProject, IEnumerable<UserSelectedProjectSettings> ScheduledProjects, IEnumerable<UserSelectedProjectSettings> OpenProjects)
+
+		Dictionary<UserSelectedProjectSettings, List<LatestChangeType>> ProjectToLatestChangeTypes;
+
+		public ScheduleWindow(
+			bool InEnabled,
+			TimeSpan InTime,
+			bool AnyOpenProject,
+			IEnumerable<UserSelectedProjectSettings> ScheduledProjects,
+			IEnumerable<UserSelectedProjectSettings> OpenProjects,
+			Dictionary<UserSelectedProjectSettings, List<LatestChangeType>> InProjectToLatestChangeTypes)
 		{
 			InitializeComponent();
 
 			EnableCheckBox.Checked = InEnabled;
-			ChangeComboBox.SelectedIndex = (int)InChange;
+
+			ProjectToLatestChangeTypes = InProjectToLatestChangeTypes;
 
 			DateTime CurrentTime = DateTime.Now;
 			TimePicker.CustomFormat = CultureInfo.CurrentCulture.DateTimeFormat.ShortTimePattern;
@@ -39,9 +49,65 @@ namespace UnrealGameSync
 				ProjectListBox.Items.Add(Project, Enabled);
 			}
 
+			ProjectListBox.ItemCheck += ProjectListBox_ItemCheck;
+
+			SyncTypeDropDown.Closed += SyncTypeDropDown_Closed;
+
 			UpdateEnabledControls();
 		}
 
+		private void UpdateSyncTypeDropDownWithProject(UserSelectedProjectSettings? SelectedProject)
+		{
+			SyncTypeDropDown.Items.Clear();
+
+			if (SelectedProject != null
+				&& ProjectToLatestChangeTypes.ContainsKey(SelectedProject))
+			{
+				foreach (LatestChangeType ChangeType in ProjectToLatestChangeTypes[SelectedProject])
+				{
+					System.Windows.Forms.ToolStripMenuItem MenuItem = new System.Windows.Forms.ToolStripMenuItem();
+					MenuItem.Name = ChangeType.Name;
+					MenuItem.Text = ChangeType.Description;
+					MenuItem.Size = new System.Drawing.Size(189, 22);
+					MenuItem.Click += (sender, e) => SyncTypeDropDown_Click(sender, e, ChangeType.Name);
+
+					SyncTypeDropDown.Items.Add(MenuItem);
+				}
+			}
+		}
+
+		private void ProjectListBox_ItemCheck(object sender, ItemCheckEventArgs e)
+		{
+			bool bIsAnyOpenProjectIndex = e.Index == 0;
+			if (e.NewValue == CheckState.Checked && !bIsAnyOpenProjectIndex)
+			{
+				UpdateSyncTypeDropDownWithProject(ProjectListBox.Items[e.Index] as UserSelectedProjectSettings);
+
+				SyncTypeDropDown.Show(
+					ProjectListBox,
+					ProjectListBox.GetItemRectangle(e.Index).Location,
+					ToolStripDropDownDirection.BelowRight);
+			}
+		}
+
+		private void SyncTypeDropDown_Closed(object sender, ToolStripDropDownClosedEventArgs e)
+		{
+			if (e.CloseReason != ToolStripDropDownCloseReason.ItemClicked)
+			{
+				ProjectListBox.SetItemChecked(ProjectListBox.SelectedIndex, false);
+			}
+		}
+
+		private void SyncTypeDropDown_Click(object? sender, EventArgs e, string SyncTypeID)
+		{
+			UserSelectedProjectSettings? ProjectSetting = ProjectListBox.Items[ProjectListBox.SelectedIndex] as UserSelectedProjectSettings;
+			if (ProjectSetting != null)
+			{
+				ProjectSetting.ScheduledSyncTypeID = SyncTypeID;
+			}
+			SyncTypeDropDown.Close();
+		}
+
 		private void AddProjects(IEnumerable<UserSelectedProjectSettings> Projects, Dictionary<string, UserSelectedProjectSettings> LocalFileToProject)
 		{
 			foreach(UserSelectedProjectSettings Project in Projects)
@@ -53,10 +119,9 @@ namespace UnrealGameSync
 			}
 		}
 
-		public void CopySettings(out bool OutEnabled, out LatestChangeType OutChange, out TimeSpan OutTime, out bool OutAnyOpenProject, out List<UserSelectedProjectSettings> OutScheduledProjects)
+		public void CopySettings(out bool OutEnabled, out TimeSpan OutTime, out bool OutAnyOpenProject, out List<UserSelectedProjectSettings> OutScheduledProjects)
 		{
 			OutEnabled = EnableCheckBox.Checked;
-			OutChange = (LatestChangeType)ChangeComboBox.SelectedIndex;
 			OutTime = TimePicker.Value.TimeOfDay;
 
 			OutAnyOpenProject = false;
@@ -83,7 +148,6 @@ namespace UnrealGameSync
 
 		private void UpdateEnabledControls()
 		{
-			ChangeComboBox.Enabled = EnableCheckBox.Checked;
 			TimePicker.Enabled = EnableCheckBox.Checked;
 			ProjectListBox.Enabled = EnableCheckBox.Checked;
 		}
diff --git a/Engine/Source/Programs/UnrealGameSync/UnrealGameSyncShared/UserSettings.cs b/Engine/Source/Programs/UnrealGameSync/UnrealGameSyncShared/UserSettings.cs
--- a/Engine/Source/Programs/UnrealGameSync/UnrealGameSyncShared/UserSettings.cs
+++ b/Engine/Source/Programs/UnrealGameSync/UnrealGameSyncShared/UserSettings.cs
@@ -51,11 +51,81 @@ namespace UnrealGameSync
 		Content
 	}
 
-	public enum LatestChangeType
+	/// <summary>
+	/// Config specified class to determine what is the Latest Change to Sync
+	/// Can be configured using badges, good and starred CLs
+	/// </summary>
+	public class LatestChangeType
 	{
-		Any,
-		Good,
-		Starred,
+		// Display Specifiers
+		public string Name = ""; // Name will the saved ID for this LatestChangeType
+		public string Description = "";
+		public int OrderIndex = Int32.MaxValue;
+
+		// What Rules to Check for.
+		public bool bGood = false;
+		public bool bStarred = false;
+		public List<string> RequiredBadges = new List<string>();
+
+		public static bool TryParseConfigEntry(string Text, [NotNullWhen(true)] out LatestChangeType? ChangeType)
+		{
+			ConfigObject DefinitionObject = new ConfigObject(Text);
+
+			string LatestChangeTypeName = DefinitionObject.GetValue("Name", "");
+			if (LatestChangeTypeName != "")
+			{
+				ChangeType = new LatestChangeType();
+				ChangeType.Name = LatestChangeTypeName;
+				ChangeType.Description = DefinitionObject.GetValue("Description", "No \"Description\" for LatestChangeType given.");
+				ChangeType.OrderIndex = DefinitionObject.GetValue("OrderIndex", Int32.MaxValue);
+
+				ChangeType.bGood = DefinitionObject.GetValue("bGood", false);
+				ChangeType.bStarred = DefinitionObject.GetValue("bStarred", false);
+				ChangeType.RequiredBadges = new List<string>(DefinitionObject.GetValue("RequiredBadges", "").Split(new char[] { ',' }, StringSplitOptions.RemoveEmptyEntries));
+				return true;
+			}
+			else
+			{
+				ChangeType = null;
+			}
+
+			return false;
+		}
+
+		// Default Template for Latest Change
+		public static LatestChangeType LatestChange()
+		{
+			LatestChangeType Latest = new LatestChangeType();
+			Latest.Name = "LatestChange";
+			Latest.Description = "Latest Change";
+			Latest.OrderIndex = -3;
+
+			return Latest;
+		}
+
+		// Default Template for Latest Good Change
+		public static LatestChangeType LatestGoodChange()
+		{
+			LatestChangeType LatestGood = new LatestChangeType();
+			LatestGood.Name = "LatestGoodChange";
+			LatestGood.Description = "Latest Good Change";
+			LatestGood.OrderIndex = -2;
+			LatestGood.bGood = true;
+
+			return LatestGood;
+		}
+
+		// Default Template for Latest Starred Change
+		public static LatestChangeType LatestStarredChange()
+		{
+			LatestChangeType LatestStarred = new LatestChangeType();
+			LatestStarred.Name = "LatestStarredChange";
+			LatestStarred.Description = "Latest Starred Change";
+			LatestStarred.OrderIndex = -1;
+			LatestStarred.bStarred = true;
+
+			return LatestStarred;
+		}
 	}
 
 	public enum UserSettingsVersion
@@ -124,6 +194,7 @@ namespace UnrealGameSync
 		public UserSelectedProjectType Type { get; set; }
 		public string? ClientPath { get; set; }
 		public string? LocalPath { get; set; }
+		public string? ScheduledSyncTypeID { get; set; }
 
 		public UserSelectedProjectSettings(string? ServerAndPort, string? UserName, UserSelectedProjectType Type, string? ClientPath, string? LocalPath)
 		{
@@ -171,9 +242,16 @@ namespace UnrealGameSync
 					ClientPath = null;
 				}
 
-				if((Type == UserSelectedProjectType.Client && ClientPath != null) || (Type == UserSelectedProjectType.Local && LocalPath != null))
+				string? ScheduledSyncTypeID = Object.GetValue("ScheduledSyncTypeID", null);
+				if (String.IsNullOrWhiteSpace(ScheduledSyncTypeID))
+				{
+					ScheduledSyncTypeID = null;
+				}
+
+				if ((Type == UserSelectedProjectType.Client && ClientPath != null) || (Type == UserSelectedProjectType.Local && LocalPath != null))
 				{
 					Project = new UserSelectedProjectSettings(ServerAndPort, UserName, Type, ClientPath, LocalPath);
+					Project.ScheduledSyncTypeID = ScheduledSyncTypeID;
 					return true;
 				}
 			}
@@ -205,6 +283,10 @@ namespace UnrealGameSync
 			{
 				Object.SetValue("LocalPath", LocalPath);
 			}
+			if (ScheduledSyncTypeID != null)
+			{
+				Object.SetValue("ScheduledSyncTypeID", ScheduledSyncTypeID);
+			}
 
 			return Object.ToString();
 		}
@@ -453,7 +535,7 @@ namespace UnrealGameSync
 		public UserSelectedProjectSettings? LastProject;
 		public List<UserSelectedProjectSettings> OpenProjects;
 		public List<UserSelectedProjectSettings> RecentProjects;
-		public LatestChangeType SyncType;
+		public string SyncTypeID;
 		public BuildConfig CompiledEditorBuildConfig; // NB: This assumes not using precompiled editor. See CurrentBuildConfig.
 		public TabLabels TabLabels;
 
@@ -471,7 +553,6 @@ namespace UnrealGameSync
 		// Schedule settings
 		public bool bScheduleEnabled;
 		public TimeSpan ScheduleTime;
-		public LatestChangeType ScheduleChange;
 		public bool ScheduleAnyOpenProject;
 		public List<UserSelectedProjectSettings> ScheduleProjects;
 
@@ -608,9 +689,22 @@ namespace UnrealGameSync
 
 			OpenProjects = ReadProjectList("General.OpenProjects", "General.OpenProjectFileNames");
 			RecentProjects = ReadProjectList("General.RecentProjects", "General.OtherProjectFileNames");
-			if(!Enum.TryParse(ConfigFile.GetValue("General.SyncType", ""), out SyncType))
+			SyncTypeID = ConfigFile.GetValue("General.SyncTypeID", "");
+			string? OldSyncTye = ConfigFile.GetValue("General.SyncType", null);
+			if (OldSyncTye != null)
 			{
-				SyncType = LatestChangeType.Good;
+				if (OldSyncTye == "Any")
+				{
+					SyncTypeID = LatestChangeType.LatestChange().Name;
+				}
+				else if (OldSyncTye == "Good")
+				{
+					SyncTypeID = LatestChangeType.LatestGoodChange().Name;
+				}
+				else if (OldSyncTye == "Starred")
+				{
+					SyncTypeID = LatestChangeType.LatestStarredChange().Name;
+				}
 			}
 
 			// Build configuration
@@ -690,10 +784,7 @@ namespace UnrealGameSync
 			{
 				ScheduleTime = new TimeSpan(6, 0, 0);
 			}
-			if(!Enum.TryParse(ConfigFile.GetValue("Schedule.Change", ""), out ScheduleChange))
-			{
-				ScheduleChange = LatestChangeType.Good;
-			}
+
 			ScheduleAnyOpenProject = ConfigFile.GetValue("Schedule.AnyOpenProject", true);
 			ScheduleProjects = ReadProjectList("Schedule.Projects", "Schedule.ProjectFileNames");
 
@@ -969,7 +1060,7 @@ namespace UnrealGameSync
 			GeneralSection.SetValues("EnabledTools", EnabledTools);
 			GeneralSection.SetValue("FilterIndex", FilterIndex);
 			GeneralSection.SetValues("RecentProjects", RecentProjects.Select(x => x.ToConfigEntry()).ToArray());
-			GeneralSection.SetValue("SyncType", SyncType.ToString());
+			GeneralSection.SetValue("SyncTypeID", SyncTypeID);
 
 			// Build configuration
 			GeneralSection.SetValue("BuildConfig", CompiledEditorBuildConfig.ToString());
@@ -991,7 +1082,6 @@ namespace UnrealGameSync
 			ScheduleSection.Clear();
 			ScheduleSection.SetValue("Enabled", bScheduleEnabled);
 			ScheduleSection.SetValue("Time", ScheduleTime.ToString());
-			ScheduleSection.SetValue("Change", ScheduleChange.ToString());
 			ScheduleSection.SetValue("AnyOpenProject", ScheduleAnyOpenProject);
 			ScheduleSection.SetValues("Projects", ScheduleProjects.Select(x => x.ToConfigEntry()).ToArray());
 
