Minor fixup to isolate UGS badges

diff --git a/Engine/Source/Programs/UnrealGameSync/UnrealGameSyncShared/UserSettings.cs b/Engine/Source/Programs/UnrealGameSync/UnrealGameSyncShared/UserSettings.cs
--- a/Engine/Source/Programs/UnrealGameSync/UnrealGameSyncShared/UserSettings.cs
+++ b/Engine/Source/Programs/UnrealGameSync/UnrealGameSyncShared/UserSettings.cs
@@ -81,7 +81,7 @@ namespace UnrealGameSync
 
 				ChangeType.bGood = DefinitionObject.GetValue("bGood", false);
 				ChangeType.bStarred = DefinitionObject.GetValue("bStarred", false);
-				ChangeType.RequiredBadges = new List<string>(DefinitionObject.GetValue("RequiredBadges", "").Split(new char[] { ',' }, StringSplitOptions.RemoveEmptyEntries));
+				ChangeType.RequiredBadges = new List<string>(DefinitionObject.GetValue("RequiredBadges", "").Split(new char[] { ',',' '}, StringSplitOptions.RemoveEmptyEntries));
 				return true;
 			}
 			else
