UGS: Fix incorrect parameter order when an error ocurrs in PerforceModalTask.

diff --git a/Engine/Source/Programs/UnrealGameSync/UnrealGameSync/PerforceModalTask.cs b/Engine/Source/Programs/UnrealGameSync/UnrealGameSync/PerforceModalTask.cs
--- a/Engine/Source/Programs/UnrealGameSync/UnrealGameSync/PerforceModalTask.cs
+++ b/Engine/Source/Programs/UnrealGameSync/UnrealGameSync/PerforceModalTask.cs
@@ -96,7 +96,7 @@ namespace UnrealGameSync
 						Password = PasswordWindow.Password;
 						continue;
 					}
-					MessageBox.Show(Owner, Title, Result.Error, MessageBoxButtons.OK);
+					MessageBox.Show(Owner, Result.Error, Title, MessageBoxButtons.OK);
 				}
 
 				return Result;
