UGS: Fix time zone display with new Perforce library.

diff --git a/Engine/Source/Programs/UnrealGameSync/UnrealGameSync/Controls/WorkspaceControl.cs b/Engine/Source/Programs/UnrealGameSync/UnrealGameSync/Controls/WorkspaceControl.cs
--- a/Engine/Source/Programs/UnrealGameSync/UnrealGameSync/Controls/WorkspaceControl.cs
+++ b/Engine/Source/Programs/UnrealGameSync/UnrealGameSync/Controls/WorkspaceControl.cs
@@ -1323,10 +1323,14 @@ namespace UnrealGameSync
 		void BuildList_AddItem(ChangesRecord Change)
 		{
 			// Get the display time for this item
-			DateTime DisplayTime = Change.Time;
+			DateTime DisplayTime;
 			if (Settings.bShowLocalTimes)
 			{
-				DisplayTime = (DisplayTime - PerforceMonitor.ServerTimeZone).ToLocalTime();
+				DisplayTime = Change.Time.ToLocalTime();
+			}
+			else
+			{
+				DisplayTime = new DateTime(Change.Time.Ticks + PerforceMonitor.ServerTimeZone.Ticks, DateTimeKind.Local);
 			}
 
 			// Find or add the new group
diff --git a/Engine/Source/Programs/UnrealGameSync/UnrealGameSync/Forms/IssueDetailsWindow.cs b/Engine/Source/Programs/UnrealGameSync/UnrealGameSync/Forms/IssueDetailsWindow.cs
--- a/Engine/Source/Programs/UnrealGameSync/UnrealGameSync/Forms/IssueDetailsWindow.cs
+++ b/Engine/Source/Programs/UnrealGameSync/UnrealGameSync/Forms/IssueDetailsWindow.cs
@@ -783,11 +783,16 @@ namespace UnrealGameSync
 
 							Item.SubItems.Add(Change.Number.ToString());
 
-							DateTime DisplayTime = Change.Time;
-							if (ServerTimeOffset.HasValue)
+							DateTime DisplayTime;
+							if (!ServerTimeOffset.HasValue)
 							{
-								DisplayTime = (DisplayTime - ServerTimeOffset.Value).ToLocalTime();
+								DisplayTime = Change.Time.ToLocalTime();
 							}
+							else
+							{
+								DisplayTime = new DateTime(Change.Time.Ticks + ServerTimeOffset.Value.Ticks, DateTimeKind.Local);
+							}
+
 							Item.SubItems.Add(DisplayTime.ToString("h\\.mmtt"));
 
 							Item.SubItems.Add(WorkspaceControl.FormatUserName(Change.User));
