UGS: Fix log output for exceptions.

diff --git a/Engine/Source/Programs/UnrealGameSync/UnrealGameSync/Controls/LogControl.cs b/Engine/Source/Programs/UnrealGameSync/UnrealGameSync/Controls/LogControl.cs
--- a/Engine/Source/Programs/UnrealGameSync/UnrealGameSync/Controls/LogControl.cs
+++ b/Engine/Source/Programs/UnrealGameSync/UnrealGameSync/Controls/LogControl.cs
@@ -930,7 +930,7 @@ namespace UnrealGameSync
 			{
 				foreach (string Line in Exception.ToString().Trim().Split('\n'))
 				{
-					LogControl.AppendLine(Message);
+					LogControl.AppendLine(Line);
 				}
 			}
 		}
