UGS: Fix workspace settings being lost when switching between projects in the same stream.

diff --git a/Engine/Source/Programs/UnrealGameSync/UnrealGameSyncShared/UserSettings.cs b/Engine/Source/Programs/UnrealGameSync/UnrealGameSyncShared/UserSettings.cs
--- a/Engine/Source/Programs/UnrealGameSync/UnrealGameSyncShared/UserSettings.cs
+++ b/Engine/Source/Programs/UnrealGameSync/UnrealGameSyncShared/UserSettings.cs
@@ -452,7 +452,7 @@ namespace UnrealGameSync
 
 		public bool IsValid(ProjectInfo ProjectInfo)
 		{
-			return ProjectInfo.Equals(ProjectInfo.ProjectIdentifier, StringComparison.Ordinal);
+			return ProjectIdentifier.Equals(ProjectInfo.ProjectIdentifier, StringComparison.OrdinalIgnoreCase);
 		}
 
 		public void SetBisectState(int Change, BisectState State)
