UGS: Remove reference to mismatched System.Linq.Async version.

diff --git a/Engine/Source/Programs/UnrealGameSync/UnrealGameSync/UnrealGameSync.csproj b/Engine/Source/Programs/UnrealGameSync/UnrealGameSync/UnrealGameSync.csproj
--- a/Engine/Source/Programs/UnrealGameSync/UnrealGameSync/UnrealGameSync.csproj
+++ b/Engine/Source/Programs/UnrealGameSync/UnrealGameSync/UnrealGameSync.csproj
@@ -27,7 +27,6 @@
     <PackageReference Include="Serilog.Extensions.Hosting" Version="4.2.0" />
     <PackageReference Include="Serilog.Extensions.Logging" Version="3.1.0" />
     <PackageReference Include="Serilog.Sinks.File" Version="5.0.0" />
-    <PackageReference Include="System.Linq.Async" Version="5.1.0" />
   </ItemGroup>
   <ItemGroup>
     <ProjectReference Include="..\..\Shared\EpicGames.Core\EpicGames.Core.csproj" />
