UGS: Fix copying text containing existing CR-LF pairs from log control getting additional newlines.

diff --git a/Engine/Source/Programs/UnrealGameSync/UnrealGameSync/Controls/LogControl.cs b/Engine/Source/Programs/UnrealGameSync/UnrealGameSync/Controls/LogControl.cs
--- a/Engine/Source/Programs/UnrealGameSync/UnrealGameSync/Controls/LogControl.cs
+++ b/Engine/Source/Programs/UnrealGameSync/UnrealGameSync/Controls/LogControl.cs
@@ -384,7 +384,7 @@ namespace UnrealGameSync
 						SelectedText.Append(Lines[LineIdx], MinIdx, MaxIdx - MinIdx);
 					}
 				}
-				Clipboard.SetText(SelectedText.ToString().Replace("\n", "\r\n"));
+				Clipboard.SetText(Regex.Replace(SelectedText.ToString(), "(?<!\r)\n", "\r\n"));
 			}
 		}
 
