UGS: Fix nullable warnings with new Registry nuget package.

diff --git a/Engine/Source/Programs/UnrealGameSync/UnrealGameSync/AutomationServer.cs b/Engine/Source/Programs/UnrealGameSync/UnrealGameSync/AutomationServer.cs
--- a/Engine/Source/Programs/UnrealGameSync/UnrealGameSync/AutomationServer.cs
+++ b/Engine/Source/Programs/UnrealGameSync/UnrealGameSync/AutomationServer.cs
@@ -200,7 +200,7 @@ namespace UnrealGameSync
 
 		public static int GetPortNumber()
 		{
-			object PortValue = Registry.GetValue("HKEY_CURRENT_USER\\Software\\Epic Games\\UnrealGameSync", "AutomationPort", null);
+			object? PortValue = Registry.GetValue("HKEY_CURRENT_USER\\Software\\Epic Games\\UnrealGameSync", "AutomationPort", null);
 			if (PortValue != null && PortValue is int)
 			{
 				return (int)PortValue;
diff --git a/Engine/Source/Programs/UnrealGameSync/UnrealGameSync/Forms/ApplicationSettingsWindow.cs b/Engine/Source/Programs/UnrealGameSync/UnrealGameSync/Forms/ApplicationSettingsWindow.cs
--- a/Engine/Source/Programs/UnrealGameSync/UnrealGameSync/Forms/ApplicationSettingsWindow.cs
+++ b/Engine/Source/Programs/UnrealGameSync/UnrealGameSync/Forms/ApplicationSettingsWindow.cs
@@ -149,8 +149,8 @@ namespace UnrealGameSync
 
 		private bool IsAutomaticallyRunAtStartup()
 		{
-			RegistryKey Key = Registry.CurrentUser.OpenSubKey("SOFTWARE\\Microsoft\\Windows\\CurrentVersion\\Run");
-			return (Key.GetValue("UnrealGameSync") != null);
+			RegistryKey? Key = Registry.CurrentUser.OpenSubKey("SOFTWARE\\Microsoft\\Windows\\CurrentVersion\\Run");
+			return (Key?.GetValue("UnrealGameSync") != null);
 		}
 
 		private void OkBtn_Click(object sender, EventArgs e)
diff --git a/Engine/Source/Programs/UnrealGameSync/UnrealGameSync/GlobalSettings.cs b/Engine/Source/Programs/UnrealGameSync/UnrealGameSync/GlobalSettings.cs
--- a/Engine/Source/Programs/UnrealGameSync/UnrealGameSync/GlobalSettings.cs
+++ b/Engine/Source/Programs/UnrealGameSync/UnrealGameSync/GlobalSettings.cs
@@ -14,7 +14,7 @@ namespace UnrealGameSync
 	{
 		public static void ReadGlobalPerforceSettings(ref string? ServerAndPort, ref string? UserName, ref string? DepotPath)
 		{
-			using (RegistryKey Key = Registry.CurrentUser.OpenSubKey("SOFTWARE\\Epic Games\\UnrealGameSync", false))
+			using (RegistryKey? Key = Registry.CurrentUser.OpenSubKey("SOFTWARE\\Epic Games\\UnrealGameSync", false))
 			{
 				if (Key != null)
 				{
@@ -38,7 +38,7 @@ namespace UnrealGameSync
 
 		public static void DeleteRegistryKey(RegistryKey RootKey, string KeyName, string ValueName)
 		{
-			using (RegistryKey Key = RootKey.OpenSubKey(KeyName, true))
+			using (RegistryKey? Key = RootKey.OpenSubKey(KeyName, true))
 			{
 				if (Key != null)
 				{
diff --git a/Engine/Source/Programs/UnrealGameSync/UnrealGameSync/UriHandlers/UriHandler.cs b/Engine/Source/Programs/UnrealGameSync/UnrealGameSync/UriHandlers/UriHandler.cs
--- a/Engine/Source/Programs/UnrealGameSync/UnrealGameSync/UriHandlers/UriHandler.cs
+++ b/Engine/Source/Programs/UnrealGameSync/UnrealGameSync/UriHandlers/UriHandler.cs
@@ -324,7 +324,7 @@ namespace UnrealGameSync
 				{
 					foreach (IGrouping<string, RegistrySetting> KeyNameGroup in RootKeyGroup.GroupBy(x => x.KeyName))
 					{
-						using (RegistryKey RegistryKey = RootKeyGroup.Key.OpenSubKey(KeyNameGroup.Key))
+						using (RegistryKey? RegistryKey = RootKeyGroup.Key.OpenSubKey(KeyNameGroup.Key))
 						{
 							if (RegistryKey == null)
 							{
