UGS: Fix project settings not being saved with new config model.

diff --git a/Engine/Source/Programs/UnrealGameSync/UnrealGameSync/Controls/WorkspaceControl.cs b/Engine/Source/Programs/UnrealGameSync/UnrealGameSync/Controls/WorkspaceControl.cs
--- a/Engine/Source/Programs/UnrealGameSync/UnrealGameSync/Controls/WorkspaceControl.cs
+++ b/Engine/Source/Programs/UnrealGameSync/UnrealGameSync/Controls/WorkspaceControl.cs
@@ -306,7 +306,7 @@ namespace UnrealGameSync
 			Settings = InSettings;
 			this.WorkspaceSettings = OpenProjectInfo.WorkspaceSettings;
 			this.WorkspaceState = OpenProjectInfo.WorkspaceState;
-			ProjectSettings = InSettings.FindOrAddProjectSettings(OpenProjectInfo.ProjectInfo);
+			ProjectSettings = InSettings.FindOrAddProjectSettings(OpenProjectInfo.ProjectInfo, OpenProjectInfo.WorkspaceSettings);
 
 			DesiredTaskbarState = Tuple.Create(TaskbarState.NoProgress, 0.0f);
 
@@ -2104,7 +2104,7 @@ namespace UnrealGameSync
 				LayoutInfo.BuildBadges = CreateBuildBadges(Change.Number, Summary);
 				LayoutInfo.CustomBadges = CreateCustomBadges(Change.Number, Summary);
 
-				ChangeNumberToLayoutInfo.Add(Change.Number, LayoutInfo);
+				ChangeNumberToLayoutInfo[Change.Number] = LayoutInfo;
 			}
 			return LayoutInfo;
 		}
@@ -5263,7 +5263,7 @@ namespace UnrealGameSync
 
 				// Save the settings
 				ProjectSettings.BuildSteps = ModifiedBuildSteps;
-				Settings.Save();
+				ProjectSettings.Save();
 
 				// Update the custom tools menu, because we might have changed it
 				UpdateBuildSteps();
@@ -5901,7 +5901,7 @@ namespace UnrealGameSync
 		{
 			ProjectSettings.FilterBadges.Clear();
 			ProjectSettings.FilterType = FilterType.None;
-			Settings.Save();
+			ProjectSettings.Save();
 
 			Settings.bShowAutomatedChanges = false;
 			Settings.Save();
@@ -5993,7 +5993,7 @@ namespace UnrealGameSync
 		private void FilterContextMenu_Type_ShowAll_Click(object sender, EventArgs e)
 		{
 			ProjectSettings.FilterType = FilterType.None;
-			Settings.Save();
+			ProjectSettings.Save();
 
 			UpdateBuildListFilter();
 		}
@@ -6001,7 +6001,7 @@ namespace UnrealGameSync
 		private void FilterContextMenu_Type_Code_Click(object sender, EventArgs e)
 		{
 			ProjectSettings.FilterType = FilterType.Code;
-			Settings.Save();
+			ProjectSettings.Save();
 
 			UpdateBuildListFilter();
 		}
@@ -6009,7 +6009,7 @@ namespace UnrealGameSync
 		private void FilterContextMenu_Type_Content_Click(object sender, EventArgs e)
 		{
 			ProjectSettings.FilterType = FilterType.Content;
-			Settings.Save();
+			ProjectSettings.Save();
 
 			UpdateBuildListFilter();
 		}
diff --git a/Engine/Source/Programs/UnrealGameSync/UnrealGameSyncShared/ConfigFile.cs b/Engine/Source/Programs/UnrealGameSync/UnrealGameSyncShared/ConfigFile.cs
--- a/Engine/Source/Programs/UnrealGameSync/UnrealGameSyncShared/ConfigFile.cs
+++ b/Engine/Source/Programs/UnrealGameSync/UnrealGameSyncShared/ConfigFile.cs
@@ -9,10 +9,13 @@ using System.Diagnostics.CodeAnalysis;
 using System.IO;
 using System.Linq;
 using System.Text;
+using System.Text.Json;
+using System.Text.Json.Serialization;
 using System.Threading.Tasks;
 
 namespace UnrealGameSync
 {
+	[JsonConverter(typeof(ConfigObjectJsonConverter))]
 	public class ConfigObject
 	{
 		const string ConfigSeparatorCharacters = "(),= \t\"";
@@ -298,6 +301,25 @@ namespace UnrealGameSync
 		}
 	}
 
+	public class ConfigObjectJsonConverter : JsonConverter<ConfigObject>
+	{
+		public override ConfigObject Read(ref Utf8JsonReader Reader, Type TypeToConvert, JsonSerializerOptions Options)
+		{
+			ConfigObject Object = new ConfigObject();
+			return Object;
+		}
+
+		public override void Write(Utf8JsonWriter Writer, ConfigObject Value, JsonSerializerOptions Options)
+		{
+			Writer.WriteStartObject();
+			foreach (KeyValuePair<string, string> Pair in Value.Pairs)
+			{
+				Writer.WriteString(Pair.Key, Pair.Value);
+			}
+			Writer.WriteEndObject();
+		}
+	}
+
 	[DebuggerDisplay("{Name}")]
 	public class ConfigSection
 	{
diff --git a/Engine/Source/Programs/UnrealGameSync/UnrealGameSyncShared/GlobalSettings.cs b/Engine/Source/Programs/UnrealGameSync/UnrealGameSyncShared/GlobalSettings.cs
--- a/Engine/Source/Programs/UnrealGameSync/UnrealGameSyncShared/GlobalSettings.cs
+++ b/Engine/Source/Programs/UnrealGameSync/UnrealGameSyncShared/GlobalSettings.cs
@@ -59,9 +59,30 @@ namespace UnrealGameSync
 		public FileReference File { get; }
 		public GlobalSettings Global { get; }
 
-		public virtual UserProjectSettings FindOrAddProjectSettings(ProjectInfo ProjectInfo)
+		public UserProjectSettings FindOrAddProjectSettings(ProjectInfo ProjectInfo, UserWorkspaceSettings Settings)
+		{
+			FileReference ConfigFile;
+			if (ProjectInfo.LocalFileName.HasExtension(".uprojectdirs"))
+			{
+				ConfigFile = FileReference.Combine(UserSettings.GetConfigDir(Settings.RootDir), "project.json");
+			}
+			else
+			{
+				ConfigFile = FileReference.Combine(UserSettings.GetConfigDir(Settings.RootDir), $"project_{ProjectInfo.LocalFileName.GetFileNameWithoutExtension()}.json");
+			}
+
+			UserProjectSettings? ProjectSettings;
+			if (!UserProjectSettings.TryLoad(ConfigFile, out ProjectSettings))
+			{
+				ProjectSettings = new UserProjectSettings(ConfigFile);
+				ImportProjectSettings(ProjectInfo, ProjectSettings);
+				ProjectSettings.Save();
+			}
+			return ProjectSettings;
+		}
+
+		protected virtual void ImportProjectSettings(ProjectInfo ProjectInfo, UserProjectSettings ProjectSettings)
 		{
-			return new UserProjectSettings();
 		}
 
 		protected virtual void ImportWorkspaceSettings(DirectoryReference RootDir, string ClientName, string BranchPath, UserWorkspaceSettings WorkspaceSettings)
diff --git a/Engine/Source/Programs/UnrealGameSync/UnrealGameSyncShared/UserSettings.cs b/Engine/Source/Programs/UnrealGameSync/UnrealGameSyncShared/UserSettings.cs
--- a/Engine/Source/Programs/UnrealGameSync/UnrealGameSyncShared/UserSettings.cs
+++ b/Engine/Source/Programs/UnrealGameSync/UnrealGameSyncShared/UserSettings.cs
@@ -497,9 +497,44 @@ namespace UnrealGameSync
 
 	public class UserProjectSettings
 	{
+		[JsonIgnore]
+		public FileReference ConfigFile { get; private set; } = null!;
+
 		public List<ConfigObject> BuildSteps { get; set; } = new List<ConfigObject>();
 		public FilterType FilterType { get; set; }
 		public HashSet<string> FilterBadges { get; set; } = new HashSet<string>(StringComparer.OrdinalIgnoreCase);
+
+		static object SyncRoot = new object();
+
+		private UserProjectSettings()
+		{
+		}
+
+		public UserProjectSettings(FileReference ConfigFile)
+		{
+			this.ConfigFile = ConfigFile;
+		}
+
+		public static bool TryLoad(FileReference ConfigFile, [NotNullWhen(true)] out UserProjectSettings? Settings)
+		{
+			if (FileReference.Exists(ConfigFile))
+			{
+				Settings = Utility.LoadJson<UserProjectSettings>(ConfigFile);
+				Settings.ConfigFile = ConfigFile;
+				return true;
+			}
+
+			Settings = null;
+			return false;
+		}
+
+		public void Save()
+		{
+			lock (SyncRoot)
+			{
+				Utility.SaveJson(ConfigFile, this);
+			}
+		}
 	}
 
 	public class UserSettings : GlobalSettingsFile
@@ -1011,13 +1046,11 @@ namespace UnrealGameSync
 			}
 		}
 
-		public override UserProjectSettings FindOrAddProjectSettings(ProjectInfo ProjectInfo)
+		protected override void ImportProjectSettings(ProjectInfo ProjectInfo, UserProjectSettings CurrentProject)
 		{
 			string ClientProjectFileName = Regex.Replace(ProjectInfo.ClientFileName, "^//[^/]+/", "");
 
 			// Read the project settings
-			UserProjectSettings CurrentProject = new UserProjectSettings();
-	
 			ConfigSection ProjectSection = ConfigFile.FindOrAddSection(ClientProjectFileName);
 			CurrentProject.BuildSteps.AddRange(ProjectSection.GetValues("BuildStep", new string[0]).Select(x => new ConfigObject(x)));
 
@@ -1029,8 +1062,6 @@ namespace UnrealGameSync
 
 			CurrentProject.FilterType = FilterType;
 			CurrentProject.FilterBadges.UnionWith(ProjectSection.GetValues("FilterBadges", new string[0]));
-
-			return CurrentProject;
 		}
 
 		public override void Save()
