UGS: Add code to read in ConfigObject instances from Json.

diff --git a/Engine/Source/Programs/UnrealGameSync/UnrealGameSyncShared/ConfigFile.cs b/Engine/Source/Programs/UnrealGameSync/UnrealGameSyncShared/ConfigFile.cs
--- a/Engine/Source/Programs/UnrealGameSync/UnrealGameSyncShared/ConfigFile.cs
+++ b/Engine/Source/Programs/UnrealGameSync/UnrealGameSyncShared/ConfigFile.cs
@@ -306,6 +306,20 @@ namespace UnrealGameSync
 		public override ConfigObject Read(ref Utf8JsonReader Reader, Type TypeToConvert, JsonSerializerOptions Options)
 		{
 			ConfigObject Object = new ConfigObject();
+			if (Reader.TokenType == JsonTokenType.StartObject)
+			{
+				while (Reader.Read() && Reader.TokenType == JsonTokenType.PropertyName)
+				{
+					string Name = Reader.GetString();
+					Reader.Read();
+					string Value = Reader.GetString();
+					Object.Pairs.Add(KeyValuePair.Create(Name, Value));
+				}
+			}
+			else
+			{
+				Reader.Skip();
+			}
 			return Object;
 		}
 
