UGS: Fix P4 server settings not being respected when opening a new project.

diff --git a/Engine/Source/Programs/UnrealGameSync/UnrealGameSync/Forms/OpenProjectWindow.cs b/Engine/Source/Programs/UnrealGameSync/UnrealGameSync/Forms/OpenProjectWindow.cs
--- a/Engine/Source/Programs/UnrealGameSync/UnrealGameSync/Forms/OpenProjectWindow.cs
+++ b/Engine/Source/Programs/UnrealGameSync/UnrealGameSync/Forms/OpenProjectWindow.cs
@@ -286,7 +286,7 @@ namespace UnrealGameSync
 			{
 				ILogger<OpenProjectInfo> Logger = ServiceProvider.GetRequiredService<ILogger<OpenProjectInfo>>();
 
-				ModalTask<OpenProjectInfo>? NewOpenProjectInfo = PerforceModalTask.Execute(this, "Opening project", "Opening project, please wait...", DefaultSettings, (x, y) => DetectSettingsAsync(x, SelectedProject, Settings, Logger, y), Logger);
+				ModalTask<OpenProjectInfo>? NewOpenProjectInfo = PerforceModalTask.Execute(this, "Opening project", "Opening project, please wait...", Perforce, (x, y) => DetectSettingsAsync(x, SelectedProject, Settings, Logger, y), Logger);
 				if (NewOpenProjectInfo != null && NewOpenProjectInfo.Succeeded)
 				{
 					OpenProjectInfo = NewOpenProjectInfo.Result;
