Horde: Fix build error in UnrealGameSyncCmd.

diff --git a/Engine/Source/Programs/UnrealGameSync/UnrealGameSyncCmd/Program.cs b/Engine/Source/Programs/UnrealGameSync/UnrealGameSyncCmd/Program.cs
--- a/Engine/Source/Programs/UnrealGameSync/UnrealGameSyncCmd/Program.cs
+++ b/Engine/Source/Programs/UnrealGameSync/UnrealGameSyncCmd/Program.cs
@@ -623,7 +623,7 @@ namespace UnrealGameSyncCmd
 				Options |= WorkspaceUpdateOptions.RemoveFilteredFiles;
 
 				ProjectInfo ProjectInfo = State.CreateProjectInfo();
-				UserProjectSettings ProjectSettings = Context.UserSettings.FindOrAddProjectSettings(ProjectInfo);
+				UserProjectSettings ProjectSettings = Context.UserSettings.FindOrAddProjectSettings(ProjectInfo, Settings);
 
 				ConfigFile ProjectConfig = await ConfigUtils.ReadProjectConfigFileAsync(PerforceClient, ProjectInfo, Logger, CancellationToken.None);
 				string[] SyncFilter = ReadSyncFilter(Settings, Context.UserSettings, ProjectConfig);
