UGS: Fix error when syncing if files are checked out.

diff --git a/Engine/Source/Programs/UnrealGameSync/UnrealGameSyncShared/WorkspaceUpdate.cs b/Engine/Source/Programs/UnrealGameSync/UnrealGameSyncShared/WorkspaceUpdate.cs
--- a/Engine/Source/Programs/UnrealGameSync/UnrealGameSyncShared/WorkspaceUpdate.cs
+++ b/Engine/Source/Programs/UnrealGameSync/UnrealGameSyncShared/WorkspaceUpdate.cs
@@ -784,6 +784,11 @@ namespace UnrealGameSync
 								{
 									return (WorkspaceUpdateResult.FailedToSync, $"Couldn't enumerate changes matching {SyncPath}.");
 								}
+								if (Response.Info != null)
+								{
+									Logger.LogInformation("Note: {0}", Response.Info.Data);
+									continue;
+								}
 
 								SyncRecord Record = Response.Data;
 
@@ -791,7 +796,7 @@ namespace UnrealGameSync
 								string RelativePath;
 								try
 								{
-									FileReference SyncFile = new FileReference(Record.Path);
+									FileReference SyncFile = new FileReference(Record.Path.ToString());
 									RelativePath = PerforceUtils.GetClientRelativePath(Project.LocalRootPath, SyncFile);
 								}
 								catch (PathTooLongException)
@@ -802,7 +807,7 @@ namespace UnrealGameSync
 
 								// Create the sync record
 								long SyncSize = (Record.Action == SyncAction.Deleted) ? 0 : Record.FileSize;
-								SyncFiles.Add(new SyncFile(Record.DepotFile, RelativePath, SyncSize));
+								SyncFiles.Add(new SyncFile(Record.DepotFile.ToString(), RelativePath, SyncSize));
 								Counter.Increment();
 							}
 
@@ -1592,7 +1597,7 @@ namespace UnrealGameSync
 		{
 			lock (State)
 			{
-				State.RemainingDepotPaths.Remove(Record.DepotFile);
+				State.RemainingDepotPaths.Remove(Record.DepotFile.ToString());
 
 				string Message = String.Format("{0} ({1}/{2})", Prefix, State.TotalDepotPaths - State.RemainingDepotPaths.Count, State.TotalDepotPaths);
 				float Fraction = Math.Min((float)(State.TotalDepotPaths - State.RemainingDepotPaths.Count) / (float)State.TotalDepotPaths, 1.0f);
@@ -1649,7 +1654,11 @@ namespace UnrealGameSync
 			foreach (PerforceResponse<SyncRecord> Response in Responses)
 			{
 				const string NoClobberPrefix = "Can't clobber writable file ";
-				if (Response.Succeeded)
+				if (Response.Info != null)
+				{
+					continue;
+				}
+				else if (Response.Succeeded)
 				{
 					SyncOutput(Response.Data);
 				}
