UGS: Fix access of TelemetrySink after it's been disposed.

diff --git a/Engine/Source/Programs/UnrealGameSync/UnrealGameSync/Program.cs b/Engine/Source/Programs/UnrealGameSync/UnrealGameSync/Program.cs
--- a/Engine/Source/Programs/UnrealGameSync/UnrealGameSync/Program.cs
+++ b/Engine/Source/Programs/UnrealGameSync/UnrealGameSync/Program.cs
@@ -149,31 +149,42 @@ namespace UnrealGameSync
 					ILogger TelemetryLogger = LoggerProvider.CreateLogger("Telemetry");
 					TelemetryLogger.LogInformation("Creating telemetry sink for session {SessionId}", SessionId);
 
-					ITelemetrySink? PrevTelemetrySink = Telemetry.ActiveSink;
 					using (ITelemetrySink TelemetrySink = DeploymentSettings.CreateTelemetrySink(UserName, SessionId, TelemetryLogger))
 					{
-						Telemetry.ActiveSink = TelemetrySink;
-
-						Telemetry.SendEvent("Startup", new { User = Environment.UserName, Machine = Environment.MachineName });
+						ITelemetrySink? PrevTelemetrySink = Telemetry.ActiveSink;
+						try
+						{
+							Telemetry.ActiveSink = TelemetrySink;
 
-						AppDomain.CurrentDomain.UnhandledException += CurrentDomain_UnhandledException;
+							Telemetry.SendEvent("Startup", new { User = Environment.UserName, Machine = Environment.MachineName });
 
-						IPerforceSettings DefaultSettings = new PerforceSettings(ServerAndPort, UserName) { PreferNativeClient = true };
+							AppDomain.CurrentDomain.UnhandledException += CurrentDomain_UnhandledException;
 
-						using (UpdateMonitor UpdateMonitor = new UpdateMonitor(DefaultSettings, UpdatePath, ServiceProvider))
-						{
-							ProgramApplicationContext Context = new ProgramApplicationContext(DefaultSettings, UpdateMonitor, DeploymentSettings.ApiUrl, DataFolder, ActivateEvent, bRestoreState, UpdateSpawn, ProjectFileName, bUnstable, ServiceProvider, Uri);
-							Application.Run(Context);
+							IPerforceSettings DefaultSettings = new PerforceSettings(ServerAndPort, UserName) { PreferNativeClient = true };
 
-							if (UpdateMonitor.IsUpdateAvailable && UpdateSpawn != null)
+							using (UpdateMonitor UpdateMonitor = new UpdateMonitor(DefaultSettings, UpdatePath, ServiceProvider))
 							{
-								InstanceMutex.Close();
-								bool bLaunchUnstable = UpdateMonitor.RelaunchUnstable ?? bUnstable;
-								Utility.SpawnProcess(UpdateSpawn, "-restorestate" + (bLaunchUnstable ? " -unstable" : ""));
+								ProgramApplicationContext Context = new ProgramApplicationContext(DefaultSettings, UpdateMonitor, DeploymentSettings.ApiUrl, DataFolder, ActivateEvent, bRestoreState, UpdateSpawn, ProjectFileName, bUnstable, ServiceProvider, Uri);
+								Application.Run(Context);
+
+								if (UpdateMonitor.IsUpdateAvailable && UpdateSpawn != null)
+								{
+									InstanceMutex.Close();
+									bool bLaunchUnstable = UpdateMonitor.RelaunchUnstable ?? bUnstable;
+									Utility.SpawnProcess(UpdateSpawn, "-restorestate" + (bLaunchUnstable ? " -unstable" : ""));
+								}
 							}
 						}
+						catch (Exception Ex)
+						{
+							Telemetry.SendEvent("Crash", new { Exception = Ex.ToString() });
+							throw;
+						}
+						finally
+						{
+							Telemetry.ActiveSink = PrevTelemetrySink;
+						}
 					}
-					Telemetry.ActiveSink = PrevTelemetrySink;
 				}
 			}
 		}
