UGS: Fix exception opening a project that is multiple subdirectories under the root.

diff --git a/Engine/Source/Programs/UnrealGameSync/UnrealGameSyncShared/OpenProjectInfo.cs b/Engine/Source/Programs/UnrealGameSync/UnrealGameSyncShared/OpenProjectInfo.cs
--- a/Engine/Source/Programs/UnrealGameSync/UnrealGameSyncShared/OpenProjectInfo.cs
+++ b/Engine/Source/Programs/UnrealGameSync/UnrealGameSyncShared/OpenProjectInfo.cs
@@ -187,7 +187,7 @@ namespace UnrealGameSync
 					if (NewSelectedClientFileName[EndIdx] == '/')
 					{
 						List<PerforceResponse<FStatRecord>> FileRecords = await PerforceClient.TryFStatAsync(FStatOptions.None, NewSelectedClientFileName.Substring(0, EndIdx) + "/Engine/Build/Build.version", CancellationToken).ToListAsync();
-						if (FileRecords.Succeeded())
+						if (FileRecords.Succeeded() && FileRecords.Count > 0)
 						{
 							FStatRecord FileRecord = FileRecords[0].Data;
 							if (FileRecord.ClientFile == null)
