UGS: Fix support for replicating the editor title bar tint color to UGS tabs.

diff --git a/Engine/Source/Programs/UnrealGameSync/UnrealGameSync/Controls/WorkspaceControl.cs b/Engine/Source/Programs/UnrealGameSync/UnrealGameSync/Controls/WorkspaceControl.cs
--- a/Engine/Source/Programs/UnrealGameSync/UnrealGameSync/Controls/WorkspaceControl.cs
+++ b/Engine/Source/Programs/UnrealGameSync/UnrealGameSync/Controls/WorkspaceControl.cs
@@ -407,12 +407,14 @@ namespace UnrealGameSync
 
 			if (SelectedFileName.HasExtension(".uproject"))
 			{
-				FileReference TintColorFileName = GetTintColorFileName();
+				DirectoryReference ConfigDir = DirectoryReference.Combine(SelectedFileName.Directory, "Saved", "Config");
 
-				DirectoryReference.CreateDirectory(TintColorFileName.Directory);
+				DirectoryReference.CreateDirectory(ConfigDir);
 
-				EditorConfigWatcher.Path = TintColorFileName.Directory.FullName;
-				EditorConfigWatcher.Filter = TintColorFileName.GetFileName();
+				EditorConfigWatcher.Path = ConfigDir.FullName;
+				EditorConfigWatcher.NotifyFilter = NotifyFilters.LastWrite | NotifyFilters.FileName;
+				EditorConfigWatcher.Filter = "*.ini";
+				EditorConfigWatcher.IncludeSubdirectories = true;
 				EditorConfigWatcher.EnableRaisingEvents = true;
 
 				TintColor = GetTintColor();
@@ -6036,70 +6038,110 @@ namespace UnrealGameSync
 			UpdateTintColor();
 		}
 
-		private FileReference GetTintColorFileName()
+		private Color? GetTintColor()
 		{
-			return FileReference.Combine(SelectedFileName.Directory, "Saved", "Config", "Windows", "EditorPerProjectUserSettings.ini");
+			try
+			{
+				return GetNewTintColor() ?? GetLegacyTintColor();
+			}
+			catch
+			{
+				return null;
+			}
 		}
 
-		private Color? GetTintColor()
+		private Color? GetNewTintColor()
 		{
-			try
+			FileReference NewConfigFile = FileReference.Combine(SelectedFileName.Directory, "Saved", "Config", "WindowsEditor", "EditorPerProjectUserSettings.ini");
+			if (!FileReference.Exists(NewConfigFile))
 			{
-				FileReference FileName = GetTintColorFileName();
-				if (!FileReference.Exists(FileName))
-				{
-					return null;
-				}
+				return null;
+			}
 
-				ConfigFile Config = new ConfigFile();
-				Config.Load(FileName);
+			ConfigFile Config = new ConfigFile();
+			Config.Load(NewConfigFile);
 
-				ConfigSection Section = Config.FindSection("/Script/EditorStyle.EditorStyleSettings");
-				if (Section == null)
-				{
-					return null;
-				}
+			ConfigSection Section = Config.FindSection("/Script/EditorStyle.EditorStyleSettings");
+			if (Section == null)
+			{
+				return null;
+			}
 
-				string? OverrideValue = Section.GetValue("EditorMainWindowBackgroundOverride", null);
-				if (OverrideValue == null)
-				{
-					return null;
-				}
+			string? EnableColor = Section.GetValue("bEnableEditorWindowBackgroundColor", null);
+			if (EnableColor == null || !String.Equals(EnableColor, "True", StringComparison.OrdinalIgnoreCase))
+			{
+				return null;
+			}
 
-				ConfigObject OverrideObject = new ConfigObject(OverrideValue);
+			string? BackgroundColor = Section.GetValue("EditorWindowBackgroundColor", null);
+			if (BackgroundColor == null)
+			{
+				return null;
+			}
 
-				string? TintColorValue = OverrideObject.GetValue("TintColor");
-				if (TintColorValue == null)
-				{
-					return null;
-				}
+			ConfigObject OverrideObject = new ConfigObject(BackgroundColor);
 
-				ConfigObject TintColorObject = new ConfigObject(TintColorValue);
-				if (TintColorObject.GetValue("ColorUseRule", "") != "UseColor_Specified")
-				{
-					return null;
-				}
+			float R, G, B;
+			if (!float.TryParse(OverrideObject.GetValue("R", ""), out R) || !float.TryParse(OverrideObject.GetValue("G", ""), out G) || !float.TryParse(OverrideObject.GetValue("B", ""), out B))
+			{
+				return null;
+			}
 
-				string? SpecifiedColorValue = TintColorObject.GetValue("SpecifiedColor");
-				if(SpecifiedColorValue == null)
-				{
-					return null;
-				}
+			return Color.FromArgb((int)(255.0f * R), (int)(255.0f * G), (int)(255.0f * B));
+		}
 
-				ConfigObject SpecifiedColorObject = new ConfigObject(SpecifiedColorValue);
+		private Color? GetLegacyTintColor()
+		{
+			FileReference FileName = FileReference.Combine(SelectedFileName.Directory, "Saved", "Config", "Windows", "EditorPerProjectUserSettings.ini");
+			if (!FileReference.Exists(FileName))
+			{
+				return null;
+			}
 
-				float R, G, B;
-				if (!float.TryParse(SpecifiedColorObject.GetValue("R", ""), out R) || !float.TryParse(SpecifiedColorObject.GetValue("G", ""), out G) || !float.TryParse(SpecifiedColorObject.GetValue("B", ""), out B))
-				{
-					return null;
-				}
+			ConfigFile Config = new ConfigFile();
+			Config.Load(FileName);
 
-				return Color.FromArgb((int)(255.0f * R), (int)(255.0f * G), (int)(255.0f * B));
+			ConfigSection Section = Config.FindSection("/Script/EditorStyle.EditorStyleSettings");
+			if (Section == null)
+			{
+				return null;
 			}
-			catch
+
+			string? OverrideValue = Section.GetValue("EditorMainWindowBackgroundOverride", null);
+			if (OverrideValue == null)
 			{
 				return null;
 			}
+
+			ConfigObject OverrideObject = new ConfigObject(OverrideValue);
+
+			string? TintColorValue = OverrideObject.GetValue("TintColor");
+			if (TintColorValue == null)
+			{
+				return null;
+			}
+
+			ConfigObject TintColorObject = new ConfigObject(TintColorValue);
+			if (TintColorObject.GetValue("ColorUseRule", "") != "UseColor_Specified")
+			{
+				return null;
+			}
+
+			string? SpecifiedColorValue = TintColorObject.GetValue("SpecifiedColor");
+			if (SpecifiedColorValue == null)
+			{
+				return null;
+			}
+
+			ConfigObject SpecifiedColorObject = new ConfigObject(SpecifiedColorValue);
+
+			float R, G, B;
+			if (!float.TryParse(SpecifiedColorObject.GetValue("R", ""), out R) || !float.TryParse(SpecifiedColorObject.GetValue("G", ""), out G) || !float.TryParse(SpecifiedColorObject.GetValue("B", ""), out B))
+			{
+				return null;
+			}
+
+			return Color.FromArgb((int)(255.0f * R), (int)(255.0f * G), (int)(255.0f * B));
 		}
 
 		private void UpdateTintColor()
