UGS: Fix build errors.

diff --git a/Engine/Source/Programs/UnrealGameSync/UnrealGameSync/Controls/WorkspaceControl.cs b/Engine/Source/Programs/UnrealGameSync/UnrealGameSync/Controls/WorkspaceControl.cs
--- a/Engine/Source/Programs/UnrealGameSync/UnrealGameSync/Controls/WorkspaceControl.cs
+++ b/Engine/Source/Programs/UnrealGameSync/UnrealGameSync/Controls/WorkspaceControl.cs
@@ -3818,7 +3818,7 @@ namespace UnrealGameSync
 						{
 							if (Idx == 1 || !await Perforce.OpenedAsync(OpenedOptions.None, FileSpecList.Any, CancellationToken).AnyAsync())
 							{
-								await Perforce.SwitchClientToStreamAsync(Perforce.Settings.ClientName!, NewStreamName, SwitchClientOptions.IgnoreOpenFiles, CancellationToken);
+								await Perforce.SwitchClientToStreamAsync(NewStreamName, SwitchClientOptions.IgnoreOpenFiles, CancellationToken);
 								return true;
 							}
 							return false;
diff --git a/Engine/Source/Programs/UnrealGameSync/UnrealGameSync/Forms/MainWindow.cs b/Engine/Source/Programs/UnrealGameSync/UnrealGameSync/Forms/MainWindow.cs
--- a/Engine/Source/Programs/UnrealGameSync/UnrealGameSync/Forms/MainWindow.cs
+++ b/Engine/Source/Programs/UnrealGameSync/UnrealGameSync/Forms/MainWindow.cs
@@ -360,7 +360,7 @@ namespace UnrealGameSync
 				// Switch the stream
 				Func<IPerforceConnection, CancellationToken, Task> SwitchTask = async (Connection, CancellationToken) =>
 				{
-					await Connection.SwitchClientToStreamAsync(Connection.Settings.ClientName!, StreamName, SwitchClientOptions.None, CancellationToken);
+					await Connection.SwitchClientToStreamAsync(StreamName, SwitchClientOptions.None, CancellationToken);
 				};
 
 				PerforceSettings Settings = new PerforceSettings(WorkspaceInfo.ServerAndPort, WorkspaceInfo.UserName);
