UGS: Fix build error.

diff --git a/Engine/Source/Programs/UnrealGameSync/UnrealGameSyncCmd/Program.cs b/Engine/Source/Programs/UnrealGameSync/UnrealGameSyncCmd/Program.cs
--- a/Engine/Source/Programs/UnrealGameSync/UnrealGameSyncCmd/Program.cs
+++ b/Engine/Source/Programs/UnrealGameSync/UnrealGameSyncCmd/Program.cs
@@ -1133,7 +1133,7 @@ namespace UnrealGameSyncCmd
 					throw new UserErrorException("Client {ClientName} has files opened. Use -Force to switch anyway.", PerforceClient.Settings.ClientName!);
 				}
 
-				await PerforceClient.SwitchClientToStreamAsync(PerforceClient.Settings.ClientName!, StreamName, SwitchClientOptions.IgnoreOpenFiles);
+				await PerforceClient.SwitchClientToStreamAsync(StreamName, SwitchClientOptions.IgnoreOpenFiles);
 
 				Logger.LogInformation("Switched to stream {StreamName}", StreamName);
 			}
