UGS: Fix overridden Perforce settings not being used when opening a project.

diff --git a/Engine/Source/Programs/UnrealGameSync/UnrealGameSync/Forms/MainWindow.cs b/Engine/Source/Programs/UnrealGameSync/UnrealGameSync/Forms/MainWindow.cs
--- a/Engine/Source/Programs/UnrealGameSync/UnrealGameSync/Forms/MainWindow.cs
+++ b/Engine/Source/Programs/UnrealGameSync/UnrealGameSync/Forms/MainWindow.cs
@@ -1130,7 +1130,9 @@ namespace UnrealGameSync
 				TaskFlags |= ModalTaskFlags.Quiet;
 			}
 
-			ModalTask<OpenProjectInfo>? SettingsTask = PerforceModalTask.Execute(this, "Opening Project", "Opening project, please wait...", DefaultPerforceSettings, (p, c) => OpenProjectWindow.DetectSettingsAsync(p, Project, Settings, ProjectLogger, c), ProjectLogger, TaskFlags);
+			PerforceSettings PerforceSettings = Utility.OverridePerforceSettings(DefaultPerforceSettings, Project.ServerAndPort, Project.UserName);
+
+			ModalTask<OpenProjectInfo>? SettingsTask = PerforceModalTask.Execute(this, "Opening Project", "Opening project, please wait...", PerforceSettings, (p, c) => OpenProjectWindow.DetectSettingsAsync(p, Project, Settings, ProjectLogger, c), ProjectLogger, TaskFlags);
 			if (SettingsTask == null || SettingsTask.Failed)
 			{
 				if(SettingsTask != null) CreateErrorPanel(ReplaceTabIdx, Project, SettingsTask.Error);
diff --git a/Engine/Source/Programs/UnrealGameSync/UnrealGameSync/Forms/OpenProjectWindow.cs b/Engine/Source/Programs/UnrealGameSync/UnrealGameSync/Forms/OpenProjectWindow.cs
--- a/Engine/Source/Programs/UnrealGameSync/UnrealGameSync/Forms/OpenProjectWindow.cs
+++ b/Engine/Source/Programs/UnrealGameSync/UnrealGameSync/Forms/OpenProjectWindow.cs
@@ -286,7 +286,9 @@ namespace UnrealGameSync
 			{
 				ILogger<OpenProjectInfo> Logger = ServiceProvider.GetRequiredService<ILogger<OpenProjectInfo>>();
 
-				ModalTask<OpenProjectInfo>? NewOpenProjectInfo = PerforceModalTask.Execute(this, "Opening project", "Opening project, please wait...", Perforce, (x, y) => DetectSettingsAsync(x, SelectedProject, Settings, Logger, y), Logger);
+				PerforceSettings NewPerforceSettings = Utility.OverridePerforceSettings(Perforce, SelectedProject.ServerAndPort, SelectedProject.UserName);
+
+				ModalTask<OpenProjectInfo>? NewOpenProjectInfo = PerforceModalTask.Execute(this, "Opening project", "Opening project, please wait...", NewPerforceSettings, (x, y) => DetectSettingsAsync(x, SelectedProject, Settings, Logger, y), Logger);
 				if (NewOpenProjectInfo != null && NewOpenProjectInfo.Succeeded)
 				{
 					OpenProjectInfo = NewOpenProjectInfo.Result;
