UGS: Fix exception performing sync with "Any open projects" selected.

diff --git a/Engine/Source/Programs/UnrealGameSync/UnrealGameSync/Controls/WorkspaceControl.cs b/Engine/Source/Programs/UnrealGameSync/UnrealGameSync/Controls/WorkspaceControl.cs
--- a/Engine/Source/Programs/UnrealGameSync/UnrealGameSync/Controls/WorkspaceControl.cs
+++ b/Engine/Source/Programs/UnrealGameSync/UnrealGameSync/Controls/WorkspaceControl.cs
@@ -4925,7 +4925,7 @@ namespace UnrealGameSync
 			{
 				// Check the project config if they have a scheduled setting selected.
 				string? ScheduledSyncTypeID = null;
-				UserSelectedProjectSettings ProjectSetting = Settings.ScheduleProjects.First(x => x.LocalPath != null && x.LocalPath.Equals(SelectedProject.LocalPath, StringComparison.OrdinalIgnoreCase));
+				UserSelectedProjectSettings ProjectSetting = Settings.ScheduleProjects.FirstOrDefault(x => x.LocalPath != null && x.LocalPath.Equals(SelectedProject.LocalPath, StringComparison.OrdinalIgnoreCase));
 				if (ProjectSetting != null)
 				{
 					ScheduledSyncTypeID = ProjectSetting.ScheduledSyncTypeID;
