UGS: Add a URI handler that opens a changelist in P4V (ugs://change?number=12345).

diff --git a/Engine/Source/Programs/UnrealGameSync/UnrealGameSync/UriHandlers/P4Handler.cs b/Engine/Source/Programs/UnrealGameSync/UnrealGameSync/UriHandlers/P4Handler.cs
--- a/Engine/Source/Programs/UnrealGameSync/UnrealGameSync/UriHandlers/P4Handler.cs
+++ b/Engine/Source/Programs/UnrealGameSync/UnrealGameSync/UriHandlers/P4Handler.cs
@@ -29,6 +29,13 @@ namespace UnrealGameSync
 			Program.SpawnP4VC(CommandLine);
 
 			return new UriResult() { Success = true };
-		}		
+		}
+
+		[UriHandler(true)]
+		public static UriResult Change(int number)
+		{
+			Program.SpawnP4VC($"change {number}");
+			return new UriResult() { Success = true };
+		}
 	}
 }
\ No newline at end of file
