UGS: Fix zombie process instances left around when interrupting a sync.

diff --git a/Engine/Source/Programs/UnrealGameSync/UnrealGameSyncShared/WorkspaceUpdate.cs b/Engine/Source/Programs/UnrealGameSync/UnrealGameSyncShared/WorkspaceUpdate.cs
--- a/Engine/Source/Programs/UnrealGameSync/UnrealGameSyncShared/WorkspaceUpdate.cs
+++ b/Engine/Source/Programs/UnrealGameSync/UnrealGameSyncShared/WorkspaceUpdate.cs
@@ -1564,21 +1564,25 @@ namespace UnrealGameSync
 				// Run one worker on the current thread
 				await StaticSyncWorker(1, Perforce, Context, State, SyncOutput, LogWrapper, CancellationToken);
 
+				// Allow the tasks to throw
+				foreach (Task ChildTask in ChildTasks)
+				{
+					await ChildTask;
+				}
+
 				// Return the result that was set on the state object
 				return (State.Result, State.StatusMessage);
 			}
 			finally
 			{
-				await Task.WhenAll(ChildTasks);
-
-				foreach (IPerforceConnection ChildConnection in ChildConnections)
+				foreach (Task ChildTask in ChildTasks)
 				{
-					ChildConnection.Dispose();
+					await ChildTask.ContinueWith(_ => { }); // Swallow exceptions until we've disposed the connections
 				}
 
-				foreach (Task ChildTask in ChildTasks)
+				foreach (IPerforceConnection ChildConnection in ChildConnections)
 				{
-					await ChildTask;
+					ChildConnection.Dispose();
 				}
 			}
 		}
