UGS: Catch more exception types thrown by PerforceMonitor.

diff --git a/Engine/Source/Programs/UnrealGameSync/UnrealGameSync/PerforceMonitor.cs b/Engine/Source/Programs/UnrealGameSync/UnrealGameSync/PerforceMonitor.cs
--- a/Engine/Source/Programs/UnrealGameSync/UnrealGameSync/PerforceMonitor.cs
+++ b/Engine/Source/Programs/UnrealGameSync/UnrealGameSync/PerforceMonitor.cs
@@ -248,7 +248,7 @@ namespace UnrealGameSync
 						await PollForUpdatesInner(CancellationToken);
 					}
 				}
-				catch (TaskCanceledException) when (CancellationToken.IsCancellationRequested)
+				catch (OperationCanceledException) when (CancellationToken.IsCancellationRequested)
 				{
 					break;
 				}
