UGS: Add separate buttons to open the application/workspace folders in the diagnostics window.

diff --git a/Engine/Source/Programs/UnrealGameSync/UnrealGameSync/Controls/WorkspaceControl.cs b/Engine/Source/Programs/UnrealGameSync/UnrealGameSync/Controls/WorkspaceControl.cs
--- a/Engine/Source/Programs/UnrealGameSync/UnrealGameSync/Controls/WorkspaceControl.cs
+++ b/Engine/Source/Programs/UnrealGameSync/UnrealGameSync/Controls/WorkspaceControl.cs
@@ -201,8 +201,9 @@ namespace UnrealGameSync
 		const int BuildListExpandCount = 250;
 
 		IWorkspaceControlOwner Owner;
+		DirectoryReference AppDataFolder;
 		string? ApiUrl;
-		DirectoryReference DataFolder;
+		DirectoryReference WorkspaceDataFolder;
 		IServiceProvider ServiceProvider;
 		ILogger Logger;
 		IPerforceSettings PerforceSettings;
@@ -290,15 +291,16 @@ namespace UnrealGameSync
 		// Placeholder text that is in the control and cleared when the user starts editing.
 		static string AuthorFilterPlaceholderText = "<username>";
 
-		public WorkspaceControl(IWorkspaceControlOwner InOwner, string? InApiUrl, OpenProjectInfo OpenProjectInfo, IServiceProvider InServiceProvider, UserSettings InSettings, OIDCTokenManager? InOidcTokenManager)
+		public WorkspaceControl(IWorkspaceControlOwner InOwner, DirectoryReference InAppDataFolder, string? InApiUrl, OpenProjectInfo OpenProjectInfo, IServiceProvider InServiceProvider, UserSettings InSettings, OIDCTokenManager? InOidcTokenManager)
 		{
 			InitializeComponent();
 
 			MainThreadSynchronizationContext = SynchronizationContext.Current!;
 
 			Owner = InOwner;
+			AppDataFolder = InAppDataFolder;
 			ApiUrl = InApiUrl;
-			DataFolder = OpenProjectInfo.ProjectInfo.DataFolder;
+			WorkspaceDataFolder = OpenProjectInfo.ProjectInfo.DataFolder;
 			ServiceProvider = InServiceProvider;
 			Logger = InServiceProvider.GetRequiredService<ILogger<WorkspaceControl>>();
 			PerforceSettings = OpenProjectInfo.PerforceSettings;
@@ -366,7 +368,7 @@ namespace UnrealGameSync
 			Workspace = new Workspace(PerforceClientSettings, Project, WorkspaceState, OpenProjectInfo.WorkspaceProjectConfigFile, OpenProjectInfo.WorkspaceProjectStreamFilter, new LogControlTextWriter(SyncLog), ServiceProvider);
 			Workspace.OnUpdateComplete += UpdateCompleteCallback;
 
-			FileReference ProjectLogBaseName = FileReference.Combine(DataFolder, "sync.log");
+			FileReference ProjectLogBaseName = FileReference.Combine(WorkspaceDataFolder, "sync.log");
 
 			ILogger PerforceLogger = ServiceProvider.GetRequiredService<ILogger<PerforceMonitor>>();
 			PerforceMonitor = new PerforceMonitor(PerforceClientSettings, OpenProjectInfo.ProjectInfo, OpenProjectInfo.LatestProjectConfigFile, OpenProjectInfo.ProjectInfo.CacheFolder, OpenProjectInfo.LocalConfigFiles, ServiceProvider);
@@ -5361,7 +5363,7 @@ namespace UnrealGameSync
 			DiagnosticsText.AppendFormat("Event monitor: {0}\n", (EventMonitor == null) ? "(inactive)" : EventMonitor.LastStatusMessage);
 			DiagnosticsText.AppendFormat("Issue monitor: {0}\n", (IssueMonitor == null) ? "(inactive)" : IssueMonitor.LastStatusMessage);
 
-			DiagnosticsWindow Diagnostics = new DiagnosticsWindow(DataFolder, DiagnosticsText.ToString(), Settings.GetCachedFilePaths());
+			DiagnosticsWindow Diagnostics = new DiagnosticsWindow(AppDataFolder, WorkspaceDataFolder, DiagnosticsText.ToString(), Settings.GetCachedFilePaths());
 			Diagnostics.ShowDialog(this);
 		}
 
diff --git a/Engine/Source/Programs/UnrealGameSync/UnrealGameSync/Forms/DiagnosticsWindow.cs b/Engine/Source/Programs/UnrealGameSync/UnrealGameSync/Forms/DiagnosticsWindow.cs
--- a/Engine/Source/Programs/UnrealGameSync/UnrealGameSync/Forms/DiagnosticsWindow.cs
+++ b/Engine/Source/Programs/UnrealGameSync/UnrealGameSync/Forms/DiagnosticsWindow.cs
@@ -13,20 +13,29 @@ namespace UnrealGameSync
 {
 	public partial class DiagnosticsWindow : Form
 	{
-		DirectoryReference DataFolder;
+		DirectoryReference AppDataFolder;
+		DirectoryReference WorkspaceDataFolder;
 		List<FileReference> ExtraFiles;
 
-		public DiagnosticsWindow(DirectoryReference InDataFolder, string InDiagnosticsText, IEnumerable<FileReference> InExtraFiles)
+		public DiagnosticsWindow(DirectoryReference InAppDataFolder, DirectoryReference InWorkspaceDataFolder, string InDiagnosticsText, IEnumerable<FileReference> InExtraFiles)
 		{
 			InitializeComponent();
-			DataFolder = InDataFolder;
+
+			AppDataFolder = InAppDataFolder;
+			WorkspaceDataFolder = InWorkspaceDataFolder;
+
 			DiagnosticsTextBox.Text = InDiagnosticsText.Replace("\n", "\r\n");
 			ExtraFiles = InExtraFiles.ToList();
 		}
 
-		private void ViewLogsButton_Click(object sender, EventArgs e)
+		private void ViewApplicationDataButton_Click(object sender, EventArgs e)
 		{
-			Process.Start("explorer.exe", DataFolder.FullName);
+			Process.Start("explorer.exe", AppDataFolder.FullName);
+		}
+
+		private void ViewWorkspaceDataButton_Click(object sender, EventArgs e)
+		{
+			Process.Start("explorer.exe", WorkspaceDataFolder.FullName);
 		}
 
 		private void SaveButton_Click(object sender, EventArgs e)
@@ -37,7 +46,7 @@ namespace UnrealGameSync
 			Dialog.FileName = Path.Combine(Dialog.InitialDirectory, "UGS-Diagnostics.zip");
 			if(Dialog.ShowDialog() == DialogResult.OK)
 			{
-				FileReference DiagnosticsFileName = FileReference.Combine(DataFolder, "Diagnostics.txt");
+				FileReference DiagnosticsFileName = FileReference.Combine(AppDataFolder, "Diagnostics.txt");
 				try
 				{
 					FileReference.WriteAllLines(DiagnosticsFileName, DiagnosticsTextBox.Lines);
@@ -53,20 +62,9 @@ namespace UnrealGameSync
 				{
 					using (ZipArchive Zip = new ZipArchive(File.OpenWrite(ZipFileName), ZipArchiveMode.Create))
 					{
-						foreach (FileReference FileName in DirectoryReference.EnumerateFiles(DataFolder))
-						{
-							if (!FileName.HasExtension(".exe") && !FileName.HasExtension(".dll"))
-							{
-								using (FileStream InputStream = FileReference.Open(FileName, FileMode.Open, FileAccess.Read, FileShare.ReadWrite))
-								{
-									ZipArchiveEntry Entry = Zip.CreateEntry(FileName.MakeRelativeTo(DataFolder).Replace('\\', '/'));
-									using (Stream OutputStream = Entry.Open())
-									{
-										InputStream.CopyTo(OutputStream);
-									}
-								}
-							}
-						}
+						AddFilesToZip(Zip, AppDataFolder, "App/");
+						AddFilesToZip(Zip, WorkspaceDataFolder, "Workspace/");
+
 						foreach (FileReference ExtraFile in ExtraFiles)
 						{
 							if(FileReference.Exists(ExtraFile))
@@ -90,5 +88,23 @@ namespace UnrealGameSync
 				}
 			}
 		}
+
+		private static void AddFilesToZip(ZipArchive Zip, DirectoryReference DataFolder, string RelativeDir)
+		{
+			foreach (FileReference FileName in DirectoryReference.EnumerateFiles(DataFolder))
+			{
+				if (!FileName.HasExtension(".exe") && !FileName.HasExtension(".dll"))
+				{
+					using (FileStream InputStream = FileReference.Open(FileName, FileMode.Open, FileAccess.Read, FileShare.ReadWrite))
+					{
+						ZipArchiveEntry Entry = Zip.CreateEntry(RelativeDir + FileName.MakeRelativeTo(DataFolder).Replace('\\', '/'));
+						using (Stream OutputStream = Entry.Open())
+						{
+							InputStream.CopyTo(OutputStream);
+						}
+					}
+				}
+			}
+		}
 	}
 }
diff --git a/Engine/Source/Programs/UnrealGameSync/UnrealGameSync/Forms/MainWindow.cs b/Engine/Source/Programs/UnrealGameSync/UnrealGameSync/Forms/MainWindow.cs
--- a/Engine/Source/Programs/UnrealGameSync/UnrealGameSync/Forms/MainWindow.cs
+++ b/Engine/Source/Programs/UnrealGameSync/UnrealGameSync/Forms/MainWindow.cs
@@ -1196,7 +1196,7 @@ namespace UnrealGameSync
 			}
 
 			// Now that we have the project settings, we can construct the tab
-			WorkspaceControl NewWorkspace = new WorkspaceControl(this, ApiUrl, OpenProjectInfo, ServiceProvider, Settings, OIDCTokenManager);
+			WorkspaceControl NewWorkspace = new WorkspaceControl(this, DataFolder, ApiUrl, OpenProjectInfo, ServiceProvider, Settings, OIDCTokenManager);
 			
 			NewWorkspace.Parent = TabPanel;
 			NewWorkspace.Dock = DockStyle.Fill;
