UGS: Fix files not being force-synced when cleaning a workspace. The previous code is a nop since the have table already has it.

diff --git a/Engine/Source/Programs/UnrealGameSync/UnrealGameSync/DeleteFilesTask.cs b/Engine/Source/Programs/UnrealGameSync/UnrealGameSync/DeleteFilesTask.cs
--- a/Engine/Source/Programs/UnrealGameSync/UnrealGameSync/DeleteFilesTask.cs
+++ b/Engine/Source/Programs/UnrealGameSync/UnrealGameSync/DeleteFilesTask.cs
@@ -28,7 +28,7 @@ namespace UnrealGameSync
 					RevisionsToSync.Add(String.Format("{0}#have", PerforceUtils.EscapePath(FileToSync.FullName)));
 				}
 
-				List<PerforceResponse<SyncRecord>> FailedRecords = await Perforce.TrySyncAsync(RevisionsToSync, CancellationToken).Where(x => x.Failed).ToListAsync(CancellationToken);
+				List<PerforceResponse<SyncRecord>> FailedRecords = await Perforce.TrySyncAsync(SyncOptions.Force, RevisionsToSync, CancellationToken).Where(x => x.Failed).ToListAsync(CancellationToken);
 				foreach (PerforceResponse<SyncRecord> FailedRecord in FailedRecords)
 				{
 					FailMessage.Append(FailedRecord.ToString());
