UGS: Allow sending metadata update notifications from PerforceMonitor while blocked on a long running operation.

diff --git a/Engine/Source/Programs/UnrealGameSync/UnrealGameSync/PerforceMonitor.cs b/Engine/Source/Programs/UnrealGameSync/UnrealGameSync/PerforceMonitor.cs
--- a/Engine/Source/Programs/UnrealGameSync/UnrealGameSync/PerforceMonitor.cs
+++ b/Engine/Source/Programs/UnrealGameSync/UnrealGameSync/PerforceMonitor.cs
@@ -497,35 +497,47 @@ namespace UnrealGameSync
 			}
 
 			// Update them in batches
-			foreach (IReadOnlyList<int> QueryChangeNumberBatch in QueryChangeNumbers.Batch(10))
+			using (CancellationTokenSource cancellationSource = new CancellationTokenSource())
 			{
-				// Skip this stuff if the user wants us to query for more changes
-				if (PendingMaxChanges != CurrentMaxChanges)
+				Task notifyTask = Task.CompletedTask;
+				foreach (IReadOnlyList<int> QueryChangeNumberBatch in QueryChangeNumbers.OrderByDescending(x => x).Batch(10))
 				{
-					break;
-				}
-
-				// If there's something to check for, find all the content changes after this changelist
-				List<DescribeRecord> DescribeRecords = await Perforce.DescribeAsync(QueryChangeNumberBatch.ToArray(), CancellationToken);
-				foreach (DescribeRecord DescribeRecord in DescribeRecords.OrderByDescending(x => x.Number))
-				{
-					int QueryChangeNumber = DescribeRecord.Number;
+					// Skip this stuff if the user wants us to query for more changes
+					if (PendingMaxChanges != CurrentMaxChanges)
+					{
+						break;
+					}
 
-					PerforceChangeDetails Details = new PerforceChangeDetails(DescribeRecord);
-					lock (this)
+					// If there's something to check for, find all the content changes after this changelist
+					List<DescribeRecord> DescribeRecords = await Perforce.DescribeAsync(QueryChangeNumberBatch.ToArray(), CancellationToken);
+					foreach (DescribeRecord DescribeRecord in DescribeRecords.OrderByDescending(x => x.Number))
 					{
-						if (!ChangeDetails.ContainsKey(QueryChangeNumber))
+						int QueryChangeNumber = DescribeRecord.Number;
+
+						PerforceChangeDetails Details = new PerforceChangeDetails(DescribeRecord);
+						lock (this)
+						{
+							if (!ChangeDetails.ContainsKey(QueryChangeNumber))
+							{
+								ChangeDetails.Add(QueryChangeNumber, Details);
+							}
+						}
+
+						// Reload the config file if it changes
+						if (DescribeRecord.Files.Any(x => x.DepotFile.EndsWith("/UnrealGameSync.ini", StringComparison.OrdinalIgnoreCase)))
 						{
-							ChangeDetails.Add(QueryChangeNumber, Details);
+							await UpdateProjectConfigFileAsync(Perforce, CancellationToken);
 						}
-					}
 
-					// Reload the config file if it changes
-					if (DescribeRecord.Files.Any(x => x.DepotFile.EndsWith("/UnrealGameSync.ini", StringComparison.OrdinalIgnoreCase)))
-					{
-						await UpdateProjectConfigFileAsync(Perforce, CancellationToken);
+						// Notify the caller after a fixed period of time, in case further updates are slow to arrive
+						if (notifyTask.IsCompleted)
+						{
+							notifyTask = Task.Delay(TimeSpan.FromSeconds(5.0), cancellationSource.Token).ContinueWith(_ => SynchronizationContext.Post(_ => OnUpdateMetadata?.Invoke(), null));
+						}
 					}
 				}
+				cancellationSource.Cancel();
+				await notifyTask.ContinueWith(_ => { }); // Ignore exceptions
 			}
 
 			// Find the last submitted code change by the current user
