PR #9034: Adding Sentry Support into UGS (Disabled by default)  (Contributed by [REDACTED_PERSON])

diff --git a/Engine/Source/Programs/UnrealGameSync/UnrealGameSync/DeploymentSettings.cs b/Engine/Source/Programs/UnrealGameSync/UnrealGameSync/DeploymentSettings.cs
--- a/Engine/Source/Programs/UnrealGameSync/UnrealGameSync/DeploymentSettings.cs
+++ b/Engine/Source/Programs/UnrealGameSync/UnrealGameSync/DeploymentSettings.cs
@@ -54,6 +54,11 @@ namespace UnrealGameSync
 		/// </summary>
 		public static readonly string? ToolsDepotPath = null;
 
+		/// <summary>
+		/// DSN for sending crash reports to Sentry.
+		/// </summary>
+		public static readonly string? SentryDsn = null;
+
 #if WITH_TELEMETRY
 		/// <summary>
 		/// Delegate used to create a new telemetry sink
diff --git a/Engine/Source/Programs/UnrealGameSync/UnrealGameSync/Program.cs b/Engine/Source/Programs/UnrealGameSync/UnrealGameSync/Program.cs
--- a/Engine/Source/Programs/UnrealGameSync/UnrealGameSync/Program.cs
+++ b/Engine/Source/Programs/UnrealGameSync/UnrealGameSync/Program.cs
@@ -4,6 +4,8 @@ using EpicGames.Core;
 using EpicGames.Perforce;
 using Microsoft.Extensions.DependencyInjection;
 using Microsoft.Extensions.Logging;
+using Sentry;
+using Sentry.Infrastructure;
 using System;
 using System.Collections.Generic;
 using System.Data.SqlClient;
@@ -19,6 +21,7 @@ using System.Text.Json;
 using System.Text.Json.Serialization;
 using System.Threading;
 using System.Threading.Tasks;
+using System.Windows;
 using System.Windows.Forms;
 
 namespace UnrealGameSync
@@ -36,6 +39,25 @@ namespace UnrealGameSync
 		[STAThread]
 		static void Main(string[] Args)
 		{
+			if (DeploymentSettings.SentryDsn != null)
+			{
+				SentryOptions sentryOptions = new SentryOptions();
+				sentryOptions.Dsn = DeploymentSettings.SentryDsn;
+				sentryOptions.StackTraceMode = StackTraceMode.Enhanced;
+				sentryOptions.AttachStacktrace = true;
+				sentryOptions.TracesSampleRate = 1.0;
+				sentryOptions.SendDefaultPii = true;
+				sentryOptions.Debug = true;
+				sentryOptions.AutoSessionTracking = true;
+				sentryOptions.DetectStartupTime = StartupTimeDetectionMode.Best;
+				sentryOptions.ReportAssembliesMode = ReportAssembliesMode.InformationalVersion;
+				sentryOptions.DiagnosticLogger = new TraceDiagnosticLogger(SentryLevel.Debug);
+				SentrySdk.Init(sentryOptions);
+
+				Application.ThreadException += Application_ThreadException_Sentry;
+				AppDomain.CurrentDomain.UnhandledException += CurrentDomain_UnhandledException_Sentry;
+			}
+
 			bool bFirstInstance;
 			using (Mutex InstanceMutex = new Mutex(true, "UnrealGameSyncRunning", out bFirstInstance))
 			{
@@ -198,6 +220,23 @@ namespace UnrealGameSync
 			}
 		}
 
+		private static void CurrentDomain_UnhandledException_Sentry(object Sender, UnhandledExceptionEventArgs Args)
+		{
+			Exception? Ex = Args.ExceptionObject as Exception;
+			if (Ex != null)
+			{
+				SentrySdk.CaptureException(Ex);
+			}
+		}
+
+		private static void Application_ThreadException_Sentry(object sender, ThreadExceptionEventArgs e)
+		{
+			SentrySdk.CaptureException(e.Exception);
+
+			ThreadExceptionDialog dialog = new ThreadExceptionDialog(e.Exception);
+			dialog.ShowDialog();
+		}
+
 		static void MergeUpdateSettings(FileReference UpdateConfigFile, ref string? UpdatePath, ref string? UpdateSpawn)
 		{
 			try
diff --git a/Engine/Source/Programs/UnrealGameSync/UnrealGameSync/UnrealGameSync.csproj b/Engine/Source/Programs/UnrealGameSync/UnrealGameSync/UnrealGameSync.csproj
--- a/Engine/Source/Programs/UnrealGameSync/UnrealGameSync/UnrealGameSync.csproj
+++ b/Engine/Source/Programs/UnrealGameSync/UnrealGameSync/UnrealGameSync.csproj
@@ -23,6 +23,7 @@
   <ItemGroup>
     <PackageReference Include="EnvDTE" Version="16.7.30328.74" />
     <PackageReference Include="IdentityModel.OidcClient" Version="3.1.2" />
+    <PackageReference Include="Sentry" Version="3.15.0" />
     <PackageReference Include="Serilog.Expressions" Version="3.2.1" />
     <PackageReference Include="Serilog.Extensions.Hosting" Version="4.2.0" />
     <PackageReference Include="Serilog.Extensions.Logging" Version="3.1.0" />
@@ -63,4 +64,4 @@
     </When>
   </Choose>
   <Import Project="..\..\..\..\Restricted\NotForLicensees\Source\Programs\UnrealGameSync\UnrealGameSync\UnrealGameSync.csproj" Condition="Exists('..\..\..\..\Restricted\NotForLicensees\Source\Programs\UnrealGameSync\UnrealGameSync\UnrealGameSync.csproj')" />
-</Project>
\ No newline at end of file
+</Project>
