UGS: Fix deleted projects being shown in the select project window.

diff --git a/Engine/Source/Programs/UnrealGameSync/UnrealGameSync/Forms/SelectProjectFromWorkspaceWindow.cs b/Engine/Source/Programs/UnrealGameSync/UnrealGameSync/Forms/SelectProjectFromWorkspaceWindow.cs
--- a/Engine/Source/Programs/UnrealGameSync/UnrealGameSync/Forms/SelectProjectFromWorkspaceWindow.cs
+++ b/Engine/Source/Programs/UnrealGameSync/UnrealGameSync/Forms/SelectProjectFromWorkspaceWindow.cs
@@ -47,7 +47,7 @@ namespace UnrealGameSync
 					string Filter = String.Format("//{0}/{1}", ClientName, Pattern);
 
 					List<FStatRecord> WildcardFileRecords = await Perforce.FStatAsync(Filter, CancellationToken).ToListAsync(CancellationToken);
-					WildcardFileRecords.RemoveAll(x => x.Action == FileAction.Delete || x.Action == FileAction.MoveDelete);
+					WildcardFileRecords.RemoveAll(x => x.HeadAction == FileAction.Delete || x.HeadAction == FileAction.MoveDelete);
 
 					FileRecords.AddRange(WildcardFileRecords);
 				}
