UGS: Fix build errors.

diff --git a/Engine/Source/Programs/UnrealGameSync/UnrealGameSyncLauncher/Program.cs b/Engine/Source/Programs/UnrealGameSync/UnrealGameSyncLauncher/Program.cs
--- a/Engine/Source/Programs/UnrealGameSync/UnrealGameSyncLauncher/Program.cs
+++ b/Engine/Source/Programs/UnrealGameSync/UnrealGameSyncLauncher/Program.cs
@@ -64,7 +64,7 @@ namespace UnrealGameSyncLauncher
 					// Try to do a sync with the current settings first
 					CaptureLogger Logger = new CaptureLogger();
 
-					IPerforceSettings Settings = PerforceSettings.Default.MergeWith(NewServerAndPort: ServerAndPort, NewUserName: UserName);
+					IPerforceSettings Settings = PerforceSettings.Default.MergeWith(newServerAndPort: ServerAndPort, newUserName: UserName);
 
 					ModalTask? Task = PerforceModalTask.Execute(null, "Updating", "Checking for updates, please wait...", Settings, (p, c) => SyncAndRun(p, DepotPath, bUnstable, Args, InstanceMutex, Logger, c), Logger);
 					if (Task == null)
