UGS: Use a separate config file for newer builds of UGS (UnrealGameSyncV2.ini), to avoid invalidating old configs if we have to revert. Config will still be read from old location if a new config file is not present.

diff --git a/Engine/Source/Programs/UnrealGameSync/UnrealGameSyncCmd/Program.cs b/Engine/Source/Programs/UnrealGameSync/UnrealGameSyncCmd/Program.cs
--- a/Engine/Source/Programs/UnrealGameSync/UnrealGameSyncCmd/Program.cs
+++ b/Engine/Source/Programs/UnrealGameSync/UnrealGameSyncCmd/Program.cs
@@ -200,7 +200,7 @@ namespace UnrealGameSyncCmd
 				GlobalSettingsFile Settings;
 				if (RuntimeInformation.IsOSPlatform(OSPlatform.Windows))
 				{
-					Settings = UserSettings.Create(FileReference.Combine(GlobalConfigFolder, "UnrealGameSync.ini"), FileReference.Combine(GlobalConfigFolder, "Global.json"), Logger);
+					Settings = UserSettings.Create(GlobalConfigFolder, Logger);
 				}
 				else
 				{
diff --git a/Engine/Source/Programs/UnrealGameSync/UnrealGameSyncShared/UserSettings.cs b/Engine/Source/Programs/UnrealGameSync/UnrealGameSyncShared/UserSettings.cs
--- a/Engine/Source/Programs/UnrealGameSync/UnrealGameSyncShared/UserSettings.cs
+++ b/Engine/Source/Programs/UnrealGameSync/UnrealGameSyncShared/UserSettings.cs
@@ -645,13 +645,20 @@ namespace UnrealGameSync
 
 		public static UserSettings Create(DirectoryReference SettingsDir, ILogger Logger)
 		{
-			return Create(FileReference.Combine(SettingsDir, "UnrealGameSync.ini"), FileReference.Combine(SettingsDir, "Global.json"), Logger);
+			return Create(FileReference.Combine(SettingsDir, "UnrealGameSyncV2.ini"), FileReference.Combine(SettingsDir, "UnrealGameSync.ini"), FileReference.Combine(SettingsDir, "Global.json"), Logger);
 		}
 
-		public static UserSettings Create(FileReference FileName, FileReference CoreFileName, ILogger Logger)
+		public static UserSettings Create(FileReference FileName, FileReference LegacyFileName, FileReference CoreFileName, ILogger Logger)
 		{
 			ConfigFile ConfigFile = new ConfigFile();
-			ConfigFile.TryLoad(FileName, Logger);
+			if (FileReference.Exists(FileName))
+			{
+				ConfigFile.TryLoad(FileName, Logger);
+			}
+			else
+			{
+				ConfigFile.TryLoad(LegacyFileName, Logger);
+			}
 
 			GlobalSettings? CoreSettingsData;
 			if(FileReference.Exists(CoreFileName))
