This adds an option to keep trying a number of times after "Aborted sync due to errors".

This option will keep trying when the perforce server times out.  That can be useful when perforce servers are overloaded and users have bad connections and continually get kicked.  It is currently defaulted to off.  In the future we may consider defaulting this number to around 1000. The setting can be found under Options->Application Settings->Advanced->Retries on sync error.

#test Pulled down entirly new branch twice.

diff --git a/Engine/Source/Programs/UnrealGameSync/UnrealGameSync/Forms/PerforceSyncSettingsWindow.cs b/Engine/Source/Programs/UnrealGameSync/UnrealGameSync/Forms/PerforceSyncSettingsWindow.cs
--- a/Engine/Source/Programs/UnrealGameSync/UnrealGameSync/Forms/PerforceSyncSettingsWindow.cs
+++ b/Engine/Source/Programs/UnrealGameSync/UnrealGameSync/Forms/PerforceSyncSettingsWindow.cs
@@ -23,6 +23,7 @@ namespace UnrealGameSync
 			numericUpDownFileBufferSize.Value = (SyncOptions.FileBufferSize > 0) ? SyncOptions.FileBufferSize / 1024 : PerforceSyncOptions.DefaultFileBufferSize / 1024;
 			numericUpDownMaxCommandsPerBatch.Value = (SyncOptions.MaxCommandsPerBatch > 0) ? SyncOptions.MaxCommandsPerBatch : PerforceSyncOptions.DefaultMaxCommandsPerBatch;
 			numericUpDownMaxSizePerBatch.Value = (SyncOptions.MaxSizePerBatch > 0) ? SyncOptions.MaxSizePerBatch / 1024 / 1024 : PerforceSyncOptions.DefaultMaxSizePerBatch / 1024 / 1024;
+			numericUpDownRetriesOnSyncError.Value = (SyncOptions.NumSyncErrorRetries > 0) ? SyncOptions.NumSyncErrorRetries : PerforceSyncOptions.DefaultNumSyncErrorRetries;
 		}
 
 		private void OkButton_Click(object sender, EventArgs e)
@@ -32,6 +33,7 @@ namespace UnrealGameSync
 			Settings.SyncOptions.FileBufferSize = (int)numericUpDownFileBufferSize.Value * 1024;
 			Settings.SyncOptions.MaxCommandsPerBatch = (int)numericUpDownMaxCommandsPerBatch.Value;
 			Settings.SyncOptions.MaxSizePerBatch = (int)numericUpDownMaxSizePerBatch.Value * 1024 * 1024;
+			Settings.SyncOptions.NumSyncErrorRetries = (int)numericUpDownRetriesOnSyncError.Value;
 			Settings.Save();
 
 			DialogResult = System.Windows.Forms.DialogResult.OK;
@@ -52,6 +54,7 @@ namespace UnrealGameSync
 			numericUpDownFileBufferSize.Value = PerforceSyncOptions.DefaultFileBufferSize / 1024;
 			numericUpDownMaxCommandsPerBatch.Value = PerforceSyncOptions.DefaultMaxCommandsPerBatch;
 			numericUpDownMaxSizePerBatch.Value = PerforceSyncOptions.DefaultMaxSizePerBatch / 1024 / 1024;
+			numericUpDownRetriesOnSyncError.Value = PerforceSyncOptions.DefaultNumSyncErrorRetries;
 		}
 	}
 }
diff --git a/Engine/Source/Programs/UnrealGameSync/UnrealGameSyncShared/UserSettings.cs b/Engine/Source/Programs/UnrealGameSync/UnrealGameSyncShared/UserSettings.cs
--- a/Engine/Source/Programs/UnrealGameSync/UnrealGameSyncShared/UserSettings.cs
+++ b/Engine/Source/Programs/UnrealGameSync/UnrealGameSyncShared/UserSettings.cs
@@ -855,6 +855,7 @@ namespace UnrealGameSync
 			SyncOptions.FileBufferSize = ConfigFile.GetValue("Perforce.FileBufferSize", PerforceSyncOptions.DefaultFileBufferSize);
 			SyncOptions.MaxCommandsPerBatch = ConfigFile.GetValue("Perforce.MaxCommandsPerBatch", PerforceSyncOptions.DefaultMaxCommandsPerBatch);
 			SyncOptions.MaxSizePerBatch = ConfigFile.GetValue("Perforce.MaxSizePerBatch", PerforceSyncOptions.DefaultMaxSizePerBatch);
+			SyncOptions.NumSyncErrorRetries = ConfigFile.GetValue("Perforce.NumSyncErrorRetries", PerforceSyncOptions.DefaultNumSyncErrorRetries);
 		}
 
 		static Dictionary<Guid, bool> GetCategorySettings(ConfigSection Section, string IncludedKey, string ExcludedKey)
@@ -1197,6 +1198,10 @@ namespace UnrealGameSync
 			{
 				PerforceSection.SetValue("MaxSizePerBatch", SyncOptions.MaxSizePerBatch);
 			}
+			if (SyncOptions.NumSyncErrorRetries > 0 && SyncOptions.NumSyncErrorRetries != PerforceSyncOptions.DefaultNumSyncErrorRetries)
+			{
+				PerforceSection.SetValue("NumSyncErrorRetries", SyncOptions.NumSyncErrorRetries);
+			}
 
 			// Save the file
 			ConfigFile.Save(FileName);
diff --git a/Engine/Source/Programs/UnrealGameSync/UnrealGameSyncShared/WorkspaceUpdate.cs b/Engine/Source/Programs/UnrealGameSync/UnrealGameSyncShared/WorkspaceUpdate.cs
--- a/Engine/Source/Programs/UnrealGameSync/UnrealGameSyncShared/WorkspaceUpdate.cs
+++ b/Engine/Source/Programs/UnrealGameSync/UnrealGameSyncShared/WorkspaceUpdate.cs
@@ -68,6 +68,7 @@ namespace UnrealGameSync
 
 		public const int DefaultMaxCommandsPerBatch = 200;
 		public const int DefaultMaxSizePerBatch = 128 * 1024 * 1024;
+		public const int DefaultNumSyncErrorRetries = 0;
 
 		public int NumRetries = DefaultNumRetries;
 		public int NumThreads = DefaultNumThreads;
@@ -77,6 +78,8 @@ namespace UnrealGameSync
 		public int MaxCommandsPerBatch = DefaultMaxCommandsPerBatch;
 		public int MaxSizePerBatch = DefaultMaxSizePerBatch;
 
+		public int NumSyncErrorRetries = DefaultNumSyncErrorRetries;
+
 
 		public PerforceSyncOptions Clone()
 		{
@@ -1630,8 +1633,21 @@ namespace UnrealGameSync
 					}
 				}
 
-				// Sync the files
-				(WorkspaceUpdateResult Result, string StatusMessage) = await StaticSyncFileRevisions(Perforce, Context, SyncCommands, Record => SyncOutput(Record, ThreadLog), CancellationToken);
+				WorkspaceUpdateResult Result = WorkspaceUpdateResult.FailedToSync;
+				string StatusMessage = "";
+
+				int Retries = (null != Context.PerforceSyncOptions) ? Context.PerforceSyncOptions!.NumSyncErrorRetries : 0;
+
+				while (Retries >= 0 && WorkspaceUpdateResult.FailedToSync == Result)
+				{
+					// Sync the files
+					(Result, StatusMessage) = await StaticSyncFileRevisions(Perforce, Context, SyncCommands, Record => SyncOutput(Record, ThreadLog), CancellationToken);
+
+					if (WorkspaceUpdateResult.FailedToSync == Result && --Retries >= 0)
+					{
+						ThreadLog.LogWarning("Sync Errors occurred.  Retrying: Remaining retries " + Retries.ToString());
+					}
+				}
 
 				// If it failed, try to set it on the state if nothing else has failed first
 				if (Result != WorkspaceUpdateResult.Success)
@@ -1672,7 +1688,9 @@ namespace UnrealGameSync
 				}
 				else
 				{
-					return (WorkspaceUpdateResult.FailedToSync, "Aborted sync due to errors.");
+					return (WorkspaceUpdateResult.FailedToSync, "Aborted sync due to errors.. Currently retries on sync error is set at " +
+						((null != Context.PerforceSyncOptions) ? Context.PerforceSyncOptions!.NumSyncErrorRetries : 0).ToString() +
+						" in Options->Application Settings...->Advanced.  You might want to set it higher if you are on a bad connection.");
 				}
 			}
 
