UGS: Prevent connections being kept around for a long time. Dispose connections in PerforceMonitor between each poll.

diff --git a/Engine/Source/Programs/UnrealGameSync/UnrealGameSync/PerforceMonitor.cs b/Engine/Source/Programs/UnrealGameSync/UnrealGameSync/PerforceMonitor.cs
--- a/Engine/Source/Programs/UnrealGameSync/UnrealGameSync/PerforceMonitor.cs
+++ b/Engine/Source/Programs/UnrealGameSync/UnrealGameSync/PerforceMonitor.cs
@@ -262,60 +262,66 @@ namespace UnrealGameSync
 
 		async Task PollForUpdatesInner(CancellationToken CancellationToken)
 		{
-			using IPerforceConnection Perforce = await PerforceConnection.CreateAsync(PerforceSettings, Logger);
-			string? StreamName = await Perforce.GetCurrentStreamAsync(CancellationToken);
-
-			// Get the perforce server settings
-			PerforceResponse<InfoRecord> InfoResponse = await Perforce.TryGetInfoAsync(InfoOptions.ShortOutput, CancellationToken);
-			if (InfoResponse.Succeeded)
+			string? StreamName;
+			using (IPerforceConnection Perforce = await PerforceConnection.CreateAsync(PerforceSettings, Logger))
 			{
-				DateTimeOffset? ServerDate = InfoResponse.Data.ServerDate;
-				if (ServerDate.HasValue)
+				StreamName = await Perforce.GetCurrentStreamAsync(CancellationToken);
+
+				// Get the perforce server settings
+				PerforceResponse<InfoRecord> InfoResponse = await Perforce.TryGetInfoAsync(InfoOptions.ShortOutput, CancellationToken);
+				if (InfoResponse.Succeeded)
 				{
-					ServerTimeZone = ServerDate.Value.Offset;
+					DateTimeOffset? ServerDate = InfoResponse.Data.ServerDate;
+					if (ServerDate.HasValue)
+					{
+						ServerTimeZone = ServerDate.Value.Offset;
+					}
 				}
-			}
 
-			// Try to update the zipped binaries list before anything else, because it causes a state change in the UI
-			await UpdateArchivesAsync(Perforce, CancellationToken);
+				// Try to update the zipped binaries list before anything else, because it causes a state change in the UI
+				await UpdateArchivesAsync(Perforce, CancellationToken);
+			}
 
 			while(!CancellationToken.IsCancellationRequested)
 			{
 				Stopwatch Timer = Stopwatch.StartNew();
 				Task NextRefreshTask = RefreshEvent.Task;
 
-				// Check we still have a valid login ticket
-				PerforceResponse<LoginRecord> LoginState = await Perforce.TryGetLoginStateAsync(CancellationToken);
-				if(!LoginState.Succeeded)
+				using (IPerforceConnection Perforce = await PerforceConnection.CreateAsync(PerforceSettings, Logger))
 				{
-					LastStatusMessage = "User is not logged in";
-					SynchronizationContext.Post(_ => OnLoginExpired?.Invoke(), null);
-				}
-				else
-				{
-					// Check we haven't switched streams
-					string? NewStreamName = await Perforce.GetCurrentStreamAsync(CancellationToken);
-					if (NewStreamName != StreamName)
-					{
-						SynchronizationContext.Post(_ => OnStreamChange?.Invoke(), null);
-					}
-
-					// Check for any p4 changes
-					if(!await UpdateChangesAsync(Perforce, CancellationToken))
-					{
-						LastStatusMessage = "Failed to update changes";
-					}
-					else if(!await UpdateChangeTypesAsync(Perforce, CancellationToken))
+					// Check we still have a valid login ticket
+					PerforceResponse<LoginRecord> LoginState = await Perforce.TryGetLoginStateAsync(CancellationToken);
+					if (!LoginState.Succeeded)
 					{
-						LastStatusMessage = "Failed to update change types";
-					}
-					else if(!await UpdateArchivesAsync(Perforce, CancellationToken))
-					{
-						LastStatusMessage = "Failed to update zipped binaries list";
+						LastStatusMessage = "User is not logged in";
+						SynchronizationContext.Post(_ => OnLoginExpired?.Invoke(), null);
 					}
 					else
 					{
-						LastStatusMessage = String.Format("Last update took {0}ms", Timer.ElapsedMilliseconds);
+						// Check we haven't switched streams
+						string? NewStreamName = await Perforce.GetCurrentStreamAsync(CancellationToken);
+						if (NewStreamName != StreamName)
+						{
+							SynchronizationContext.Post(_ => OnStreamChange?.Invoke(), null);
+						}
+
+						// Check for any p4 changes
+						if (!await UpdateChangesAsync(Perforce, CancellationToken))
+						{
+							LastStatusMessage = "Failed to update changes";
+						}
+						else if (!await UpdateChangeTypesAsync(Perforce, CancellationToken))
+						{
+							LastStatusMessage = "Failed to update change types";
+						}
+						else if (!await UpdateArchivesAsync(Perforce, CancellationToken))
+						{
+							LastStatusMessage = "Failed to update zipped binaries list";
+						}
+						else
+						{
+							LastStatusMessage = String.Format("Last update took {0}ms", Timer.ElapsedMilliseconds);
+						}
 					}
 				}
 
