UGS: Look for a project-specific receipt when launching a content-only project, rather than falling back to looking for the engine receipt.

diff --git a/Engine/Source/Programs/UnrealGameSync/UnrealGameSyncShared/ConfigUtils.cs b/Engine/Source/Programs/UnrealGameSync/UnrealGameSyncShared/ConfigUtils.cs
--- a/Engine/Source/Programs/UnrealGameSync/UnrealGameSyncShared/ConfigUtils.cs
+++ b/Engine/Source/Programs/UnrealGameSync/UnrealGameSyncShared/ConfigUtils.cs
@@ -248,7 +248,7 @@ namespace UnrealGameSync
 			string TargetName = TargetFile.GetFileNameWithoutAnyExtensions();
 
 			DirectoryReference? ProjectDir = ProjectInfo.ProjectDir;
-			if (ProjectDir != null && TargetFile.IsUnderDirectory(ProjectDir))
+			if (ProjectDir != null)
 			{
 				return GetReceiptFile(ProjectDir, TargetName, Configuration);
 			}
