UGS: Add a setting to globally disable alerts. These are inconsistent with Horde, and we only want notifications coming from one place.

diff --git a/Engine/Source/Programs/UnrealGameSync/UnrealGameSync/Controls/WorkspaceControl.cs b/Engine/Source/Programs/UnrealGameSync/UnrealGameSync/Controls/WorkspaceControl.cs
--- a/Engine/Source/Programs/UnrealGameSync/UnrealGameSync/Controls/WorkspaceControl.cs
+++ b/Engine/Source/Programs/UnrealGameSync/UnrealGameSync/Controls/WorkspaceControl.cs
@@ -1759,6 +1759,11 @@ namespace UnrealGameSync
 
 		void UpdateBuildFailureNotification()
 		{
+			if (!DeploymentSettings.EnableAlerts)
+			{
+				return;
+			}
+
 			// Ignore this if we're using the new build health system
 			if (GetDefaultIssueFilter() != null)
 			{
diff --git a/Engine/Source/Programs/UnrealGameSync/UnrealGameSync/DeploymentSettings.cs b/Engine/Source/Programs/UnrealGameSync/UnrealGameSync/DeploymentSettings.cs
--- a/Engine/Source/Programs/UnrealGameSync/UnrealGameSync/DeploymentSettings.cs
+++ b/Engine/Source/Programs/UnrealGameSync/UnrealGameSync/DeploymentSettings.cs
@@ -59,6 +59,11 @@ namespace UnrealGameSync
 		/// </summary>
 		public static readonly string? SentryDsn = null;
 
+		/// <summary>
+		/// Whether to allow notifications about build failures
+		/// </summary>
+		public static readonly bool EnableAlerts = true;
+
 #if WITH_TELEMETRY
 		/// <summary>
 		/// Delegate used to create a new telemetry sink
diff --git a/Engine/Source/Programs/UnrealGameSync/UnrealGameSync/Forms/MainWindow.cs b/Engine/Source/Programs/UnrealGameSync/UnrealGameSync/Forms/MainWindow.cs
--- a/Engine/Source/Programs/UnrealGameSync/UnrealGameSync/Forms/MainWindow.cs
+++ b/Engine/Source/Programs/UnrealGameSync/UnrealGameSync/Forms/MainWindow.cs
@@ -1471,6 +1471,11 @@ namespace UnrealGameSync
 
 		public void UpdateAlertWindows()
 		{
+			if (!DeploymentSettings.EnableAlerts)
+			{
+				return;
+			}
+
 			HashSet<IssueData> AllIssues = new HashSet<IssueData>();
 			foreach(IssueMonitor IssueMonitor in WorkspaceIssueMonitors.Select(x => x.IssueMonitor))
 			{
