UGS: Remove another long-lived connection to P4.

diff --git a/Engine/Source/Programs/UnrealGameSync/UnrealGameSync/UpdateMonitor.cs b/Engine/Source/Programs/UnrealGameSync/UnrealGameSync/UpdateMonitor.cs
--- a/Engine/Source/Programs/UnrealGameSync/UnrealGameSync/UpdateMonitor.cs
+++ b/Engine/Source/Programs/UnrealGameSync/UnrealGameSync/UpdateMonitor.cs
@@ -67,15 +67,17 @@ namespace UnrealGameSync
 
 		async Task PollForUpdatesAsync(IPerforceSettings PerforceSettings, string WatchPath, CancellationToken CancellationToken)
 		{
-			using IPerforceConnection Perforce = await PerforceConnection.CreateAsync(PerforceSettings, Logger);
 			for (; ; )
 			{
 				await Task.Delay(TimeSpan.FromMinutes(5.0), CancellationToken);
 
-				PerforceResponseList<ChangesRecord> Changes = await Perforce.TryGetChangesAsync(ChangesOptions.None, -1, ChangeStatus.Submitted, WatchPath, CancellationToken);
-				if (Changes.Succeeded && Changes.Data.Count > 0)
+				using (IPerforceConnection Perforce = await PerforceConnection.CreateAsync(PerforceSettings, Logger))
 				{
-					TriggerUpdate(UpdateType.Background, null);
+					PerforceResponseList<ChangesRecord> Changes = await Perforce.TryGetChangesAsync(ChangesOptions.None, -1, ChangeStatus.Submitted, WatchPath, CancellationToken);
+					if (Changes.Succeeded && Changes.Data.Count > 0)
+					{
+						TriggerUpdate(UpdateType.Background, null);
+					}
 				}
 			}
 		}
