UGS: Fix import of custom build steps.

diff --git a/Engine/Source/Programs/UnrealGameSync/UnrealGameSyncShared/UserSettings.cs b/Engine/Source/Programs/UnrealGameSync/UnrealGameSyncShared/UserSettings.cs
--- a/Engine/Source/Programs/UnrealGameSync/UnrealGameSyncShared/UserSettings.cs
+++ b/Engine/Source/Programs/UnrealGameSync/UnrealGameSyncShared/UserSettings.cs
@@ -1056,7 +1056,7 @@ namespace UnrealGameSync
 
 		protected override void ImportProjectSettings(ProjectInfo ProjectInfo, UserProjectSettings CurrentProject)
 		{
-			string ClientProjectFileName = Regex.Replace(ProjectInfo.ClientFileName, "^//[^/]+/", "");
+			string ClientProjectFileName = ProjectInfo.ClientFileName;
 
 			// Read the project settings
 			ConfigSection ProjectSection = ConfigFile.FindOrAddSection(ClientProjectFileName);
