Add an option for custom sync options to not sync beyond the CL that matched the UGS badge list

diff --git a/Engine/Source/Programs/UnrealGameSync/UnrealGameSync/Controls/WorkspaceControl.cs b/Engine/Source/Programs/UnrealGameSync/UnrealGameSync/Controls/WorkspaceControl.cs
--- a/Engine/Source/Programs/UnrealGameSync/UnrealGameSync/Controls/WorkspaceControl.cs
+++ b/Engine/Source/Programs/UnrealGameSync/UnrealGameSync/Controls/WorkspaceControl.cs
@@ -4999,7 +4999,15 @@ namespace UnrealGameSync
 					if (CanSyncChange(Change.Number)
 						&& CanSyncChangeType(ChangeType, Change, Summary))
 					{
-						ChangeNumber = FindNewestGoodContentChange(Change.Number);
+						if (ChangeType.bFindNewestGoodContent)
+						{
+							ChangeNumber = FindNewestGoodContentChange(Change.Number);
+						}
+						else
+						{
+							ChangeNumber = Change.Number;
+						}
+						
 						return true;
 					}
 				}
diff --git a/Engine/Source/Programs/UnrealGameSync/UnrealGameSyncShared/UserSettings.cs b/Engine/Source/Programs/UnrealGameSync/UnrealGameSyncShared/UserSettings.cs
--- a/Engine/Source/Programs/UnrealGameSync/UnrealGameSyncShared/UserSettings.cs
+++ b/Engine/Source/Programs/UnrealGameSync/UnrealGameSyncShared/UserSettings.cs
@@ -65,6 +65,7 @@ namespace UnrealGameSync
 		// What Rules to Check for.
 		public bool bGood = false;
 		public bool bStarred = false;
+		public bool bFindNewestGoodContent = false;
 		public List<string> RequiredBadges = new List<string>();
 
 		public static bool TryParseConfigEntry(string Text, [NotNullWhen(true)] out LatestChangeType? ChangeType)
@@ -81,6 +82,7 @@ namespace UnrealGameSync
 
 				ChangeType.bGood = DefinitionObject.GetValue("bGood", false);
 				ChangeType.bStarred = DefinitionObject.GetValue("bStarred", false);
+				ChangeType.bFindNewestGoodContent = DefinitionObject.GetValue("bFindNewestGoodContent", false);
 				ChangeType.RequiredBadges = new List<string>(DefinitionObject.GetValue("RequiredBadges", "").Split(new char[] { ',',' '}, StringSplitOptions.RemoveEmptyEntries));
 				return true;
 			}
@@ -99,6 +101,7 @@ namespace UnrealGameSync
 			Latest.Name = "LatestChange";
 			Latest.Description = "Latest Change";
 			Latest.OrderIndex = -3;
+			Latest.bFindNewestGoodContent = true;
 
 			return Latest;
 		}
@@ -111,6 +114,7 @@ namespace UnrealGameSync
 			LatestGood.Description = "Latest Good Change";
 			LatestGood.OrderIndex = -2;
 			LatestGood.bGood = true;
+			LatestGood.bFindNewestGoodContent = true;
 
 			return LatestGood;
 		}
@@ -123,6 +127,7 @@ namespace UnrealGameSync
 			LatestStarred.Description = "Latest Starred Change";
 			LatestStarred.OrderIndex = -1;
 			LatestStarred.bStarred = true;
+			LatestStarred.bFindNewestGoodContent = true;
 
 			return LatestStarred;
 		}
