UGS: Fix potential race for tasks to be disposed.

diff --git a/Engine/Source/Programs/UnrealGameSync/UnrealGameSync/AsyncDisposer.cs b/Engine/Source/Programs/UnrealGameSync/UnrealGameSync/AsyncDisposer.cs
--- a/Engine/Source/Programs/UnrealGameSync/UnrealGameSync/AsyncDisposer.cs
+++ b/Engine/Source/Programs/UnrealGameSync/UnrealGameSync/AsyncDisposer.cs
@@ -27,11 +27,11 @@ namespace UnrealGameSync
 
 		public void Add(Task Task)
 		{
+			Task ContinuationTask = Task.ContinueWith(Remove);
 			lock (LockObject)
 			{
-				Tasks.Add(Task);
+				Tasks.Add(ContinuationTask);
 			}
-			Task.ContinueWith(Remove);
 		}
 
 		private void Remove(Task Task)
@@ -40,10 +40,6 @@ namespace UnrealGameSync
 			{
 				Logger.LogError(Task.Exception, "Exception while disposing task");
 			}
-			lock(LockObject)
-			{
-				Tasks.Remove(Task);
-			}
 		}
 
 		public async ValueTask DisposeAsync()
@@ -51,7 +47,7 @@ namespace UnrealGameSync
 			List<Task> TasksCopy;
 			lock (LockObject)
 			{
-				TasksCopy = Tasks;
+				TasksCopy = new List<Task>(Tasks);
 			}
 
 			Task WaitTask = Task.WhenAll(TasksCopy);
