UGS: Additional guards against null references.

diff --git a/Engine/Source/Programs/UnrealGameSync/UnrealGameSync/Controls/WorkspaceControl.cs b/Engine/Source/Programs/UnrealGameSync/UnrealGameSync/Controls/WorkspaceControl.cs
--- a/Engine/Source/Programs/UnrealGameSync/UnrealGameSync/Controls/WorkspaceControl.cs
+++ b/Engine/Source/Programs/UnrealGameSync/UnrealGameSync/Controls/WorkspaceControl.cs
@@ -1047,6 +1047,11 @@ namespace UnrealGameSync
 
 		void UpdateComplete(WorkspaceUpdateContext Context, WorkspaceUpdateResult Result, string ResultMessage)
 		{
+			if(bIsDisposing)
+			{
+				return;
+			}
+
 			UpdateTimer.Stop();
 
 			if (Result == WorkspaceUpdateResult.FilesToResolve)
@@ -3250,11 +3255,16 @@ namespace UnrealGameSync
 
 		private void UpdateStatusPanel_CrossThread()
 		{
-			MainThreadSynchronizationContext.Post((o) => { if (!IsDisposed) { UpdateStatusPanel(); } }, null);
+			MainThreadSynchronizationContext.Post((o) => { if (!IsDisposed && !bIsDisposing) { UpdateStatusPanel(); } }, null);
 		}
 
 		private void UpdateStatusPanel()
 		{
+			if(Workspace == null)
+			{
+				return;
+			}
+
 			int NewContentWidth = Math.Max(TextRenderer.MeasureText(String.Format("Opened {0}  |  Browse...  |  Connect...", SelectedFileName), StatusPanel.Font).Width, 400);
 			StatusPanel.SetContentWidth(NewContentWidth);
 
