UGS: Add a "version" command to print out the current application version.

diff --git a/Engine/Source/Programs/UnrealGameSync/UnrealGameSyncCmd/Program.cs b/Engine/Source/Programs/UnrealGameSync/UnrealGameSyncCmd/Program.cs
--- a/Engine/Source/Programs/UnrealGameSync/UnrealGameSyncCmd/Program.cs
+++ b/Engine/Source/Programs/UnrealGameSync/UnrealGameSyncCmd/Program.cs
@@ -11,6 +11,7 @@ using System.Globalization;
 using System.IO;
 using System.Linq;
 using System.Net;
+using System.Reflection;
 using System.Runtime.InteropServices;
 using System.Text;
 using System.Text.RegularExpressions;
@@ -101,7 +102,11 @@ namespace UnrealGameSyncCmd
 			new CommandInfo("status", typeof(StatusCommand),
 				"ugs status [-update]",
 				"Shows the status of the currently synced branch."
-			)
+			),
+			new CommandInfo("version", typeof(VersionCommand),
+				"ugs version",
+				"Prints the current application version"
+			),
 		};
 
 		class CommandContext
@@ -1145,5 +1150,14 @@ namespace UnrealGameSyncCmd
 				Logger.LogInformation("Switched to project {ProjectPath}", Settings.ClientProjectPath);
 			}
 		}
+
+		class VersionCommand : Command
+		{
+			public override Task ExecuteAsync(CommandContext Context)
+			{
+				Console.WriteLine("UnrealGameSync {0}", Assembly.GetExecutingAssembly().GetName().Version?.ToString());
+				return Task.CompletedTask;
+			}
+		}
 	}
 }
