UGS: Ignore deleted files when detecting available projects in a workspace.

diff --git a/Engine/Source/Programs/UnrealGameSync/UnrealGameSyncCmd/Program.cs b/Engine/Source/Programs/UnrealGameSync/UnrealGameSyncCmd/Program.cs
--- a/Engine/Source/Programs/UnrealGameSync/UnrealGameSyncCmd/Program.cs
+++ b/Engine/Source/Programs/UnrealGameSync/UnrealGameSyncCmd/Program.cs
@@ -363,6 +363,7 @@ namespace UnrealGameSyncCmd
 			}
 
 			List<FStatRecord> ProjectFileRecords = await PerforceClient.FStatAsync(FStatOptions.ClientFileInPerforceSyntax, SearchPath).ToListAsync();
+			ProjectFileRecords.RemoveAll(x => x.HeadAction == FileAction.Delete || x.HeadAction == FileAction.MoveDelete);
 			ProjectFileRecords.RemoveAll(x => !x.IsMapped);
 
 			List<string> Paths = ProjectFileRecords.Select(x => PerforceUtils.GetClientRelativePath(x.ClientFile!)).Distinct(StringComparer.Ordinal).ToList();
