UGS: Fix stripping of progress text from tool output.

diff --git a/Engine/Source/Programs/UnrealGameSync/UnrealGameSyncShared/WorkspaceUpdate.cs b/Engine/Source/Programs/UnrealGameSync/UnrealGameSyncShared/WorkspaceUpdate.cs
--- a/Engine/Source/Programs/UnrealGameSync/UnrealGameSyncShared/WorkspaceUpdate.cs
+++ b/Engine/Source/Programs/UnrealGameSync/UnrealGameSyncShared/WorkspaceUpdate.cs
@@ -1404,7 +1404,7 @@ namespace UnrealGameSync
 			string? ParsedLine = ProgressTextWriter.ParseLine(Line, Progress);
 			if (ParsedLine != null)
 			{
-				Logger.LogInformation("{Prefix}{Message}", Prefix, Line);
+				Logger.LogInformation("{Prefix}{Message}", Prefix, ParsedLine);
 			}
 		}
 
