UGS: Disable links to custom tools that aren't currently installed.

diff --git a/Engine/Source/Programs/UnrealGameSync/UnrealGameSync/Controls/WorkspaceControl.cs b/Engine/Source/Programs/UnrealGameSync/UnrealGameSync/Controls/WorkspaceControl.cs
--- a/Engine/Source/Programs/UnrealGameSync/UnrealGameSync/Controls/WorkspaceControl.cs
+++ b/Engine/Source/Programs/UnrealGameSync/UnrealGameSync/Controls/WorkspaceControl.cs
@@ -2656,6 +2656,7 @@ namespace UnrealGameSync
 
 		public void UpdateSettings()
 		{
+			UpdateBuildSteps();
 			UpdateStatusPanel();
 		}
 
@@ -3255,7 +3256,14 @@ namespace UnrealGameSync
 
 		private void UpdateStatusPanel_CrossThread()
 		{
-			MainThreadSynchronizationContext.Post((o) => { if (!IsDisposed && !bIsDisposing) { UpdateStatusPanel(); } }, null);
+			MainThreadSynchronizationContext.Post((o) => 
+			{ 
+				if (!IsDisposed && !bIsDisposing)
+				{ 
+					UpdateBuildSteps();
+					UpdateStatusPanel(); 
+				} 
+			}, null);
 		}
 
 		private void UpdateStatusPanel()
@@ -5147,7 +5155,7 @@ namespace UnrealGameSync
 					Dictionary<string, ContextMenuStrip> NameToMenu = new Dictionary<string, ContextMenuStrip>();
 					foreach (BuildStep Step in UserSteps)
 					{
-						if (!String.IsNullOrEmpty(Step.StatusPanelLink))
+						if (!String.IsNullOrEmpty(Step.StatusPanelLink) && CanRunStep(Step))
 						{
 							int BaseMenuIdx = Step.StatusPanelLink.IndexOf('|');
 							if (BaseMenuIdx == -1)
@@ -5190,6 +5198,23 @@ namespace UnrealGameSync
 			MoreActionsContextMenu_CustomToolSeparator.Visible = (MoreToolsItemCount > 0);
 		}
 
+		private bool CanRunStep(BuildStep Step)
+		{
+			if(Step.ToolId != Guid.Empty)
+			{
+				string? ToolName = Owner.ToolUpdateMonitor.GetToolName(Step.ToolId);
+				if (ToolName == null)
+				{
+					return false;
+				}
+				if (Owner.ToolUpdateMonitor.GetToolPath(ToolName) == null)
+				{
+					return false;
+				}
+			}
+			return true;
+		}
+
 		private void RunCustomTool(BuildStep Step, List<BuildStep> AllSteps)
 		{
 			if (Workspace != null)
