UGS: Close stdin when spawning child processes, so they aren't left waiting for input from the user.

diff --git a/Engine/Source/Programs/UnrealGameSync/UnrealGameSyncShared/Utility.cs b/Engine/Source/Programs/UnrealGameSync/UnrealGameSyncShared/Utility.cs
--- a/Engine/Source/Programs/UnrealGameSync/UnrealGameSyncShared/Utility.cs
+++ b/Engine/Source/Programs/UnrealGameSync/UnrealGameSyncShared/Utility.cs
@@ -211,7 +211,7 @@
 
 		public static async Task<int> ExecuteProcessAsync(string FileName, string? WorkingDir, string CommandLine, Action<string> OutputLine, CancellationToken CancellationToken)
 		{
-			using (ManagedProcess NewProcess = new ManagedProcess(null, FileName, CommandLine, WorkingDir, null, ProcessPriorityClass.Normal))
+			using (ManagedProcess NewProcess = new ManagedProcess(null, FileName, CommandLine, WorkingDir, null, null, ProcessPriorityClass.Normal))
 			{
 				for (; ; )
 				{
