UGS: Fix exception when passing the -Clobber option to UGS CLI.

diff --git a/Engine/Source/Programs/UnrealGameSync/UnrealGameSyncShared/WorkspaceUpdate.cs b/Engine/Source/Programs/UnrealGameSync/UnrealGameSyncShared/WorkspaceUpdate.cs
--- a/Engine/Source/Programs/UnrealGameSync/UnrealGameSyncShared/WorkspaceUpdate.cs
+++ b/Engine/Source/Programs/UnrealGameSync/UnrealGameSyncShared/WorkspaceUpdate.cs
@@ -1720,7 +1720,8 @@ namespace UnrealGameSync
 				}
 				foreach (string TamperedFile in TamperedFiles)
 				{
-					if (Context.ClobberFiles[TamperedFile])
+					bool bShouldClobber = (Context.Options & WorkspaceUpdateOptions.Clobber) != 0 || Context.ClobberFiles[TamperedFile];
+					if (bShouldClobber)
 					{
 						List<PerforceResponse<SyncRecord>> Response = await Perforce.TrySyncAsync(SyncOptions.Force, -1, TamperedFile, CancellationToken).ToListAsync();
 						if (!Response.Succeeded())
