UGS: Break out of the P4 change update loop if the cancellation token is set.

diff --git a/Engine/Source/Programs/UnrealGameSync/UnrealGameSync/PerforceMonitor.cs b/Engine/Source/Programs/UnrealGameSync/UnrealGameSync/PerforceMonitor.cs
--- a/Engine/Source/Programs/UnrealGameSync/UnrealGameSync/PerforceMonitor.cs
+++ b/Engine/Source/Programs/UnrealGameSync/UnrealGameSync/PerforceMonitor.cs
@@ -508,6 +508,8 @@ namespace UnrealGameSync
 				Task notifyTask = Task.CompletedTask;
 				foreach (IReadOnlyList<int> QueryChangeNumberBatch in QueryChangeNumbers.OrderByDescending(x => x).Batch(10))
 				{
+					CancellationToken.ThrowIfCancellationRequested();
+
 					// Skip this stuff if the user wants us to query for more changes
 					if (PendingMaxChanges != CurrentMaxChanges)
 					{
