UGS: Fix sync state not being updated correctly.

diff --git a/Engine/Source/Programs/UnrealGameSync/UnrealGameSync/Workspace.cs b/Engine/Source/Programs/UnrealGameSync/UnrealGameSync/Workspace.cs
--- a/Engine/Source/Programs/UnrealGameSync/UnrealGameSync/Workspace.cs
+++ b/Engine/Source/Programs/UnrealGameSync/UnrealGameSync/Workspace.cs
@@ -149,11 +149,7 @@ namespace UnrealGameSync
 			{
 				WorkspaceUpdateContext Context = Update.Context;
 
-				State.LastSyncChangeNumber = Context.ChangeNumber;
-				State.LastSyncResult = Result;
-				State.LastSyncResultMessage = StatusMessage;
-				State.LastSyncTime = DateTime.UtcNow;
-				State.LastSyncDurationSeconds = (int)(State.LastSyncTime.Value - Context.StartTime).TotalSeconds;
+				State.SetLastSyncState(Result, Context, StatusMessage);
 				State.Save();
 
 				OnUpdateComplete?.Invoke(Context, Result, StatusMessage);
diff --git a/Engine/Source/Programs/UnrealGameSync/UnrealGameSyncCmd/Program.cs b/Engine/Source/Programs/UnrealGameSync/UnrealGameSyncCmd/Program.cs
--- a/Engine/Source/Programs/UnrealGameSync/UnrealGameSyncCmd/Program.cs
+++ b/Engine/Source/Programs/UnrealGameSync/UnrealGameSyncCmd/Program.cs
@@ -651,6 +651,9 @@ namespace UnrealGameSyncCmd
 				{
 					Logger.LogError("{Message} (Result: {Result})", Message, Result);
 				}
+
+				State.SetLastSyncState(Result, UpdateContext, Message);
+				State.Save();
 			}
 		}
 
diff --git a/Engine/Source/Programs/UnrealGameSync/UnrealGameSyncShared/UserSettings.cs b/Engine/Source/Programs/UnrealGameSync/UnrealGameSyncShared/UserSettings.cs
--- a/Engine/Source/Programs/UnrealGameSync/UnrealGameSyncShared/UserSettings.cs
+++ b/Engine/Source/Programs/UnrealGameSync/UnrealGameSyncShared/UserSettings.cs
@@ -487,6 +487,15 @@ namespace UnrealGameSync
 			}
 		}
 
+		public void SetLastSyncState(WorkspaceUpdateResult Result, WorkspaceUpdateContext Context, string StatusMessage)
+		{
+			LastSyncChangeNumber = Context.ChangeNumber;
+			LastSyncResult = Result;
+			LastSyncResultMessage = StatusMessage;
+			LastSyncTime = DateTime.UtcNow;
+			LastSyncDurationSeconds = (int)(LastSyncTime.Value - Context.StartTime).TotalSeconds;
+		}
+
 		static object SyncRoot = new object();
 
 		public void Save()
