UGS: Prevent creation of an empty .ugs directory when loading (or trying to load) config files.

diff --git a/Engine/Source/Programs/UnrealGameSync/UnrealGameSyncShared/GlobalSettings.cs b/Engine/Source/Programs/UnrealGameSync/UnrealGameSyncShared/GlobalSettings.cs
--- a/Engine/Source/Programs/UnrealGameSync/UnrealGameSyncShared/GlobalSettings.cs
+++ b/Engine/Source/Programs/UnrealGameSync/UnrealGameSyncShared/GlobalSettings.cs
@@ -71,6 +71,8 @@ namespace UnrealGameSync
 				ConfigFile = FileReference.Combine(UserSettings.GetConfigDir(Settings.RootDir), $"project_{ProjectInfo.LocalFileName.GetFileNameWithoutExtension()}.json");
 			}
 
+			UserSettings.CreateConfigDir(ConfigFile.Directory);
+
 			UserProjectSettings? ProjectSettings;
 			if (!UserProjectSettings.TryLoad(ConfigFile, out ProjectSettings))
 			{
diff --git a/Engine/Source/Programs/UnrealGameSync/UnrealGameSyncShared/UserSettings.cs b/Engine/Source/Programs/UnrealGameSync/UnrealGameSyncShared/UserSettings.cs
--- a/Engine/Source/Programs/UnrealGameSync/UnrealGameSyncShared/UserSettings.cs
+++ b/Engine/Source/Programs/UnrealGameSync/UnrealGameSyncShared/UserSettings.cs
@@ -379,6 +379,7 @@ namespace UnrealGameSync
 		{
 			lock (SyncRoot)
 			{
+				UserSettings.CreateConfigDir(ConfigFile.Directory);
 				Utility.SaveJson(ConfigFile, this);
 				LastModifiedTimeUtc = FileReference.GetLastWriteTimeUtc(ConfigFile).Ticks;
 			}
@@ -502,6 +503,7 @@ namespace UnrealGameSync
 		{
 			lock (SyncRoot)
 			{
+				UserSettings.CreateConfigDir(ConfigFile.Directory);
 				Utility.SaveJson(ConfigFile, this);
 			}
 		}
@@ -938,7 +940,11 @@ namespace UnrealGameSync
 		public static DirectoryReference GetConfigDir(DirectoryReference WorkspaceDir)
 		{
 			DirectoryReference ConfigDir = DirectoryReference.Combine(WorkspaceDir, ".ugs");
+			return ConfigDir;
+		}
 
+		public static void CreateConfigDir(DirectoryReference ConfigDir)
+		{
 			DirectoryInfo ConfigDirInfo = ConfigDir.ToDirectoryInfo();
 			if (!ConfigDirInfo.Exists)
 			{
@@ -948,8 +954,6 @@ namespace UnrealGameSync
 					ConfigDirInfo.Attributes = FileAttributes.Directory | FileAttributes.Hidden;
 				}
 			}
-
-			return ConfigDir;
 		}
 
 		protected override void ImportWorkspaceState(DirectoryReference RootDir, string ClientName, string BranchPath, UserWorkspaceState CurrentWorkspace)
