UGS: Fix version string displayed by the "UGS version" command.

diff --git a/Engine/Source/Programs/UnrealGameSync/UnrealGameSyncCmd/Program.cs b/Engine/Source/Programs/UnrealGameSync/UnrealGameSyncCmd/Program.cs
--- a/Engine/Source/Programs/UnrealGameSync/UnrealGameSyncCmd/Program.cs
+++ b/Engine/Source/Programs/UnrealGameSync/UnrealGameSyncCmd/Program.cs
@@ -1159,7 +1159,11 @@ namespace UnrealGameSyncCmd
 		{
 			public override Task ExecuteAsync(CommandContext Context)
 			{
-				Console.WriteLine("UnrealGameSync {0}", Assembly.GetExecutingAssembly().GetName().Version?.ToString());
+				ILogger Logger = Context.Logger;
+ 
+				AssemblyInformationalVersionAttribute? Version = Assembly.GetExecutingAssembly().GetCustomAttribute<AssemblyInformationalVersionAttribute>();
+				Logger.LogInformation("UnrealGameSync {Version}", Version?.InformationalVersion ?? "Unknown");
+
 				return Task.CompletedTask;
 			}
 		}
