UGS: Guard against exceptions when reading config files.

diff --git a/Engine/Source/Programs/UnrealGameSync/UnrealGameSyncShared/UserSettings.cs b/Engine/Source/Programs/UnrealGameSync/UnrealGameSyncShared/UserSettings.cs
--- a/Engine/Source/Programs/UnrealGameSync/UnrealGameSyncShared/UserSettings.cs
+++ b/Engine/Source/Programs/UnrealGameSync/UnrealGameSyncShared/UserSettings.cs
@@ -676,12 +676,20 @@ namespace UnrealGameSync
 				ConfigFile.TryLoad(LegacyFileName, Logger);
 			}
 
-			GlobalSettings? CoreSettingsData;
-			if(FileReference.Exists(CoreFileName))
+			GlobalSettings? CoreSettingsData = null;
+			if (FileReference.Exists(CoreFileName))
 			{
-				CoreSettingsData = Utility.LoadJson<GlobalSettings>(CoreFileName);
+				try
+				{
+					CoreSettingsData = Utility.LoadJson<GlobalSettings>(CoreFileName);
+				}
+				catch (Exception Ex)
+				{
+					Logger.LogError(Ex, "Error while reading {File}.", CoreFileName);
+				}
 			}
-			else
+
+			if (CoreSettingsData == null)
 			{
 				CoreSettingsData = new GlobalSettings();
 				CoreSettingsData.Filter.View = ConfigFile.GetValues("General.SyncFilter", new string[0]).ToList();
