UGS: Fix rounding error (yuck) when parsing integer parameters from URIs.

diff --git a/Engine/Source/Programs/UnrealGameSync/UnrealGameSync/UriHandlers/UriHandler.cs b/Engine/Source/Programs/UnrealGameSync/UnrealGameSync/UriHandlers/UriHandler.cs
--- a/Engine/Source/Programs/UnrealGameSync/UnrealGameSync/UriHandlers/UriHandler.cs
+++ b/Engine/Source/Programs/UnrealGameSync/UnrealGameSync/UriHandlers/UriHandler.cs
@@ -81,25 +81,24 @@ namespace UnrealGameSync
 							return new UriResult() { Error = string.Format("Uri {0} bool parameter {1} must be true or false", Uri.Host, Param.Name) };
 						}
 					}
-					else if (Param.ParameterType == typeof(int) || Param.ParameterType == typeof(float))
+					else if (Param.ParameterType == typeof(int))
+					{
+						int NumberValue;
+						if (!int.TryParse(Value, out NumberValue))
+						{
+							return new UriResult() { Error = string.Format("Uri {0} invalid integer parameter {1} : {2}", Uri.Host, Param.Name, Value) };
+						}
+						Parameters.Add(NumberValue);
+					}
+					else if (Param.ParameterType == typeof(float))
 					{
 						float NumberValue;
 						if (!float.TryParse(Value, out NumberValue))
 						{
 							return new UriResult() { Error = string.Format("Uri {0} invalid number parameter {1} : {2}", Uri.Host, Param.Name, Value) };
 						}
-
-						if (Param.ParameterType == typeof(int))
-						{
-							Parameters.Add(Convert.ToInt32(NumberValue));
-						}
-						else
-						{
-							Parameters.Add(NumberValue);
-						}
-
+						Parameters.Add(NumberValue);
 					}
-
 				}
 
 				return (UriResult)Info.Invoke(null, Parameters.ToArray())!;					
