UGS memory management fix: P4Automation.PrintToTempFileAsync() will now tear down any new perforce connections it creates.

diff --git a/Engine/Source/Programs/UnrealGameSync/UnrealGameSync/Automation/P4Automation.cs b/Engine/Source/Programs/UnrealGameSync/UnrealGameSync/Automation/P4Automation.cs
--- a/Engine/Source/Programs/UnrealGameSync/UnrealGameSync/Automation/P4Automation.cs
+++ b/Engine/Source/Programs/UnrealGameSync/UnrealGameSync/Automation/P4Automation.cs
@@ -34,30 +34,44 @@ namespace UnrealGameSync
 
 		public static async Task<string> PrintToTempFileAsync(IPerforceConnection? Connection, string DepotPath, CancellationToken CancellationToken, ILogger Logger)
 		{
-			if (Connection == null)
+			bool bCreateNewConnection = (Connection == null);
+
+			try
 			{
-				IPerforceSettings Settings = GetConnectionSettings();
-				Connection = await PerforceConnection.CreateAsync(Logger);
+				if (Connection == null)
+				{
+					IPerforceSettings Settings = GetConnectionSettings();
+					Connection = await PerforceConnection.CreateAsync(Logger);
+				}
+
+				string DepotFileName = Path.GetFileName(DepotPath);
+
+				// Reorder CL and extension
+				int Index = DepotFileName.IndexOf('@');
+				if (Index == -1)
+				{
+					DepotFileName += "@Latest";
+					Index = DepotFileName.IndexOf('@');
+				}
+
+				string CL = DepotFileName.Substring(Index + 1);
+				string FileName = DepotFileName.Substring(0, Index);
+				string TempFileName = string.Format("{0}@{1}{2}", Path.GetFileNameWithoutExtension(FileName), CL, Path.GetExtension(FileName));
+
+				TempFileName = Path.Combine(Path.GetTempPath(), TempFileName);
+				await Connection.PrintAsync(TempFileName, DepotFileName, CancellationToken);
+
+				return TempFileName;
 			}
-
-			string DepotFileName = Path.GetFileName(DepotPath);
-
-			// Reorder CL and extension
-			int Index = DepotFileName.IndexOf('@');
-			if (Index == -1)
+			finally
 			{
-				DepotFileName += "@Latest";
-				Index = DepotFileName.IndexOf('@');
+				// If we created a new connection, tear it down now.
+				if (bCreateNewConnection)
+				{
+					Connection?.Dispose();
+					Connection = null;
+				}
 			}
-
-			string CL = DepotFileName.Substring(Index + 1);
-			string FileName = DepotFileName.Substring(0, Index);
-			string TempFileName = string.Format("{0}@{1}{2}", Path.GetFileNameWithoutExtension(FileName), CL, Path.GetExtension(FileName));
-
-			TempFileName = Path.Combine(Path.GetTempPath(), TempFileName);
-			await Connection.PrintAsync(TempFileName, DepotFileName, CancellationToken);
-
-			return TempFileName;
 		}
 
 	}
