UnrealGameSync: net6.0 upgrade

#rnx

diff --git a/Engine/Source/Programs/UnrealGameSync/Installer/Build.bat b/Engine/Source/Programs/UnrealGameSync/Installer/Build.bat
--- a/Engine/Source/Programs/UnrealGameSync/Installer/Build.bat
+++ b/Engine/Source/Programs/UnrealGameSync/Installer/Build.bat
@@ -3,7 +3,7 @@ setlocal
 
 pushd %~dp0
 set WIXDIR=..\..\..\ThirdParty\WiX\3.8
-set SRCDIR=..\UnrealGameSyncLauncher\bin\Release\netcoreapp3.1
+set SRCDIR=..\UnrealGameSyncLauncher\bin\Release\net6.0-windows
 
 %WIXDIR%\heat.exe dir "%SRCDIR%" -cg UGSLauncher_Project -dr INSTALLFOLDER -scom -sreg -srd -var var.BasePath -gg -sfrag -out obj/Release/UGSLauncher.wxs -nologo 
 if errorlevel 1 goto :eof
diff --git a/Engine/Source/Programs/UnrealGameSync/UnrealGameSync/Automation/VisualStudioAutomation.cs b/Engine/Source/Programs/UnrealGameSync/UnrealGameSync/Automation/VisualStudioAutomation.cs
--- a/Engine/Source/Programs/UnrealGameSync/UnrealGameSync/Automation/VisualStudioAutomation.cs
+++ b/Engine/Source/Programs/UnrealGameSync/UnrealGameSync/Automation/VisualStudioAutomation.cs
@@ -62,7 +62,7 @@ namespace UnrealGameSync
 			ErrorMessage = "";
 			VisualStudioInstallation? Install = VisualStudioInstallations.GetPreferredInstallation();
 
-			if (Install == null)
+			if (Install == null || Install.DevEnvPath == null)
 			{
 				ErrorMessage = string.Format("Unable to get Visual Studio installation");
 				return false;
@@ -215,7 +215,7 @@ namespace UnrealGameSync
 				return CachedInstalls.First();
 			}
 
-			VisualStudioInstallation Installation = CachedInstalls.FirstOrDefault(Install => { return Install.MajorVersion == MajorVersion; });
+			VisualStudioInstallation? Installation = CachedInstalls.FirstOrDefault(Install => { return Install.MajorVersion == MajorVersion; });
 
 			return Installation;
 
diff --git a/Engine/Source/Programs/UnrealGameSync/UnrealGameSync/Controls/BuildListControl.cs b/Engine/Source/Programs/UnrealGameSync/UnrealGameSync/Controls/BuildListControl.cs
--- a/Engine/Source/Programs/UnrealGameSync/UnrealGameSync/Controls/BuildListControl.cs
+++ b/Engine/Source/Programs/UnrealGameSync/UnrealGameSync/Controls/BuildListControl.cs
@@ -57,7 +57,7 @@ namespace UnrealGameSync
 
 		public BuildListControl()
 		{
-			Font = SystemFonts.IconTitleFont;
+			Font = SystemFonts.IconTitleFont!;
 		}
 
 		public int GetFirstVisibleIndex()
diff --git a/Engine/Source/Programs/UnrealGameSync/UnrealGameSync/Controls/CustomListViewControl.cs b/Engine/Source/Programs/UnrealGameSync/UnrealGameSync/Controls/CustomListViewControl.cs
--- a/Engine/Source/Programs/UnrealGameSync/UnrealGameSync/Controls/CustomListViewControl.cs
+++ b/Engine/Source/Programs/UnrealGameSync/UnrealGameSync/Controls/CustomListViewControl.cs
@@ -400,7 +400,7 @@ namespace UnrealGameSync
 
 		public void DrawNormalSubItem(DrawListViewSubItemEventArgs e)
 		{
-			DrawText(e.Graphics, e.SubItem.Bounds, Columns[e.ColumnIndex].TextAlign, SystemColors.WindowText, e.SubItem.Text);
+			DrawText(e.Graphics, e.SubItem!.Bounds, Columns[e.ColumnIndex].TextAlign, SystemColors.WindowText, e.SubItem.Text);
 		}
 
 		public void DrawCustomSubItem(Graphics Graphics, ListViewItem.ListViewSubItem SubItem)
diff --git a/Engine/Source/Programs/UnrealGameSync/UnrealGameSync/Controls/LogSplitContainer.cs b/Engine/Source/Programs/UnrealGameSync/UnrealGameSync/Controls/LogSplitContainer.cs
--- a/Engine/Source/Programs/UnrealGameSync/UnrealGameSync/Controls/LogSplitContainer.cs
+++ b/Engine/Source/Programs/UnrealGameSync/UnrealGameSync/Controls/LogSplitContainer.cs
@@ -49,7 +49,7 @@ namespace UnrealGameSync
 
 		protected Font CaptionFont
 		{
-			get { return SystemFonts.IconTitleFont; }
+			get { return SystemFonts.IconTitleFont!; }
 		}
 
 		protected override void OnCreateControl()
@@ -68,7 +68,7 @@ namespace UnrealGameSync
 			Invalidate();
 		}
 
-		public void OnSplitterMoved(object Obj, EventArgs Args)
+		public void OnSplitterMoved(object? Obj, EventArgs Args)
 		{
 			bSplitterMoved = true;
 
diff --git a/Engine/Source/Programs/UnrealGameSync/UnrealGameSync/Controls/StatusPanel.cs b/Engine/Source/Programs/UnrealGameSync/UnrealGameSync/Controls/StatusPanel.cs
--- a/Engine/Source/Programs/UnrealGameSync/UnrealGameSync/Controls/StatusPanel.cs
+++ b/Engine/Source/Programs/UnrealGameSync/UnrealGameSync/Controls/StatusPanel.cs
@@ -767,7 +767,7 @@ namespace UnrealGameSync
 				}
 				if(Alert != null)
 				{
-					e.Graphics.DrawLine(AlertDividerPen, 0, AlertDividerY, Width, AlertDividerY);
+					e.Graphics.DrawLine(AlertDividerPen!, 0, AlertDividerY, Width, AlertDividerY);
 					Alert.Draw(e.Graphics, Resources!);
 				}
 			}
diff --git a/Engine/Source/Programs/UnrealGameSync/UnrealGameSync/Controls/WorkspaceControl.cs b/Engine/Source/Programs/UnrealGameSync/UnrealGameSync/Controls/WorkspaceControl.cs
--- a/Engine/Source/Programs/UnrealGameSync/UnrealGameSync/Controls/WorkspaceControl.cs
+++ b/Engine/Source/Programs/UnrealGameSync/UnrealGameSync/Controls/WorkspaceControl.cs
@@ -915,7 +915,7 @@ namespace UnrealGameSync
 
 			string[] CombinedSyncFilter = UserSettings.GetCombinedSyncFilter(GetSyncCategories(), Settings.Global.Filter, WorkspaceSettings.Filter);
 
-			ConfigSection PerforceSection = PerforceMonitor.LatestProjectConfigFile.FindSection("Perforce");
+			ConfigSection? PerforceSection = PerforceMonitor.LatestProjectConfigFile.FindSection("Perforce");
 			if (PerforceSection != null)
 			{
 				IEnumerable<string> AdditionalPaths = PerforceSection.GetValues("AdditionalPathsToSync", new string[0]);
@@ -949,7 +949,7 @@ namespace UnrealGameSync
 					string[]? ZippedBinariesSyncFilter;
 					if (TryGetProjectSetting(PerforceMonitor.LatestProjectConfigFile, "ZippedBinariesSyncFilter", out ZippedBinariesSyncFilter) && ZippedBinariesSyncFilter.Length > 0)
 					{
-						Context.SyncFilter = Enumerable.Concat(Context.SyncFilter, ZippedBinariesSyncFilter).ToArray();
+						Context.SyncFilter = Enumerable.Concat(Context.SyncFilter!, ZippedBinariesSyncFilter).ToArray();
 					}
 
 					Context.ArchiveTypeToArchive[Archive.Type] = new Tuple<IArchiveInfo, string>(Archive, ArchivePath);
@@ -1720,7 +1720,7 @@ namespace UnrealGameSync
 				try
 				{
 					string JsonString = File.ReadAllText(UncontrolledChangelistPersistencyFilePath);
-					UncontrolledChangelistPersistency UCLPersistency = JsonSerializer.Deserialize<UncontrolledChangelistPersistency>(JsonString, Utility.DefaultJsonSerializerOptions);
+					UncontrolledChangelistPersistency UCLPersistency = JsonSerializer.Deserialize<UncontrolledChangelistPersistency>(JsonString, Utility.DefaultJsonSerializerOptions)!;
 
 					foreach (UncontrolledChangelist UCL in UCLPersistency.Changelists)
 					{
@@ -2149,7 +2149,7 @@ namespace UnrealGameSync
 
 		private void BuildList_DrawSubItem(object sender, DrawListViewSubItemEventArgs e)
 		{
-			ChangesRecord Change = (ChangesRecord)e.Item.Tag;
+			ChangesRecord Change = (ChangesRecord)e.Item!.Tag;
 			if (Change == null)
 			{
 				return;
@@ -2207,7 +2207,7 @@ namespace UnrealGameSync
 						int FirstPass, FirstFail;
 						GetRemainingBisectRange(out FirstPass, out FirstFail);
 
-						BisectEntry Entry = WorkspaceState.BisectChanges.FirstOrDefault(x => x.Change == Change.Number);
+						BisectEntry? Entry = WorkspaceState.BisectChanges.FirstOrDefault(x => x.Change == Change.Number);
 						if (Entry == null || Entry.State == BisectState.Exclude)
 						{
 							QualityIcon = new Rectangle(0, 0, 0, 0);
@@ -2258,11 +2258,11 @@ namespace UnrealGameSync
 			}
 			else if (e.ColumnIndex == ChangeColumn.Index || e.ColumnIndex == TimeColumn.Index)
 			{
-				TextRenderer.DrawText(e.Graphics, e.SubItem.Text, CurrentFont, e.Bounds, TextColor, TextFormatFlags.EndEllipsis | TextFormatFlags.SingleLine | TextFormatFlags.HorizontalCenter | TextFormatFlags.VerticalCenter | TextFormatFlags.NoPrefix);
+				TextRenderer.DrawText(e.Graphics, e.SubItem!.Text, CurrentFont, e.Bounds, TextColor, TextFormatFlags.EndEllipsis | TextFormatFlags.SingleLine | TextFormatFlags.HorizontalCenter | TextFormatFlags.VerticalCenter | TextFormatFlags.NoPrefix);
 			}
 			else if (e.ColumnIndex == AuthorColumn.Index)
 			{
-				TextRenderer.DrawText(e.Graphics, e.SubItem.Text, CurrentFont, e.Bounds, TextColor, TextFormatFlags.EndEllipsis | TextFormatFlags.SingleLine | TextFormatFlags.VerticalCenter | TextFormatFlags.NoPrefix);
+				TextRenderer.DrawText(e.Graphics, e.SubItem!.Text, CurrentFont, e.Bounds, TextColor, TextFormatFlags.EndEllipsis | TextFormatFlags.SingleLine | TextFormatFlags.VerticalCenter | TextFormatFlags.NoPrefix);
 			}
 			else if (e.ColumnIndex == DescriptionColumn.Index)
 			{
@@ -2282,7 +2282,7 @@ namespace UnrealGameSync
 					RemainingBounds = new Rectangle(RemainingBounds.Left, RemainingBounds.Top, ListLocation.X - RemainingBounds.Left - (int)(2 * DpiScaleX), RemainingBounds.Height);
 				}
 
-				TextRenderer.DrawText(e.Graphics, e.SubItem.Text, CurrentFont, RemainingBounds, TextColor, TextFormatFlags.EndEllipsis | TextFormatFlags.SingleLine | TextFormatFlags.VerticalCenter | TextFormatFlags.NoPrefix);
+				TextRenderer.DrawText(e.Graphics, e.SubItem!.Text, CurrentFont, RemainingBounds, TextColor, TextFormatFlags.EndEllipsis | TextFormatFlags.SingleLine | TextFormatFlags.VerticalCenter | TextFormatFlags.NoPrefix);
 			}
 			else if (e.ColumnIndex == TypeColumn.Index)
 			{
@@ -2311,7 +2311,7 @@ namespace UnrealGameSync
 			}
 			else if (e.ColumnIndex == StatusColumn.Index)
 			{
-				float MaxX = e.SubItem.Bounds.Right;
+				float MaxX = e.SubItem!.Bounds.Right;
 
 				if (Change.Number == Workspace.PendingChangeNumber && Workspace.IsBusy())
 				{
@@ -2613,7 +2613,7 @@ namespace UnrealGameSync
 				ProcessStartInfo StartInfo = new ProcessStartInfo();
 				StartInfo.FileName = Url;
 				StartInfo.UseShellExecute = true;
-				using Process _ = Process.Start(StartInfo);
+				using Process? _ = Process.Start(StartInfo);
 			}
 			catch
 			{
@@ -2629,7 +2629,7 @@ namespace UnrealGameSync
 				StartInfo.FileName = Url;
 				StartInfo.Arguments = Arguments;
 				StartInfo.UseShellExecute = true;
-				using Process _ = Process.Start(StartInfo);
+				using Process? _ = Process.Start(StartInfo);
 			}
 			catch
 			{
@@ -2646,7 +2646,7 @@ namespace UnrealGameSync
 				StartInfo.Arguments = Arguments;
 				StartInfo.WorkingDirectory = WorkingDir;
 				StartInfo.UseShellExecute = true;
-				using Process _ = Process.Start(StartInfo);
+				using Process? _ = Process.Start(StartInfo);
 			}
 			catch
 			{
@@ -4090,14 +4090,14 @@ namespace UnrealGameSync
 						bool bIsBusy = Workspace.IsBusy();
 						bool bIsCurrentChange = (ContextMenuChange.Number == Workspace.CurrentChangeNumber);
 						BuildListContextMenu_Sync.Visible = !bIsBusy;
-						BuildListContextMenu_Sync.Font = new Font(SystemFonts.MenuFont, bIsCurrentChange ? FontStyle.Regular : FontStyle.Bold);
+						BuildListContextMenu_Sync.Font = new Font(SystemFonts.MenuFont!, bIsCurrentChange ? FontStyle.Regular : FontStyle.Bold);
 						BuildListContextMenu_SyncContentOnly.Visible = !bIsBusy && ShouldSyncPrecompiledEditor;
 						BuildListContextMenu_SyncOnlyThisChange.Visible = !bIsBusy && !bIsCurrentChange && ContextMenuChange.Number > Workspace.CurrentChangeNumber && Workspace.CurrentChangeNumber != -1;
 						BuildListContextMenu_Build.Visible = !bIsBusy && bIsCurrentChange && !ShouldSyncPrecompiledEditor;
 						BuildListContextMenu_Rebuild.Visible = !bIsBusy && bIsCurrentChange && !ShouldSyncPrecompiledEditor;
 						BuildListContextMenu_GenerateProjectFiles.Visible = !bIsBusy && bIsCurrentChange;
 						BuildListContextMenu_LaunchEditor.Visible = !bIsBusy && ContextMenuChange.Number == Workspace.CurrentChangeNumber;
-						BuildListContextMenu_LaunchEditor.Font = new Font(SystemFonts.MenuFont, FontStyle.Bold);
+						BuildListContextMenu_LaunchEditor.Font = new Font(SystemFonts.MenuFont!, FontStyle.Bold);
 						BuildListContextMenu_OpenVisualStudio.Visible = !bIsBusy && bIsCurrentChange;
 						BuildListContextMenu_Cancel.Visible = bIsBusy;
 
@@ -4271,7 +4271,7 @@ namespace UnrealGameSync
 			Dictionary<string, KeyValuePair<IArchiveInfo, int>> ArchiveTypeToSelection = new Dictionary<string, KeyValuePair<IArchiveInfo, int>>();
 			foreach (IArchiveInfo Archive in Archives)
 			{
-				ArchiveSettings ArchiveSettings = Settings.Archives.FirstOrDefault(x => x.Type == Archive.Type);
+				ArchiveSettings? ArchiveSettings = Settings.Archives.FirstOrDefault(x => x.Type == Archive.Type);
 				if (ArchiveSettings != null && ArchiveSettings.bEnabled)
 				{
 					int Preference = ArchiveSettings.Order.IndexOf(Archive.Name);
@@ -4303,7 +4303,7 @@ namespace UnrealGameSync
 
 		private void SetSelectedArchive(IArchiveInfo Archive, bool bSelected)
 		{
-			ArchiveSettings ArchiveSettings = Settings.Archives.FirstOrDefault(x => x.Type == Archive.Type);
+			ArchiveSettings? ArchiveSettings = Settings.Archives.FirstOrDefault(x => x.Type == Archive.Type);
 			if (ArchiveSettings == null)
 			{
 				ArchiveSettings = new ArchiveSettings(bSelected, Archive.Type, new string[] { Archive.Name });
@@ -4950,7 +4950,7 @@ namespace UnrealGameSync
 			{
 				// Check the project config if they have a scheduled setting selected.
 				string? ScheduledSyncTypeID = null;
-				UserSelectedProjectSettings ProjectSetting = Settings.ScheduleProjects.FirstOrDefault(x => x.LocalPath != null && x.LocalPath.Equals(SelectedProject.LocalPath, StringComparison.OrdinalIgnoreCase));
+				UserSelectedProjectSettings? ProjectSetting = Settings.ScheduleProjects.FirstOrDefault(x => x.LocalPath != null && x.LocalPath.Equals(SelectedProject.LocalPath, StringComparison.OrdinalIgnoreCase));
 				if (ProjectSetting != null)
 				{
 					ScheduledSyncTypeID = ProjectSetting.ScheduledSyncTypeID;
@@ -5234,7 +5234,7 @@ namespace UnrealGameSync
 						while (Stack.Count > 0)
 						{
 							Guid Id = Stack.Pop();
-							BuildStep NextStep = AllSteps.FirstOrDefault(x => x.UniqueId == Id);
+							BuildStep? NextStep = AllSteps.FirstOrDefault(x => x.UniqueId == Id);
 							if (NextStep != null)
 							{
 								foreach (Guid RequiresId in NextStep.Requires)
@@ -5382,7 +5382,7 @@ namespace UnrealGameSync
 		{
 			if (OptionsContextMenu_SyncPrecompiledBinaries.DropDownItems.Count == 0)
 			{
-				IArchiveInfo EditorArchive = GetArchives().FirstOrDefault(x => x.Type == IArchiveInfo.EditorArchiveType);
+				IArchiveInfo? EditorArchive = GetArchives().FirstOrDefault(x => x.Type == IArchiveInfo.EditorArchiveType);
 				if (EditorArchive != null)
 				{
 					SetSelectedArchive(EditorArchive, !OptionsContextMenu_SyncPrecompiledBinaries.Checked);
@@ -6111,7 +6111,7 @@ namespace UnrealGameSync
 			ConfigFile Config = new ConfigFile();
 			Config.Load(NewConfigFile);
 
-			ConfigSection Section = Config.FindSection("/Script/EditorStyle.EditorStyleSettings");
+			ConfigSection? Section = Config.FindSection("/Script/EditorStyle.EditorStyleSettings");
 			if (Section == null)
 			{
 				return null;
@@ -6151,7 +6151,7 @@ namespace UnrealGameSync
 			ConfigFile Config = new ConfigFile();
 			Config.Load(FileName);
 
-			ConfigSection Section = Config.FindSection("/Script/EditorStyle.EditorStyleSettings");
+			ConfigSection? Section = Config.FindSection("/Script/EditorStyle.EditorStyleSettings");
 			if (Section == null)
 			{
 				return null;
diff --git a/Engine/Source/Programs/UnrealGameSync/UnrealGameSync/EventMonitor.cs b/Engine/Source/Programs/UnrealGameSync/UnrealGameSync/EventMonitor.cs
--- a/Engine/Source/Programs/UnrealGameSync/UnrealGameSync/EventMonitor.cs
+++ b/Engine/Source/Programs/UnrealGameSync/UnrealGameSync/EventMonitor.cs
@@ -1142,19 +1142,19 @@ namespace UnrealGameSync
 		public bool IsUnderInvestigation(int ChangeNumber)
 		{
 			UpdateActiveInvestigations();
-			return ActiveInvestigations.Any(x => x.Change <= ChangeNumber);
+			return ActiveInvestigations!.Any(x => x.Change <= ChangeNumber);
 		}
 
 		public bool IsUnderInvestigationByCurrentUser(int ChangeNumber)
 		{
 			UpdateActiveInvestigations();
-			return ActiveInvestigations.Any(x => x.Change <= ChangeNumber && String.Compare(x.UserName, CurrentUserName, true) == 0);
+			return ActiveInvestigations!.Any(x => x.Change <= ChangeNumber && String.Compare(x.UserName, CurrentUserName, true) == 0);
 		}
 
 		public IEnumerable<string> GetInvestigatingUsers(int ChangeNumber)
 		{
 			UpdateActiveInvestigations();
-			return ActiveInvestigations.Where(x => ChangeNumber >= x.Change).Select(x => x.UserName);
+			return ActiveInvestigations!.Where(x => ChangeNumber >= x.Change).Select(x => x.UserName);
 		}
 
 		public int GetInvestigationStartChangeNumber(int LastChangeNumber)
diff --git a/Engine/Source/Programs/UnrealGameSync/UnrealGameSync/Forms/BuildStepWindow.cs b/Engine/Source/Programs/UnrealGameSync/UnrealGameSync/Forms/BuildStepWindow.cs
--- a/Engine/Source/Programs/UnrealGameSync/UnrealGameSync/Forms/BuildStepWindow.cs
+++ b/Engine/Source/Programs/UnrealGameSync/UnrealGameSync/Forms/BuildStepWindow.cs
@@ -230,7 +230,7 @@ namespace UnrealGameSync
 			}
 		}
 
-		private void OnClosedVariablesWindow(object sender, EventArgs e)
+		private void OnClosedVariablesWindow(object? sender, EventArgs e)
 		{
 			VariablesWindow = null;
 		}
diff --git a/Engine/Source/Programs/UnrealGameSync/UnrealGameSync/Forms/CleanWorkspaceWindow.cs b/Engine/Source/Programs/UnrealGameSync/UnrealGameSync/Forms/CleanWorkspaceWindow.cs
--- a/Engine/Source/Programs/UnrealGameSync/UnrealGameSync/Forms/CleanWorkspaceWindow.cs
+++ b/Engine/Source/Programs/UnrealGameSync/UnrealGameSync/Forms/CleanWorkspaceWindow.cs
@@ -512,7 +512,7 @@ namespace UnrealGameSync
 			}
 			else if (NodeData.File != null)
 			{
-				Process.Start("explorer.exe", String.Format("\"{0}\"", NodeData.File.Directory.FullName));
+				Process.Start("explorer.exe", String.Format("\"{0}\"", NodeData.File.Directory!.FullName));
 			}
 		}
 	}
diff --git a/Engine/Source/Programs/UnrealGameSync/UnrealGameSync/Forms/IssueBrowserWindow.cs b/Engine/Source/Programs/UnrealGameSync/UnrealGameSync/Forms/IssueBrowserWindow.cs
--- a/Engine/Source/Programs/UnrealGameSync/UnrealGameSync/Forms/IssueBrowserWindow.cs
+++ b/Engine/Source/Programs/UnrealGameSync/UnrealGameSync/Forms/IssueBrowserWindow.cs
@@ -262,7 +262,7 @@ namespace UnrealGameSync
 
 		public static void Show(Form Owner, IssueMonitor IssueMonitor, IPerforceSettings PerforceSettings, TimeSpan? ServerTimeOffset, IServiceProvider ServiceProvider, string? CurrentStream, Dictionary<string, Func<IssueData, bool>> CustomFilters, string? DefaultFilter)
 		{
-			IssueBrowserWindow Window = ExistingWindows.FirstOrDefault(x => x.IssueMonitor == IssueMonitor);
+			IssueBrowserWindow? Window = ExistingWindows.FirstOrDefault(x => x.IssueMonitor == IssueMonitor);
 			if(Window == null)
 			{
 				Window = new IssueBrowserWindow(IssueMonitor, PerforceSettings, ServerTimeOffset, ServiceProvider, CurrentStream, CustomFilters, DefaultFilter);
@@ -403,7 +403,7 @@ namespace UnrealGameSync
 
 		private void IssueListView_DrawSubItem(object sender, DrawListViewSubItemEventArgs e)
 		{
-			IssueData Issue = (IssueData)e.Item.Tag;
+			IssueData Issue = (IssueData)e.Item!.Tag;
 			if (e.ColumnIndex == IconHeader.Index)
 			{
 				if(!Issue.ResolvedAt.HasValue && Issue.FixChange == 0)
diff --git a/Engine/Source/Programs/UnrealGameSync/UnrealGameSync/Forms/IssueDetailsWindow.cs b/Engine/Source/Programs/UnrealGameSync/UnrealGameSync/Forms/IssueDetailsWindow.cs
--- a/Engine/Source/Programs/UnrealGameSync/UnrealGameSync/Forms/IssueDetailsWindow.cs
+++ b/Engine/Source/Programs/UnrealGameSync/UnrealGameSync/Forms/IssueDetailsWindow.cs
@@ -455,7 +455,7 @@ namespace UnrealGameSync
 			foreach (IGrouping<long, IssueDiagnosticData> Group in Diagnostics.GroupBy(x => x.BuildId ?? -1))
 			{
 				// Step 'Foo'
-				IssueBuildData Build = IssueBuilds.FirstOrDefault(x => x.Id == Group.Key);
+				IssueBuildData? Build = IssueBuilds.FirstOrDefault(x => x.Id == Group.Key);
 				if (Build != null)
 				{
 					RichText.Append(@"\pard");   // Paragraph default
@@ -527,7 +527,7 @@ namespace UnrealGameSync
 		{
 			UpdateSummaryTextIfChanged(SummaryTextBox, Issue.Summary.ToString());
 
-			IssueBuildData FirstFailingBuild = IssueBuilds.FirstOrDefault(x => x.ErrorUrl != null);
+			IssueBuildData? FirstFailingBuild = IssueBuilds.FirstOrDefault(x => x.ErrorUrl != null);
 			BuildLinkLabel.Text = (FirstFailingBuild != null)? FirstFailingBuild.JobName : "Unknown";
 
 			StringBuilder Status = new StringBuilder();
@@ -1044,7 +1044,7 @@ namespace UnrealGameSync
 
 		public static void Show(Form Owner, IssueMonitor IssueMonitor, IPerforceSettings PerforceSettings, TimeSpan? ServerTimeOffset, IssueData Issue, List<IssueBuildData> IssueBuilds, IServiceProvider ServiceProvider, string? CurrentStream)
 		{
-			IssueDetailsWindow Window = ExistingWindows.FirstOrDefault(x => x.IssueMonitor == IssueMonitor && x.Issue.Id == Issue.Id);
+			IssueDetailsWindow? Window = ExistingWindows.FirstOrDefault(x => x.IssueMonitor == IssueMonitor && x.Issue.Id == Issue.Id);
 			if(Window == null)
 			{
 				List<IssueDiagnosticData> Diagnostics = new List<IssueDiagnosticData>();
@@ -1145,11 +1145,11 @@ namespace UnrealGameSync
 
 			Color TextColor = SystemColors.WindowText;
 
-			if(e.SubItem.Tag is CustomListViewWidget)
+			if(e.SubItem!.Tag is CustomListViewWidget)
 			{
 				BuildListView.DrawCustomSubItem(e.Graphics, e.SubItem);
 			}
-			else if(e.Item.Tag is ChangesRecord)
+			else if(e.Item!.Tag is ChangesRecord)
 			{
 				ChangesRecord Change = (ChangesRecord)e.Item.Tag;
 
@@ -1369,7 +1369,7 @@ namespace UnrealGameSync
 
 		private void DescriptionLabel_LinkClicked(object sender, LinkLabelLinkClickedEventArgs e)
 		{
-			IssueBuildData LastBuild = IssueBuilds.Where(x => x.Stream == SelectedStream).OrderByDescending(x => x.Change).ThenByDescending(x => x.ErrorUrl).FirstOrDefault();
+			IssueBuildData? LastBuild = IssueBuilds.Where(x => x.Stream == SelectedStream).OrderByDescending(x => x.Change).ThenByDescending(x => x.ErrorUrl).FirstOrDefault();
 			if(LastBuild != null)
 			{
 				Utility.OpenUrl(LastBuild.ErrorUrl);
@@ -1417,7 +1417,7 @@ namespace UnrealGameSync
 
 		private void BuildLinkLabel_LinkClicked(object sender, LinkLabelLinkClickedEventArgs e)
 		{
-			IssueBuildData Build = IssueBuilds.FirstOrDefault(x => x.ErrorUrl != null);
+			IssueBuildData? Build = IssueBuilds.FirstOrDefault(x => x.ErrorUrl != null);
 			if (Build != null)
 			{
 				Utility.OpenUrl(Build.ErrorUrl);
@@ -1426,7 +1426,7 @@ namespace UnrealGameSync
 
 		private void DetailsTextBox_LinkClicked(object sender, LinkClickedEventArgs e)
 		{
-			Utility.OpenUrl(e.LinkText);
+			Utility.OpenUrl(e.LinkText!);
 		}
 	}
 }
diff --git a/Engine/Source/Programs/UnrealGameSync/UnrealGameSync/Forms/MainWindow.cs b/Engine/Source/Programs/UnrealGameSync/UnrealGameSync/Forms/MainWindow.cs
--- a/Engine/Source/Programs/UnrealGameSync/UnrealGameSync/Forms/MainWindow.cs
+++ b/Engine/Source/Programs/UnrealGameSync/UnrealGameSync/Forms/MainWindow.cs
@@ -1519,7 +1519,7 @@ namespace UnrealGameSync
 						}
 					}
 
-					IssueAlertWindow AlertWindow = AlertWindows.FirstOrDefault(x => x.IssueMonitor == IssueMonitor && x.Issue.Id == Issue.Id);
+					IssueAlertWindow? AlertWindow = AlertWindows.FirstOrDefault(x => x.IssueMonitor == IssueMonitor && x.Issue.Id == Issue.Id);
 					if(AlertWindow == null)
 					{
 						if(Reason != 0)
@@ -1668,7 +1668,7 @@ namespace UnrealGameSync
 
 		public IssueMonitor CreateIssueMonitor(string? ApiUrl, string UserName)
 		{
-			WorkspaceIssueMonitor WorkspaceIssueMonitor = WorkspaceIssueMonitors.FirstOrDefault(x => String.Compare(x.IssueMonitor.ApiUrl, ApiUrl, StringComparison.OrdinalIgnoreCase) == 0 && String.Compare(x.IssueMonitor.UserName, UserName, StringComparison.OrdinalIgnoreCase) == 0);
+			WorkspaceIssueMonitor? WorkspaceIssueMonitor = WorkspaceIssueMonitors.FirstOrDefault(x => String.Compare(x.IssueMonitor.ApiUrl, ApiUrl, StringComparison.OrdinalIgnoreCase) == 0 && String.Compare(x.IssueMonitor.UserName, UserName, StringComparison.OrdinalIgnoreCase) == 0);
 			if (WorkspaceIssueMonitor == null)
 			{
 				string ServerId = ApiUrl != null ? Regex.Replace(ApiUrl, @"^.*://", "") : "noserver";
diff --git a/Engine/Source/Programs/UnrealGameSync/UnrealGameSync/Forms/ModifyBuildStepsWindow.cs b/Engine/Source/Programs/UnrealGameSync/UnrealGameSync/Forms/ModifyBuildStepsWindow.cs
--- a/Engine/Source/Programs/UnrealGameSync/UnrealGameSync/Forms/ModifyBuildStepsWindow.cs
+++ b/Engine/Source/Programs/UnrealGameSync/UnrealGameSync/Forms/ModifyBuildStepsWindow.cs
@@ -149,7 +149,7 @@ namespace UnrealGameSync
 			}
 			else
 			{
-				BuildStep Task = (BuildStep)e.Item.Tag;
+				BuildStep Task = (BuildStep)e.Item!.Tag;
 
 				bool bEnabled;
 				if(e.ColumnIndex == 1)
diff --git a/Engine/Source/Programs/UnrealGameSync/UnrealGameSync/Forms/NotificationWindow.cs b/Engine/Source/Programs/UnrealGameSync/UnrealGameSync/Forms/NotificationWindow.cs
--- a/Engine/Source/Programs/UnrealGameSync/UnrealGameSync/Forms/NotificationWindow.cs
+++ b/Engine/Source/Programs/UnrealGameSync/UnrealGameSync/Forms/NotificationWindow.cs
@@ -234,7 +234,7 @@ namespace UnrealGameSync
 			e.Graphics.DrawRectangle(Pens.Black, new Rectangle(0, 0, Width - 1, Height - 1));
 
 			e.Graphics.InterpolationMode = InterpolationMode.HighQualityBicubic;
-			e.Graphics.DrawImage(LogoBitmap, LogoBounds);
+			e.Graphics.DrawImage(LogoBitmap!, LogoBounds);
 		}
 
 		protected override void OnPaint(PaintEventArgs e)
diff --git a/Engine/Source/Programs/UnrealGameSync/UnrealGameSync/Forms/ScheduleWindow.cs b/Engine/Source/Programs/UnrealGameSync/UnrealGameSync/Forms/ScheduleWindow.cs
--- a/Engine/Source/Programs/UnrealGameSync/UnrealGameSync/Forms/ScheduleWindow.cs
+++ b/Engine/Source/Programs/UnrealGameSync/UnrealGameSync/Forms/ScheduleWindow.cs
@@ -76,7 +76,7 @@ namespace UnrealGameSync
 			}
 		}
 
-		private void ProjectListBox_ItemCheck(object sender, ItemCheckEventArgs e)
+		private void ProjectListBox_ItemCheck(object? sender, ItemCheckEventArgs e)
 		{
 			bool bIsAnyOpenProjectIndex = e.Index == 0;
 			if (e.NewValue == CheckState.Checked && !bIsAnyOpenProjectIndex)
@@ -90,7 +90,7 @@ namespace UnrealGameSync
 			}
 		}
 
-		private void SyncTypeDropDown_Closed(object sender, ToolStripDropDownClosedEventArgs e)
+		private void SyncTypeDropDown_Closed(object? sender, ToolStripDropDownClosedEventArgs e)
 		{
 			if (e.CloseReason != ToolStripDropDownCloseReason.ItemClicked)
 			{
diff --git a/Engine/Source/Programs/UnrealGameSync/UnrealGameSync/Forms/SdkInfoWindow.cs b/Engine/Source/Programs/UnrealGameSync/UnrealGameSync/Forms/SdkInfoWindow.cs
--- a/Engine/Source/Programs/UnrealGameSync/UnrealGameSync/Forms/SdkInfoWindow.cs
+++ b/Engine/Source/Programs/UnrealGameSync/UnrealGameSync/Forms/SdkInfoWindow.cs
@@ -147,13 +147,13 @@ namespace UnrealGameSync
 		{
 			if(e.ColumnIndex != columnHeader3.Index)
 			{
-				TextRenderer.DrawText(e.Graphics, e.SubItem.Text, SdkListView.Font, e.Bounds, SdkListView.ForeColor, TextFormatFlags.EndEllipsis | TextFormatFlags.SingleLine | TextFormatFlags.Left | TextFormatFlags.VerticalCenter | TextFormatFlags.NoPrefix);
+				TextRenderer.DrawText(e.Graphics, e.SubItem!.Text, SdkListView.Font, e.Bounds, SdkListView.ForeColor, TextFormatFlags.EndEllipsis | TextFormatFlags.SingleLine | TextFormatFlags.Left | TextFormatFlags.VerticalCenter | TextFormatFlags.NoPrefix);
 			}
 			else
 			{
 				e.Graphics.SmoothingMode = SmoothingMode.AntiAlias;
 
-				List<BadgeInfo> Badges = GetBadges(e.Item, e.SubItem);
+				List<BadgeInfo> Badges = GetBadges(e.Item!, e.SubItem!);
 				for(int Idx = 0; Idx < Badges.Count; Idx++)
 				{
 					Color BadgeColor = (HoverBadgeUniqueId == Badges[Idx].UniqueId)? Color.FromArgb(140, 180, 230) : Color.FromArgb(112, 146, 190);
diff --git a/Engine/Source/Programs/UnrealGameSync/UnrealGameSync/Forms/SelectProjectFromWorkspaceWindow.cs b/Engine/Source/Programs/UnrealGameSync/UnrealGameSync/Forms/SelectProjectFromWorkspaceWindow.cs
--- a/Engine/Source/Programs/UnrealGameSync/UnrealGameSync/Forms/SelectProjectFromWorkspaceWindow.cs
+++ b/Engine/Source/Programs/UnrealGameSync/UnrealGameSync/Forms/SelectProjectFromWorkspaceWindow.cs
@@ -225,7 +225,7 @@ namespace UnrealGameSync
 
 		private void ProjectTreeView_AfterSelect(object sender, TreeViewEventArgs e)
 		{
-			OkBtn.Enabled = (e.Node.Tag != null);
+			OkBtn.Enabled = (e.Node!.Tag != null);
 		}
 
 		private void OkBtn_Click(object sender, EventArgs e)
diff --git a/Engine/Source/Programs/UnrealGameSync/UnrealGameSync/IssueMonitor.cs b/Engine/Source/Programs/UnrealGameSync/UnrealGameSync/IssueMonitor.cs
--- a/Engine/Source/Programs/UnrealGameSync/UnrealGameSync/IssueMonitor.cs
+++ b/Engine/Source/Programs/UnrealGameSync/UnrealGameSync/IssueMonitor.cs
@@ -427,7 +427,7 @@ namespace UnrealGameSync
 						{
 							NewIssue.bIsWarning = !Builds.Any(x => x.Outcome != IssueBuildOutcome.Warning);
 
-							IssueBuildData LastBuild = Builds.OrderByDescending(x => x.Change).FirstOrDefault();
+							IssueBuildData? LastBuild = Builds.OrderByDescending(x => x.Change).FirstOrDefault();
 							if (LastBuild != null && !String.IsNullOrEmpty(LastBuild.ErrorUrl))
 							{
 								NewIssue.BuildUrl = LastBuild.ErrorUrl;
diff --git a/Engine/Source/Programs/UnrealGameSync/UnrealGameSync/PerforceMonitor.cs b/Engine/Source/Programs/UnrealGameSync/UnrealGameSync/PerforceMonitor.cs
--- a/Engine/Source/Programs/UnrealGameSync/UnrealGameSync/PerforceMonitor.cs
+++ b/Engine/Source/Programs/UnrealGameSync/UnrealGameSync/PerforceMonitor.cs
@@ -368,7 +368,7 @@ namespace UnrealGameSync
 				}
 
 				// Add in additional paths property
-				ConfigSection ProjectConfigSection = LatestProjectConfigFile.FindSection("Perforce");
+				ConfigSection? ProjectConfigSection = LatestProjectConfigFile.FindSection("Perforce");
 				if (ProjectConfigSection != null)
 				{
 					IEnumerable<string> AdditionalPaths = ProjectConfigSection.GetValues("AdditionalPathsToSync", new string[0]);
@@ -580,7 +580,7 @@ namespace UnrealGameSync
 			List<PerforceArchiveInfo> NewArchives = new List<PerforceArchiveInfo>();
 
 			// Find all the zipped binaries under this stream
-			ConfigSection ProjectConfigSection = LatestProjectConfigFile.FindSection(SelectedProjectIdentifier);
+			ConfigSection? ProjectConfigSection = LatestProjectConfigFile.FindSection(SelectedProjectIdentifier);
 			if (ProjectConfigSection != null)
 			{
 				// Legacy
diff --git a/Engine/Source/Programs/UnrealGameSync/UnrealGameSync/Program.cs b/Engine/Source/Programs/UnrealGameSync/UnrealGameSync/Program.cs
--- a/Engine/Source/Programs/UnrealGameSync/UnrealGameSync/Program.cs
+++ b/Engine/Source/Programs/UnrealGameSync/UnrealGameSync/Program.cs
@@ -65,6 +65,7 @@ namespace UnrealGameSync
 				{
 					Application.EnableVisualStyles();
 					Application.SetCompatibleTextRenderingDefault(false);
+					Application.SetHighDpiMode(HighDpiMode.PerMonitorV2);
 				}
 
 				using (EventWaitHandle ActivateEvent = new EventWaitHandle(false, EventResetMode.AutoReset, "ActivateUnrealGameSync"))
@@ -120,7 +121,7 @@ namespace UnrealGameSync
 			// Set the current working directory to the update directory to prevent child-process file handles from disrupting auto-updates
 			if (UpdateSpawn != null)
 			{
-				Directory.SetCurrentDirectory(Path.GetDirectoryName(UpdateSpawn));
+				Directory.SetCurrentDirectory(Path.GetDirectoryName(UpdateSpawn)!);
 			}
 
 			string SyncVersionFile = Path.Combine(Path.GetDirectoryName(Assembly.GetExecutingAssembly().Location!)!, "SyncVersion.txt");
diff --git a/Engine/Source/Programs/UnrealGameSync/UnrealGameSync/ProgramApplicationContext.cs b/Engine/Source/Programs/UnrealGameSync/UnrealGameSync/ProgramApplicationContext.cs
--- a/Engine/Source/Programs/UnrealGameSync/UnrealGameSync/ProgramApplicationContext.cs
+++ b/Engine/Source/Programs/UnrealGameSync/UnrealGameSync/ProgramApplicationContext.cs
@@ -260,7 +260,7 @@ namespace UnrealGameSync
 			MainWindowInstance.FormClosed += MainWindowInstance_FormClosed;
 		}
 
-		private void MainWindowInstance_FormClosed(object sender, FormClosedEventArgs e)
+		private void MainWindowInstance_FormClosed(object? sender, FormClosedEventArgs e)
 		{
 			ExitThread();
 		}
@@ -345,7 +345,7 @@ namespace UnrealGameSync
 			}
 		}
 
-		private void NotifyIcon_MouseDown(object sender, MouseEventArgs e)
+		private void NotifyIcon_MouseDown(object? sender, MouseEventArgs e)
 		{
 			// Have to set up this stuff here, because the menu is laid out before Opening() is called on it after mouse-up.
 			bool bCanSyncNow = MainWindowInstance != null && MainWindowInstance.CanSyncNow();
diff --git a/Engine/Source/Programs/UnrealGameSync/UnrealGameSync/Properties/app.manifest b/Engine/Source/Programs/UnrealGameSync/UnrealGameSync/Properties/app.manifest
--- a/Engine/Source/Programs/UnrealGameSync/UnrealGameSync/Properties/app.manifest
+++ b/Engine/Source/Programs/UnrealGameSync/UnrealGameSync/Properties/app.manifest
@@ -38,11 +38,6 @@
       <!--<supportedOS Id="{1f676c76-80e1-4239-95bb-83d0f6d0da78}"/>-->
     </application>
   </compatibility>
-  <asmv3:application>
-    <asmv3:windowsSettings xmlns="http://schemas.microsoft.com/SMI/2005/WindowsSettings">
-      <dpiAware>true</dpiAware>
-    </asmv3:windowsSettings>
-  </asmv3:application>
 
 	<!-- Enable themes for Windows common controls and dialogs (Windows XP and later) -->
   <!-- <dependency>
diff --git a/Engine/Source/Programs/UnrealGameSync/UnrealGameSync/RESTApi.cs b/Engine/Source/Programs/UnrealGameSync/UnrealGameSync/RESTApi.cs
--- a/Engine/Source/Programs/UnrealGameSync/UnrealGameSync/RESTApi.cs
+++ b/Engine/Source/Programs/UnrealGameSync/UnrealGameSync/RESTApi.cs
@@ -76,7 +76,7 @@ namespace UnrealGameSync
 
 		public static async Task<T> GetAsync<T>(string Url, CancellationToken CancellationToken)
 		{
-			return JsonSerializer.Deserialize<T>(await GetAsync(Url, CancellationToken), Utility.DefaultJsonSerializerOptions);
+			return JsonSerializer.Deserialize<T>(await GetAsync(Url, CancellationToken), Utility.DefaultJsonSerializerOptions)!;
 		}
 
 		public static Task<string> PutAsync<T>(string Url, T Object, CancellationToken CancellationToken)
diff --git a/Engine/Source/Programs/UnrealGameSync/UnrealGameSync/UnrealGameSync.csproj b/Engine/Source/Programs/UnrealGameSync/UnrealGameSync/UnrealGameSync.csproj
--- a/Engine/Source/Programs/UnrealGameSync/UnrealGameSync/UnrealGameSync.csproj
+++ b/Engine/Source/Programs/UnrealGameSync/UnrealGameSync/UnrealGameSync.csproj
@@ -1,7 +1,7 @@
-<Project Sdk="Microsoft.NET.Sdk.WindowsDesktop">
+<Project Sdk="Microsoft.NET.Sdk">
   <PropertyGroup>
     <OutputType>WinExe</OutputType>
-    <TargetFramework>netcoreapp3.1</TargetFramework>
+    <TargetFramework>net6.0-windows</TargetFramework>
     <UseWindowsForms>true</UseWindowsForms>
     <Deterministic>true</Deterministic>
     <Product>UnrealGameSync</Product>
diff --git a/Engine/Source/Programs/UnrealGameSync/UnrealGameSync/UriHandlers/UriHandler.cs b/Engine/Source/Programs/UnrealGameSync/UnrealGameSync/UriHandlers/UriHandler.cs
--- a/Engine/Source/Programs/UnrealGameSync/UnrealGameSync/UriHandlers/UriHandler.cs
+++ b/Engine/Source/Programs/UnrealGameSync/UnrealGameSync/UriHandlers/UriHandler.cs
@@ -49,7 +49,7 @@ namespace UnrealGameSync
 
 				foreach (ParameterInfo Param in Info.GetParameters())
 				{					
-					string Value = Query.Get(Param.Name);
+					string? Value = Query.Get(Param.Name);
 
 					if (Value == null)
 					{
@@ -296,7 +296,7 @@ namespace UnrealGameSync
 
 		public static string CurrentProcessFilePath
 		{
-			get { return Path.GetFullPath(Process.GetCurrentProcess().MainModule.FileName); }
+			get { return Path.GetFullPath(Process.GetCurrentProcess().MainModule!.FileName!); }
 		}
 
 		static List<RegistrySetting> GetRegistrySettings()
diff --git a/Engine/Source/Programs/UnrealGameSync/UnrealGameSyncAutomation/UnrealGameSyncAutomation.csproj b/Engine/Source/Programs/UnrealGameSync/UnrealGameSyncAutomation/UnrealGameSyncAutomation.csproj
--- a/Engine/Source/Programs/UnrealGameSync/UnrealGameSyncAutomation/UnrealGameSyncAutomation.csproj
+++ b/Engine/Source/Programs/UnrealGameSync/UnrealGameSyncAutomation/UnrealGameSyncAutomation.csproj
@@ -1,11 +1,11 @@
-﻿<?xml version="1.0" encoding="utf-8"?>
-<Project Sdk="Microsoft.NET.Sdk.WindowsDesktop">
+<Project Sdk="Microsoft.NET.Sdk">
   <PropertyGroup>
     <OutputType>WinExe</OutputType>
-    <TargetFramework>netcoreapp3.1</TargetFramework>
+    <TargetFramework>net6.0-windows</TargetFramework>
     <UseWindowsForms>true</UseWindowsForms>
     <GenerateAssemblyInfo>false</GenerateAssemblyInfo>
     <RootNamespace>UnrealGameSyncAutomation</RootNamespace>
     <AssemblyName>UnrealGameSyncAutomation</AssemblyName>
+    <NoWarn>CA1416</NoWarn> <!-- Suppress warnings about functions only available on Windows -->
   </PropertyGroup>
 </Project>
\ No newline at end of file
diff --git a/Engine/Source/Programs/UnrealGameSync/UnrealGameSyncCmd/Program.cs b/Engine/Source/Programs/UnrealGameSync/UnrealGameSyncCmd/Program.cs
--- a/Engine/Source/Programs/UnrealGameSync/UnrealGameSyncCmd/Program.cs
+++ b/Engine/Source/Programs/UnrealGameSync/UnrealGameSyncCmd/Program.cs
@@ -333,7 +333,7 @@ namespace UnrealGameSyncCmd
 			Dictionary<Guid, WorkspaceSyncCategory> SyncCategories = ConfigUtils.GetSyncCategories(ProjectConfig);
 			string[] CombinedSyncFilter = GlobalSettingsFile.GetCombinedSyncFilter(SyncCategories, UserSettings.Global.Filter, WorkspaceSettings.Filter);
 
-			ConfigSection PerforceSection = ProjectConfig.FindSection("Perforce");
+			ConfigSection? PerforceSection = ProjectConfig.FindSection("Perforce");
 			if (PerforceSection != null)
 			{
 				IEnumerable<string> AdditionalPaths = PerforceSection.GetValues("AdditionalPathsToSync", new string[0]);
diff --git a/Engine/Source/Programs/UnrealGameSync/UnrealGameSyncCmd/UnrealGameSyncCmd.csproj b/Engine/Source/Programs/UnrealGameSync/UnrealGameSyncCmd/UnrealGameSyncCmd.csproj
--- a/Engine/Source/Programs/UnrealGameSync/UnrealGameSyncCmd/UnrealGameSyncCmd.csproj
+++ b/Engine/Source/Programs/UnrealGameSync/UnrealGameSyncCmd/UnrealGameSyncCmd.csproj
@@ -2,7 +2,7 @@
 
   <PropertyGroup>
     <OutputType>Exe</OutputType>
-    <TargetFramework>netcoreapp3.1</TargetFramework>
+    <TargetFramework>net6.0</TargetFramework>
     <Nullable>enable</Nullable>
     <AssemblyName>ugs</AssemblyName>
   </PropertyGroup>
diff --git a/Engine/Source/Programs/UnrealGameSync/UnrealGameSyncLauncher/Program.cs b/Engine/Source/Programs/UnrealGameSync/UnrealGameSyncLauncher/Program.cs
--- a/Engine/Source/Programs/UnrealGameSync/UnrealGameSyncLauncher/Program.cs
+++ b/Engine/Source/Programs/UnrealGameSync/UnrealGameSyncLauncher/Program.cs
@@ -26,6 +26,7 @@ namespace UnrealGameSyncLauncher
 		{
 			Application.EnableVisualStyles();
 			Application.SetCompatibleTextRenderingDefault(false);
+			Application.SetHighDpiMode(HighDpiMode.PerMonitorV2);
 
 			bool bFirstInstance;
 			using(Mutex InstanceMutex = new Mutex(true, "UnrealGameSyncRunning", out bFirstInstance))
diff --git a/Engine/Source/Programs/UnrealGameSync/UnrealGameSyncLauncher/UnrealGameSyncLauncher.csproj b/Engine/Source/Programs/UnrealGameSync/UnrealGameSyncLauncher/UnrealGameSyncLauncher.csproj
--- a/Engine/Source/Programs/UnrealGameSync/UnrealGameSyncLauncher/UnrealGameSyncLauncher.csproj
+++ b/Engine/Source/Programs/UnrealGameSync/UnrealGameSyncLauncher/UnrealGameSyncLauncher.csproj
@@ -1,10 +1,9 @@
-<?xml version="1.0" encoding="utf-8"?>
-<Project Sdk="Microsoft.NET.Sdk.WindowsDesktop">
+<Project Sdk="Microsoft.NET.Sdk">
 
   <PropertyGroup>
     <OutputType>WinExe</OutputType>
     <DefineConstants>UGS_LAUNCHER</DefineConstants>
-    <TargetFramework>netcoreapp3.1</TargetFramework>
+    <TargetFramework>net6.0-windows</TargetFramework>
     <UseWindowsForms>true</UseWindowsForms>
     <GenerateAssemblyInfo>false</GenerateAssemblyInfo>
     <RootNamespace>UnrealGameSyncLauncher</RootNamespace>
@@ -14,6 +13,7 @@
     <ApplicationManifest>app.manifest</ApplicationManifest>
     <RuntimeIdentifiers>win10-x64</RuntimeIdentifiers>
     <Nullable>enable</Nullable>
+    <NoWarn>CA1416</NoWarn> <!-- Suppress warnings about functions only available on Windows -->
   </PropertyGroup>
   <ItemGroup>
     <Compile Include="..\UnrealGameSync\Forms\ModalTaskWindow.cs">
diff --git a/Engine/Source/Programs/UnrealGameSync/UnrealGameSyncLauncher/app.manifest b/Engine/Source/Programs/UnrealGameSync/UnrealGameSyncLauncher/app.manifest
--- a/Engine/Source/Programs/UnrealGameSync/UnrealGameSyncLauncher/app.manifest
+++ b/Engine/Source/Programs/UnrealGameSync/UnrealGameSyncLauncher/app.manifest
@@ -45,16 +45,6 @@
     </application>
   </compatibility>
 
-  <!-- Indicates that the application is DPI-aware and will not be automatically scaled by Windows at higher
-       DPIs. Windows Presentation Foundation (WPF) applications are automatically DPI-aware and do not need 
-       to opt in. Windows Forms applications targeting .NET Framework 4.6 that opt into this setting, should 
-       also set the 'EnableWindowsFormsHighDpiAutoResizing' setting to 'true' in their app.config. -->
-  <application xmlns="urn:schemas-microsoft-com:asm.v3">
-    <windowsSettings>
-      <dpiAware xmlns="http://schemas.microsoft.com/SMI/2005/WindowsSettings">true</dpiAware>
-    </windowsSettings>
-  </application>
-
   <!-- Enable themes for Windows common controls and dialogs (Windows XP and later) -->
   <!--
   <dependency>
diff --git a/Engine/Source/Programs/UnrealGameSync/UnrealGameSyncShared/ConfigFile.cs b/Engine/Source/Programs/UnrealGameSync/UnrealGameSyncShared/ConfigFile.cs
--- a/Engine/Source/Programs/UnrealGameSync/UnrealGameSyncShared/ConfigFile.cs
+++ b/Engine/Source/Programs/UnrealGameSync/UnrealGameSyncShared/ConfigFile.cs
@@ -310,9 +310,9 @@ namespace UnrealGameSync
 			{
 				while (Reader.Read() && Reader.TokenType == JsonTokenType.PropertyName)
 				{
-					string Name = Reader.GetString();
+					string Name = Reader.GetString()!;
 					Reader.Read();
-					string Value = Reader.GetString();
+					string Value = Reader.GetString()!;
 					Object.Pairs.Add(KeyValuePair.Create(Name, Value));
 				}
 			}
@@ -591,14 +591,14 @@ namespace UnrealGameSync
 			FileReference.Move(TempFileName, FileName);
 		}
 
-		public ConfigSection FindSection(string Name)
+		public ConfigSection? FindSection(string Name)
 		{
 			return Sections.FirstOrDefault(x => String.Compare(x.Name, Name, true) == 0);
 		}
 
 		public ConfigSection FindOrAddSection(string Name)
 		{
-			ConfigSection Section = FindSection(Name);
+			ConfigSection? Section = FindSection(Name);
 			if(Section == null)
 			{
 				Section = new ConfigSection(Name);
@@ -638,14 +638,14 @@ namespace UnrealGameSync
 		public bool GetValue(string Key, bool DefaultValue)
 		{
 			int DotIdx = Key.IndexOf('.');
-			ConfigSection Section = FindSection(Key.Substring(0, DotIdx));
+			ConfigSection? Section = FindSection(Key.Substring(0, DotIdx));
 			return (Section == null)? DefaultValue : Section.GetValue(Key.Substring(DotIdx + 1), DefaultValue);
 		}
 
 		public int GetValue(string Key, int DefaultValue)
 		{
 			int DotIdx = Key.IndexOf('.');
-			ConfigSection Section = FindSection(Key.Substring(0, DotIdx));
+			ConfigSection? Section = FindSection(Key.Substring(0, DotIdx));
 			return (Section == null)? DefaultValue : Section.GetValue(Key.Substring(DotIdx + 1), DefaultValue);
 		}
 
@@ -653,7 +653,7 @@ namespace UnrealGameSync
 		public string? GetValue(string Key, string? DefaultValue)
 		{
 			int DotIdx = Key.IndexOf('.');
-			ConfigSection Section = FindSection(Key.Substring(0, DotIdx));
+			ConfigSection? Section = FindSection(Key.Substring(0, DotIdx));
 			return (Section == null)? DefaultValue : Section.GetValue(Key.Substring(DotIdx + 1), DefaultValue);
 		}
 
@@ -661,7 +661,7 @@ namespace UnrealGameSync
 		public string[]? GetValues(string Key, string[]? DefaultValue)
 		{
 			int DotIdx = Key.IndexOf('.');
-			ConfigSection Section = FindSection(Key.Substring(0, DotIdx));
+			ConfigSection? Section = FindSection(Key.Substring(0, DotIdx));
 			return (Section == null)? DefaultValue : Section.GetValues(Key.Substring(DotIdx + 1), DefaultValue);
 		}
 
@@ -669,7 +669,7 @@ namespace UnrealGameSync
 		public Guid[]? GetGuidValues(string Key, Guid[]? DefaultValue)
 		{
 			int DotIdx = Key.IndexOf('.');
-			ConfigSection Section = FindSection(Key.Substring(0, DotIdx));
+			ConfigSection? Section = FindSection(Key.Substring(0, DotIdx));
 			return (Section == null)? DefaultValue : Section.GetValues(Key.Substring(DotIdx + 1), DefaultValue);
 		}
 	}
diff --git a/Engine/Source/Programs/UnrealGameSync/UnrealGameSyncShared/ConfigUtils.cs b/Engine/Source/Programs/UnrealGameSync/UnrealGameSyncShared/ConfigUtils.cs
--- a/Engine/Source/Programs/UnrealGameSync/UnrealGameSyncShared/ConfigUtils.cs
+++ b/Engine/Source/Programs/UnrealGameSync/UnrealGameSyncShared/ConfigUtils.cs
@@ -369,7 +369,7 @@ namespace UnrealGameSync
 			string Path = SelectedProjectIdentifier;
 			for (; ; )
 			{
-				ConfigSection ProjectSection = ProjectConfigFile.FindSection(Path);
+				ConfigSection? ProjectSection = ProjectConfigFile.FindSection(Path);
 				if (ProjectSection != null)
 				{
 					string? NewValue = ProjectSection.GetValue(Name, null);
@@ -389,7 +389,7 @@ namespace UnrealGameSync
 				Path = Path.Substring(0, LastSlash);
 			}
 
-			ConfigSection DefaultSection = ProjectConfigFile.FindSection("Default");
+			ConfigSection? DefaultSection = ProjectConfigFile.FindSection("Default");
 			if (DefaultSection != null)
 			{
 				string? NewValue = DefaultSection.GetValue(Name, null);
diff --git a/Engine/Source/Programs/UnrealGameSync/UnrealGameSyncShared/FileFilter.cs b/Engine/Source/Programs/UnrealGameSync/UnrealGameSyncShared/FileFilter.cs
--- a/Engine/Source/Programs/UnrealGameSync/UnrealGameSyncShared/FileFilter.cs
+++ b/Engine/Source/Programs/UnrealGameSync/UnrealGameSyncShared/FileFilter.cs
@@ -266,7 +266,7 @@ namespace UnrealGameSync
 			FileFilterNode LastNode = RootNode;
 			foreach (string BranchPattern in BranchPatterns)
 			{
-				FileFilterNode NextNode = LastNode.Branches.FirstOrDefault(x => x.Pattern == BranchPattern);
+				FileFilterNode? NextNode = LastNode.Branches.FirstOrDefault(x => x.Pattern == BranchPattern);
 				if (NextNode == null)
 				{
 					NextNode = new FileFilterNode(LastNode, BranchPattern);
diff --git a/Engine/Source/Programs/UnrealGameSync/UnrealGameSyncShared/Telemetry.cs b/Engine/Source/Programs/UnrealGameSync/UnrealGameSyncShared/Telemetry.cs
--- a/Engine/Source/Programs/UnrealGameSync/UnrealGameSyncShared/Telemetry.cs
+++ b/Engine/Source/Programs/UnrealGameSync/UnrealGameSyncShared/Telemetry.cs
@@ -180,7 +180,7 @@ namespace UnrealGameSync
 				catch (WebException Ex)
 				{
 					// Handle errors. Any non-200 responses automatically generate a WebException.
-					HttpWebResponse Response = (HttpWebResponse)Ex.Response;
+					HttpWebResponse? Response = (HttpWebResponse?)Ex.Response;
 					if (Response == null)
 					{
 						Logger.LogError(Ex, "Exception while attempting to send event");
diff --git a/Engine/Source/Programs/UnrealGameSync/UnrealGameSyncShared/UnrealGameSyncShared.csproj b/Engine/Source/Programs/UnrealGameSync/UnrealGameSyncShared/UnrealGameSyncShared.csproj
--- a/Engine/Source/Programs/UnrealGameSync/UnrealGameSyncShared/UnrealGameSyncShared.csproj
+++ b/Engine/Source/Programs/UnrealGameSync/UnrealGameSyncShared/UnrealGameSyncShared.csproj
@@ -1,7 +1,7 @@
 <Project Sdk="Microsoft.NET.Sdk">
 
   <PropertyGroup>
-    <TargetFramework>netcoreapp3.1</TargetFramework>
+    <TargetFramework>net6.0</TargetFramework>
     <Nullable>enable</Nullable>
     <RootNamespace>UnrealGameSync</RootNamespace>
   </PropertyGroup>
diff --git a/Engine/Source/Programs/UnrealGameSync/UnrealGameSyncShared/UserSettings.cs b/Engine/Source/Programs/UnrealGameSync/UnrealGameSyncShared/UserSettings.cs
--- a/Engine/Source/Programs/UnrealGameSync/UnrealGameSyncShared/UserSettings.cs
+++ b/Engine/Source/Programs/UnrealGameSync/UnrealGameSyncShared/UserSettings.cs
@@ -463,7 +463,7 @@ namespace UnrealGameSync
 
 		public void SetBisectState(int Change, BisectState State)
 		{
-			BisectEntry Entry = BisectChanges.FirstOrDefault(x => x.Change == Change);
+			BisectEntry? Entry = BisectChanges.FirstOrDefault(x => x.Change == Change);
 			if (Entry == null)
 			{
 				Entry = new BisectEntry();
@@ -882,7 +882,7 @@ namespace UnrealGameSync
 			SyncOptions.NumSyncErrorRetries = ConfigFile.GetValue("Perforce.NumSyncErrorRetries", PerforceSyncOptions.DefaultNumSyncErrorRetries);
 		}
 
-		static Dictionary<Guid, bool> GetCategorySettings(ConfigSection Section, string IncludedKey, string ExcludedKey)
+		static Dictionary<Guid, bool> GetCategorySettings(ConfigSection? Section, string IncludedKey, string ExcludedKey)
 		{
 			Dictionary<Guid, bool> Result = new Dictionary<Guid, bool>();
 			if (Section != null)
@@ -967,7 +967,7 @@ namespace UnrealGameSync
 		protected override void ImportWorkspaceState(DirectoryReference RootDir, string ClientName, string BranchPath, UserWorkspaceState CurrentWorkspace)
 		{
 			// Read the workspace settings
-			ConfigSection WorkspaceSection = ConfigFile.FindSection(ClientName + BranchPath);
+			ConfigSection? WorkspaceSection = ConfigFile.FindSection(ClientName + BranchPath);
 			if(WorkspaceSection == null)
 			{
 				string LegacyBranchAndClientKey = ClientName + BranchPath;
@@ -1066,7 +1066,7 @@ namespace UnrealGameSync
 
 		protected override void ImportWorkspaceSettings(DirectoryReference RootDir, string ClientName, string BranchPath, UserWorkspaceSettings CurrentWorkspace)
 		{
-			ConfigSection WorkspaceSection = ConfigFile.FindSection(ClientName + BranchPath);
+			ConfigSection? WorkspaceSection = ConfigFile.FindSection(ClientName + BranchPath);
 			if (WorkspaceSection != null)
 			{
 				CurrentWorkspace.Filter.View = WorkspaceSection.GetValues("SyncFilter", new string[0]).ToList();
diff --git a/Engine/Source/Programs/UnrealGameSync/UnrealGameSyncShared/Utility.cs b/Engine/Source/Programs/UnrealGameSync/UnrealGameSyncShared/Utility.cs
--- a/Engine/Source/Programs/UnrealGameSync/UnrealGameSyncShared/Utility.cs
+++ b/Engine/Source/Programs/UnrealGameSync/UnrealGameSyncShared/Utility.cs
@@ -121,12 +121,12 @@ namespace UnrealGameSync
 		public static T LoadJson<T>(FileReference File)
 		{
 			byte[] Data = FileReference.ReadAllBytes(File);
-			return JsonSerializer.Deserialize<T>(Data, DefaultJsonSerializerOptions);
+			return JsonSerializer.Deserialize<T>(Data, DefaultJsonSerializerOptions)!;
 		}
 
 		public static void SaveJson<T>(FileReference File, T Object)
 		{
-			JsonSerializerOptions Options = new JsonSerializerOptions { IgnoreNullValues = true, WriteIndented = true };
+			JsonSerializerOptions Options = new JsonSerializerOptions { DefaultIgnoreCondition = JsonIgnoreCondition.WhenWritingNull, WriteIndented = true };
 			Options.Converters.Add(new JsonStringEnumConverter());
 
 			using (Stream Stream = FileReference.Open(File, FileMode.Create, FileAccess.Write, FileShare.Read))
@@ -140,7 +140,7 @@ namespace UnrealGameSync
 
 		public static string GetPathWithCorrectCase(FileInfo Info)
 		{
-			DirectoryInfo ParentInfo = Info.Directory;
+			DirectoryInfo ParentInfo = Info.Directory!;
 			if(Info.Exists)
 			{
 				return Path.Combine(GetPathWithCorrectCase(ParentInfo), ParentInfo.GetFiles(Info.Name)[0].Name); 
@@ -153,7 +153,7 @@ namespace UnrealGameSync
 
 		public static string GetPathWithCorrectCase(DirectoryInfo Info)
 		{
-			DirectoryInfo ParentInfo = Info.Parent;
+			DirectoryInfo ParentInfo = Info.Parent!;
 			if(ParentInfo == null)
 			{
 				return Info.FullName.ToUpperInvariant();
@@ -316,7 +316,7 @@ namespace UnrealGameSync
 				Options.PropertyNameCaseInsensitive = true;
 				Options.Converters.Add(new JsonStringEnumConverter());
 
-				ProjectJson Project = JsonSerializer.Deserialize<ProjectJson>(Text, Options);
+				ProjectJson Project = JsonSerializer.Deserialize<ProjectJson>(Text, Options)!;
 
 				return Project.Enterprise;
 			}
@@ -382,7 +382,7 @@ namespace UnrealGameSync
 
 			if (ProjectFile.Name.EndsWith(".uproject", StringComparison.OrdinalIgnoreCase))
 			{
-				AddLocalConfigPaths_WithExtensionDirs(ProjectFile.Directory, "Build", "UnrealGameSync.ini", SearchPaths);
+				AddLocalConfigPaths_WithExtensionDirs(ProjectFile.Directory!, "Build", "UnrealGameSync.ini", SearchPaths);
 			}
 			else
 			{
@@ -606,7 +606,7 @@ namespace UnrealGameSync
 			ProcessStartInfo StartInfo = new ProcessStartInfo();
 			StartInfo.FileName = Url;
 			StartInfo.UseShellExecute = true;
-			using Process _ = Process.Start(StartInfo);
+			using Process? _ = Process.Start(StartInfo);
 		}
 	}
 }
diff --git a/Engine/Source/Programs/UnrealGameSync/UnrealGameSyncShared/WorkspaceUpdate.cs b/Engine/Source/Programs/UnrealGameSync/UnrealGameSyncShared/WorkspaceUpdate.cs
--- a/Engine/Source/Programs/UnrealGameSync/UnrealGameSyncShared/WorkspaceUpdate.cs
+++ b/Engine/Source/Programs/UnrealGameSync/UnrealGameSyncShared/WorkspaceUpdate.cs
@@ -612,7 +612,7 @@ namespace UnrealGameSync
 
 					// Check if the new sync filter matches the previous one. If not, we'll enumerate all files in the workspace and make sure there's nothing extra there.
 					string? NextSyncFilterHash = null;
-					using (SHA1Managed SHA = new SHA1Managed())
+					using (SHA1 SHA = SHA1.Create())
 					{
 						StringBuilder CombinedFilter = new StringBuilder();
 						foreach (string RelativeSyncPath in RelativeSyncPaths)
@@ -1774,7 +1774,7 @@ namespace UnrealGameSync
 				return null;
 			}
 
-			return Response.Data.Contents.Select(x => x.Trim()).Where(x => x.Length > 0).ToList().AsReadOnly();
+			return Response.Data.Contents!.Select(x => x.Trim()).Where(x => x.Length > 0).ToList().AsReadOnly();
 		}
 
 		static string FormatTime(long Seconds)
@@ -1895,7 +1895,7 @@ namespace UnrealGameSync
 
 		static string UpdateBuildVersion(string Text, int Changelist, int CodeChangelist, string BranchOrStreamName, bool bIsLicenseeVersion)
 		{
-			Dictionary<string, object> Object = JsonSerializer.Deserialize<Dictionary<string, object>>(Text, Utility.DefaultJsonSerializerOptions);
+			Dictionary<string, object> Object = JsonSerializer.Deserialize<Dictionary<string, object>>(Text, Utility.DefaultJsonSerializerOptions)!;
 
 			int PrevCompatibleChangelist = 0;
 			if (Object.TryGetValue("CompatibleChangelist", out object? PrevCompatibleChangelistObj))
@@ -1937,7 +1937,7 @@ namespace UnrealGameSync
 				}
 				else
 				{
-					Directory.CreateDirectory(Path.GetDirectoryName(LocalPath));
+					Directory.CreateDirectory(Path.GetDirectoryName(LocalPath)!);
 					Utility.ForceDeleteFile(LocalPath);
 					if (DepotPath != null)
 					{
