UGS: Send unobserved task exceptions to Sentry.

diff --git a/Engine/Source/Programs/UnrealGameSync/UnrealGameSync/Program.cs b/Engine/Source/Programs/UnrealGameSync/UnrealGameSync/Program.cs
--- a/Engine/Source/Programs/UnrealGameSync/UnrealGameSync/Program.cs
+++ b/Engine/Source/Programs/UnrealGameSync/UnrealGameSync/Program.cs
@@ -56,6 +56,7 @@ namespace UnrealGameSync
 
 				Application.ThreadException += Application_ThreadException_Sentry;
 				AppDomain.CurrentDomain.UnhandledException += CurrentDomain_UnhandledException_Sentry;
+				TaskScheduler.UnobservedTaskException += Application_UnobservedException_Sentry;
 			}
 
 			bool bFirstInstance;
@@ -237,6 +238,15 @@ namespace UnrealGameSync
 			dialog.ShowDialog();
 		}
 
+		private static void Application_UnobservedException_Sentry(object? sender, UnobservedTaskExceptionEventArgs args)
+		{
+			Exception? innerException = args.Exception?.InnerException;
+			if (innerException != null)
+			{
+				SentrySdk.CaptureException(innerException, s => s.SetTag("Unobserved", "1"));
+			}
+		}
+
 		static void MergeUpdateSettings(FileReference UpdateConfigFile, ref string? UpdatePath, ref string? UpdateSpawn)
 		{
 			try
