UGS: Serialize config file to memory before writing to disk. Suspect that there's some exception thrown during serialization that causes file to be left empty.

diff --git a/Engine/Source/Programs/UnrealGameSync/UnrealGameSyncShared/Utility.cs b/Engine/Source/Programs/UnrealGameSync/UnrealGameSyncShared/Utility.cs
--- a/Engine/Source/Programs/UnrealGameSync/UnrealGameSyncShared/Utility.cs
+++ b/Engine/Source/Programs/UnrealGameSync/UnrealGameSyncShared/Utility.cs
@@ -129,13 +129,17 @@ namespace UnrealGameSync
 			JsonSerializerOptions Options = new JsonSerializerOptions { IgnoreNullValues = true, WriteIndented = true };
 			Options.Converters.Add(new JsonStringEnumConverter());
 
-			using (Stream Stream = FileReference.Open(File, FileMode.Create, FileAccess.Write, FileShare.Read))
+			byte[] Buffer;
+			using (MemoryStream Stream = new MemoryStream())
 			{
 				using (Utf8JsonWriter Writer = new Utf8JsonWriter(Stream, new JsonWriterOptions { Indented = true }))
 				{
 					JsonSerializer.Serialize(Writer, Object, Options);
 				}
+				Buffer = Stream.ToArray();
 			}
+
+			FileReference.WriteAllBytes(File, Buffer);
 		}
 
 		public static string GetPathWithCorrectCase(FileInfo Info)
