UGS: Fixed ultra-rare crash with empty log windows.

diff --git a/Engine/Source/Programs/UnrealGameSync/UnrealGameSync/Controls/LogControl.cs b/Engine/Source/Programs/UnrealGameSync/UnrealGameSync/Controls/LogControl.cs
--- a/Engine/Source/Programs/UnrealGameSync/UnrealGameSync/Controls/LogControl.cs
+++ b/Engine/Source/Programs/UnrealGameSync/UnrealGameSync/Controls/LogControl.cs
@@ -820,7 +820,11 @@ namespace UnrealGameSync
 					if(Selection.End.LineIdx < VerticalScroll.nPos + VerticalScroll.nPage)
 					{
 						int LineIdx = Math.Min(VerticalScroll.nPos + VerticalScroll.nPage, Lines.Count - 1);
-						Selection.End = new TextLocation(LineIdx, Lines[LineIdx].Length);
+
+						if (LineIdx >= 0)
+						{
+							Selection.End = new TextLocation(LineIdx, Lines[LineIdx].Length);
+						}
 					}
 				}
 
