UGS: Disable innocuous warning about syncing a file more than once. This can happen legitimately when files are checked out.

diff --git a/Engine/Source/Programs/UnrealGameSync/UnrealGameSyncShared/WorkspaceUpdate.cs b/Engine/Source/Programs/UnrealGameSync/UnrealGameSyncShared/WorkspaceUpdate.cs
--- a/Engine/Source/Programs/UnrealGameSync/UnrealGameSyncShared/WorkspaceUpdate.cs
+++ b/Engine/Source/Programs/UnrealGameSync/UnrealGameSyncShared/WorkspaceUpdate.cs
@@ -457,12 +457,7 @@ namespace UnrealGameSync
 				int Idx = Path.IndexOf('/');
 				if (Idx == -1)
 				{
-					long OldSize;
-					if (IncludedFiles.TryGetValue(Path, out OldSize))
-					{
-						Logger.LogWarning("File {FullPath} added twice to sync tree (old size {OldSize}, new size {Size})", FullPath, OldSize, Size);
-					}
-					else
+					if (!IncludedFiles.ContainsKey(Path))
 					{
 						IncludedFiles.Add(Path, Size);
 					}
