UGS: Fix uncaught exceptions thrown from thread that polls for new versions of UGS.

diff --git a/Engine/Source/Programs/UnrealGameSync/UnrealGameSync/Program.cs b/Engine/Source/Programs/UnrealGameSync/UnrealGameSync/Program.cs
--- a/Engine/Source/Programs/UnrealGameSync/UnrealGameSync/Program.cs
+++ b/Engine/Source/Programs/UnrealGameSync/UnrealGameSync/Program.cs
@@ -36,6 +36,14 @@ namespace UnrealGameSync
 
 		public static string? SyncVersion = null;
 
+		public static void CaptureException(Exception exception)
+		{
+			if (DeploymentSettings.SentryDsn != null)
+			{
+				SentrySdk.CaptureException(exception);
+			}
+		}
+
 		[STAThread]
 		static void Main(string[] Args)
 		{
diff --git a/Engine/Source/Programs/UnrealGameSync/UnrealGameSync/UpdateMonitor.cs b/Engine/Source/Programs/UnrealGameSync/UnrealGameSync/UpdateMonitor.cs
--- a/Engine/Source/Programs/UnrealGameSync/UnrealGameSync/UpdateMonitor.cs
+++ b/Engine/Source/Programs/UnrealGameSync/UnrealGameSync/UpdateMonitor.cs
@@ -71,14 +71,30 @@ namespace UnrealGameSync
 			{
 				await Task.Delay(TimeSpan.FromMinutes(5.0), CancellationToken);
 
-				using (IPerforceConnection Perforce = await PerforceConnection.CreateAsync(PerforceSettings, Logger))
+				IPerforceConnection? Perforce = null;
+				try
 				{
+					Perforce = await PerforceConnection.CreateAsync(PerforceSettings, Logger);
+
 					PerforceResponseList<ChangesRecord> Changes = await Perforce.TryGetChangesAsync(ChangesOptions.None, -1, ChangeStatus.Submitted, WatchPath, CancellationToken);
 					if (Changes.Succeeded && Changes.Data.Count > 0)
 					{
 						TriggerUpdate(UpdateType.Background, null);
 					}
 				}
+				catch (PerforceException ex)
+				{
+					Logger.LogInformation(ex, "Perforce exception while attempting to poll for updates.");
+				}
+				catch (Exception ex)
+				{
+					Logger.LogWarning(ex, "Exception while attempting to poll for updates.");
+					Program.CaptureException(ex);
+				}
+				finally
+				{
+					Perforce?.Dispose();
+				}
 			}
 		}
 
