UGS: Send exceptions on Perforce update thread to Sentry.

diff --git a/Engine/Source/Programs/UnrealGameSync/UnrealGameSync/PerforceMonitor.cs b/Engine/Source/Programs/UnrealGameSync/UnrealGameSync/PerforceMonitor.cs
--- a/Engine/Source/Programs/UnrealGameSync/UnrealGameSync/PerforceMonitor.cs
+++ b/Engine/Source/Programs/UnrealGameSync/UnrealGameSync/PerforceMonitor.cs
@@ -241,7 +241,7 @@ namespace UnrealGameSync
 				{
 					if (bCoolDown)
 					{
-						await Task.Delay(TimeSpan.FromSeconds(5.0), CancellationToken);
+						await Task.Delay(TimeSpan.FromSeconds(20.0), CancellationToken);
 					}
 					else
 					{
@@ -255,6 +255,10 @@ namespace UnrealGameSync
 				catch (Exception Ex)
 				{
 					Logger.LogError(Ex, "Unhandled exception in PollForUpdatesInner()");
+					if (!(Ex is PerforceException))
+					{
+						Program.CaptureException(Ex);
+					}
 					bCoolDown = true;
 				}
 			}
