UGS: Fix exception when the UpdateSpawn setting in AutoUpdate.ini is invalid.

diff --git a/Engine/Source/Programs/UnrealGameSync/UnrealGameSync/Program.cs b/Engine/Source/Programs/UnrealGameSync/UnrealGameSync/Program.cs
--- a/Engine/Source/Programs/UnrealGameSync/UnrealGameSync/Program.cs
+++ b/Engine/Source/Programs/UnrealGameSync/UnrealGameSync/Program.cs
@@ -129,7 +129,10 @@ namespace UnrealGameSync
 			// Set the current working directory to the update directory to prevent child-process file handles from disrupting auto-updates
 			if (UpdateSpawn != null)
 			{
-				Directory.SetCurrentDirectory(Path.GetDirectoryName(UpdateSpawn));
+				if (File.Exists(UpdateSpawn))
+				{
+					Directory.SetCurrentDirectory(Path.GetDirectoryName(UpdateSpawn));
+				}
 			}
 
 			string SyncVersionFile = Path.Combine(Path.GetDirectoryName(Assembly.GetExecutingAssembly().Location!)!, "SyncVersion.txt");
