UGS: Send exceptions to Sentry when opening new projects at startup.

diff --git a/Engine/Source/Programs/UnrealGameSync/UnrealGameSync/ProgramApplicationContext.cs b/Engine/Source/Programs/UnrealGameSync/UnrealGameSync/ProgramApplicationContext.cs
--- a/Engine/Source/Programs/UnrealGameSync/UnrealGameSync/ProgramApplicationContext.cs
+++ b/Engine/Source/Programs/UnrealGameSync/UnrealGameSync/ProgramApplicationContext.cs
@@ -191,7 +191,7 @@ namespace UnrealGameSync
 				Task<OpenProjectInfo> StartupTask = Task.Run(() => OpenProjectInfo.CreateAsync(DefaultPerforceSettings, ProjectSettings, Settings, Logger, StartupCancellationSource.Token), StartupCancellationSource.Token);
 				StartupTasks.Add((ProjectSettings, new ModalTask<OpenProjectInfo>(StartupTask)));
 			}
-			StartupTask = Task.WhenAll(StartupTasks.Select(x => x.Item2.Task));
+			StartupTask = Task.Run(() => WaitForStartupTasks(StartupTasks));
 
 			StartupWindow = new ModalTaskWindow("Opening Projects", "Opening projects, please wait...", FormStartPosition.CenterScreen, StartupTask, StartupCancellationSource);
 			Components.Add(StartupWindow);
@@ -211,6 +211,24 @@ namespace UnrealGameSync
 			StartupWindow.FormClosed += (s, e) => OnStartupComplete(StartupTasks);
 		}
 
+		static async Task WaitForStartupTasks(List<(UserSelectedProjectSettings, ModalTask<OpenProjectInfo>)> StartupTasks)
+		{
+			foreach ((_, ModalTask<OpenProjectInfo> ModalTask) in StartupTasks)
+			{
+				try
+				{
+					await ModalTask.Task;
+				}
+				catch (PerforceException)
+				{
+				}
+				catch (Exception ex)
+				{
+					Program.CaptureException(ex);
+				}
+			}
+		}
+
 		private UserSelectedProjectSettings UpgradeSelectedProjectSettings(UserSelectedProjectSettings Project)
 		{
 			if (Project.ServerAndPort == null || String.Compare(Project.ServerAndPort, DefaultPerforceSettings.ServerAndPort, StringComparison.OrdinalIgnoreCase) == 0)
