UGS: Fix null reference when no screen is attached. (?)

diff --git a/Engine/Source/Programs/UnrealGameSync/UnrealGameSync/Forms/MainWindow.cs b/Engine/Source/Programs/UnrealGameSync/UnrealGameSync/Forms/MainWindow.cs
--- a/Engine/Source/Programs/UnrealGameSync/UnrealGameSync/Forms/MainWindow.cs
+++ b/Engine/Source/Programs/UnrealGameSync/UnrealGameSync/Forms/MainWindow.cs
@@ -1650,13 +1650,16 @@ namespace UnrealGameSync
 
 		private void SetAlertWindowPositions()
 		{
-			Rectangle NewPrimaryWorkArea = Screen.PrimaryScreen.WorkingArea;
-			if(NewPrimaryWorkArea != PrimaryWorkArea)
+			if (Screen.PrimaryScreen != null)
 			{
-				PrimaryWorkArea = NewPrimaryWorkArea;
-				for(int Idx = 0; Idx < AlertWindows.Count; Idx++)
+				Rectangle NewPrimaryWorkArea = Screen.PrimaryScreen.WorkingArea;
+				if (NewPrimaryWorkArea != PrimaryWorkArea)
 				{
-					SetAlertWindowPosition(Idx);
+					PrimaryWorkArea = NewPrimaryWorkArea;
+					for (int Idx = 0; Idx < AlertWindows.Count; Idx++)
+					{
+						SetAlertWindowPosition(Idx);
+					}
 				}
 			}
 		}
