UGS: Fix application context not being disposed when exiting the application.

diff --git a/Engine/Source/Programs/UnrealGameSync/UnrealGameSync/Program.cs b/Engine/Source/Programs/UnrealGameSync/UnrealGameSync/Program.cs
--- a/Engine/Source/Programs/UnrealGameSync/UnrealGameSync/Program.cs
+++ b/Engine/Source/Programs/UnrealGameSync/UnrealGameSync/Program.cs
@@ -198,7 +198,7 @@ namespace UnrealGameSync
 
 							using (UpdateMonitor UpdateMonitor = new UpdateMonitor(DefaultSettings, UpdatePath, ServiceProvider))
 							{
-								ProgramApplicationContext Context = new ProgramApplicationContext(DefaultSettings, UpdateMonitor, DeploymentSettings.ApiUrl, DataFolder, ActivateEvent, bRestoreState, UpdateSpawn, ProjectFileName, bUnstable, ServiceProvider, Uri);
+								using ProgramApplicationContext Context = new ProgramApplicationContext(DefaultSettings, UpdateMonitor, DeploymentSettings.ApiUrl, DataFolder, ActivateEvent, bRestoreState, UpdateSpawn, ProjectFileName, bUnstable, ServiceProvider, Uri);
 								Application.Run(Context);
 
 								if (UpdateMonitor.IsUpdateAvailable && UpdateSpawn != null)
