UGS: Output errors when unable to write config files (eg. because disk is full) rather than crashing.

diff --git a/Engine/Source/Programs/UnrealGameSync/UnrealGameSync/Controls/WorkspaceControl.cs b/Engine/Source/Programs/UnrealGameSync/UnrealGameSync/Controls/WorkspaceControl.cs
--- a/Engine/Source/Programs/UnrealGameSync/UnrealGameSync/Controls/WorkspaceControl.cs
+++ b/Engine/Source/Programs/UnrealGameSync/UnrealGameSync/Controls/WorkspaceControl.cs
@@ -308,7 +308,7 @@ namespace UnrealGameSync
 			Settings = InSettings;
 			this.WorkspaceSettings = OpenProjectInfo.WorkspaceSettings;
 			this.WorkspaceState = OpenProjectInfo.WorkspaceState;
-			ProjectSettings = InSettings.FindOrAddProjectSettings(OpenProjectInfo.ProjectInfo, OpenProjectInfo.WorkspaceSettings);
+			ProjectSettings = InSettings.FindOrAddProjectSettings(OpenProjectInfo.ProjectInfo, OpenProjectInfo.WorkspaceSettings, Logger);
 
 			DesiredTaskbarState = Tuple.Create(TaskbarState.NoProgress, 0.0f);
 
@@ -1020,7 +1020,7 @@ namespace UnrealGameSync
 				WorkspaceState.LastSyncResultMessage = null;
 				WorkspaceState.LastSyncTime = null;
 				WorkspaceState.LastSyncDurationSeconds = 0;
-				WorkspaceState.Save();
+				WorkspaceState.Save(Logger);
 
 				Workspace.CancelUpdate();
 
@@ -3658,7 +3658,7 @@ namespace UnrealGameSync
 			string? BuildHealthProject;
 			TryGetProjectSetting(PerforceMonitor.LatestProjectConfigFile, "BuildHealthProject", out BuildHealthProject);
 
-			IssueSettingsWindow IssueSettings = new IssueSettingsWindow(Settings, BuildHealthProject ?? "");
+			IssueSettingsWindow IssueSettings = new IssueSettingsWindow(Settings, BuildHealthProject ?? "", Logger);
 			if (IssueSettings.ShowDialog(this) == DialogResult.OK)
 			{
 				Owner.UpdateAlertWindows();
@@ -4296,7 +4296,7 @@ namespace UnrealGameSync
 			{
 				ArchiveSettings.bEnabled = false;
 			}
-			Settings.Save();
+			Settings.Save(Logger);
 
 			UpdateSelectedArchives();
 		}
@@ -4319,7 +4319,7 @@ namespace UnrealGameSync
 			{
 				ArchiveSettings.bEnabled = false;
 			}
-			Settings.Save();
+			Settings.Save(Logger);
 
 			UpdateSelectedArchives();
 		}
@@ -4412,7 +4412,7 @@ namespace UnrealGameSync
 		private void BuildAfterSyncCheckBox_CheckedChanged(object sender, EventArgs e)
 		{
 			Settings.bBuildAfterSync = BuildAfterSyncCheckBox.Checked;
-			Settings.Save();
+			Settings.Save(Logger);
 
 			UpdateSyncActionCheckboxes();
 		}
@@ -4420,7 +4420,7 @@ namespace UnrealGameSync
 		private void RunAfterSyncCheckBox_CheckedChanged(object sender, EventArgs e)
 		{
 			Settings.bRunAfterSync = RunAfterSyncCheckBox.Checked;
-			Settings.Save();
+			Settings.Save(Logger);
 
 			UpdateSyncActionCheckboxes();
 		}
@@ -4428,7 +4428,7 @@ namespace UnrealGameSync
 		private void OpenSolutionAfterSyncCheckBox_CheckedChanged(object sender, EventArgs e)
 		{
 			Settings.bOpenSolutionAfterSync = OpenSolutionAfterSyncCheckBox.Checked;
-			Settings.Save();
+			Settings.Save(Logger);
 
 			UpdateSyncActionCheckboxes();
 		}
@@ -4833,7 +4833,7 @@ namespace UnrealGameSync
 		private void Splitter_OnVisibilityChanged(bool bVisible)
 		{
 			Settings.bShowLogWindow = bVisible;
-			Settings.Save();
+			Settings.Save(Logger);
 
 			UpdateStatusPanel();
 		}
@@ -4842,7 +4842,7 @@ namespace UnrealGameSync
 		{
 			OptionsContextMenu_AutoResolveConflicts.Checked ^= true;
 			Settings.bAutoResolveConflicts = OptionsContextMenu_AutoResolveConflicts.Checked;
-			Settings.Save();
+			Settings.Save(Logger);
 		}
 
 		private void OptionsContextMenu_EditorArguments_Click(object sender, EventArgs e)
@@ -4857,7 +4857,7 @@ namespace UnrealGameSync
 			{
 				Settings.EditorArguments = Arguments.GetItems();
 				Settings.bEditorArgumentsPrompt = Arguments.PromptBeforeLaunch;
-				Settings.Save();
+				Settings.Save(Logger);
 				return true;
 			}
 			return false;
@@ -4921,7 +4921,7 @@ namespace UnrealGameSync
 		void UpdateBuildConfig(BuildConfig NewConfig)
 		{
 			Settings.CompiledEditorBuildConfig = NewConfig;
-			Settings.Save();
+			Settings.Save(Logger);
 			UpdateCheckedBuildConfig();
 		}
 
@@ -5087,7 +5087,7 @@ namespace UnrealGameSync
 		private void BuildListContextMenu_ShowServerTimes_Click(object sender, EventArgs e)
 		{
 			Settings.bShowLocalTimes = false;
-			Settings.Save();
+			Settings.Save(Logger);
 			BuildList.Items.Clear(); // Need to clear list to rebuild groups, populate time column again
 			UpdateBuildList();
 		}
@@ -5095,7 +5095,7 @@ namespace UnrealGameSync
 		private void BuildListContextMenu_ShowLocalTimes_Click(object sender, EventArgs e)
 		{
 			Settings.bShowLocalTimes = true;
-			Settings.Save();
+			Settings.Save(Logger);
 			BuildList.Items.Clear(); // Need to clear list to rebuild groups, populate time column again
 			UpdateBuildList();
 		}
@@ -5315,7 +5315,7 @@ namespace UnrealGameSync
 
 				// Save the settings
 				ProjectSettings.BuildSteps = ModifiedBuildSteps;
-				ProjectSettings.Save();
+				ProjectSettings.Save(Logger);
 
 				// Update the custom tools menu, because we might have changed it
 				UpdateBuildSteps();
@@ -5421,10 +5421,10 @@ namespace UnrealGameSync
 			if (Filter.ShowDialog() == DialogResult.OK)
 			{
 				Settings.Global.Filter = Filter.GlobalFilter;
-				Settings.Save();
+				Settings.Save(Logger);
 
 				WorkspaceSettings.Filter = Filter.WorkspaceFilter;
-				WorkspaceSettings.Save();
+				WorkspaceSettings.Save(Logger);
 			}
 		}
 
@@ -5489,7 +5489,7 @@ namespace UnrealGameSync
 		private void SyncContextMenu_LatestChangeType_Click(object? sender, EventArgs e, string SyncTypeID)
 		{
 			Settings.SyncTypeID = SyncTypeID;
-			Settings.Save();
+			Settings.Save(Logger);
 		}
 
 		private void SyncContextMenu_EnterChangelist_Click(object sender, EventArgs e)
@@ -5565,13 +5565,13 @@ namespace UnrealGameSync
 		private void RecentMenu_ClearList_Click(object sender, EventArgs e)
 		{
 			Settings.RecentProjects.Clear();
-			Settings.Save();
+			Settings.Save(Logger);
 		}
 
 		private void OptionsContextMenu_ShowChanges_ShowUnreviewed_Click(object sender, EventArgs e)
 		{
 			Settings.bShowUnreviewedChanges ^= true;
-			Settings.Save();
+			Settings.Save(Logger);
 
 			UpdateBuildList();
 			ShrinkNumRequestedBuilds();
@@ -5580,7 +5580,7 @@ namespace UnrealGameSync
 		private void OptionsContextMenu_ShowChanges_ShowAutomated_Click(object sender, EventArgs e)
 		{
 			Settings.bShowAutomatedChanges ^= true;
-			Settings.Save();
+			Settings.Save(Logger);
 
 			UpdateBuildListFilter();
 		}
@@ -5788,7 +5788,7 @@ namespace UnrealGameSync
 				ChangeNumberToBisectState[ChangeNumberToBisectState.Keys.Max()] = BisectState.Fail;
 
 				WorkspaceState.BisectChanges = ChangeNumberToBisectState.Select(x => new BisectEntry {  Change = x.Key, State = x.Value }).ToList();
-				WorkspaceState.Save();
+				WorkspaceState.Save(Logger);
 
 				UpdateBuildList();
 				UpdateStatusPanel();
@@ -5798,7 +5798,7 @@ namespace UnrealGameSync
 		private void CancelBisectMode()
 		{
 			WorkspaceState.BisectChanges.Clear();
-			WorkspaceState.Save();
+			WorkspaceState.Save(Logger);
 
 			UpdateBuildList();
 			UpdateStatusPanel();
@@ -5815,7 +5815,7 @@ namespace UnrealGameSync
 				}
 			}
 
-			WorkspaceState.Save();
+			WorkspaceState.Save(Logger);
 
 			ChangeNumberToLayoutInfo.Clear();
 			BuildList.Invalidate();
@@ -5944,7 +5944,7 @@ namespace UnrealGameSync
 			{
 				ProjectSettings.FilterBadges.Add(BadgeName);
 			}
-			Settings.Save();
+			Settings.Save(Logger);
 
 			UpdateBuildListFilter();
 		}
@@ -5953,10 +5953,10 @@ namespace UnrealGameSync
 		{
 			ProjectSettings.FilterBadges.Clear();
 			ProjectSettings.FilterType = FilterType.None;
-			ProjectSettings.Save();
+			ProjectSettings.Save(Logger);
 
 			Settings.bShowAutomatedChanges = false;
-			Settings.Save();
+			Settings.Save(Logger);
 
 			UpdateBuildListFilter();
 		}
@@ -5964,7 +5964,7 @@ namespace UnrealGameSync
 		private void FilterContextMenu_ShowBuildMachineChanges_Click(object sender, EventArgs e)
 		{
 			Settings.bShowAutomatedChanges ^= true;
-			Settings.Save();
+			Settings.Save(Logger);
 
 			UpdateBuildListFilter();
 		}
@@ -5972,28 +5972,28 @@ namespace UnrealGameSync
 		private void FilterContextMenu_Robomerge_ShowAll_Click(object sender, EventArgs e)
 		{
 			Settings.ShowRobomerge = UserSettings.RobomergeShowChangesOption.All;	
-			Settings.Save();
+			Settings.Save(Logger);
 			UpdateBuildListFilter();
 		}
 
 		private void FilterContextMenu_Robomerge_ShowBadged_Click(object sender, EventArgs e)
 		{
 			Settings.ShowRobomerge = UserSettings.RobomergeShowChangesOption.Badged;
-			Settings.Save();
+			Settings.Save(Logger);
 			UpdateBuildListFilter();
 		}
 
 		private void FilterContextMenu_Robomerge_ShowNone_Click(object sender, EventArgs e)
 		{
 			Settings.ShowRobomerge = UserSettings.RobomergeShowChangesOption.None;
-			Settings.Save();
+			Settings.Save(Logger);
 			UpdateBuildListFilter();
 		}
 
 		private void FilterContextMenu_Robomerge_Annotate_Click(object sender, EventArgs e)
 		{
 			Settings.bAnnotateRobmergeChanges ^= true;
-			Settings.Save();
+			Settings.Save(Logger);
 			BuildList.Items.Clear();	// need to reload to update user names
 
 			UpdateBuildListFilter();
@@ -6045,7 +6045,7 @@ namespace UnrealGameSync
 		private void FilterContextMenu_Type_ShowAll_Click(object sender, EventArgs e)
 		{
 			ProjectSettings.FilterType = FilterType.None;
-			ProjectSettings.Save();
+			ProjectSettings.Save(Logger);
 
 			UpdateBuildListFilter();
 		}
@@ -6053,7 +6053,7 @@ namespace UnrealGameSync
 		private void FilterContextMenu_Type_Code_Click(object sender, EventArgs e)
 		{
 			ProjectSettings.FilterType = FilterType.Code;
-			ProjectSettings.Save();
+			ProjectSettings.Save(Logger);
 
 			UpdateBuildListFilter();
 		}
@@ -6061,7 +6061,7 @@ namespace UnrealGameSync
 		private void FilterContextMenu_Type_Content_Click(object sender, EventArgs e)
 		{
 			ProjectSettings.FilterType = FilterType.Content;
-			ProjectSettings.Save();
+			ProjectSettings.Save(Logger);
 
 			UpdateBuildListFilter();
 		}
diff --git a/Engine/Source/Programs/UnrealGameSync/UnrealGameSync/Forms/ApplicationSettingsWindow.cs b/Engine/Source/Programs/UnrealGameSync/UnrealGameSync/Forms/ApplicationSettingsWindow.cs
--- a/Engine/Source/Programs/UnrealGameSync/UnrealGameSync/Forms/ApplicationSettingsWindow.cs
+++ b/Engine/Source/Programs/UnrealGameSync/UnrealGameSync/Forms/ApplicationSettingsWindow.cs
@@ -225,7 +225,7 @@ namespace UnrealGameSync
 			{
 				Settings.SyncOptions.NumThreads = (int)ParallelSyncThreadsSpinner.Value;
 				Settings.bKeepInTray = KeepInTrayCheckBox.Checked;
-				Settings.Save();
+				Settings.Save(Logger);
 			}
 
 			List<Guid> NewEnabledTools = new List<Guid>();
@@ -239,7 +239,7 @@ namespace UnrealGameSync
 			if (!NewEnabledTools.SequenceEqual(Settings.EnabledTools))
 			{
 				Settings.EnabledTools = NewEnabledTools.ToArray();
-				Settings.Save();
+				Settings.Save(Logger);
 				ToolUpdateMonitor.UpdateNow();
 			}
 
@@ -275,7 +275,7 @@ namespace UnrealGameSync
 
 		private void AdvancedBtn_Click(object sender, EventArgs e)
 		{
-			PerforceSyncSettingsWindow Window = new PerforceSyncSettingsWindow(Settings);
+			PerforceSyncSettingsWindow Window = new PerforceSyncSettingsWindow(Settings, Logger);
 			Window.ShowDialog();
 		}
 	}
diff --git a/Engine/Source/Programs/UnrealGameSync/UnrealGameSync/Forms/IssueSettingsWindow.cs b/Engine/Source/Programs/UnrealGameSync/UnrealGameSync/Forms/IssueSettingsWindow.cs
--- a/Engine/Source/Programs/UnrealGameSync/UnrealGameSync/Forms/IssueSettingsWindow.cs
+++ b/Engine/Source/Programs/UnrealGameSync/UnrealGameSync/Forms/IssueSettingsWindow.cs
@@ -1,5 +1,6 @@
 // Copyright Epic Games, Inc. All Rights Reserved.
 
+using Microsoft.Extensions.Logging;
 using System;
 using System.Collections.Generic;
 using System.ComponentModel;
@@ -15,10 +16,12 @@ namespace UnrealGameSync
 	partial class IssueSettingsWindow : Form
 	{
 		UserSettings Settings;
+		ILogger Logger;
 
-		public IssueSettingsWindow(UserSettings Settings, string CurrentProject)
+		public IssueSettingsWindow(UserSettings Settings, string CurrentProject, ILogger Logger)
 		{
 			this.Settings = Settings;
+			this.Logger = Logger;
 
 			InitializeComponent();
 
@@ -145,7 +148,7 @@ namespace UnrealGameSync
 			Settings.NotifyUnresolvedMinutes = NewNotifyUnresolvedMinutes;
 			Settings.NotifyUnacknowledgedMinutes = NewNotifyUnacknowledgedMinutes;
 			Settings.NotifyUnassignedMinutes = NewNotifyUnassignedMinutes;
-			Settings.Save();
+			Settings.Save(Logger);
 
 			DialogResult = DialogResult.OK;
 			Close();
diff --git a/Engine/Source/Programs/UnrealGameSync/UnrealGameSync/Forms/MainWindow.cs b/Engine/Source/Programs/UnrealGameSync/UnrealGameSync/Forms/MainWindow.cs
--- a/Engine/Source/Programs/UnrealGameSync/UnrealGameSync/Forms/MainWindow.cs
+++ b/Engine/Source/Programs/UnrealGameSync/UnrealGameSync/Forms/MainWindow.cs
@@ -670,7 +670,7 @@ namespace UnrealGameSync
 				Settings.WindowBounds = RestoreBounds;
 			}
 
-			Settings.Save();
+			Settings.Save(Logger);
 		}
 
 		private void SetupDefaultControl()
@@ -781,7 +781,7 @@ namespace UnrealGameSync
 				Activate();
 
 				Settings.bWindowVisible = Visible;
-				Settings.Save();
+				Settings.Save(Logger);
 			}
 		}
 
@@ -877,7 +877,7 @@ namespace UnrealGameSync
 			if(Schedule.ShowDialog() == System.Windows.Forms.DialogResult.OK)
 			{
 				Schedule.CopySettings(out Settings.bScheduleEnabled, out Settings.ScheduleTime, out Settings.ScheduleAnyOpenProject, out Settings.ScheduleProjects);
-				Settings.Save();
+				Settings.Save(Logger);
 			}
 
 			StartScheduleTimer();
@@ -1009,7 +1009,7 @@ namespace UnrealGameSync
 				DefaultControl.Hide();
 			}
 
-			Settings.Save();
+			Settings.Save(Logger);
 
 			if(CurrentTabPanel != null)
 			{
@@ -1045,7 +1045,7 @@ namespace UnrealGameSync
 
 		public void OpenNewProject()
 		{
-			OpenProjectInfo? OpenProjectInfo = OpenProjectWindow.ShowModal(this, null, Settings, DataFolder, CacheFolder, DefaultPerforceSettings, ServiceProvider);
+			OpenProjectInfo? OpenProjectInfo = OpenProjectWindow.ShowModal(this, null, Settings, DataFolder, CacheFolder, DefaultPerforceSettings, ServiceProvider, Logger);
 			if(OpenProjectInfo != null)
 			{
 				int NewTabIdx = TryOpenProject(OpenProjectInfo, -1, OpenProjectOptions.None);
@@ -1069,7 +1069,7 @@ namespace UnrealGameSync
 				Settings.RecentProjects.RemoveRange(MaxRecentProjects, Settings.RecentProjects.Count - MaxRecentProjects);
 			}
 
-			Settings.Save();
+			Settings.Save(Logger);
 		}
 
 		public void EditSelectedProject(int TabIdx)
@@ -1107,7 +1107,7 @@ namespace UnrealGameSync
 
 		public void EditSelectedProject(int TabIdx, UserSelectedProjectSettings SelectedProject)
 		{
-			OpenProjectInfo? OpenProjectInfo = OpenProjectWindow.ShowModal(this, SelectedProject, Settings, DataFolder, CacheFolder, DefaultPerforceSettings, ServiceProvider);
+			OpenProjectInfo? OpenProjectInfo = OpenProjectWindow.ShowModal(this, SelectedProject, Settings, DataFolder, CacheFolder, DefaultPerforceSettings, ServiceProvider, Logger);
 			if(OpenProjectInfo != null)
 			{
 				int NewTabIdx = TryOpenProject(OpenProjectInfo, TabIdx, OpenProjectOptions.None);
@@ -1255,7 +1255,7 @@ namespace UnrealGameSync
 				IMainWindowTabPanel TabPanel = (IMainWindowTabPanel)TabControl.GetTabData(TabIdx)!;
 				Settings.OpenProjects.Add(TabPanel.SelectedProject);
 			}
-			Settings.Save();
+			Settings.Save(Logger);
 		}
 
 		void TabControl_OnNewTabClick(Point Location, MouseButtons Buttons)
@@ -1302,7 +1302,7 @@ namespace UnrealGameSync
 			if(Settings.TabLabels != NewTabNames)
 			{
 				Settings.TabLabels = NewTabNames;
-				Settings.Save();
+				Settings.Save(Logger);
 
 				for(int Idx = 0; Idx < TabControl.GetTabCount(); Idx++)
 				{
@@ -1349,7 +1349,7 @@ namespace UnrealGameSync
 		private void TabMenu_RecentProjects_ClearList_Click(object sender, EventArgs e)
 		{
 			Settings.RecentProjects.Clear();
-			Settings.Save();
+			Settings.Save(Logger);
 		}
 
 		private void TabMenu_Closed(object sender, ToolStripDropDownClosedEventArgs e)
@@ -1360,7 +1360,7 @@ namespace UnrealGameSync
 		private void RecentMenu_ClearList_Click(object sender, EventArgs e)
 		{
 			Settings.RecentProjects.Clear();
-			Settings.Save();
+			Settings.Save(Logger);
 		}
 
 		public void UpdateProgress()
diff --git a/Engine/Source/Programs/UnrealGameSync/UnrealGameSync/Forms/OpenProjectWindow.cs b/Engine/Source/Programs/UnrealGameSync/UnrealGameSync/Forms/OpenProjectWindow.cs
--- a/Engine/Source/Programs/UnrealGameSync/UnrealGameSync/Forms/OpenProjectWindow.cs
+++ b/Engine/Source/Programs/UnrealGameSync/UnrealGameSync/Forms/OpenProjectWindow.cs
@@ -29,8 +29,9 @@ namespace UnrealGameSync
 		IPerforceSettings DefaultSettings;
 		IServiceProvider ServiceProvider;
 		UserSettings Settings;
+		ILogger Logger;
 
-		private OpenProjectWindow(UserSelectedProjectSettings? Project, UserSettings Settings, IPerforceSettings DefaultSettings, IServiceProvider ServiceProvider)
+		private OpenProjectWindow(UserSelectedProjectSettings? Project, UserSettings Settings, IPerforceSettings DefaultSettings, IServiceProvider ServiceProvider, ILogger Logger)
 		{
 			InitializeComponent();
 
@@ -38,6 +39,7 @@ namespace UnrealGameSync
 			this.OpenProjectInfo = null;
 			this.DefaultSettings = DefaultSettings;
 			this.ServiceProvider = ServiceProvider;
+			this.Logger = Logger;
 
 			if (Project == null)
 			{
@@ -90,9 +92,9 @@ namespace UnrealGameSync
 			get => Utility.OverridePerforceSettings(DefaultSettings, ServerAndPortOverride, UserNameOverride);
 		}
 
-		public static OpenProjectInfo? ShowModal(IWin32Window Owner, UserSelectedProjectSettings? Project, UserSettings Settings, DirectoryReference DataFolder, DirectoryReference CacheFolder, IPerforceSettings DefaultPerforceSettings, IServiceProvider ServiceProvider)
+		public static OpenProjectInfo? ShowModal(IWin32Window Owner, UserSelectedProjectSettings? Project, UserSettings Settings, DirectoryReference DataFolder, DirectoryReference CacheFolder, IPerforceSettings DefaultPerforceSettings, IServiceProvider ServiceProvider, ILogger Logger)
 		{
-			OpenProjectWindow Window = new OpenProjectWindow(Project, Settings, DefaultPerforceSettings, ServiceProvider);
+			OpenProjectWindow Window = new OpenProjectWindow(Project, Settings, DefaultPerforceSettings, ServiceProvider, Logger);
 			if(Window.ShowDialog(Owner) == DialogResult.OK)
 			{
 				return Window.OpenProjectInfo;
@@ -344,7 +346,7 @@ namespace UnrealGameSync
 				string FullName = Path.GetFullPath(Dialog.FileName);
 
 				Settings.FilterIndex = Dialog.FilterIndex;
-				Settings.Save();
+				Settings.Save(Logger);
 
 				LocalFileTextBox.Text = FullName;
 				UpdateOkButton();
diff --git a/Engine/Source/Programs/UnrealGameSync/UnrealGameSync/Forms/PerforceSyncSettingsWindow.cs b/Engine/Source/Programs/UnrealGameSync/UnrealGameSync/Forms/PerforceSyncSettingsWindow.cs
--- a/Engine/Source/Programs/UnrealGameSync/UnrealGameSync/Forms/PerforceSyncSettingsWindow.cs
+++ b/Engine/Source/Programs/UnrealGameSync/UnrealGameSync/Forms/PerforceSyncSettingsWindow.cs
@@ -1,5 +1,6 @@
 // Copyright Epic Games, Inc. All Rights Reserved.
 
+using Microsoft.Extensions.Logging;
 using System;
 using System.Windows.Forms;
 
@@ -8,10 +9,13 @@ namespace UnrealGameSync
 	partial class PerforceSyncSettingsWindow : Form
 	{
 		UserSettings Settings;
+		ILogger Logger;
 
-		public PerforceSyncSettingsWindow(UserSettings Settings)
+		public PerforceSyncSettingsWindow(UserSettings Settings, ILogger Logger)
 		{
 			this.Settings = Settings;
+			this.Logger = Logger;
+
 			InitializeComponent();
 		}
 
@@ -34,7 +38,7 @@ namespace UnrealGameSync
 			Settings.SyncOptions.MaxCommandsPerBatch = (int)numericUpDownMaxCommandsPerBatch.Value;
 			Settings.SyncOptions.MaxSizePerBatch = (int)numericUpDownMaxSizePerBatch.Value * 1024 * 1024;
 			Settings.SyncOptions.NumSyncErrorRetries = (int)numericUpDownRetriesOnSyncError.Value;
-			Settings.Save();
+			Settings.Save(Logger);
 
 			DialogResult = System.Windows.Forms.DialogResult.OK;
 			Close();
diff --git a/Engine/Source/Programs/UnrealGameSync/UnrealGameSync/ProgramApplicationContext.cs b/Engine/Source/Programs/UnrealGameSync/UnrealGameSync/ProgramApplicationContext.cs
--- a/Engine/Source/Programs/UnrealGameSync/UnrealGameSync/ProgramApplicationContext.cs
+++ b/Engine/Source/Programs/UnrealGameSync/UnrealGameSync/ProgramApplicationContext.cs
@@ -116,7 +116,7 @@ namespace UnrealGameSync
 
 				// Save the new settings
 				Settings.Version = UserSettingsVersion.Latest;
-				Settings.Save();
+				Settings.Save(Logger);
 			}
 
 			// Register the update listener
diff --git a/Engine/Source/Programs/UnrealGameSync/UnrealGameSync/Workspace.cs b/Engine/Source/Programs/UnrealGameSync/UnrealGameSync/Workspace.cs
--- a/Engine/Source/Programs/UnrealGameSync/UnrealGameSync/Workspace.cs
+++ b/Engine/Source/Programs/UnrealGameSync/UnrealGameSync/Workspace.cs
@@ -150,7 +150,7 @@ namespace UnrealGameSync
 				WorkspaceUpdateContext Context = Update.Context;
 
 				State.SetLastSyncState(Result, Context, StatusMessage);
-				State.Save();
+				State.Save(Logger);
 
 				OnUpdateComplete?.Invoke(Context, Result, StatusMessage);
 				CurrentUpdate = null;
diff --git a/Engine/Source/Programs/UnrealGameSync/UnrealGameSyncCmd/Program.cs b/Engine/Source/Programs/UnrealGameSync/UnrealGameSyncCmd/Program.cs
--- a/Engine/Source/Programs/UnrealGameSync/UnrealGameSyncCmd/Program.cs
+++ b/Engine/Source/Programs/UnrealGameSync/UnrealGameSyncCmd/Program.cs
@@ -295,13 +295,13 @@ namespace UnrealGameSyncCmd
 
 		public static async Task<UserWorkspaceState> ReadWorkspaceState(IPerforceConnection PerforceClient, UserWorkspaceSettings Settings, GlobalSettingsFile UserSettings, ILogger Logger)
 		{
-			UserWorkspaceState State = UserSettings.FindOrAddWorkspaceState(Settings);
+			UserWorkspaceState State = UserSettings.FindOrAddWorkspaceState(Settings, Logger);
 			if (State.SettingsTimeUtc != Settings.LastModifiedTimeUtc)
 			{
 				Logger.LogDebug("Updating state due to modified settings timestamp");
 				ProjectInfo Info = await ProjectInfo.CreateAsync(PerforceClient, Settings, CancellationToken.None);
 				State.UpdateCachedProjectInfo(Info, Settings.LastModifiedTimeUtc);
-				State.Save();
+				State.Save(Logger);
 			}
 			return State;
 		}
@@ -464,7 +464,7 @@ namespace UnrealGameSyncCmd
 				Settings.RootDir = ClientDir;
 				Settings.Init(Perforce.Settings.ServerAndPort, Perforce.Settings.UserName, ClientName, BranchPath, ProjectPath);
 				Options.ApplyTo(Settings);
-				Settings.Save();
+				Settings.Save(Logger);
 
 				Logger.LogInformation("Initialized {ClientName} with root at {RootDir}", ClientName, ClientDir);
 			}
@@ -523,7 +523,7 @@ namespace UnrealGameSyncCmd
 				Settings.RootDir = ClientDir;
 				Settings.Init(Perforce.Settings.ServerAndPort, Perforce.Settings.UserName, ClientName, BranchPath, ProjectPath);
 				Options.ApplyTo(Settings);
-				Settings.Save();
+				Settings.Save(Logger);
 
 				Logger.LogInformation("Initialized workspace at {RootDir} for {ClientProject}", ClientDir, Settings.ClientProjectPath);
 			}
@@ -629,7 +629,7 @@ namespace UnrealGameSyncCmd
 				Options |= WorkspaceUpdateOptions.RemoveFilteredFiles;
 
 				ProjectInfo ProjectInfo = State.CreateProjectInfo();
-				UserProjectSettings ProjectSettings = Context.UserSettings.FindOrAddProjectSettings(ProjectInfo, Settings);
+				UserProjectSettings ProjectSettings = Context.UserSettings.FindOrAddProjectSettings(ProjectInfo, Settings, Logger);
 
 				ConfigFile ProjectConfig = await ConfigUtils.ReadProjectConfigFileAsync(PerforceClient, ProjectInfo, Logger, CancellationToken.None);
 				string[] SyncFilter = ReadSyncFilter(Settings, Context.UserSettings, ProjectConfig);
@@ -653,7 +653,7 @@ namespace UnrealGameSyncCmd
 				}
 
 				State.SetLastSyncState(Result, UpdateContext, Message);
-				State.Save();
+				State.Save(Logger);
 			}
 		}
 
@@ -828,7 +828,7 @@ namespace UnrealGameSyncCmd
 					Context.Arguments.CheckAllArgumentsUsed(Context.Logger);
 
 					Options.ApplyTo(Settings);
-					Settings.Save();
+					Settings.Save(Logger);
 
 					Logger.LogInformation("Updated {ConfigFile}", Settings.ConfigFile);
 				}
@@ -892,12 +892,12 @@ namespace UnrealGameSyncCmd
 				if (Options.Global)
 				{
 					ApplyCommandOptions(Context.UserSettings.Global.Filter, Options, SyncCategories.Values, Logger);
-					Context.UserSettings.Save();
+					Context.UserSettings.Save(Logger);
 				}
 				else
 				{
 					ApplyCommandOptions(WorkspaceSettings.Filter, Options, SyncCategories.Values, Logger);
-					WorkspaceSettings.Save();
+					WorkspaceSettings.Save(Logger);
 				}
 
 				Dictionary<Guid, bool> GlobalCategories = GlobalFilter.GetCategories();
@@ -1150,7 +1150,7 @@ namespace UnrealGameSyncCmd
 			public async Task SwitchProjectAsync(IPerforceConnection PerforceClient, UserWorkspaceSettings Settings, string ProjectName, ILogger Logger)
 			{
 				Settings.ProjectPath = await FindProjectPathAsync(PerforceClient, Settings.ClientName, Settings.BranchPath, ProjectName);
-				Settings.Save();
+				Settings.Save(Logger);
 				Logger.LogInformation("Switched to project {ProjectPath}", Settings.ClientProjectPath);
 			}
 		}
diff --git a/Engine/Source/Programs/UnrealGameSync/UnrealGameSyncShared/GlobalSettings.cs b/Engine/Source/Programs/UnrealGameSync/UnrealGameSyncShared/GlobalSettings.cs
--- a/Engine/Source/Programs/UnrealGameSync/UnrealGameSyncShared/GlobalSettings.cs
+++ b/Engine/Source/Programs/UnrealGameSync/UnrealGameSyncShared/GlobalSettings.cs
@@ -1,6 +1,7 @@
 // Copyright Epic Games, Inc. All Rights Reserved.
 
 using EpicGames.Core;
+using Microsoft.Extensions.Logging;
 using System;
 using System.Collections.Generic;
 using System.Linq;
@@ -59,7 +60,7 @@ namespace UnrealGameSync
 		public FileReference File { get; }
 		public GlobalSettings Global { get; }
 
-		public UserProjectSettings FindOrAddProjectSettings(ProjectInfo ProjectInfo, UserWorkspaceSettings Settings)
+		public UserProjectSettings FindOrAddProjectSettings(ProjectInfo ProjectInfo, UserWorkspaceSettings Settings, ILogger Logger)
 		{
 			FileReference ConfigFile;
 			if (ProjectInfo.LocalFileName.HasExtension(".uprojectdirs"))
@@ -78,7 +79,7 @@ namespace UnrealGameSync
 			{
 				ProjectSettings = new UserProjectSettings(ConfigFile);
 				ImportProjectSettings(ProjectInfo, ProjectSettings);
-				ProjectSettings.Save();
+				ProjectSettings.Save(Logger);
 			}
 			return ProjectSettings;
 		}
@@ -111,17 +112,26 @@ namespace UnrealGameSync
 			return new GlobalSettingsFile(File, Data);
 		}
 
-		public virtual void Save()
+		public virtual bool Save(ILogger Logger)
 		{
-			Utility.SaveJson(File, Global);
+			try
+			{
+				Utility.SaveJson(File, Global);
+				return true;
+			}
+			catch (Exception Ex)
+			{
+				Logger.LogError(Ex, "Unable to save {File}: {Message}", File, Ex.Message);
+				return false;
+			}
 		}
 
-		public UserWorkspaceState FindOrAddWorkspaceState(UserWorkspaceSettings Settings)
+		public UserWorkspaceState FindOrAddWorkspaceState(UserWorkspaceSettings Settings, ILogger Logger)
 		{
-			return FindOrAddWorkspaceState(Settings.RootDir, Settings.ClientName, Settings.BranchPath);
+			return FindOrAddWorkspaceState(Settings.RootDir, Settings.ClientName, Settings.BranchPath, Logger);
 		}
 
-		public UserWorkspaceState FindOrAddWorkspaceState(DirectoryReference RootDir, string ClientName, string BranchPath)
+		public UserWorkspaceState FindOrAddWorkspaceState(DirectoryReference RootDir, string ClientName, string BranchPath, ILogger Logger)
 		{
 			UserWorkspaceState? State;
 			if (!UserWorkspaceState.TryLoad(RootDir, out State))
@@ -129,14 +139,14 @@ namespace UnrealGameSync
 				State = new UserWorkspaceState();
 				State.RootDir = RootDir;
 				ImportWorkspaceState(RootDir, ClientName, BranchPath, State);
-				State.Save();
+				State.Save(Logger);
 			}
 			return State;
 		}
 
-		public UserWorkspaceState FindOrAddWorkspaceState(ProjectInfo ProjectInfo, UserWorkspaceSettings Settings)
+		public UserWorkspaceState FindOrAddWorkspaceState(ProjectInfo ProjectInfo, UserWorkspaceSettings Settings, ILogger Logger)
 		{
-			UserWorkspaceState State = FindOrAddWorkspaceState(ProjectInfo.LocalRootPath, ProjectInfo.ClientName, ProjectInfo.BranchPath);
+			UserWorkspaceState State = FindOrAddWorkspaceState(ProjectInfo.LocalRootPath, ProjectInfo.ClientName, ProjectInfo.BranchPath, Logger);
 			if (!State.IsValid(ProjectInfo))
 			{
 				State = new UserWorkspaceState();
@@ -145,7 +155,7 @@ namespace UnrealGameSync
 			return State;
 		}
 
-		public UserWorkspaceSettings FindOrAddWorkspaceSettings(DirectoryReference RootDir, string? ServerAndPort, string? UserName, string ClientName, string BranchPath, string ProjectPath)
+		public UserWorkspaceSettings FindOrAddWorkspaceSettings(DirectoryReference RootDir, string? ServerAndPort, string? UserName, string ClientName, string BranchPath, string ProjectPath, ILogger Logger)
 		{
 			ProjectInfo.ValidateBranchPath(BranchPath);
 			ProjectInfo.ValidateProjectPath(ProjectPath);
@@ -159,7 +169,7 @@ namespace UnrealGameSync
 			}
 
 			Settings.Init(ServerAndPort, UserName, ClientName, BranchPath, ProjectPath);
-			Settings.Save();
+			Settings.Save(Logger);
 
 			return Settings;
 		}
diff --git a/Engine/Source/Programs/UnrealGameSync/UnrealGameSyncShared/OpenProjectInfo.cs b/Engine/Source/Programs/UnrealGameSync/UnrealGameSyncShared/OpenProjectInfo.cs
--- a/Engine/Source/Programs/UnrealGameSync/UnrealGameSyncShared/OpenProjectInfo.cs
+++ b/Engine/Source/Programs/UnrealGameSync/UnrealGameSyncShared/OpenProjectInfo.cs
@@ -212,13 +212,13 @@ namespace UnrealGameSync
 				int BranchIdx = BranchClientPath.IndexOf('/', 2);
 				string BranchPath = (BranchIdx == -1) ? String.Empty : BranchClientPath.Substring(BranchIdx);
 				string ProjectPath = NewSelectedClientFileName.Substring(BranchClientPath.Length);
-				UserWorkspaceSettings UserWorkspaceSettings = UserSettings.FindOrAddWorkspaceSettings(BranchDirectoryName, PerforceSettings.ServerAndPort, PerforceSettings.UserName, PerforceSettings.ClientName!, BranchPath, ProjectPath);
+				UserWorkspaceSettings UserWorkspaceSettings = UserSettings.FindOrAddWorkspaceSettings(BranchDirectoryName, PerforceSettings.ServerAndPort, PerforceSettings.UserName, PerforceSettings.ClientName!, BranchPath, ProjectPath, Logger);
 
 				// Now compute the updated project info
 				ProjectInfo ProjectInfo = await ProjectInfo.CreateAsync(PerforceClient, UserWorkspaceSettings, CancellationToken);
 
 				// Update the cached workspace state
-				UserWorkspaceState UserWorkspaceState = UserSettings.FindOrAddWorkspaceState(ProjectInfo, UserWorkspaceSettings);
+				UserWorkspaceState UserWorkspaceState = UserSettings.FindOrAddWorkspaceState(ProjectInfo, UserWorkspaceSettings, Logger);
 
 				// Read the initial config file
 				List<KeyValuePair<FileReference, DateTime>> LocalConfigFiles = new List<KeyValuePair<FileReference, DateTime>>();
diff --git a/Engine/Source/Programs/UnrealGameSync/UnrealGameSyncShared/UserSettings.cs b/Engine/Source/Programs/UnrealGameSync/UnrealGameSyncShared/UserSettings.cs
--- a/Engine/Source/Programs/UnrealGameSync/UnrealGameSyncShared/UserSettings.cs
+++ b/Engine/Source/Programs/UnrealGameSync/UnrealGameSyncShared/UserSettings.cs
@@ -375,7 +375,21 @@ namespace UnrealGameSync
 
 		static object SyncRoot = new object();
 
-		public void Save()
+		public bool Save(ILogger Logger)
+		{
+			try
+			{
+				SaveInternal();
+				return true;
+			}
+			catch (Exception Ex)
+			{
+				Logger.LogError(Ex, "Unable to save {ConfigFile}: {Message}", ConfigFile, Ex.Message);
+				return false;
+			}
+		}
+
+		private void SaveInternal()
 		{
 			lock (SyncRoot)
 			{
@@ -499,7 +513,21 @@ namespace UnrealGameSync
 
 		static object SyncRoot = new object();
 
-		public void Save()
+		public bool Save(ILogger Logger)
+		{
+			try
+			{
+				SaveInternal();
+				return true;
+			}
+			catch (Exception Ex)
+			{
+				Logger.LogError(Ex, "Unable to save {ConfigFile}: {Message}", ConfigFile, Ex.Message);
+				return false;
+			}
+		}
+
+		private void SaveInternal()
 		{
 			lock (SyncRoot)
 			{
@@ -544,7 +572,21 @@ namespace UnrealGameSync
 			return false;
 		}
 
-		public void Save()
+		public bool Save(ILogger Logger)
+		{
+			try
+			{
+				SaveInternal();
+				return true;
+			}
+			catch (Exception Ex)
+			{
+				Logger.LogError(Ex, "Unable to save {ConfigFile}: {Message}", ConfigFile, Ex.Message);
+				return false;
+			}
+		}
+		
+		void SaveInternal()
 		{
 			lock (SyncRoot)
 			{
@@ -1098,9 +1140,12 @@ namespace UnrealGameSync
 			CurrentProject.FilterBadges.UnionWith(ProjectSection.GetValues("FilterBadges", new string[0]));
 		}
 
-		public override void Save()
+		public override bool Save(ILogger Logger)
 		{
-			base.Save();
+			if (!base.Save(Logger))
+			{
+				return false;
+			}
 
 			// General settings
 			ConfigSection GeneralSection = ConfigFile.FindOrAddSection("General");
@@ -1230,7 +1275,16 @@ namespace UnrealGameSync
 			}
 
 			// Save the file
-			ConfigFile.Save(FileName);
+			try
+			{
+				ConfigFile.Save(FileName);
+				return true;
+			}
+			catch (Exception Ex)
+			{
+				Logger.LogError(Ex, "Unable to save config file {FileName}: {Message}", FileName, Ex.Message);
+				return false;
+			}
 		}
 
 		[return: NotNullIfNotNull("Text")]
diff --git a/Engine/Source/Programs/UnrealGameSync/UnrealGameSyncShared/WorkspaceUpdate.cs b/Engine/Source/Programs/UnrealGameSync/UnrealGameSyncShared/WorkspaceUpdate.cs
--- a/Engine/Source/Programs/UnrealGameSync/UnrealGameSyncShared/WorkspaceUpdate.cs
+++ b/Engine/Source/Programs/UnrealGameSync/UnrealGameSyncShared/WorkspaceUpdate.cs
@@ -709,7 +709,7 @@ namespace UnrealGameSync
 							{
 								// Clear the current sync filter hash. If the sync is canceled, we'll be in an indeterminate state, and we should always clean next time round.
 								State.CurrentSyncFilterHash = "INVALID";
-								State.Save();
+								State.Save(Logger);
 
 								// Find all the depot paths that will be synced
 								HashSet<string> RemainingDepotPathsToRemove = new HashSet<string>(StringComparer.OrdinalIgnoreCase);
@@ -728,7 +728,7 @@ namespace UnrealGameSync
 
 							// Update the sync filter hash. We've removed any files we need to at this point.
 							State.CurrentSyncFilterHash = NextSyncFilterHash;
-							State.Save();
+							State.Save(Logger);
 						}
 					}
 
@@ -871,7 +871,7 @@ namespace UnrealGameSync
 					{
 						State.CurrentChangeNumber = -1;
 						State.AdditionalChangeNumbers.Clear();
-						State.Save();
+						State.Save(Logger);
 					}
 
 					// Find all the depot paths that will be synced
@@ -1069,7 +1069,7 @@ namespace UnrealGameSync
 						State.CurrentChangeNumber = Context.ChangeNumber;
 						State.CurrentCodeChangeNumber = VersionChangeNumber;
 					}
-					State.Save();
+					State.Save(Logger);
 
 					// Update the timing info
 					Times.Add(new Tuple<string, TimeSpan>("Sync", SyncTelemetryStopwatch.Stop("Success")));
@@ -1130,7 +1130,7 @@ namespace UnrealGameSync
 
 					// Update the state
 					State.ExpandedArchiveTypes = Context.ArchiveTypeToArchive.Where(x => x.Value != null).Select(x => x.Key).ToArray();
-					State.Save();
+					State.Save(Logger);
 
 					// Add the finish time
 					Times.Add(new Tuple<string, TimeSpan>("Archive", Stopwatch.Stop("Success")));
@@ -1372,7 +1372,7 @@ namespace UnrealGameSync
 				if (Context.CustomBuildSteps == null || Context.CustomBuildSteps.Count == 0)
 				{
 					State.LastBuiltChangeNumber = State.CurrentChangeNumber;
-					State.Save();
+					State.Save(Logger);
 				}
 			}
 
