UGS: Ignore user error exceptions for Sentry. These are user failure conditions rather than application errors.

diff --git a/Engine/Source/Programs/UnrealGameSync/UnrealGameSync/ProgramApplicationContext.cs b/Engine/Source/Programs/UnrealGameSync/UnrealGameSync/ProgramApplicationContext.cs
--- a/Engine/Source/Programs/UnrealGameSync/UnrealGameSync/ProgramApplicationContext.cs
+++ b/Engine/Source/Programs/UnrealGameSync/UnrealGameSync/ProgramApplicationContext.cs
@@ -222,6 +222,9 @@ namespace UnrealGameSync
 				catch (PerforceException)
 				{
 				}
+				catch (UserErrorException)
+				{
+				}
 				catch (Exception ex)
 				{
 					Program.CaptureException(ex);
