UGS: Do not send cancellation exceptions to Sentry.

diff --git a/Engine/Source/Programs/UnrealGameSync/UnrealGameSync/ProgramApplicationContext.cs b/Engine/Source/Programs/UnrealGameSync/UnrealGameSync/ProgramApplicationContext.cs
--- a/Engine/Source/Programs/UnrealGameSync/UnrealGameSync/ProgramApplicationContext.cs
+++ b/Engine/Source/Programs/UnrealGameSync/UnrealGameSync/ProgramApplicationContext.cs
@@ -222,6 +222,9 @@ namespace UnrealGameSync
 				catch (PerforceException)
 				{
 				}
+				catch (OperationCanceledException)
+				{
+				}
 				catch (UserErrorException)
 				{
 				}
