UGS: Only look at the first 100 files when determining types for each changelist. This prevents the situation where the client spends a huge amount of time fetching data for merge commits.

diff --git a/Engine/Source/Programs/UnrealGameSync/UnrealGameSync/PerforceMonitor.cs b/Engine/Source/Programs/UnrealGameSync/UnrealGameSync/PerforceMonitor.cs
--- a/Engine/Source/Programs/UnrealGameSync/UnrealGameSync/PerforceMonitor.cs
+++ b/Engine/Source/Programs/UnrealGameSync/UnrealGameSync/PerforceMonitor.cs
@@ -521,7 +521,7 @@ namespace UnrealGameSync
 					}
 
 					// If there's something to check for, find all the content changes after this changelist
-					List<DescribeRecord> DescribeRecords = await Perforce.DescribeAsync(QueryChangeNumberBatch.ToArray(), CancellationToken);
+					List<DescribeRecord> DescribeRecords = await Perforce.DescribeAsync(DescribeOptions.None, 100, QueryChangeNumberBatch.ToArray(), CancellationToken);
 					foreach (DescribeRecord DescribeRecord in DescribeRecords.OrderByDescending(x => x.Number))
 					{
 						int QueryChangeNumber = DescribeRecord.Number;
