UGS: Fixed crash on shutdown due to a null dereference.

diff --git a/Engine/Source/Programs/UnrealGameSync/UnrealGameSync/Controls/WorkspaceControl.cs b/Engine/Source/Programs/UnrealGameSync/UnrealGameSync/Controls/WorkspaceControl.cs
--- a/Engine/Source/Programs/UnrealGameSync/UnrealGameSync/Controls/WorkspaceControl.cs
+++ b/Engine/Source/Programs/UnrealGameSync/UnrealGameSync/Controls/WorkspaceControl.cs
@@ -681,7 +681,10 @@ namespace UnrealGameSync
 		/// <param name="disposing">true if managed resources should be disposed; otherwise, false.</param>
 		protected override void Dispose(bool disposing)
 		{
-			Owner.ToolUpdateMonitor.OnChange -= UpdateStatusPanel_CrossThread;
+			if (Owner?.ToolUpdateMonitor?.OnChange != null)
+			{
+				Owner.ToolUpdateMonitor.OnChange -= UpdateStatusPanel_CrossThread;
+			}
 
 			bIsDisposing = true;
 
@@ -707,7 +710,10 @@ namespace UnrealGameSync
 			if (IssueMonitor != null)
 			{
 				IssueMonitor.OnIssuesChanged -= IssueMonitor_OnIssuesChangedAsync;
-				Owner.ReleaseIssueMonitor(IssueMonitor);
+				if (Owner != null)
+				{
+					Owner.ReleaseIssueMonitor(IssueMonitor);
+				}
 				IssueMonitor = null!;
 			}
 			if (StartupTimer != null)
