UGS: Fix possibility for large changelists to be miscategorized as code/content.

diff --git a/Engine/Source/Programs/UnrealGameSync/UnrealGameSync/PerforceMonitor.cs b/Engine/Source/Programs/UnrealGameSync/UnrealGameSync/PerforceMonitor.cs
--- a/Engine/Source/Programs/UnrealGameSync/UnrealGameSync/PerforceMonitor.cs
+++ b/Engine/Source/Programs/UnrealGameSync/UnrealGameSync/PerforceMonitor.cs
@@ -507,6 +507,7 @@ namespace UnrealGameSync
 			}
 
 			// Update them in batches
+			bool bUpdatedConfigFile = false;
 			using (CancellationTokenSource cancellationSource = new CancellationTokenSource())
 			{
 				Task notifyTask = Task.CompletedTask;
@@ -521,12 +522,20 @@ namespace UnrealGameSync
 					}
 
 					// If there's something to check for, find all the content changes after this changelist
-					List<DescribeRecord> DescribeRecords = await Perforce.DescribeAsync(DescribeOptions.None, 100, QueryChangeNumberBatch.ToArray(), CancellationToken);
+					const int MaxFiles = 100;
+
+					List<DescribeRecord> DescribeRecords = await Perforce.DescribeAsync(DescribeOptions.None, MaxFiles, QueryChangeNumberBatch.ToArray(), CancellationToken);
 					foreach (DescribeRecord DescribeRecord in DescribeRecords.OrderByDescending(x => x.Number))
 					{
 						int QueryChangeNumber = DescribeRecord.Number;
 
 						PerforceChangeDetails Details = new PerforceChangeDetails(DescribeRecord);
+						if (DescribeRecord.Files.Count >= MaxFiles)
+						{
+							// Assume it's a code/content change if it has more files than we queried
+							Details.bContainsCode = Details.bContainsContent = true;
+						}
+
 						lock (this)
 						{
 							if (!ChangeDetails.ContainsKey(QueryChangeNumber))
@@ -536,9 +545,10 @@ namespace UnrealGameSync
 						}
 
 						// Reload the config file if it changes
-						if (DescribeRecord.Files.Any(x => x.DepotFile.EndsWith("/UnrealGameSync.ini", StringComparison.OrdinalIgnoreCase)))
+						if (DescribeRecord.Files.Any(x => x.DepotFile.EndsWith("/UnrealGameSync.ini", StringComparison.OrdinalIgnoreCase)) && !bUpdatedConfigFile)
 						{
 							await UpdateProjectConfigFileAsync(Perforce, CancellationToken);
+							bUpdatedConfigFile = true;
 						}
 
 						// Notify the caller after a fixed period of time, in case further updates are slow to arrive
