UGS: Fix PerforceMonitor stopping updates after an exception is thrown.

diff --git a/Engine/Source/Programs/UnrealGameSync/UnrealGameSync/PerforceMonitor.cs b/Engine/Source/Programs/UnrealGameSync/UnrealGameSync/PerforceMonitor.cs
--- a/Engine/Source/Programs/UnrealGameSync/UnrealGameSync/PerforceMonitor.cs
+++ b/Engine/Source/Programs/UnrealGameSync/UnrealGameSync/PerforceMonitor.cs
@@ -234,19 +234,11 @@ namespace UnrealGameSync
 
 		async Task PollForUpdates(CancellationToken CancellationToken)
 		{
-			bool bCoolDown = false;
 			while (!CancellationToken.IsCancellationRequested)
 			{
 				try
 				{
-					if (bCoolDown)
-					{
-						await Task.Delay(TimeSpan.FromSeconds(20.0), CancellationToken);
-					}
-					else
-					{
-						await PollForUpdatesInner(CancellationToken);
-					}
+					await PollForUpdatesInner(CancellationToken);
 				}
 				catch (OperationCanceledException) when (CancellationToken.IsCancellationRequested)
 				{
@@ -259,7 +251,7 @@ namespace UnrealGameSync
 					{
 						Program.CaptureException(Ex);
 					}
-					bCoolDown = true;
+					await Task.Delay(TimeSpan.FromSeconds(20.0), CancellationToken).ContinueWith(x => { });
 				}
 			}
 		}
