UGS: Fix main window losing maximized state when restoring from minimized.

diff --git a/Engine/Source/Programs/UnrealGameSync/UnrealGameSync/Forms/MainWindow.cs b/Engine/Source/Programs/UnrealGameSync/UnrealGameSync/Forms/MainWindow.cs
--- a/Engine/Source/Programs/UnrealGameSync/UnrealGameSync/Forms/MainWindow.cs
+++ b/Engine/Source/Programs/UnrealGameSync/UnrealGameSync/Forms/MainWindow.cs
@@ -769,6 +769,11 @@ namespace UnrealGameSync
 			UpdateProgress();
 		}
 
+		[DllImport("user32.dll")]
+		private static extern int ShowWindow(IntPtr hWnd, uint Msg);
+
+		private const uint SW_RESTORE = 0x09;
+
 		public void ShowAndActivate()
 		{
 			if (!IsDisposed)
@@ -776,7 +781,7 @@ namespace UnrealGameSync
 				Show();
 				if (WindowState == FormWindowState.Minimized)
 				{
-					WindowState = FormWindowState.Normal;
+					ShowWindow(Handle, SW_RESTORE);
 				}
 				Activate();
 
