PR #9037: Fixing issue with UGS not launching the correct executable when updating. (Contributed by [REDACTED_PERSON])

diff --git a/Engine/Source/Programs/UnrealGameSync/UnrealGameSync/ProgramApplicationContext.cs b/Engine/Source/Programs/UnrealGameSync/UnrealGameSync/ProgramApplicationContext.cs
--- a/Engine/Source/Programs/UnrealGameSync/UnrealGameSync/ProgramApplicationContext.cs
+++ b/Engine/Source/Programs/UnrealGameSync/UnrealGameSync/ProgramApplicationContext.cs
@@ -265,14 +265,18 @@ namespace UnrealGameSync
 			}
 
 			// Get the application path
-			string? OriginalExe = UpdateSpawn;
-			if (OriginalExe == null)
+			string OriginalExe = Assembly.GetExecutingAssembly().Location;
+			if (Path.GetExtension(OriginalExe).Equals(".dll", StringComparison.OrdinalIgnoreCase))
 			{
-				OriginalExe = Path.ChangeExtension(Assembly.GetExecutingAssembly().Location, ".exe");
+				string NewExecutable = Path.ChangeExtension(OriginalExe, ".exe");
+				if (File.Exists(NewExecutable))
+				{
+					OriginalExe = NewExecutable;
+				}
 			}
 
-			// Create the main window
-			MainWindowInstance = new MainWindow(UpdateMonitor, ApiUrl, DataFolder, CacheFolder, bRestoreState, OriginalExe, bUnstable, StartupTasks, DefaultPerforceSettings, ServiceProvider, Settings, Uri, OIDCTokenManager);
+			// Create the main window 
+			MainWindowInstance = new MainWindow(UpdateMonitor, ApiUrl, DataFolder, CacheFolder, bRestoreState, UpdateSpawn ?? OriginalExe, bUnstable, StartupTasks, DefaultPerforceSettings, ServiceProvider, Settings, Uri, OIDCTokenManager);
 			if(bVisible)
 			{
 				MainWindowInstance.Show();
diff --git a/Engine/Source/Programs/UnrealGameSync/UnrealGameSyncLauncher/Program.cs b/Engine/Source/Programs/UnrealGameSync/UnrealGameSyncLauncher/Program.cs
--- a/Engine/Source/Programs/UnrealGameSync/UnrealGameSyncLauncher/Program.cs
+++ b/Engine/Source/Programs/UnrealGameSync/UnrealGameSyncLauncher/Program.cs
@@ -170,7 +170,17 @@ namespace UnrealGameSyncLauncher
 				Logger.LogInformation("");
 
 				// Build the command line for the synced application, including the sync path to monitor for updates
-				StringBuilder NewCommandLine = new StringBuilder(String.Format("-updatepath=\"{0}@>{1}\" -updatespawn=\"{2}\"{3}", SyncPath, RequiredChangeNumber, Assembly.GetEntryAssembly()!.Location, bUnstable ? " -unstable" : ""));
+				string OriginalExecutable = Assembly.GetEntryAssembly()!.Location;
+                if (Path.GetExtension(OriginalExecutable).Equals(".dll", StringComparison.OrdinalIgnoreCase))
+                {
+                    string NewExecutable = Path.ChangeExtension(OriginalExecutable, ".exe");
+                    if (File.Exists(NewExecutable))
+                    {
+                        OriginalExecutable = NewExecutable;
+                    }
+                }
+
+				StringBuilder NewCommandLine = new StringBuilder(String.Format("-updatepath=\"{0}@>{1}\" -updatespawn=\"{2}\"{3}", SyncPath, RequiredChangeNumber, OriginalExecutable, bUnstable ? " -unstable" : ""));
 				foreach (string Arg in Args)
 				{
 					NewCommandLine.AppendFormat(" {0}", QuoteArgument(Arg));
