UGS: Increase the number of files queried for each changelist until we can determine whether it's a code change.

diff --git a/Engine/Source/Programs/UnrealGameSync/UnrealGameSync/PerforceMonitor.cs b/Engine/Source/Programs/UnrealGameSync/UnrealGameSync/PerforceMonitor.cs
--- a/Engine/Source/Programs/UnrealGameSync/UnrealGameSync/PerforceMonitor.cs
+++ b/Engine/Source/Programs/UnrealGameSync/UnrealGameSync/PerforceMonitor.cs
@@ -517,15 +517,28 @@ namespace UnrealGameSync
 					const int MaxFiles = 100;
 
 					List<DescribeRecord> DescribeRecords = await Perforce.DescribeAsync(DescribeOptions.None, MaxFiles, QueryChangeNumberBatch.ToArray(), CancellationToken);
-					foreach (DescribeRecord DescribeRecord in DescribeRecords.OrderByDescending(x => x.Number))
+					foreach (DescribeRecord DescribeRecordLoop in DescribeRecords.OrderByDescending(x => x.Number))
 					{
+						DescribeRecord DescribeRecord = DescribeRecordLoop;
 						int QueryChangeNumber = DescribeRecord.Number;
 
 						PerforceChangeDetails Details = new PerforceChangeDetails(DescribeRecord);
-						if (DescribeRecord.Files.Count >= MaxFiles)
+
+						// Content only changes must be flagged accurately, because code changes invalidate precompiled binaries. Increase the number of files fetched until we can classify it correctly.
+						int maxFiles = MaxFiles;
+						while (DescribeRecord.Files.Count >= maxFiles && !Details.bContainsCode)
 						{
-							// Assume it's a code/content change if it has more files than we queried
-							Details.bContainsCode = Details.bContainsContent = true;
+							CancellationToken.ThrowIfCancellationRequested();
+							maxFiles *= 10;
+
+							List<DescribeRecord> NewDescribeRecords = await Perforce.DescribeAsync(DescribeOptions.None, maxFiles, new int[] { QueryChangeNumber }, CancellationToken);
+							if (NewDescribeRecords.Count == 0)
+							{
+								break;
+							}
+
+							DescribeRecord = NewDescribeRecords[0];
+							Details = new PerforceChangeDetails(DescribeRecord);
 						}
 
 						lock (this)
