Add nodes to publish UGS zip files to Perforce.

diff --git a/Engine/Source/Programs/UnrealGameSync/BuildUGS.xml b/Engine/Source/Programs/UnrealGameSync/BuildUGS.xml
--- a/Engine/Source/Programs/UnrealGameSync/BuildUGS.xml
+++ b/Engine/Source/Programs/UnrealGameSync/BuildUGS.xml
@@ -9,6 +9,7 @@
 
 	<Property Name="UgsDir" Value="$(RootDir)/Engine/Source/Programs/UnrealGameSync"/>
 	<Property Name="UgsOutputDir" Value="$(RootDir)/LocalBuilds/UnrealGameSync"/>
+	<Property Name="UgsZipOutputDir" Value="$(RootDir)/LocalBuilds/UnrealGameSyncZip"/>
 	<Property Name="WixDir" Value="$(RootDir)/Engine/Source/ThirdParty/WiX/3.8"/>
 
 	<Property Name="Version" Value="$(EngineMajorVersion).$(EngineMinorVersion).$(EnginePatchVersion)"/>
@@ -19,6 +20,9 @@
 	<Agent Name="UnrealGameSync Agent" Type="Win64">
 		<Property Name="UgsPublishDir" Value="$(UgsOutputDir)"/>
 		<Property Name="UgsPublishDir" Value="$(UgsPublishDir)/$(SubmitPath)" If="'$(SubmitPath)' != ''"/>
+
+		<Property Name="UgsZipPublishDir" Value="$(UgsZipOutputDir)"/>
+		<Property Name="UgsZipPublishDir" Value="$(UgsZipPublishDir)/$(SubmitPath)" If="'$(SubmitPath)' != ''"/>
 	
 		<Node Name="Build UnrealGameSync">
 			<Delete Files="$(UgsOutputDir)/..."/>
@@ -30,9 +34,17 @@
 		<Node Name="Publish Unstable" Requires="Build UnrealGameSync">
 			<Submit Description="Updating UGS from CL $(Change)" Workspace="Temp" Files="$(UgsPublishDir)/..." RootDir="$(UgsOutputDir)" Stream="$(SubmitToStream)" Branch="$(SubmitToBranch)/UnstableRelease"/>
 		</Node>
+		<Node Name="Publish Unstable Zip" Requires="Build UnrealGameSync">
+			<Zip FromDir="$(UgsOutputDir)" ZipFile="$(UgsZipPublishDir)/UnstableRelease.zip"/>
+			<Submit Description="Updating UGS from CL $(Change) (Zip)" Workspace="Temp" Files="$(UgsZipPublishDir)/UnstableRelease.zip" RootDir="$(UgsZipOutputDir)" Stream="$(SubmitToStream)" Branch="$(SubmitToBranch)"/>
+		</Node>
 		<Node Name="Publish Live" Requires="Build UnrealGameSync">
 			<Submit Description="Updating UGS from CL $(Change)" Workspace="Temp" Files="$(UgsPublishDir)/..." RootDir="$(UgsOutputDir)" Stream="$(SubmitToStream)" Branch="$(SubmitToBranch)/Release"/>
 		</Node>
+		<Node Name="Publish Live Zip" Requires="Build UnrealGameSync">
+			<Zip FromDir="$(UgsOutputDir)" ZipFile="$(UgsZipPublishDir)/Release.zip"/>
+			<Submit Description="Updating UGS from CL $(Change) (Zip)" Workspace="Temp" Files="$(UgsZipPublishDir)/Release.zip" RootDir="$(UgsZipOutputDir)" Stream="$(SubmitToStream)" Branch="$(SubmitToBranch)"/>
+		</Node>
 	</Agent>
 	
 	<Agent Name="UnrealGameSync Installer" Type="Win64">
