UGS: Add aggregates for publishing UGS loose and zip files together.

diff --git a/Engine/Source/Programs/UnrealGameSync/BuildUGS.xml b/Engine/Source/Programs/UnrealGameSync/BuildUGS.xml
--- a/Engine/Source/Programs/UnrealGameSync/BuildUGS.xml
+++ b/Engine/Source/Programs/UnrealGameSync/BuildUGS.xml
@@ -62,4 +62,7 @@
 		</Node>
 	</Agent>
 
+	<Aggregate Name="Publish Unstable With Zip" Requires="Publish Unstable;Publish Unstable Zip"/>
+	<Aggregate Name="Publish Live With Zip" Requires="Publish Live;Publish Live Zip"/>
+
 </BuildGraph>
