UGS: Support for downloading UGS binaries via a zip file. This reduces chattiness, and will allows us to distribute self-contained bundles.

diff --git a/Engine/Source/Programs/UnrealGameSync/UnrealGameSyncLauncher/Program.cs b/Engine/Source/Programs/UnrealGameSync/UnrealGameSyncLauncher/Program.cs
--- a/Engine/Source/Programs/UnrealGameSync/UnrealGameSyncLauncher/Program.cs
+++ b/Engine/Source/Programs/UnrealGameSync/UnrealGameSyncLauncher/Program.cs
@@ -9,6 +9,7 @@ using System.Collections.Generic;
 using System.Diagnostics;
 using System.Diagnostics.CodeAnalysis;
 using System.IO;
+using System.IO.Compression;
 using System.Linq;
 using System.Reflection;
 using System.Text;
@@ -64,7 +65,7 @@ namespace UnrealGameSyncLauncher
 					// Try to do a sync with the current settings first
 					CaptureLogger Logger = new CaptureLogger();
 
-					IPerforceSettings Settings = PerforceSettings.Default.MergeWith(newServerAndPort: ServerAndPort, newUserName: UserName);
+					IPerforceSettings Settings = new PerforceSettings(PerforceSettings.Default) { PreferNativeClient = true }.MergeWith(newServerAndPort: ServerAndPort, newUserName: UserName);
 
 					ModalTask? Task = PerforceModalTask.Execute(null, "Updating", "Checking for updates, please wait...", Settings, (p, c) => SyncAndRun(p, DepotPath, bUnstable, Args, InstanceMutex, Logger, c), Logger);
 					if (Task == null)
@@ -95,7 +96,22 @@ namespace UnrealGameSyncLauncher
 					throw new UserErrorException($"Invalid setting for sync path");
 				}
 
-				string SyncPath = BaseDepotPath.TrimEnd('/') + (bUnstable ? "/UnstableRelease/..." : "/Release/...");
+				string BaseDepotPathPrefix = BaseDepotPath.TrimEnd('/');
+
+				// Find the most recent changelist
+				string SyncPath = BaseDepotPathPrefix + (bUnstable ? "/UnstableRelease.zip" : "/Release.zip");
+				List<ChangesRecord> Changes = await Perforce.GetChangesAsync(ChangesOptions.None, 1, ChangeStatus.Submitted, SyncPath, CancellationToken);
+				if (Changes.Count == 0)
+				{
+					SyncPath = BaseDepotPathPrefix + (bUnstable ? "/UnstableRelease/..." : "/Release/...");
+					Changes = await Perforce.GetChangesAsync(ChangesOptions.None, 1, ChangeStatus.Submitted, SyncPath, CancellationToken);
+					if (Changes.Count == 0)
+					{
+						throw new UserErrorException($"Unable to find any UGS binaries under {SyncPath}");
+					}
+				}
+
+				int RequiredChangeNumber = Changes[0].Number;
 				Logger.LogInformation("Syncing from {SyncPath}", SyncPath);
 
 				// Create the target folder
@@ -105,10 +121,6 @@ namespace UnrealGameSyncLauncher
 					throw new UserErrorException($"Couldn't create directory: {ApplicationFolder}");
 				}
 
-				// Find the most recent changelist
-				List<ChangesRecord> Changes = await Perforce.GetChangesAsync(ChangesOptions.None, 1, ChangeStatus.Submitted, SyncPath, CancellationToken);
-				int RequiredChangeNumber = Changes[0].Number;
-
 				// Read the current version
 				string SyncVersionFile = Path.Combine(ApplicationFolder, "SyncVersion.txt");
 				string RequiredSyncText = String.Format("{0}\n{1}@{2}", Perforce.Settings.ServerAndPort ?? "", SyncPath, RequiredChangeNumber);
@@ -155,6 +167,13 @@ namespace UnrealGameSyncLauncher
 						await Perforce.PrintAsync(LocalPath, FileRecord.DepotFile, CancellationToken);
 					}
 
+					// If it was a zip file, extract it
+					if (SyncPath.EndsWith(".zip", StringComparison.OrdinalIgnoreCase))
+					{
+						string LocalPath = Path.Combine(ApplicationFolder, SyncPath.Substring(DepotPathPrefix.Length).Replace('/', Path.DirectorySeparatorChar));
+						ZipFile.ExtractToDirectory(LocalPath, ApplicationFolder);
+					}
+
 					// Check the application exists
 					if (!File.Exists(ApplicationExe))
 					{
