UGS: Flag some threads as background threads to prevent them from blocking application exit.

diff --git a/Engine/Source/Programs/UnrealGameSync/UnrealGameSync/ActivationListener.cs b/Engine/Source/Programs/UnrealGameSync/UnrealGameSync/ActivationListener.cs
--- a/Engine/Source/Programs/UnrealGameSync/UnrealGameSync/ActivationListener.cs
+++ b/Engine/Source/Programs/UnrealGameSync/UnrealGameSync/ActivationListener.cs
@@ -28,6 +28,7 @@ namespace UnrealGameSync
 			if(BackgroundThread == null)
 			{
 				BackgroundThread = new Thread(x => ThreadProc());
+				BackgroundThread.IsBackground = true;
 				BackgroundThread.Start();
 			}
 		}
diff --git a/Engine/Source/Programs/UnrealGameSync/UnrealGameSync/AutomationServer.cs b/Engine/Source/Programs/UnrealGameSync/UnrealGameSync/AutomationServer.cs
--- a/Engine/Source/Programs/UnrealGameSync/UnrealGameSync/AutomationServer.cs
+++ b/Engine/Source/Programs/UnrealGameSync/UnrealGameSync/AutomationServer.cs
@@ -167,6 +167,7 @@ namespace UnrealGameSync
 						Listener = new TcpListener(IPAddress.Loopback, PortNumber);
 						Listener.Start();
 						TcpThread = new Thread(() => RunTcp());
+						TcpThread.IsBackground = true;
 						TcpThread.Start();
 
 					}
@@ -178,6 +179,7 @@ namespace UnrealGameSync
 				}
 
 				UriThread = new Thread(() => RunUri());
+				UriThread.IsBackground = true;
 				UriThread.Start();
 			}
 			catch (Exception Ex)
diff --git a/Engine/Source/Programs/UnrealGameSync/UnrealGameSync/Forms/ProgramsRunningWindow.cs b/Engine/Source/Programs/UnrealGameSync/UnrealGameSync/Forms/ProgramsRunningWindow.cs
--- a/Engine/Source/Programs/UnrealGameSync/UnrealGameSync/Forms/ProgramsRunningWindow.cs
+++ b/Engine/Source/Programs/UnrealGameSync/UnrealGameSync/Forms/ProgramsRunningWindow.cs
@@ -36,6 +36,7 @@ namespace UnrealGameSync
 			TerminateEvent = new ManualResetEvent(false);
 
 			BackgroundThread = new Thread(() => ExecuteBackgroundWork());
+			BackgroundThread.IsBackground = true;
 			BackgroundThread.Start();
 		}
 
