UGS: Rename the "unstable" setting to "preview", and persist the setting across UGS runs.

diff --git a/Engine/Source/Programs/UnrealGameSync/UnrealGameSync/Automation/P4Automation.cs b/Engine/Source/Programs/UnrealGameSync/UnrealGameSync/Automation/P4Automation.cs
--- a/Engine/Source/Programs/UnrealGameSync/UnrealGameSync/Automation/P4Automation.cs
+++ b/Engine/Source/Programs/UnrealGameSync/UnrealGameSync/Automation/P4Automation.cs
@@ -20,8 +20,9 @@ namespace UnrealGameSync
 			string? ServerAndPort = null;
 			string? UserName = null;
 			string? DepotPathSettings = null;
+			bool bPreview = false;
 
-			GlobalPerforceSettings.ReadGlobalPerforceSettings(ref ServerAndPort, ref UserName, ref DepotPathSettings);
+			GlobalPerforceSettings.ReadGlobalPerforceSettings(ref ServerAndPort, ref UserName, ref DepotPathSettings, ref bPreview);
 
 			return Utility.OverridePerforceSettings(PerforceSettings.Default, ServerAndPort, UserName);
 
diff --git a/Engine/Source/Programs/UnrealGameSync/UnrealGameSync/Forms/ApplicationSettingsWindow.cs b/Engine/Source/Programs/UnrealGameSync/UnrealGameSync/Forms/ApplicationSettingsWindow.cs
--- a/Engine/Source/Programs/UnrealGameSync/UnrealGameSync/Forms/ApplicationSettingsWindow.cs
+++ b/Engine/Source/Programs/UnrealGameSync/UnrealGameSync/Forms/ApplicationSettingsWindow.cs
@@ -43,11 +43,11 @@ namespace UnrealGameSync
 		string? InitialServerAndPort;
 		string? InitialUserName;
 		string? InitialDepotPath;
-		bool bInitialUnstable;
+		bool bInitialPreview;
 		int InitialAutomationPortNumber;
 		ProtocolHandlerState InitialProtocolHandlerState;
 
-		bool? bRestartUnstable;
+		bool? bRestartPreview;
 
 		ToolUpdateMonitor ToolUpdateMonitor;
 
@@ -66,7 +66,7 @@ namespace UnrealGameSync
 			}
 		}
 
-		private ApplicationSettingsWindow(IPerforceSettings DefaultPerforceSettings, bool bUnstable, string OriginalExecutableFileName, UserSettings Settings, ToolUpdateMonitor ToolUpdateMonitor, ILogger<ApplicationSettingsWindow> Logger)
+		private ApplicationSettingsWindow(IPerforceSettings DefaultPerforceSettings, bool bPreview, string OriginalExecutableFileName, UserSettings Settings, ToolUpdateMonitor ToolUpdateMonitor, ILogger<ApplicationSettingsWindow> Logger)
 		{
 			InitializeComponent();
 
@@ -76,8 +76,8 @@ namespace UnrealGameSync
 			this.ToolUpdateMonitor = ToolUpdateMonitor;
 			this.Logger = Logger;
 
-			GlobalPerforceSettings.ReadGlobalPerforceSettings(ref InitialServerAndPort, ref InitialUserName, ref InitialDepotPath);
-			bInitialUnstable = bUnstable;
+			GlobalPerforceSettings.ReadGlobalPerforceSettings(ref InitialServerAndPort, ref InitialUserName, ref InitialDepotPath, ref bPreview);
+			bInitialPreview = bPreview;
 
 			InitialAutomationPortNumber = AutomationServer.GetPortNumber();
 			InitialProtocolHandlerState = ProtocolHandlerUtils.GetState();
@@ -99,7 +99,7 @@ namespace UnrealGameSync
 			this.DepotPathTextBox.Select(DepotPathTextBox.TextLength, 0);
 			this.DepotPathTextBox.CueBanner = DeploymentSettings.DefaultDepotPath ?? String.Empty;
 
-			this.UseUnstableBuildCheckBox.Checked = bUnstable;
+			this.UsePreviewBuildCheckBox.Checked = bPreview;
 
 			if(InitialAutomationPortNumber > 0)
 			{
@@ -134,12 +134,12 @@ namespace UnrealGameSync
 			}
 		}
 
-		public static bool? ShowModal(IWin32Window Owner, IPerforceSettings DefaultPerforceSettings, bool bUnstable, string OriginalExecutableFileName, UserSettings Settings, ToolUpdateMonitor ToolUpdateMonitor, ILogger<ApplicationSettingsWindow> Logger)
+		public static bool? ShowModal(IWin32Window Owner, IPerforceSettings DefaultPerforceSettings, bool bPreview, string OriginalExecutableFileName, UserSettings Settings, ToolUpdateMonitor ToolUpdateMonitor, ILogger<ApplicationSettingsWindow> Logger)
 		{
-			ApplicationSettingsWindow ApplicationSettings = new ApplicationSettingsWindow(DefaultPerforceSettings, bUnstable, OriginalExecutableFileName, Settings, ToolUpdateMonitor, Logger);
+			ApplicationSettingsWindow ApplicationSettings = new ApplicationSettingsWindow(DefaultPerforceSettings, bPreview, OriginalExecutableFileName, Settings, ToolUpdateMonitor, Logger);
 			if(ApplicationSettings.ShowDialog() == DialogResult.OK)
 			{
-				return ApplicationSettings.bRestartUnstable;
+				return ApplicationSettings.bRestartPreview;
 			}
 			else
 			{
@@ -174,7 +174,7 @@ namespace UnrealGameSync
 				DepotPath = null;
 			}
 
-			bool bUnstable = UseUnstableBuildCheckBox.Checked;
+			bool bPreview = UsePreviewBuildCheckBox.Checked;
 
 
 			int AutomationPortNumber;
@@ -183,7 +183,7 @@ namespace UnrealGameSync
 				AutomationPortNumber = -1;
 			}
 			
-			if(ServerAndPort != InitialServerAndPort || UserName != InitialUserName || DepotPath != InitialDepotPath || bUnstable != bInitialUnstable || AutomationPortNumber != InitialAutomationPortNumber)
+			if(ServerAndPort != InitialServerAndPort || UserName != InitialUserName || DepotPath != InitialDepotPath || bPreview != bInitialPreview || AutomationPortNumber != InitialAutomationPortNumber)
 			{
 				// Try to log in to the new server, and check the application is there
 				if(ServerAndPort != InitialServerAndPort || UserName != InitialUserName || DepotPath != InitialDepotPath)
@@ -206,8 +206,8 @@ namespace UnrealGameSync
 					return;
 				}
 
-				bRestartUnstable = UseUnstableBuildCheckBox.Checked;
-				GlobalPerforceSettings.SaveGlobalPerforceSettings(ServerAndPort, UserName, DepotPath);
+				bRestartPreview = UsePreviewBuildCheckBox.Checked;
+				GlobalPerforceSettings.SaveGlobalPerforceSettings(ServerAndPort, UserName, DepotPath, bPreview);
 				AutomationServer.SetPortNumber(AutomationPortNumber);
 			}
 
diff --git a/Engine/Source/Programs/UnrealGameSync/UnrealGameSync/Forms/MainWindow.cs b/Engine/Source/Programs/UnrealGameSync/UnrealGameSync/Forms/MainWindow.cs
--- a/Engine/Source/Programs/UnrealGameSync/UnrealGameSync/Forms/MainWindow.cs
+++ b/Engine/Source/Programs/UnrealGameSync/UnrealGameSync/Forms/MainWindow.cs
@@ -106,7 +106,7 @@ namespace UnrealGameSync
 		System.Threading.Timer? ScheduleSettledTimer;
 
 		string OriginalExecutableFileName;
-		bool bUnstable;
+		bool bPreview;
 
 		IMainWindowTabPanel? CurrentTabPanel;
 
@@ -119,7 +119,7 @@ namespace UnrealGameSync
 
 		public ToolUpdateMonitor ToolUpdateMonitor { get; private set; }
 
-		public MainWindow(UpdateMonitor InUpdateMonitor, string? InApiUrl, DirectoryReference InDataFolder, DirectoryReference InCacheFolder, bool bInRestoreStateOnLoad, string InOriginalExecutableFileName, bool bInUnstable, List<(UserSelectedProjectSettings, ModalTask<OpenProjectInfo>)> StartupTasks, IPerforceSettings InDefaultPerforceSettings, IServiceProvider InServiceProvider, UserSettings InSettings, string? InUri, OIDCTokenManager? InOidcTokenManager)
+		public MainWindow(UpdateMonitor InUpdateMonitor, string? InApiUrl, DirectoryReference InDataFolder, DirectoryReference InCacheFolder, bool bInRestoreStateOnLoad, string InOriginalExecutableFileName, bool bInPreview, List<(UserSelectedProjectSettings, ModalTask<OpenProjectInfo>)> StartupTasks, IPerforceSettings InDefaultPerforceSettings, IServiceProvider InServiceProvider, UserSettings InSettings, string? InUri, OIDCTokenManager? InOidcTokenManager)
 		{
 			ServiceProvider = InServiceProvider;
 			Logger = ServiceProvider.GetRequiredService<ILogger<MainWindow>>();
@@ -135,7 +135,7 @@ namespace UnrealGameSync
 			CacheFolder = InCacheFolder;
 			bRestoreStateOnLoad = bInRestoreStateOnLoad;
 			OriginalExecutableFileName = InOriginalExecutableFileName;
-			bUnstable = bInUnstable;
+			bPreview = bInPreview;
 			DefaultPerforceSettings = InDefaultPerforceSettings;
 			ToolUpdateMonitor = new ToolUpdateMonitor(DefaultPerforceSettings, DataFolder, InSettings, ServiceProvider);
 
@@ -187,7 +187,7 @@ namespace UnrealGameSync
 
 			StartScheduleTimer();
 
-			if(bUnstable)
+			if(bPreview)
 			{
 				Text += $" {Program.GetVersionString()} (UNSTABLE)";
 			}
@@ -1425,10 +1425,10 @@ namespace UnrealGameSync
 
 		public void ModifyApplicationSettings()
 		{
-			bool? bRelaunchUnstable = ApplicationSettingsWindow.ShowModal(this, DefaultPerforceSettings, bUnstable, OriginalExecutableFileName, Settings, ToolUpdateMonitor, ServiceProvider.GetRequiredService<ILogger<ApplicationSettingsWindow>>());
-			if(bRelaunchUnstable.HasValue)
+			bool? bRelaunchPreview = ApplicationSettingsWindow.ShowModal(this, DefaultPerforceSettings, bPreview, OriginalExecutableFileName, Settings, ToolUpdateMonitor, ServiceProvider.GetRequiredService<ILogger<ApplicationSettingsWindow>>());
+			if(bRelaunchPreview.HasValue)
 			{
-				UpdateMonitor.TriggerUpdate(UpdateType.UserInitiated, bRelaunchUnstable);
+				UpdateMonitor.TriggerUpdate(UpdateType.UserInitiated, bRelaunchPreview);
 			}
 
 			for (int Idx = 0; Idx < TabControl.GetTabCount(); Idx++)
diff --git a/Engine/Source/Programs/UnrealGameSync/UnrealGameSync/GlobalSettings.cs b/Engine/Source/Programs/UnrealGameSync/UnrealGameSync/GlobalSettings.cs
--- a/Engine/Source/Programs/UnrealGameSync/UnrealGameSync/GlobalSettings.cs
+++ b/Engine/Source/Programs/UnrealGameSync/UnrealGameSync/GlobalSettings.cs
@@ -12,7 +12,7 @@ namespace UnrealGameSync
 {
 	internal class GlobalPerforceSettings
 	{
-		public static void ReadGlobalPerforceSettings(ref string? ServerAndPort, ref string? UserName, ref string? DepotPath)
+		public static void ReadGlobalPerforceSettings(ref string? ServerAndPort, ref string? UserName, ref string? DepotPath, ref bool bPreview)
 		{
 			using (RegistryKey? Key = Registry.CurrentUser.OpenSubKey("SOFTWARE\\Epic Games\\UnrealGameSync", false))
 			{
@@ -21,6 +21,7 @@ namespace UnrealGameSync
 					ServerAndPort = Key.GetValue("ServerAndPort", ServerAndPort) as string;
 					UserName = Key.GetValue("UserName", UserName) as string;
 					DepotPath = Key.GetValue("DepotPath", DepotPath) as string;
+					bPreview = ((Key.GetValue("Preview", bPreview? 1 : 0) as int?) ?? 0) != 0;
 
 					// Fix corrupted depot path string
 					if (DepotPath != null)
@@ -29,7 +30,7 @@ namespace UnrealGameSync
 						if (Match.Success)
 						{
 							DepotPath = Match.Groups[1].Value;
-							SaveGlobalPerforceSettings(ServerAndPort, UserName, DepotPath);
+							SaveGlobalPerforceSettings(ServerAndPort, UserName, DepotPath, bPreview);
 						}
 					}
 				}
@@ -62,7 +63,7 @@ namespace UnrealGameSync
 			}
 		}
 
-		public static void SaveGlobalPerforceSettings(string? ServerAndPort, string? UserName, string? DepotPath)
+		public static void SaveGlobalPerforceSettings(string? ServerAndPort, string? UserName, string? DepotPath, bool bPreview)
 		{
 			try
 			{
@@ -97,6 +98,15 @@ namespace UnrealGameSync
 					{
 						Key.SetValue("DepotPath", DepotPath);
 					}
+
+					if (bPreview)
+					{
+						Key.SetValue("Preview", 1);
+					}
+					else
+					{
+						try { Key.DeleteValue("Preview"); } catch (Exception) { }
+					}
 				}
 			}
 			catch (Exception Ex)
diff --git a/Engine/Source/Programs/UnrealGameSync/UnrealGameSync/Program.cs b/Engine/Source/Programs/UnrealGameSync/UnrealGameSync/Program.cs
--- a/Engine/Source/Programs/UnrealGameSync/UnrealGameSync/Program.cs
+++ b/Engine/Source/Programs/UnrealGameSync/UnrealGameSync/Program.cs
@@ -86,7 +86,7 @@ namespace UnrealGameSync
 
 					if (bFirstInstance)
 					{
-						InnerMainAsync(InstanceMutex, ActivateEvent, Args).GetAwaiter().GetResult();
+						GuardedInnerMainAsync(InstanceMutex, ActivateEvent, Args);
 					}
 					else
 					{
@@ -96,12 +96,25 @@ namespace UnrealGameSync
 			}
 		}
 
+		static void GuardedInnerMainAsync(Mutex InstanceMutex, EventWaitHandle ActivateEvent, string[] Args)
+		{
+			try
+			{
+				InnerMainAsync(InstanceMutex, ActivateEvent, Args).GetAwaiter().GetResult();
+			}
+			catch (Exception ex)
+			{
+				CaptureException(ex);
+			}
+		}
+
 		static async Task InnerMainAsync(Mutex InstanceMutex, EventWaitHandle ActivateEvent, string[] Args)
 		{
 			string? ServerAndPort = null;
 			string? UserName = null;
 			string? BaseUpdatePath = null;
-			GlobalPerforceSettings.ReadGlobalPerforceSettings(ref ServerAndPort, ref UserName, ref BaseUpdatePath);
+			bool bPreviewSetting = false;
+			GlobalPerforceSettings.ReadGlobalPerforceSettings(ref ServerAndPort, ref UserName, ref BaseUpdatePath, ref bPreviewSetting);
 
 			List<string> RemainingArgs = new List<string>(Args);
 
@@ -116,6 +129,9 @@ namespace UnrealGameSync
 
 			bool bUnstable;
 			ParseOption(RemainingArgs, "-unstable", out bUnstable);
+			bool bPreview;
+			ParseOption(RemainingArgs, "-preview", out bPreview);
+			bPreview |= bUnstable;
 
             string? ProjectFileName;
             ParseArgument(RemainingArgs, "-project=", out ProjectFileName);
@@ -177,7 +193,7 @@ namespace UnrealGameSync
 						Logger.LogInformation("Missing server settings; finding defaults.");
 						ServerAndPort ??= PerforceSettings.Default.ServerAndPort;
 						UserName ??= PerforceSettings.Default.UserName;
-						GlobalPerforceSettings.SaveGlobalPerforceSettings(ServerAndPort, UserName, BaseUpdatePath);
+						GlobalPerforceSettings.SaveGlobalPerforceSettings(ServerAndPort, UserName, BaseUpdatePath, bPreviewSetting);
 					}
 
 					ILogger TelemetryLogger = LoggerProvider.CreateLogger("Telemetry");
@@ -198,14 +214,14 @@ namespace UnrealGameSync
 
 							using (UpdateMonitor UpdateMonitor = new UpdateMonitor(DefaultSettings, UpdatePath, ServiceProvider))
 							{
-								using ProgramApplicationContext Context = new ProgramApplicationContext(DefaultSettings, UpdateMonitor, DeploymentSettings.ApiUrl, DataFolder, ActivateEvent, bRestoreState, UpdateSpawn, ProjectFileName, bUnstable, ServiceProvider, Uri);
+								using ProgramApplicationContext Context = new ProgramApplicationContext(DefaultSettings, UpdateMonitor, DeploymentSettings.ApiUrl, DataFolder, ActivateEvent, bRestoreState, UpdateSpawn, ProjectFileName, bPreview, ServiceProvider, Uri);
 								Application.Run(Context);
 
 								if (UpdateMonitor.IsUpdateAvailable && UpdateSpawn != null)
 								{
 									InstanceMutex.Close();
-									bool bLaunchUnstable = UpdateMonitor.RelaunchUnstable ?? bUnstable;
-									Utility.SpawnProcess(UpdateSpawn, "-restorestate" + (bLaunchUnstable ? " -unstable" : ""));
+									bool bLaunchPreview = UpdateMonitor.RelaunchPreview ?? bPreview;
+									Utility.SpawnProcess(UpdateSpawn, "-restorestate" + (bLaunchPreview ? " -unstable" : ""));
 								}
 							}
 						}
diff --git a/Engine/Source/Programs/UnrealGameSync/UnrealGameSync/ProgramApplicationContext.cs b/Engine/Source/Programs/UnrealGameSync/UnrealGameSync/ProgramApplicationContext.cs
--- a/Engine/Source/Programs/UnrealGameSync/UnrealGameSync/ProgramApplicationContext.cs
+++ b/Engine/Source/Programs/UnrealGameSync/UnrealGameSync/ProgramApplicationContext.cs
@@ -32,7 +32,7 @@ namespace UnrealGameSync
 		DirectoryReference CacheFolder;
 		bool bRestoreState;
 		string? UpdateSpawn;
-		bool bUnstable;
+		bool bPreview;
 		bool bIsClosing;
 		string? Uri;
 
@@ -58,7 +58,7 @@ namespace UnrealGameSync
 
 		OIDCTokenManager? OIDCTokenManager;
 
-		public ProgramApplicationContext(IPerforceSettings DefaultPerforceSettings, UpdateMonitor UpdateMonitor, string? ApiUrl, DirectoryReference DataFolder, EventWaitHandle ActivateEvent, bool bRestoreState, string? UpdateSpawn, string? ProjectFileName, bool bUnstable, IServiceProvider ServiceProvider, string? Uri)
+		public ProgramApplicationContext(IPerforceSettings DefaultPerforceSettings, UpdateMonitor UpdateMonitor, string? ApiUrl, DirectoryReference DataFolder, EventWaitHandle ActivateEvent, bool bRestoreState, string? UpdateSpawn, string? ProjectFileName, bool bPreview, IServiceProvider ServiceProvider, string? Uri)
 		{
 			this.DefaultPerforceSettings = DefaultPerforceSettings;
 			this.UpdateMonitor = UpdateMonitor;
@@ -67,7 +67,7 @@ namespace UnrealGameSync
 			this.CacheFolder = DirectoryReference.Combine(DataFolder, "Cache");
 			this.bRestoreState = bRestoreState;
 			this.UpdateSpawn = UpdateSpawn;
-			this.bUnstable = bUnstable;
+			this.bPreview = bPreview;
 			this.ServiceProvider = ServiceProvider;
 			this.Logger = ServiceProvider.GetRequiredService<ILogger<ProgramApplicationContext>>();
 			this.Uri = Uri;
@@ -276,7 +276,7 @@ namespace UnrealGameSync
 			}
 
 			// Create the main window 
-			MainWindowInstance = new MainWindow(UpdateMonitor, ApiUrl, DataFolder, CacheFolder, bRestoreState, UpdateSpawn ?? OriginalExe, bUnstable, StartupTasks, DefaultPerforceSettings, ServiceProvider, Settings, Uri, OIDCTokenManager);
+			MainWindowInstance = new MainWindow(UpdateMonitor, ApiUrl, DataFolder, CacheFolder, bRestoreState, UpdateSpawn ?? OriginalExe, bPreview, StartupTasks, DefaultPerforceSettings, ServiceProvider, Settings, Uri, OIDCTokenManager);
 			if(bVisible)
 			{
 				MainWindowInstance.Show();
diff --git a/Engine/Source/Programs/UnrealGameSync/UnrealGameSync/UpdateMonitor.cs b/Engine/Source/Programs/UnrealGameSync/UnrealGameSync/UpdateMonitor.cs
--- a/Engine/Source/Programs/UnrealGameSync/UnrealGameSync/UpdateMonitor.cs
+++ b/Engine/Source/Programs/UnrealGameSync/UnrealGameSync/UpdateMonitor.cs
@@ -29,7 +29,7 @@ namespace UnrealGameSync
 
 		public Action<UpdateType>? OnUpdateAvailable;
 
-		public bool? RelaunchUnstable
+		public bool? RelaunchPreview
 		{
 			get;
 			private set;
@@ -98,9 +98,9 @@ namespace UnrealGameSync
 			}
 		}
 
-		public void TriggerUpdate(UpdateType UpdateType, bool? RelaunchUnstable)
+		public void TriggerUpdate(UpdateType UpdateType, bool? RelaunchPreview)
 		{
-			this.RelaunchUnstable = RelaunchUnstable;
+			this.RelaunchPreview = RelaunchPreview;
 			IsUpdateAvailable = true;
 			if(OnUpdateAvailable != null)
 			{
diff --git a/Engine/Source/Programs/UnrealGameSync/UnrealGameSyncLauncher/Forms/SettingsWindow.cs b/Engine/Source/Programs/UnrealGameSync/UnrealGameSyncLauncher/Forms/SettingsWindow.cs
--- a/Engine/Source/Programs/UnrealGameSync/UnrealGameSyncLauncher/Forms/SettingsWindow.cs
+++ b/Engine/Source/Programs/UnrealGameSync/UnrealGameSyncLauncher/Forms/SettingsWindow.cs
@@ -24,14 +24,14 @@ namespace UnrealGameSyncLauncher
 		[DllImport("user32.dll")]
 		private static extern IntPtr SendMessage(IntPtr hWnd, int Msg, int wParam, [MarshalAs(UnmanagedType.LPWStr)] string lParam);
 
-		public delegate Task SyncAndRunDelegate(IPerforceConnection Perforce, string? DepotPath, bool bUnstable, ILogger LogWriter, CancellationToken CancellationToken);
+		public delegate Task SyncAndRunDelegate(IPerforceConnection Perforce, string? DepotPath, bool bPreview, ILogger LogWriter, CancellationToken CancellationToken);
 
 		const int EM_SETCUEBANNER = 0x1501;
 
 		string? LogText;
 		SyncAndRunDelegate SyncAndRun;
 
-		public SettingsWindow(string? Prompt, string? LogText, string? ServerAndPort, string? UserName, string? DepotPath, bool bUnstable, SyncAndRunDelegate SyncAndRun)
+		public SettingsWindow(string? Prompt, string? LogText, string? ServerAndPort, string? UserName, string? DepotPath, bool bPreview, SyncAndRunDelegate SyncAndRun)
 		{
 			InitializeComponent();
 
@@ -44,7 +44,7 @@ namespace UnrealGameSyncLauncher
 			this.ServerTextBox.Text = ServerAndPort ?? String.Empty;
 			this.UserNameTextBox.Text = UserName ?? String.Empty;
 			this.DepotPathTextBox.Text = DepotPath ?? String.Empty;
-			this.UseUnstableBuildCheckBox.Checked = bUnstable;
+			this.UsePreviewBuildCheckBox.Checked = bPreview;
 			this.SyncAndRun = SyncAndRun;
 
 			ViewLogBtn.Visible = LogText != null;
@@ -102,7 +102,7 @@ namespace UnrealGameSyncLauncher
 			CaptureLogger Logger = new CaptureLogger();
 
 			// Create the task for connecting to this server
-			ModalTask? Task = PerforceModalTask.Execute(this, "Updating", "Checking for updates, please wait...", PerforceSettings, (p, c) => SyncAndRun(p, DepotPath, UseUnstableBuildCheckBox.Checked, Logger, c), Logger);
+			ModalTask? Task = PerforceModalTask.Execute(this, "Updating", "Checking for updates, please wait...", PerforceSettings, (p, c) => SyncAndRun(p, DepotPath, UsePreviewBuildCheckBox.Checked, Logger, c), Logger);
 			if (Task != null)
 			{
 				if(Task.Succeeded)
diff --git a/Engine/Source/Programs/UnrealGameSync/UnrealGameSyncLauncher/Program.cs b/Engine/Source/Programs/UnrealGameSync/UnrealGameSyncLauncher/Program.cs
--- a/Engine/Source/Programs/UnrealGameSync/UnrealGameSyncLauncher/Program.cs
+++ b/Engine/Source/Programs/UnrealGameSync/UnrealGameSyncLauncher/Program.cs
@@ -41,20 +41,20 @@ namespace UnrealGameSyncLauncher
 				}
 
 				// Figure out if we should sync the unstable build by default
-				bool bUnstable = Args.Contains("-unstable", StringComparer.InvariantCultureIgnoreCase);
+				bool bPreview = Args.Contains("-unstable", StringComparer.InvariantCultureIgnoreCase) || Args.Contains("-preview", StringComparer.InvariantCultureIgnoreCase);
 
 				// Read the settings
 				string? ServerAndPort = null;
 				string? UserName = null;
 				string? DepotPath = DeploymentSettings.DefaultDepotPath;
-				GlobalPerforceSettings.ReadGlobalPerforceSettings(ref ServerAndPort, ref UserName, ref DepotPath);
+				GlobalPerforceSettings.ReadGlobalPerforceSettings(ref ServerAndPort, ref UserName, ref DepotPath, ref bPreview);
 
 				// If the shift key is held down, immediately show the settings window
-				SettingsWindow.SyncAndRunDelegate SyncAndRunWrapper = (Perforce, DepotParam, bUnstableParam, LogWriter, CancellationToken) => SyncAndRun(Perforce, DepotParam, bUnstableParam, Args, InstanceMutex, LogWriter, CancellationToken);
+				SettingsWindow.SyncAndRunDelegate SyncAndRunWrapper = (Perforce, DepotParam, bPreviewParam, LogWriter, CancellationToken) => SyncAndRun(Perforce, DepotParam, bPreviewParam, Args, InstanceMutex, LogWriter, CancellationToken);
 				if ((Control.ModifierKeys & Keys.Shift) != 0)
 				{
 					// Show the settings window immediately
-					SettingsWindow UpdateError = new SettingsWindow(null, null, ServerAndPort, UserName, DepotPath, bUnstable, SyncAndRunWrapper);
+					SettingsWindow UpdateError = new SettingsWindow(null, null, ServerAndPort, UserName, DepotPath, bPreview, SyncAndRunWrapper);
 					if(UpdateError.ShowDialog() == DialogResult.OK)
 					{
 						return 0;
@@ -67,7 +67,7 @@ namespace UnrealGameSyncLauncher
 
 					IPerforceSettings Settings = new PerforceSettings(PerforceSettings.Default) { PreferNativeClient = true }.MergeWith(newServerAndPort: ServerAndPort, newUserName: UserName);
 
-					ModalTask? Task = PerforceModalTask.Execute(null, "Updating", "Checking for updates, please wait...", Settings, (p, c) => SyncAndRun(p, DepotPath, bUnstable, Args, InstanceMutex, Logger, c), Logger);
+					ModalTask? Task = PerforceModalTask.Execute(null, "Updating", "Checking for updates, please wait...", Settings, (p, c) => SyncAndRun(p, DepotPath, bPreview, Args, InstanceMutex, Logger, c), Logger);
 					if (Task == null)
 					{
 						Logger.LogInformation("Canceled by user");
@@ -77,7 +77,7 @@ namespace UnrealGameSyncLauncher
 						return 0;
 					}
 
-					SettingsWindow UpdateError = new SettingsWindow("Unable to update UnrealGameSync from Perforce. Verify that your connection settings are correct.", Logger.Render(Environment.NewLine), ServerAndPort, UserName, DepotPath, bUnstable, SyncAndRunWrapper);
+					SettingsWindow UpdateError = new SettingsWindow("Unable to update UnrealGameSync from Perforce. Verify that your connection settings are correct.", Logger.Render(Environment.NewLine), ServerAndPort, UserName, DepotPath, bPreview, SyncAndRunWrapper);
 					if(UpdateError.ShowDialog() == DialogResult.OK)
 					{
 						return 0;
@@ -87,7 +87,7 @@ namespace UnrealGameSyncLauncher
 			return 1;
 		}
 
-		public static async Task SyncAndRun(IPerforceConnection Perforce, string? BaseDepotPath, bool bUnstable, string[] Args, Mutex InstanceMutex, ILogger Logger, CancellationToken CancellationToken)
+		public static async Task SyncAndRun(IPerforceConnection Perforce, string? BaseDepotPath, bool bPreview, string[] Args, Mutex InstanceMutex, ILogger Logger, CancellationToken CancellationToken)
 		{
 			try
 			{
@@ -99,11 +99,11 @@ namespace UnrealGameSyncLauncher
 				string BaseDepotPathPrefix = BaseDepotPath.TrimEnd('/');
 
 				// Find the most recent changelist
-				string SyncPath = BaseDepotPathPrefix + (bUnstable ? "/UnstableRelease.zip" : "/Release.zip");
+				string SyncPath = BaseDepotPathPrefix + (bPreview ? "/UnstableRelease.zip" : "/Release.zip");
 				List<ChangesRecord> Changes = await Perforce.GetChangesAsync(ChangesOptions.None, 1, ChangeStatus.Submitted, SyncPath, CancellationToken);
 				if (Changes.Count == 0)
 				{
-					SyncPath = BaseDepotPathPrefix + (bUnstable ? "/UnstableRelease/..." : "/Release/...");
+					SyncPath = BaseDepotPathPrefix + (bPreview ? "/UnstableRelease/..." : "/Release/...");
 					Changes = await Perforce.GetChangesAsync(ChangesOptions.None, 1, ChangeStatus.Submitted, SyncPath, CancellationToken);
 					if (Changes.Count == 0)
 					{
@@ -199,7 +199,7 @@ namespace UnrealGameSyncLauncher
                     }
                 }
 
-				StringBuilder NewCommandLine = new StringBuilder(String.Format("-updatepath=\"{0}@>{1}\" -updatespawn=\"{2}\"{3}", SyncPath, RequiredChangeNumber, OriginalExecutable, bUnstable ? " -unstable" : ""));
+				StringBuilder NewCommandLine = new StringBuilder(String.Format("-updatepath=\"{0}@>{1}\" -updatespawn=\"{2}\"{3}", SyncPath, RequiredChangeNumber, OriginalExecutable, bPreview ? " -unstable" : ""));
 				foreach (string Arg in Args)
 				{
 					NewCommandLine.AppendFormat(" {0}", QuoteArgument(Arg));
