Missing file.

diff --git a/Engine/Source/Programs/UnrealGameSync/UnrealGameSyncLauncher/Forms/SettingsWindow.cs b/Engine/Source/Programs/UnrealGameSync/UnrealGameSyncLauncher/Forms/SettingsWindow.cs
--- a/Engine/Source/Programs/UnrealGameSync/UnrealGameSyncLauncher/Forms/SettingsWindow.cs
+++ b/Engine/Source/Programs/UnrealGameSync/UnrealGameSyncLauncher/Forms/SettingsWindow.cs
@@ -85,7 +85,9 @@ namespace UnrealGameSyncLauncher
 				DepotPath = null;
 			}
 
-			GlobalPerforceSettings.SaveGlobalPerforceSettings(ServerAndPort, UserName, DepotPath);
+			bool bPreview = UsePreviewBuildCheckBox.Checked;
+
+			GlobalPerforceSettings.SaveGlobalPerforceSettings(ServerAndPort, UserName, DepotPath, bPreview);
 
 			PerforceSettings PerforceSettings = new PerforceSettings(PerforceSettings.Default);
 			if (!String.IsNullOrEmpty(ServerAndPort))
@@ -102,12 +104,12 @@ namespace UnrealGameSyncLauncher
 			CaptureLogger Logger = new CaptureLogger();
 
 			// Create the task for connecting to this server
-			ModalTask? Task = PerforceModalTask.Execute(this, "Updating", "Checking for updates, please wait...", PerforceSettings, (p, c) => SyncAndRun(p, DepotPath, UsePreviewBuildCheckBox.Checked, Logger, c), Logger);
+			ModalTask? Task = PerforceModalTask.Execute(this, "Updating", "Checking for updates, please wait...", PerforceSettings, (p, c) => SyncAndRun(p, DepotPath, bPreview, Logger, c), Logger);
 			if (Task != null)
 			{
 				if(Task.Succeeded)
 				{
-					GlobalPerforceSettings.SaveGlobalPerforceSettings(ServerAndPort, UserName, DepotPath);
+					GlobalPerforceSettings.SaveGlobalPerforceSettings(ServerAndPort, UserName, DepotPath, bPreview);
 					DialogResult = DialogResult.OK;
 					Close();
 				}
