UGS: Reformat to standard C# style conventions.

diff --git a/Engine/Source/Programs/UnrealGameSync/UnrealGameSync/ActivationListener.cs b/Engine/Source/Programs/UnrealGameSync/UnrealGameSync/ActivationListener.cs
--- a/Engine/Source/Programs/UnrealGameSync/UnrealGameSync/ActivationListener.cs
+++ b/Engine/Source/Programs/UnrealGameSync/UnrealGameSync/ActivationListener.cs
@@ -11,36 +11,36 @@ namespace UnrealGameSync
 {
 	class ActivationListener : IDisposable
 	{
-		Thread? BackgroundThread;
-		EventWaitHandle ActivateEventHandle;
-		EventWaitHandle QuitEventHandle;
+		Thread? _backgroundThread;
+		EventWaitHandle _activateEventHandle;
+		EventWaitHandle _quitEventHandle;
 
 		public event Action? OnActivate;
 
-		public ActivationListener(EventWaitHandle InActivateEventHandle)
+		public ActivationListener(EventWaitHandle inActivateEventHandle)
 		{
-			ActivateEventHandle = InActivateEventHandle;
-			QuitEventHandle = new EventWaitHandle(false, EventResetMode.ManualReset);
+			_activateEventHandle = inActivateEventHandle;
+			_quitEventHandle = new EventWaitHandle(false, EventResetMode.ManualReset);
 		}
 
 		public void Start()
 		{
-			if(BackgroundThread == null)
+			if(_backgroundThread == null)
 			{
-				BackgroundThread = new Thread(x => ThreadProc());
-				BackgroundThread.IsBackground = true;
-				BackgroundThread.Start();
+				_backgroundThread = new Thread(x => ThreadProc());
+				_backgroundThread.IsBackground = true;
+				_backgroundThread.Start();
 			}
 		}
 
 		public void Stop()
 		{
-			if(BackgroundThread != null)
+			if(_backgroundThread != null)
 			{
-				QuitEventHandle.Set();
+				_quitEventHandle.Set();
 
-				BackgroundThread.Join();
-				BackgroundThread = null;
+				_backgroundThread.Join();
+				_backgroundThread = null;
 			}
 		}
 
@@ -48,16 +48,16 @@ namespace UnrealGameSync
 		{
 			Stop();
 
-			QuitEventHandle.Dispose();
-			ActivateEventHandle.Dispose();
+			_quitEventHandle.Dispose();
+			_activateEventHandle.Dispose();
 		}
 
 		void ThreadProc()
 		{
 			for(;;)
 			{
-				int Index = EventWaitHandle.WaitAny(new WaitHandle[]{ ActivateEventHandle, QuitEventHandle }, Timeout.Infinite);
-				if(Index == 0)
+				int index = EventWaitHandle.WaitAny(new WaitHandle[]{ _activateEventHandle, _quitEventHandle }, Timeout.Infinite);
+				if(index == 0)
 				{
 					OnActivate?.Invoke();
 				}
diff --git a/Engine/Source/Programs/UnrealGameSync/UnrealGameSync/AsyncDisposer.cs b/Engine/Source/Programs/UnrealGameSync/UnrealGameSync/AsyncDisposer.cs
--- a/Engine/Source/Programs/UnrealGameSync/UnrealGameSync/AsyncDisposer.cs
+++ b/Engine/Source/Programs/UnrealGameSync/UnrealGameSync/AsyncDisposer.cs
@@ -11,50 +11,50 @@ namespace UnrealGameSync
 {
 	interface IAsyncDisposer
 	{
-		void Add(Task Task);
+		void Add(Task task);
 	}
 
 	internal class AsyncDisposer : IAsyncDisposer, IAsyncDisposable
 	{
-		object LockObject = new object();
-		List<Task> Tasks = new List<Task>();
-		ILogger Logger;
+		object _lockObject = new object();
+		List<Task> _tasks = new List<Task>();
+		ILogger _logger;
 
-		public AsyncDisposer(ILogger<AsyncDisposer> Logger)
+		public AsyncDisposer(ILogger<AsyncDisposer> logger)
 		{
-			this.Logger = Logger;
+			this._logger = logger;
 		}
 
-		public void Add(Task Task)
+		public void Add(Task task)
 		{
-			Task ContinuationTask = Task.ContinueWith(Remove);
-			lock (LockObject)
+			Task continuationTask = task.ContinueWith(Remove);
+			lock (_lockObject)
 			{
-				Tasks.Add(ContinuationTask);
+				_tasks.Add(continuationTask);
 			}
 		}
 
-		private void Remove(Task Task)
+		private void Remove(Task task)
 		{
-			if (Task.IsFaulted)
+			if (task.IsFaulted)
 			{
-				Logger.LogError(Task.Exception, "Exception while disposing task");
+				_logger.LogError(task.Exception, "Exception while disposing task");
 			}
 		}
 
 		public async ValueTask DisposeAsync()
 		{
-			List<Task> TasksCopy;
-			lock (LockObject)
+			List<Task> tasksCopy;
+			lock (_lockObject)
 			{
-				TasksCopy = new List<Task>(Tasks);
+				tasksCopy = new List<Task>(_tasks);
 			}
 
-			Task WaitTask = Task.WhenAll(TasksCopy);
-			while (!WaitTask.IsCompleted)
+			Task waitTask = Task.WhenAll(tasksCopy);
+			while (!waitTask.IsCompleted)
 			{
-				Logger.LogInformation("Waiting for {NumTasks} tasks to complete", TasksCopy.Count);
-				await Task.WhenAny(WaitTask, Task.Delay(TimeSpan.FromSeconds(5.0)));
+				_logger.LogInformation("Waiting for {NumTasks} tasks to complete", tasksCopy.Count);
+				await Task.WhenAny(waitTask, Task.Delay(TimeSpan.FromSeconds(5.0)));
 			}
 		}
 	}
diff --git a/Engine/Source/Programs/UnrealGameSync/UnrealGameSync/Automation/P4Automation.cs b/Engine/Source/Programs/UnrealGameSync/UnrealGameSync/Automation/P4Automation.cs
--- a/Engine/Source/Programs/UnrealGameSync/UnrealGameSync/Automation/P4Automation.cs
+++ b/Engine/Source/Programs/UnrealGameSync/UnrealGameSync/Automation/P4Automation.cs
@@ -17,60 +17,60 @@ namespace UnrealGameSync
 		public static IPerforceSettings GetConnectionSettings()
 		{
 			// Read the settings
-			string? ServerAndPort = null;
-			string? UserName = null;
-			string? DepotPathSettings = null;
-			bool bPreview = false;
+			string? serverAndPort = null;
+			string? userName = null;
+			string? depotPathSettings = null;
+			bool preview = false;
 
-			GlobalPerforceSettings.ReadGlobalPerforceSettings(ref ServerAndPort, ref UserName, ref DepotPathSettings, ref bPreview);
+			GlobalPerforceSettings.ReadGlobalPerforceSettings(ref serverAndPort, ref userName, ref depotPathSettings, ref preview);
 
-			return Utility.OverridePerforceSettings(PerforceSettings.Default, ServerAndPort, UserName);
+			return Utility.OverridePerforceSettings(PerforceSettings.Default, serverAndPort, userName);
 
 		}
 
-		public static Task<string> PrintToTempFile(IPerforceConnection? Connection, string DepotPath, ILogger Logger)
+		public static Task<string> PrintToTempFile(IPerforceConnection? connection, string depotPath, ILogger logger)
 		{
-			return PrintToTempFileAsync(Connection, DepotPath, CancellationToken.None, Logger);
+			return PrintToTempFileAsync(connection, depotPath, CancellationToken.None, logger);
 		}
 
-		public static async Task<string> PrintToTempFileAsync(IPerforceConnection? Connection, string DepotPath, CancellationToken CancellationToken, ILogger Logger)
+		public static async Task<string> PrintToTempFileAsync(IPerforceConnection? connection, string depotPath, CancellationToken cancellationToken, ILogger logger)
 		{
-			bool bCreateNewConnection = (Connection == null);
+			bool createNewConnection = (connection == null);
 
 			try
 			{
-				if (Connection == null)
+				if (connection == null)
 				{
-					IPerforceSettings Settings = GetConnectionSettings();
-					Connection = await PerforceConnection.CreateAsync(Logger);
+					IPerforceSettings settings = GetConnectionSettings();
+					connection = await PerforceConnection.CreateAsync(logger);
 				}
 
-				string DepotFileName = Path.GetFileName(DepotPath);
+				string depotFileName = Path.GetFileName(depotPath);
 
 				// Reorder CL and extension
-				int Index = DepotFileName.IndexOf('@');
-				if (Index == -1)
+				int index = depotFileName.IndexOf('@');
+				if (index == -1)
 				{
-					DepotFileName += "@Latest";
-					Index = DepotFileName.IndexOf('@');
+					depotFileName += "@Latest";
+					index = depotFileName.IndexOf('@');
 				}
 
-				string CL = DepotFileName.Substring(Index + 1);
-				string FileName = DepotFileName.Substring(0, Index);
-				string TempFileName = string.Format("{0}@{1}{2}", Path.GetFileNameWithoutExtension(FileName), CL, Path.GetExtension(FileName));
+				string cl = depotFileName.Substring(index + 1);
+				string fileName = depotFileName.Substring(0, index);
+				string tempFileName = string.Format("{0}@{1}{2}", Path.GetFileNameWithoutExtension(fileName), cl, Path.GetExtension(fileName));
 
-				TempFileName = Path.Combine(Path.GetTempPath(), TempFileName);
-				await Connection.PrintAsync(TempFileName, DepotFileName, CancellationToken);
+				tempFileName = Path.Combine(Path.GetTempPath(), tempFileName);
+				await connection.PrintAsync(tempFileName, depotFileName, cancellationToken);
 
-				return TempFileName;
+				return tempFileName;
 			}
 			finally
 			{
 				// If we created a new connection, tear it down now.
-				if (bCreateNewConnection)
+				if (createNewConnection)
 				{
-					Connection?.Dispose();
-					Connection = null;
+					connection?.Dispose();
+					connection = null;
 				}
 			}
 		}
diff --git a/Engine/Source/Programs/UnrealGameSync/UnrealGameSync/Automation/VisualStudioAutomation.cs b/Engine/Source/Programs/UnrealGameSync/UnrealGameSync/Automation/VisualStudioAutomation.cs
--- a/Engine/Source/Programs/UnrealGameSync/UnrealGameSync/Automation/VisualStudioAutomation.cs
+++ b/Engine/Source/Programs/UnrealGameSync/UnrealGameSync/Automation/VisualStudioAutomation.cs
@@ -16,34 +16,34 @@ namespace UnrealGameSync
 
 	static class VisualStudioAutomation
 	{
-		public static bool OpenFile(string FileName, out string? ErrorMessage, int Line = -1)
+		public static bool OpenFile(string fileName, out string? errorMessage, int line = -1)
 		{
-			ErrorMessage = null;
+			errorMessage = null;
 
 			// first try to open via DTE
-			DTE? DTE = VisualStudioAccessor.GetDTE();
-			if (DTE != null)
+			DTE? dte = VisualStudioAccessor.GetDte();
+			if (dte != null)
 			{
-				DTE.ItemOperations.OpenFile(FileName);
-				DTE.MainWindow.Activate();
+				dte.ItemOperations.OpenFile(fileName);
+				dte.MainWindow.Activate();
 
-				if (Line != -1)
+				if (line != -1)
 				{
-					(DTE.ActiveDocument.Selection as TextSelection)?.GotoLine(Line);
+					(dte.ActiveDocument.Selection as TextSelection)?.GotoLine(line);
 				}
 
 				return true;
 			}
 
 			// unable to get DTE connection, so launch nesw VS instance
-			string Arguments = string.Format("\"{0}\"", FileName);
-			if (Line != -1)
+			string arguments = string.Format("\"{0}\"", fileName);
+			if (line != -1)
 			{
-				Arguments += string.Format(" /command \"edit.goto {0}\"", Line);
+				arguments += string.Format(" /command \"edit.goto {0}\"", line);
 			}
 
 			// Launch new visual studio instance
-			if (!VisualStudioAccessor.LaunchVisualStudio(Arguments, out ErrorMessage))
+			if (!VisualStudioAccessor.LaunchVisualStudio(arguments, out errorMessage))
 			{
 				return false;
 			}
@@ -57,26 +57,26 @@ namespace UnrealGameSync
 	/// </summary>
 	static class VisualStudioAccessor
 	{
-		public static bool LaunchVisualStudio(string Arguments, out string ErrorMessage)
+		public static bool LaunchVisualStudio(string arguments, out string errorMessage)
 		{
-			ErrorMessage = "";
-			VisualStudioInstallation? Install = VisualStudioInstallations.GetPreferredInstallation();
+			errorMessage = "";
+			VisualStudioInstallation? install = VisualStudioInstallations.GetPreferredInstallation();
 
-			if (Install == null)
+			if (install == null)
 			{
-				ErrorMessage = string.Format("Unable to get Visual Studio installation");
+				errorMessage = string.Format("Unable to get Visual Studio installation");
 				return false;
 			}
 			try
 			{
 
-				System.Diagnostics.Process VSProcess = new System.Diagnostics.Process { StartInfo = new ProcessStartInfo(Install.DevEnvPath, Arguments) };
-				VSProcess.Start();
-				VSProcess.WaitForInputIdle();
+				System.Diagnostics.Process vsProcess = new System.Diagnostics.Process { StartInfo = new ProcessStartInfo(install.DevEnvPath, arguments) };
+				vsProcess.Start();
+				vsProcess.WaitForInputIdle();
 			}
-			catch (Exception Ex)
+			catch (Exception ex)
 			{
-				ErrorMessage = Ex.Message;
+				errorMessage = ex.Message;
 				return false;
 			}
 
@@ -85,38 +85,38 @@ namespace UnrealGameSync
 		}
 
 		[STAThread]
-		public static DTE? GetDTE()
+		public static DTE? GetDte()
 		{
-			IRunningObjectTable Table;
-			if (Succeeded(GetRunningObjectTable(0, out Table)) && Table != null)
+			IRunningObjectTable table;
+			if (Succeeded(GetRunningObjectTable(0, out table)) && table != null)
 			{
-				IEnumMoniker MonikersTable;
-				Table.EnumRunning(out MonikersTable);
+				IEnumMoniker monikersTable;
+				table.EnumRunning(out monikersTable);
 
-				if (MonikersTable == null)
+				if (monikersTable == null)
 				{
 					return null;
 				}
 
-				MonikersTable.Reset();
+				monikersTable.Reset();
 
 				// Look for all visual studio instances in the ROT
-				IMoniker[] Monikers = new IMoniker[1];
-				while (MonikersTable.Next(1, Monikers, IntPtr.Zero) == 0)
+				IMoniker[] monikers = new IMoniker[1];
+				while (monikersTable.Next(1, monikers, IntPtr.Zero) == 0)
 				{
-					IBindCtx BindContext;
-					string OutDisplayName;
-					IMoniker CurrentMoniker = Monikers[0];
+					IBindCtx bindContext;
+					string outDisplayName;
+					IMoniker currentMoniker = monikers[0];
 
-					if (!Succeeded(CreateBindCtx(0, out BindContext)))
+					if (!Succeeded(CreateBindCtx(0, out bindContext)))
 					{
 						continue;
 					}
 
 					try
 					{
-						CurrentMoniker.GetDisplayName(BindContext, null, out OutDisplayName);
-						if (string.IsNullOrEmpty(OutDisplayName) || !IsVisualStudioDTEMoniker(OutDisplayName))
+						currentMoniker.GetDisplayName(bindContext, null, out outDisplayName);
+						if (string.IsNullOrEmpty(outDisplayName) || !IsVisualStudioDteMoniker(outDisplayName))
 						{
 							continue;
 						}
@@ -127,13 +127,13 @@ namespace UnrealGameSync
 						continue;
 					}
 
-					object ComObject;
-					if (!Succeeded(Table.GetObject(CurrentMoniker, out ComObject)))
+					object comObject;
+					if (!Succeeded(table.GetObject(currentMoniker, out comObject)))
 					{
 						continue;
 					}
 
-					return ComObject as DTE;
+					return comObject as DTE;
 				}
 			}
 
@@ -141,14 +141,14 @@ namespace UnrealGameSync
 
 		}
 
-		static bool IsVisualStudioDTEMoniker(string InName)
+		static bool IsVisualStudioDteMoniker(string inName)
 		{
-			VisualStudioInstallation[] Installs = VisualStudioInstallations.Installs;
+			VisualStudioInstallation[] installs = VisualStudioInstallations.Installs;
 
-			for (int Idx = 0; Idx < Installs.Length; Idx++)
+			for (int idx = 0; idx < installs.Length; idx++)
 			{
-				string? Moniker = Installs[Idx].ROTMoniker;
-				if (Moniker != null && InName.StartsWith(Moniker))
+				string? moniker = installs[idx].RotMoniker;
+				if (moniker != null && inName.StartsWith(moniker))
 				{
 					return true;
 				}
@@ -157,17 +157,17 @@ namespace UnrealGameSync
 			return false;
 		}
 
-		static bool Succeeded(int Result)
+		static bool Succeeded(int result)
 		{
-			return Result >= 0;
+			return result >= 0;
 		}
 
 
 		[DllImport("ole32.dll")]
-		public static extern int CreateBindCtx(int Reserved, out IBindCtx BindCtx);
+		public static extern int CreateBindCtx(int reserved, out IBindCtx bindCtx);
 
 		[DllImport("ole32.dll")]
-		public static extern int GetRunningObjectTable(int Reserved, out IRunningObjectTable ROT);
+		public static extern int GetRunningObjectTable(int reserved, out IRunningObjectTable rot);
 
 	}
 
@@ -192,7 +192,7 @@ namespace UnrealGameSync
 		/// <summary>
 		/// Running Object Table moniker for this installation
 		/// </summary>
-		public string? ROTMoniker;
+		public string? RotMoniker;
 
 	}
 
@@ -202,7 +202,7 @@ namespace UnrealGameSync
 	static class VisualStudioInstallations
 	{
 
-		public static VisualStudioInstallation? GetPreferredInstallation(int MajorVersion = 0)
+		public static VisualStudioInstallation? GetPreferredInstallation(int majorVersion = 0)
 		{
 
 			if (CachedInstalls.Count == 0)
@@ -210,14 +210,14 @@ namespace UnrealGameSync
 				return null;
 			}
 
-			if (MajorVersion == 0)
+			if (majorVersion == 0)
 			{
 				return CachedInstalls.First();
 			}
 
-			VisualStudioInstallation Installation = CachedInstalls.FirstOrDefault(Install => { return Install.MajorVersion == MajorVersion; });
+			VisualStudioInstallation installation = CachedInstalls.FirstOrDefault(install => { return install.MajorVersion == majorVersion; });
 
-			return Installation;
+			return installation;
 
 		}
 
@@ -237,59 +237,59 @@ namespace UnrealGameSync
 
 			try
 			{
-				SetupConfiguration Setup = new SetupConfiguration();
-				IEnumSetupInstances Enumerator = Setup.EnumAllInstances();
+				SetupConfiguration setup = new SetupConfiguration();
+				IEnumSetupInstances enumerator = setup.EnumAllInstances();
 
-				ISetupInstance[] Instances = new ISetupInstance[1];
+				ISetupInstance[] instances = new ISetupInstance[1];
 				for (; ; )
 				{
-					int NumFetched;
-					Enumerator.Next(1, Instances, out NumFetched);
+					int numFetched;
+					enumerator.Next(1, instances, out numFetched);
 
-					if (NumFetched == 0)
+					if (numFetched == 0)
 					{
 						break;
 					}
 
-					ISetupInstance2 Instance = (ISetupInstance2)Instances[0];
-					if ((Instance.GetState() & InstanceState.Local) == InstanceState.Local)
+					ISetupInstance2 instance = (ISetupInstance2)instances[0];
+					if ((instance.GetState() & InstanceState.Local) == InstanceState.Local)
 					{
-						string VersionString = Instance.GetInstallationVersion();
-						string[] Components = VersionString.Split('.');
+						string versionString = instance.GetInstallationVersion();
+						string[] components = versionString.Split('.');
 
-						if (Components.Length == 0)
+						if (components.Length == 0)
 						{
 							continue;
 						}
 
-						int MajorVersion;
-						string InstallationPath = Instance.GetInstallationPath();
-						string DevEnvPath = Path.Combine(InstallationPath, "Common7\\IDE\\devenv.exe");
+						int majorVersion;
+						string installationPath = instance.GetInstallationPath();
+						string devEnvPath = Path.Combine(installationPath, "Common7\\IDE\\devenv.exe");
 
-						if (!int.TryParse(Components[0], out MajorVersion) || (MajorVersion != 15 && MajorVersion != 16))
+						if (!int.TryParse(components[0], out majorVersion) || (majorVersion != 15 && majorVersion != 16))
 						{
 							continue;
 						}
 
 
-						if (!File.Exists(DevEnvPath))
+						if (!File.Exists(devEnvPath))
 						{
 							continue;
 						}
 
-						VisualStudioInstallation Installation = new VisualStudioInstallation() { BaseDir = InstallationPath, DevEnvPath = DevEnvPath, MajorVersion = MajorVersion, ROTMoniker = string.Format("!VisualStudio.DTE.{0}.0", MajorVersion) };
+						VisualStudioInstallation installation = new VisualStudioInstallation() { BaseDir = installationPath, DevEnvPath = devEnvPath, MajorVersion = majorVersion, RotMoniker = string.Format("!VisualStudio.DTE.{0}.0", majorVersion) };
 
-						CachedInstalls.Add(Installation);
+						CachedInstalls.Add(installation);
 					}
 				}
 			}
-			catch (Exception Ex)
+			catch (Exception ex)
 			{
-				MessageBox.Show(string.Format("Exception while finding Visual Studio installations {0}", Ex.Message));
+				MessageBox.Show(string.Format("Exception while finding Visual Studio installations {0}", ex.Message));
 			}
 
 			// prefer newer versions
-			CachedInstalls.Sort((A, B) => { return -A.MajorVersion.CompareTo(B.MajorVersion); });
+			CachedInstalls.Sort((a, b) => { return -a.MajorVersion.CompareTo(b.MajorVersion); });
 
 			return CachedInstalls;
 		}
diff --git a/Engine/Source/Programs/UnrealGameSync/UnrealGameSync/AutomationServer.cs b/Engine/Source/Programs/UnrealGameSync/UnrealGameSync/AutomationServer.cs
--- a/Engine/Source/Programs/UnrealGameSync/UnrealGameSync/AutomationServer.cs
+++ b/Engine/Source/Programs/UnrealGameSync/UnrealGameSync/AutomationServer.cs
@@ -27,30 +27,30 @@ namespace UnrealGameSync
 		public AutomationRequestType Type;
 		public byte[] Data;
 
-		public AutomationRequestInput(AutomationRequestType Type, byte[] Data)
+		public AutomationRequestInput(AutomationRequestType type, byte[] data)
 		{
-			this.Type = Type;
-			this.Data = Data;
+			this.Type = type;
+			this.Data = data;
 		}
 
-		public static AutomationRequestInput Read(Stream InputStream)
+		public static AutomationRequestInput Read(Stream inputStream)
 		{
-			BinaryReader Reader = new BinaryReader(InputStream);
+			BinaryReader reader = new BinaryReader(inputStream);
 			
-			int Type = Reader.ReadInt32();
-			int InputSize = Reader.ReadInt32();
-			byte[] Input = Reader.ReadBytes(InputSize);
+			int type = reader.ReadInt32();
+			int inputSize = reader.ReadInt32();
+			byte[] input = reader.ReadBytes(inputSize);
 
-			return new AutomationRequestInput((AutomationRequestType)Type, Input);
+			return new AutomationRequestInput((AutomationRequestType)type, input);
 		}
 
-		public void Write(Stream OutputStream)
+		public void Write(Stream outputStream)
 		{
-			BinaryWriter Writer = new BinaryWriter(OutputStream);
+			BinaryWriter writer = new BinaryWriter(outputStream);
 
-			Writer.Write((int)Type);
-			Writer.Write(Data.Length);
-			Writer.Write(Data);
+			writer.Write((int)Type);
+			writer.Write(Data.Length);
+			writer.Write(Data);
 		}
 	}
 
@@ -69,36 +69,36 @@ namespace UnrealGameSync
 		public AutomationRequestResult Result;
 		public byte[] Data;
 
-		public AutomationRequestOutput(AutomationRequestResult Result)
+		public AutomationRequestOutput(AutomationRequestResult result)
 		{
-			this.Result = Result;
+			this.Result = result;
 			this.Data = new byte[0];
 		}
 
-		public AutomationRequestOutput(AutomationRequestResult Result, byte[] Data)
+		public AutomationRequestOutput(AutomationRequestResult result, byte[] data)
 		{
-			this.Result = Result;
-			this.Data = Data;
+			this.Result = result;
+			this.Data = data;
 		}
 
-		public static AutomationRequestOutput Read(Stream InputStream)
+		public static AutomationRequestOutput Read(Stream inputStream)
 		{
-			using(BinaryReader Reader = new BinaryReader(InputStream))
+			using(BinaryReader reader = new BinaryReader(inputStream))
 			{
-				AutomationRequestResult Result = (AutomationRequestResult)Reader.ReadInt32();
-				int DataSize = Reader.ReadInt32();
-				byte[] Data = Reader.ReadBytes(DataSize);
-				return new AutomationRequestOutput(Result, Data);
+				AutomationRequestResult result = (AutomationRequestResult)reader.ReadInt32();
+				int dataSize = reader.ReadInt32();
+				byte[] data = reader.ReadBytes(dataSize);
+				return new AutomationRequestOutput(result, data);
 			}
 		}
 
-		public void Write(Stream OutputStream)
+		public void Write(Stream outputStream)
 		{
-			using(BinaryWriter Writer = new BinaryWriter(OutputStream))
+			using(BinaryWriter writer = new BinaryWriter(outputStream))
 			{
-				Writer.Write((int)Result);
-				Writer.Write(Data.Length);
-				Writer.Write(Data);
+				writer.Write((int)Result);
+				writer.Write(Data.Length);
+				writer.Write(Data);
 			}
 		}
 	}
@@ -109,15 +109,15 @@ namespace UnrealGameSync
 		public AutomationRequestOutput? Output;
 		public ManualResetEventSlim Complete;
 
-		public AutomationRequest(AutomationRequestInput Input)
+		public AutomationRequest(AutomationRequestInput input)
 		{
-			this.Input = Input;
+			this.Input = input;
 			this.Complete = new ManualResetEventSlim(false);
 		}
 
-		public void SetOutput(AutomationRequestOutput Output)
+		public void SetOutput(AutomationRequestOutput output)
 		{
-			this.Output = Output;
+			this.Output = output;
 			Complete?.Set();
 		}
 
@@ -129,83 +129,83 @@ namespace UnrealGameSync
 
 	class AutomationServer : IDisposable
 	{
-		TcpListener? Listener;
+		TcpListener? _listener;
 		public const int DefaultPortNumber = 30422;
 
-		NamedPipeServerStream? IPCStream;
-		static UnicodeEncoding StreamEncoding = new UnicodeEncoding();
-		const string UGSChannel = @"\.\pipe\UGSChannel";
+		NamedPipeServerStream? _ipcStream;
+		static UnicodeEncoding _streamEncoding = new UnicodeEncoding();
+		const string UgsChannel = @"\.\pipe\UGSChannel";
 
-		Thread? UriThread;
-		Thread? TcpThread;
+		Thread? _uriThread;
+		Thread? _tcpThread;
 
-		EventWaitHandle ShutdownEvent;		
-		Action<AutomationRequest> PostRequest;
+		EventWaitHandle _shutdownEvent;		
+		Action<AutomationRequest> _postRequest;
 
-		bool bDisposing;
-		ILogger Logger;
-		string? CommandLineUri;
+		bool _disposing;
+		ILogger _logger;
+		string? _commandLineUri;
 
-		public AutomationServer(Action<AutomationRequest> PostRequest, string? Uri, ILogger<AutomationServer> Logger)
+		public AutomationServer(Action<AutomationRequest> postRequest, string? uri, ILogger<AutomationServer> logger)
 		{
-			ShutdownEvent = new ManualResetEvent(false);
-			this.PostRequest = PostRequest;
-			this.CommandLineUri = Uri;
-			this.Logger = Logger;
+			_shutdownEvent = new ManualResetEvent(false);
+			this._postRequest = postRequest;
+			this._commandLineUri = uri;
+			this._logger = logger;
 
 			try
 			{
 				// IPC named pipe
-				IPCStream = new NamedPipeServerStream(UGSChannel, PipeDirection.In, 1, PipeTransmissionMode.Message, PipeOptions.Asynchronous);
+				_ipcStream = new NamedPipeServerStream(UgsChannel, PipeDirection.In, 1, PipeTransmissionMode.Message, PipeOptions.Asynchronous);
 
 				// TCP listener setup
-				int PortNumber = GetPortNumber();
-				if (PortNumber > 0)
+				int portNumber = GetPortNumber();
+				if (portNumber > 0)
 				{
 					try
 					{
-						Listener = new TcpListener(IPAddress.Loopback, PortNumber);
-						Listener.Start();
-						TcpThread = new Thread(() => RunTcp());
-						TcpThread.IsBackground = true;
-						TcpThread.Start();
+						_listener = new TcpListener(IPAddress.Loopback, portNumber);
+						_listener.Start();
+						_tcpThread = new Thread(() => RunTcp());
+						_tcpThread.IsBackground = true;
+						_tcpThread.Start();
 
 					}
-					catch (Exception Ex)
+					catch (Exception ex)
 					{
-						Listener = null;
-						Logger.LogError(Ex, "Unable to start automation server tcp listener");
+						_listener = null;
+						logger.LogError(ex, "Unable to start automation server tcp listener");
 					}
 				}
 
-				UriThread = new Thread(() => RunUri());
-				UriThread.IsBackground = true;
-				UriThread.Start();
+				_uriThread = new Thread(() => RunUri());
+				_uriThread.IsBackground = true;
+				_uriThread.Start();
 			}
-			catch (Exception Ex)
+			catch (Exception ex)
 			{
-				Logger.LogError(Ex, "Unable to start automation server");
+				logger.LogError(ex, "Unable to start automation server");
 			}
 		}
 
-		public static void SetPortNumber(int PortNumber)
+		public static void SetPortNumber(int portNumber)
 		{
-			if(PortNumber <= 0)
+			if(portNumber <= 0)
 			{
 				GlobalPerforceSettings.DeleteRegistryKey(Registry.CurrentUser, "Software\\Epic Games\\UnrealGameSync", "AutomationPort");
 			}
 			else
 			{
-				Registry.SetValue("HKEY_CURRENT_USER\\Software\\Epic Games\\UnrealGameSync", "AutomationPort", PortNumber);
+				Registry.SetValue("HKEY_CURRENT_USER\\Software\\Epic Games\\UnrealGameSync", "AutomationPort", portNumber);
 			}
 		}
 
 		public static int GetPortNumber()
 		{
-			object? PortValue = Registry.GetValue("HKEY_CURRENT_USER\\Software\\Epic Games\\UnrealGameSync", "AutomationPort", null);
-			if (PortValue != null && PortValue is int)
+			object? portValue = Registry.GetValue("HKEY_CURRENT_USER\\Software\\Epic Games\\UnrealGameSync", "AutomationPort", null);
+			if (portValue != null && portValue is int)
 			{
-				return (int)PortValue;
+				return (int)portValue;
 			}
 			else
 			{
@@ -216,9 +216,9 @@ namespace UnrealGameSync
 		void RunUri()
 		{
 			// Handle main process command line URI request
-			if (!string.IsNullOrEmpty(CommandLineUri))
+			if (!string.IsNullOrEmpty(_commandLineUri))
 			{
-				HandleUri(CommandLineUri);
+				HandleUri(_commandLineUri);
 			}
 
 			for (;;)
@@ -226,42 +226,42 @@ namespace UnrealGameSync
 				try
 				{
 
-					IAsyncResult IPCResult = IPCStream!.BeginWaitForConnection(null, null);
+					IAsyncResult ipcResult = _ipcStream!.BeginWaitForConnection(null, null);
 
-					int WaitResult = WaitHandle.WaitAny(new WaitHandle[] { ShutdownEvent, IPCResult.AsyncWaitHandle });
+					int waitResult = WaitHandle.WaitAny(new WaitHandle[] { _shutdownEvent, ipcResult.AsyncWaitHandle });
 
 					// Shutting down
-					if (WaitResult == 0)
+					if (waitResult == 0)
 					{
 						break;
 					}
 
 					try
 					{
-						IPCStream.EndWaitForConnection(IPCResult);
+						_ipcStream.EndWaitForConnection(ipcResult);
 
-						Logger.LogInformation("Accepted Uri connection");
+						_logger.LogInformation("Accepted Uri connection");
 
 						// Read URI
-						string Uri = ReadString(IPCStream);
+						string uri = ReadString(_ipcStream);
 
-						Logger.LogInformation("Received Uri: {Uri}", Uri);
+						_logger.LogInformation("Received Uri: {Uri}", uri);
 
-						IPCStream.Disconnect();
+						_ipcStream.Disconnect();
 
-						HandleUri(Uri);
+						HandleUri(uri);
 
 					}
-					catch (Exception Ex)
+					catch (Exception ex)
 					{
-						Logger.LogError(Ex, "Error during automation connection");
+						_logger.LogError(ex, "Error during automation connection");
 					}
 				}
-				catch (Exception Ex)
+				catch (Exception ex)
 				{
-					if (!bDisposing)
+					if (!_disposing)
 					{
-						Logger.LogError(Ex, "Error during automation connection");
+						_logger.LogError(ex, "Error during automation connection");
 					}
 				}
 			}
@@ -276,77 +276,77 @@ namespace UnrealGameSync
 				try
 				{
 
-					IAsyncResult TCPResult = Listener!.BeginAcceptTcpClient(null, null);
+					IAsyncResult tcpResult = _listener!.BeginAcceptTcpClient(null, null);
 
-					int WaitResult = WaitHandle.WaitAny(new WaitHandle[] { ShutdownEvent, TCPResult.AsyncWaitHandle });
+					int waitResult = WaitHandle.WaitAny(new WaitHandle[] { _shutdownEvent, tcpResult.AsyncWaitHandle });
 
 					// Shutting down
-					if (WaitResult == 0)
+					if (waitResult == 0)
 					{
 						break;
 					}
 
 					try
 					{
-						TcpClient Client = Listener.EndAcceptTcpClient(TCPResult);
+						TcpClient client = _listener.EndAcceptTcpClient(tcpResult);
 
-						Logger.LogInformation("Accepted connection from {Remote}", Client.Client.RemoteEndPoint);
+						_logger.LogInformation("Accepted connection from {Remote}", client.Client.RemoteEndPoint);
 
-						NetworkStream Stream = Client.GetStream();
+						NetworkStream stream = client.GetStream();
 
-						AutomationRequestInput Input = AutomationRequestInput.Read(Stream);
-						Logger.LogInformation("Received input: {Type} (+{NumBytes} bytes)", Input.Type, Input.Data.Length);
+						AutomationRequestInput input = AutomationRequestInput.Read(stream);
+						_logger.LogInformation("Received input: {Type} (+{NumBytes} bytes)", input.Type, input.Data.Length);
 
-						AutomationRequestOutput Output;
-						using (AutomationRequest Request = new AutomationRequest(Input))
+						AutomationRequestOutput output;
+						using (AutomationRequest request = new AutomationRequest(input))
 						{
-							PostRequest(Request);
-							Request.Complete.Wait();
-							Output = Request.Output!;
+							_postRequest(request);
+							request.Complete.Wait();
+							output = request.Output!;
 						}
 
-						Output.Write(Stream);
-						Logger.LogInformation("Sent output: {Result} (+{NumBytes} bytes)", Output.Result, Output.Data.Length);
+						output.Write(stream);
+						_logger.LogInformation("Sent output: {Result} (+{NumBytes} bytes)", output.Result, output.Data.Length);
 					}
-					catch (Exception Ex)
+					catch (Exception ex)
 					{
-						Logger.LogError(Ex, "Exception during automation");
+						_logger.LogError(ex, "Exception during automation");
 					}
 					finally
 					{
-						Logger.LogInformation("Closed connection.");
+						_logger.LogInformation("Closed connection.");
 					}
 				}
-				catch (Exception Ex)
+				catch (Exception ex)
 				{
-					if (!bDisposing)
+					if (!_disposing)
 					{
-						Logger.LogError(Ex, "Exception during automation operation");
+						_logger.LogError(ex, "Exception during automation operation");
 					}
 				}
 			}
 		}
 
-		void HandleUri(string Uri)
+		void HandleUri(string uri)
 		{
 			try
 			{
-				UriResult Result = UriHandler.HandleUri(Uri);
-				if (!Result.Success)
+				UriResult result = UriHandler.HandleUri(uri);
+				if (!result.Success)
 				{
-					if (!string.IsNullOrEmpty(Result.Error))
+					if (!string.IsNullOrEmpty(result.Error))
 					{
-						MessageBox.Show(String.Format("Error handling uri: {0}", Result.Error));
+						MessageBox.Show(String.Format("Error handling uri: {0}", result.Error));
 					}
 
 					return;
 				}
 
-				if (Result.Request != null)
+				if (result.Request != null)
 				{
-					PostRequest(Result.Request);
-					Result.Request.Complete.Wait();
-					Result.Request.Dispose();
+					_postRequest(result.Request);
+					result.Request.Complete.Wait();
+					result.Request.Dispose();
 				}
 			}
 			catch { }
@@ -355,14 +355,14 @@ namespace UnrealGameSync
 		/// <summary>
 		/// Sends UGS scope URI from secondary process to main for handling
 		/// </summary>		
-		public static void SendUri(string Uri)
+		public static void SendUri(string uri)
 		{
-			using (NamedPipeClientStream ClientStream = new NamedPipeClientStream(".", UGSChannel, PipeDirection.Out, PipeOptions.None))
+			using (NamedPipeClientStream clientStream = new NamedPipeClientStream(".", UgsChannel, PipeDirection.Out, PipeOptions.None))
 			{
 				try
 				{
-					ClientStream.Connect(5000);
-					WriteString(ClientStream, Uri);
+					clientStream.Connect(5000);
+					WriteString(clientStream, uri);
 				}
 				catch (Exception)
 				{
@@ -371,71 +371,71 @@ namespace UnrealGameSync
 			}
 		}
 
-		static string ReadString(Stream Stream)
+		static string ReadString(Stream stream)
 		{
-			int Len = Stream.ReadByte() * 256;
-			Len += Stream.ReadByte();
-			byte[] InBuffer = new byte[Len];
-			Stream.Read(InBuffer, 0, Len);
+			int len = stream.ReadByte() * 256;
+			len += stream.ReadByte();
+			byte[] inBuffer = new byte[len];
+			stream.Read(inBuffer, 0, len);
 
-			return StreamEncoding.GetString(InBuffer);
+			return _streamEncoding.GetString(inBuffer);
 		}
 
-		static void WriteString(Stream Stream, string Output)
+		static void WriteString(Stream stream, string output)
 		{
-			byte[] OutBuffer = StreamEncoding.GetBytes(Output);
+			byte[] outBuffer = _streamEncoding.GetBytes(output);
 
-			int Len = OutBuffer.Length;
+			int len = outBuffer.Length;
 
-			if (Len > ushort.MaxValue)
+			if (len > ushort.MaxValue)
 			{
-				Len = ushort.MaxValue;
+				len = ushort.MaxValue;
 			}
 
-			Stream.WriteByte((byte)(Len / 256));
-			Stream.WriteByte((byte)(Len & 255));
-			Stream.Write(OutBuffer, 0, Len);
-			Stream.Flush();
+			stream.WriteByte((byte)(len / 256));
+			stream.WriteByte((byte)(len & 255));
+			stream.Write(outBuffer, 0, len);
+			stream.Flush();
 		}
 
 
 		public void Dispose()
 		{
-			const int Timeout = 5000;
-			bDisposing = true;
-			ShutdownEvent.Set();
+			const int timeout = 5000;
+			_disposing = true;
+			_shutdownEvent.Set();
 
-			if (UriThread != null)
+			if (_uriThread != null)
 			{
-				if (!UriThread.Join(Timeout))
+				if (!_uriThread.Join(timeout))
 				{
-					try { UriThread.Abort(); } catch { }
+					try { _uriThread.Abort(); } catch { }
 				}
 
-				UriThread = null;
+				_uriThread = null;
 			}
 
-			if (TcpThread != null)
+			if (_tcpThread != null)
 			{
-				if (!TcpThread.Join(Timeout))
+				if (!_tcpThread.Join(timeout))
 				{
-					try { TcpThread.Abort(); } catch { }
+					try { _tcpThread.Abort(); } catch { }
 				}
 
-				TcpThread = null;
+				_tcpThread = null;
 			}
 
 
 			// clean up IPC stream
-			if (IPCStream != null)
+			if (_ipcStream != null)
 			{
-				IPCStream.Dispose();
+				_ipcStream.Dispose();
 			}
 
-			if (Listener != null)
+			if (_listener != null)
 			{
-				Listener.Stop();
-				Listener = null;
+				_listener.Stop();
+				_listener = null;
 			}
 
 		}
diff --git a/Engine/Source/Programs/UnrealGameSync/UnrealGameSync/Controls/AlertButtonControl.cs b/Engine/Source/Programs/UnrealGameSync/UnrealGameSync/Controls/AlertButtonControl.cs
--- a/Engine/Source/Programs/UnrealGameSync/UnrealGameSync/Controls/AlertButtonControl.cs
+++ b/Engine/Source/Programs/UnrealGameSync/UnrealGameSync/Controls/AlertButtonControl.cs
@@ -69,22 +69,22 @@ namespace UnrealGameSync
 			}
 		}
 
-		AlertButtonTheme ThemeValue;
-		AlertButtonColors Colors;
-		AlertButtonColors CustomColorsValue;
-		bool bMouseOver;
-		bool bMouseDown;
+		AlertButtonTheme _themeValue;
+		AlertButtonColors _colors;
+		AlertButtonColors _customColorsValue;
+		bool _mouseOver;
+		bool _mouseDown;
 
 		public AlertButtonTheme Theme
 		{
-			get { return ThemeValue; }
-			set { ThemeValue = value; UpdateThemeColors(); }
+			get { return _themeValue; }
+			set { _themeValue = value; UpdateThemeColors(); }
 		}
 
 		public AlertButtonColors CustomColors
 		{
-			get { return CustomColorsValue; }
-			set { CustomColorsValue = value; UpdateThemeColors(); }
+			get { return _customColorsValue; }
+			set { _customColorsValue = value; UpdateThemeColors(); }
 		}
 
 		public AlertButtonControl()
@@ -98,51 +98,51 @@ namespace UnrealGameSync
 			switch(Theme)
 			{
 				case AlertButtonTheme.Normal:
-					Colors.ForeColor = Color.FromArgb(64, 86, 106);
-					Colors.BorderColor = Color.FromArgb(230, 232, 235);
-					Colors.BackgroundColor1 = Color.FromArgb(255, 255, 255);
-					Colors.BackgroundColor2 = Color.FromArgb(244, 245, 247);
-					Colors.BackgroundColorHover1 = Color.FromArgb(244, 245, 247);
-					Colors.BackgroundColorHover2 = Color.FromArgb(244, 245, 247);
-					Colors.BackgroundColorDown1 = Color.FromArgb(234, 235, 237);
-					Colors.BackgroundColorDown2 = Color.FromArgb(234, 235, 237);
+					_colors.ForeColor = Color.FromArgb(64, 86, 106);
+					_colors.BorderColor = Color.FromArgb(230, 232, 235);
+					_colors.BackgroundColor1 = Color.FromArgb(255, 255, 255);
+					_colors.BackgroundColor2 = Color.FromArgb(244, 245, 247);
+					_colors.BackgroundColorHover1 = Color.FromArgb(244, 245, 247);
+					_colors.BackgroundColorHover2 = Color.FromArgb(244, 245, 247);
+					_colors.BackgroundColorDown1 = Color.FromArgb(234, 235, 237);
+					_colors.BackgroundColorDown2 = Color.FromArgb(234, 235, 237);
 					break;
 				case AlertButtonTheme.Green:
-					Colors.ForeColor = Color.FromArgb(255, 255, 255);
-					Colors.BorderColor = Color.FromArgb(143, 199, 156);
-					Colors.BackgroundColor1 = Color.FromArgb(116, 192, 134);
-					Colors.BackgroundColor2 = Color.FromArgb(99, 175, 117);
-					Colors.BackgroundColorHover1 = Color.FromArgb(99, 175, 117);
-					Colors.BackgroundColorHover2 = Color.FromArgb(99, 175, 117);
-					Colors.BackgroundColorDown1 = Color.FromArgb(90, 165, 107);
-					Colors.BackgroundColorDown2 = Color.FromArgb(90, 165, 107);
+					_colors.ForeColor = Color.FromArgb(255, 255, 255);
+					_colors.BorderColor = Color.FromArgb(143, 199, 156);
+					_colors.BackgroundColor1 = Color.FromArgb(116, 192, 134);
+					_colors.BackgroundColor2 = Color.FromArgb(99, 175, 117);
+					_colors.BackgroundColorHover1 = Color.FromArgb(99, 175, 117);
+					_colors.BackgroundColorHover2 = Color.FromArgb(99, 175, 117);
+					_colors.BackgroundColorDown1 = Color.FromArgb(90, 165, 107);
+					_colors.BackgroundColorDown2 = Color.FromArgb(90, 165, 107);
 					break;
 				case AlertButtonTheme.Red:
-					Colors.ForeColor = Color.FromArgb(255, 255, 255);
-					Colors.BorderColor = Color.FromArgb(230, 232, 235);
-					Colors.BackgroundColor1 = Color.FromArgb(222, 108, 86);
-					Colors.BackgroundColor2 = Color.FromArgb(214, 69, 64);
-					Colors.BackgroundColorHover1 = Color.FromArgb(214, 69, 64);
-					Colors.BackgroundColorHover2 = Color.FromArgb(214, 69, 64);
-					Colors.BackgroundColorDown1 = Color.FromArgb(204, 59, 54);
-					Colors.BackgroundColorDown2 = Color.FromArgb(204, 59, 54);
+					_colors.ForeColor = Color.FromArgb(255, 255, 255);
+					_colors.BorderColor = Color.FromArgb(230, 232, 235);
+					_colors.BackgroundColor1 = Color.FromArgb(222, 108, 86);
+					_colors.BackgroundColor2 = Color.FromArgb(214, 69, 64);
+					_colors.BackgroundColorHover1 = Color.FromArgb(214, 69, 64);
+					_colors.BackgroundColorHover2 = Color.FromArgb(214, 69, 64);
+					_colors.BackgroundColorDown1 = Color.FromArgb(204, 59, 54);
+					_colors.BackgroundColorDown2 = Color.FromArgb(204, 59, 54);
 					break;
 				case AlertButtonTheme.Strong:
-					Colors.ForeColor = Color.FromArgb(255, 255, 255);
-					Colors.BorderColor = Color.FromArgb(230, 232, 235);
-					Colors.BackgroundColor1 = Color.FromArgb(200, 74, 49);
-					Colors.BackgroundColor2 = Color.FromArgb(200, 74, 49); 
-					Colors.BackgroundColorHover1 = Color.FromArgb(222, 108, 86);
-					Colors.BackgroundColorHover2 = Color.FromArgb(222, 108, 86);
-					Colors.BackgroundColorDown1 = Color.FromArgb(204, 59, 54);
-					Colors.BackgroundColorDown2 = Color.FromArgb(204, 59, 54);
+					_colors.ForeColor = Color.FromArgb(255, 255, 255);
+					_colors.BorderColor = Color.FromArgb(230, 232, 235);
+					_colors.BackgroundColor1 = Color.FromArgb(200, 74, 49);
+					_colors.BackgroundColor2 = Color.FromArgb(200, 74, 49); 
+					_colors.BackgroundColorHover1 = Color.FromArgb(222, 108, 86);
+					_colors.BackgroundColorHover2 = Color.FromArgb(222, 108, 86);
+					_colors.BackgroundColorDown1 = Color.FromArgb(204, 59, 54);
+					_colors.BackgroundColorDown2 = Color.FromArgb(204, 59, 54);
 					break;
 				case AlertButtonTheme.Custom:
-					Colors = CustomColorsValue;
+					_colors = _customColorsValue;
 					break;
 			}
 
-			base.ForeColor = Colors.ForeColor;
+			base.ForeColor = _colors.ForeColor;
 			Invalidate();
 		}
 
@@ -150,7 +150,7 @@ namespace UnrealGameSync
 		{
 			base.OnMouseDown(mevent);
 
-			bMouseDown = true;
+			_mouseDown = true;
 			Invalidate();
 		}
 
@@ -158,7 +158,7 @@ namespace UnrealGameSync
 		{
 			base.OnMouseUp(mevent);
 
-			bMouseDown = false;
+			_mouseDown = false;
 			Invalidate();
 		}
 
@@ -166,7 +166,7 @@ namespace UnrealGameSync
 		{
 			base.OnMouseHover(e);
 
-			bMouseOver = true;
+			_mouseOver = true;
 			Invalidate();
 		}
 
@@ -174,38 +174,38 @@ namespace UnrealGameSync
 		{
 			base.OnMouseLeave(e);
 
-			bMouseOver = false;
+			_mouseOver = false;
 			Invalidate();
 		}
 
-		protected override void OnPaint(PaintEventArgs Event)
+		protected override void OnPaint(PaintEventArgs e)
 		{
-			Event.Graphics.FillRectangle(SystemBrushes.Window, 0, 0, Width, Height);
-			using(GraphicsPath Path = new GraphicsPath())
+			e.Graphics.FillRectangle(SystemBrushes.Window, 0, 0, Width, Height);
+			using(GraphicsPath path = new GraphicsPath())
 			{
-				const int Diameter = 4;
+				const int diameter = 4;
 
-				Path.StartFigure();
-				Path.AddArc(Width - 1 - Diameter, Height - 1 - Diameter, Diameter, Diameter, 0, 90);
-				Path.AddArc(0, Height - 1 - Diameter, Diameter, Diameter, 90, 90);
-				Path.AddArc(0, 0, Diameter, Diameter, 180, 90);
-				Path.AddArc(Width - 1 - Diameter, 0, Diameter, Diameter, 270, 90);
-				Path.CloseFigure();
+				path.StartFigure();
+				path.AddArc(Width - 1 - diameter, Height - 1 - diameter, diameter, diameter, 0, 90);
+				path.AddArc(0, Height - 1 - diameter, diameter, diameter, 90, 90);
+				path.AddArc(0, 0, diameter, diameter, 180, 90);
+				path.AddArc(Width - 1 - diameter, 0, diameter, diameter, 270, 90);
+				path.CloseFigure();
 
-				Color BackgroundColorMin = (bMouseDown && bMouseOver)? Colors.BackgroundColorDown1 : bMouseOver? Colors.BackgroundColorHover1 : Colors.BackgroundColor1;
-				Color BackgroundColorMax = (bMouseDown && bMouseOver)? Colors.BackgroundColorDown2 : bMouseOver? Colors.BackgroundColorHover2 : Colors.BackgroundColor2;
+				Color backgroundColorMin = (_mouseDown && _mouseOver)? _colors.BackgroundColorDown1 : _mouseOver? _colors.BackgroundColorHover1 : _colors.BackgroundColor1;
+				Color backgroundColorMax = (_mouseDown && _mouseOver)? _colors.BackgroundColorDown2 : _mouseOver? _colors.BackgroundColorHover2 : _colors.BackgroundColor2;
 
-				using(LinearGradientBrush Brush = new LinearGradientBrush(new Point(0, 0), new Point(0, Height), BackgroundColorMin, BackgroundColorMax))
+				using(LinearGradientBrush brush = new LinearGradientBrush(new Point(0, 0), new Point(0, Height), backgroundColorMin, backgroundColorMax))
 				{
-					Event.Graphics.FillPath(Brush, Path);
+					e.Graphics.FillPath(brush, path);
 				}
-				using(Pen SolidPen = new Pen(Colors.BorderColor))
+				using(Pen solidPen = new Pen(_colors.BorderColor))
 				{
-					Event.Graphics.DrawPath(SolidPen, Path);
+					e.Graphics.DrawPath(solidPen, path);
 				}
 			}
 
-			TextRenderer.DrawText(Event.Graphics, Text, Font, new Rectangle(0, 0, Width, Height), ForeColor, TextFormatFlags.HorizontalCenter | TextFormatFlags.VerticalCenter | TextFormatFlags.SingleLine);
+			TextRenderer.DrawText(e.Graphics, Text, Font, new Rectangle(0, 0, Width, Height), ForeColor, TextFormatFlags.HorizontalCenter | TextFormatFlags.VerticalCenter | TextFormatFlags.SingleLine);
 		}
 	}
 }
diff --git a/Engine/Source/Programs/UnrealGameSync/UnrealGameSync/Controls/AlertImageControl.cs b/Engine/Source/Programs/UnrealGameSync/UnrealGameSync/Controls/AlertImageControl.cs
--- a/Engine/Source/Programs/UnrealGameSync/UnrealGameSync/Controls/AlertImageControl.cs
+++ b/Engine/Source/Programs/UnrealGameSync/UnrealGameSync/Controls/AlertImageControl.cs
@@ -16,33 +16,33 @@ namespace UnrealGameSync
 {
 	partial class AlertImageControl : UserControl
 	{
-		Image? ImageValue;
+		Image? _imageValue;
 
 		public Image? Image
 		{
-			get { return ImageValue; }
-			set { ImageValue = value; Invalidate(); }
+			get { return _imageValue; }
+			set { _imageValue = value; Invalidate(); }
 		}
 
 		public AlertImageControl()
 		{
 		}
 
-		protected override void OnPaint(PaintEventArgs Event)
+		protected override void OnPaint(PaintEventArgs e)
 		{
-			Event.Graphics.FillRectangle(SystemBrushes.Window, 0, 0, Width, Height);
+			e.Graphics.FillRectangle(SystemBrushes.Window, 0, 0, Width, Height);
 
 			if(Image != null)
 			{
-				float Scale = Math.Min((float)Width / Image.Width, (float)Height / Image.Height);
+				float scale = Math.Min((float)Width / Image.Width, (float)Height / Image.Height);
 
-				int ImageW = (int)(Image.Width * Scale);
-				int ImageH = (int)(Image.Height * Scale);
+				int imageW = (int)(Image.Width * scale);
+				int imageH = (int)(Image.Height * scale);
 
-				Event.Graphics.InterpolationMode = InterpolationMode.HighQualityBicubic;
-				Event.Graphics.SmoothingMode = SmoothingMode.HighQuality;
-				Event.Graphics.PixelOffsetMode = PixelOffsetMode.HighQuality;
-				Event.Graphics.DrawImage(Image, (Width - ImageW) / 2, (Height - ImageH) / 2, ImageW, ImageH); 
+				e.Graphics.InterpolationMode = InterpolationMode.HighQualityBicubic;
+				e.Graphics.SmoothingMode = SmoothingMode.HighQuality;
+				e.Graphics.PixelOffsetMode = PixelOffsetMode.HighQuality;
+				e.Graphics.DrawImage(Image, (Width - imageW) / 2, (Height - imageH) / 2, imageW, imageH); 
 			}
 		}
 	}
diff --git a/Engine/Source/Programs/UnrealGameSync/UnrealGameSync/Controls/BuildListControl.cs b/Engine/Source/Programs/UnrealGameSync/UnrealGameSync/Controls/BuildListControl.cs
--- a/Engine/Source/Programs/UnrealGameSync/UnrealGameSync/Controls/BuildListControl.cs
+++ b/Engine/Source/Programs/UnrealGameSync/UnrealGameSync/Controls/BuildListControl.cs
@@ -15,7 +15,7 @@ namespace UnrealGameSync
     class BuildListControl : CustomListViewControl
 	{
 		[StructLayout(LayoutKind.Sequential)]
-		public class SCROLLINFO
+		public class Scrollinfo
 		{
 			public int cbSize;
 			public int fMask;
@@ -26,30 +26,30 @@ namespace UnrealGameSync
 			public int nTrackPos;
 		}
 
-		const int WM_VSCROLL = 0x115;
-		const int WM_MOUSEWHEEL = 0x020A;
+		const int WmVscroll = 0x115;
+		const int WmMousewheel = 0x020A;
 
-		const int SB_HORZ = 0;
-		const int SB_VERT = 1;
+		const int SbHorz = 0;
+		const int SbVert = 1;
 
-		const int SIF_RANGE = 0x0001;
-		const int SIF_PAGE = 0x0002;
-		const int SIF_POS = 0x0004;
-		const int SIF_DISABLENOSCROLL = 0x0008;
-		const int SIF_TRACKPOS = 0x0010;
-		const int SIF_ALL = (SIF_RANGE | SIF_PAGE | SIF_POS | SIF_TRACKPOS);        
+		const int SifRange = 0x0001;
+		const int SifPage = 0x0002;
+		const int SifPos = 0x0004;
+		const int SifDisablenoscroll = 0x0008;
+		const int SifTrackpos = 0x0010;
+		const int SifAll = (SifRange | SifPage | SifPos | SifTrackpos);        
 
-		const int LVM_GETTOPINDEX = 0x1000 + 39;
-		const int LVM_GETCOUNTPERPAGE = 0x1000 + 40;
+		const int LvmGettopindex = 0x1000 + 39;
+		const int LvmGetcountperpage = 0x1000 + 40;
 
 		[DllImport("user32.dll")]
-		private static extern int GetScrollInfo(IntPtr hwnd, int fnBar, SCROLLINFO lpsi);
+		private static extern int GetScrollInfo(IntPtr hwnd, int fnBar, Scrollinfo lpsi);
 
 		[DllImport("user32.dll")]
-		private static extern int SetScrollInfo(IntPtr hwnd, int fnBar, SCROLLINFO lpsi, int Redraw);
+		private static extern int SetScrollInfo(IntPtr hwnd, int fnBar, Scrollinfo lpsi, int redraw);
 
 		[DllImport("user32.dll", CharSet = CharSet.Auto)]
-		private static extern IntPtr SendMessage(IntPtr hWnd, UInt32 Msg, IntPtr wParam, IntPtr lParam);
+		private static extern IntPtr SendMessage(IntPtr hWnd, UInt32 msg, IntPtr wParam, IntPtr lParam);
 
 		public delegate void OnScrollDelegate();
 
@@ -66,7 +66,7 @@ namespace UnrealGameSync
 			{
 				return 0;
 			}
-			return SendMessage(Handle, LVM_GETTOPINDEX, IntPtr.Zero, IntPtr.Zero).ToInt32();
+			return SendMessage(Handle, LvmGettopindex, IntPtr.Zero, IntPtr.Zero).ToInt32();
 		}
 
 		public int GetVisibleItemsPerPage()
@@ -75,41 +75,41 @@ namespace UnrealGameSync
 			{
 				return Height / Font.Height;
 			}
-			return SendMessage(Handle, LVM_GETCOUNTPERPAGE, IntPtr.Zero, IntPtr.Zero).ToInt32();
+			return SendMessage(Handle, LvmGetcountperpage, IntPtr.Zero, IntPtr.Zero).ToInt32();
 		}
 
-		public bool GetScrollPosition(out int ScrollY)
+		public bool GetScrollPosition(out int scrollY)
 		{
 			if(!IsHandleCreated)
 			{
-				ScrollY = 0;
+				scrollY = 0;
 				return false;
 			}
 
-			SCROLLINFO ScrollInfo = new SCROLLINFO();
-			ScrollInfo.cbSize = Marshal.SizeOf(ScrollInfo);
-			ScrollInfo.fMask = SIF_ALL;
-			if(GetScrollInfo(Handle, SB_VERT, ScrollInfo) == 0)
+			Scrollinfo scrollInfo = new Scrollinfo();
+			scrollInfo.cbSize = Marshal.SizeOf(scrollInfo);
+			scrollInfo.fMask = SifAll;
+			if(GetScrollInfo(Handle, SbVert, scrollInfo) == 0)
 			{
-				ScrollY = 0;
+				scrollY = 0;
 				return false;
 			}
 			else
 			{
-				ScrollY = ScrollInfo.nPos;
+				scrollY = scrollInfo.nPos;
 				return true;
 			}
 		}
 
-		public void SetScrollPosition(int ScrollY)
+		public void SetScrollPosition(int scrollY)
 		{
 			if(IsHandleCreated)
 			{
-				SCROLLINFO ScrollInfo = new SCROLLINFO();
-				ScrollInfo.cbSize = Marshal.SizeOf(ScrollInfo);
-				ScrollInfo.nPos = ScrollY;
-				ScrollInfo.fMask = SIF_POS;
-				SetScrollInfo(Handle, SB_VERT, ScrollInfo, 0);
+				Scrollinfo scrollInfo = new Scrollinfo();
+				scrollInfo.cbSize = Marshal.SizeOf(scrollInfo);
+				scrollInfo.nPos = scrollY;
+				scrollInfo.fMask = SifPos;
+				SetScrollInfo(Handle, SbVert, scrollInfo, 0);
 			}
 		}
 
@@ -117,10 +117,10 @@ namespace UnrealGameSync
 		{
 			if(IsHandleCreated)
 			{
-				SCROLLINFO ScrollInfo = new SCROLLINFO();
-				ScrollInfo.cbSize = Marshal.SizeOf(ScrollInfo);
-				ScrollInfo.fMask = SIF_ALL;
-				if(GetScrollInfo(Handle, SB_VERT, ScrollInfo) != 0 && ScrollInfo.nPos >= ScrollInfo.nMax - ScrollInfo.nPage)
+				Scrollinfo scrollInfo = new Scrollinfo();
+				scrollInfo.cbSize = Marshal.SizeOf(scrollInfo);
+				scrollInfo.fMask = SifAll;
+				if(GetScrollInfo(Handle, SbVert, scrollInfo) != 0 && scrollInfo.nPos >= scrollInfo.nMax - scrollInfo.nPage)
 				{
 					return true;
 				}
@@ -128,14 +128,14 @@ namespace UnrealGameSync
 			return false;
 		}
 
-		protected override void WndProc(ref Message Message)
+		protected override void WndProc(ref Message message)
 		{
-			base.WndProc(ref Message);
+			base.WndProc(ref message);
 
-			switch(Message.Msg)
+			switch(message.Msg)
 			{
-				case WM_VSCROLL:
-				case WM_MOUSEWHEEL:
+				case WmVscroll:
+				case WmMousewheel:
 					if(OnScroll != null)
 					{
 						OnScroll();
diff --git a/Engine/Source/Programs/UnrealGameSync/UnrealGameSync/Controls/CustomListViewControl.cs b/Engine/Source/Programs/UnrealGameSync/UnrealGameSync/Controls/CustomListViewControl.cs
--- a/Engine/Source/Programs/UnrealGameSync/UnrealGameSync/Controls/CustomListViewControl.cs
+++ b/Engine/Source/Programs/UnrealGameSync/UnrealGameSync/Controls/CustomListViewControl.cs
@@ -21,7 +21,7 @@ namespace UnrealGameSync
 			private set;
 		}
 
-		Rectangle? PreviousBounds;
+		Rectangle? _previousBounds;
 
 		public Cursor Cursor
 		{
@@ -29,35 +29,35 @@ namespace UnrealGameSync
 			set;
 		}
 
-		public CustomListViewWidget(ListViewItem Item)
+		public CustomListViewWidget(ListViewItem item)
 		{
-			this.Item = Item;
+			this.Item = item;
 			this.Cursor = Cursors.Arrow;
 		}
 
 		public void Invalidate()
 		{
-			PreviousBounds = null;
+			_previousBounds = null;
 		}
 
 		public virtual bool RequiresLayout()
 		{
-			return !PreviousBounds.HasValue;
+			return !_previousBounds.HasValue;
 		}
 
-		public void ConditionalLayout(Control Owner, Rectangle Bounds)
+		public void ConditionalLayout(Control owner, Rectangle bounds)
 		{
-			if(RequiresLayout() || PreviousBounds == null || PreviousBounds.Value != Bounds)
+			if(RequiresLayout() || _previousBounds == null || _previousBounds.Value != bounds)
 			{
-				using(Graphics Graphics = Owner.CreateGraphics())
+				using(Graphics graphics = owner.CreateGraphics())
 				{
-					Layout(Graphics, Bounds);
+					Layout(graphics, bounds);
 				}
-				PreviousBounds = Bounds;
+				_previousBounds = bounds;
 			}
 		}
 
-		public virtual void OnMouseMove(Point Location)
+		public virtual void OnMouseMove(Point location)
 		{
 		}
 
@@ -65,16 +65,16 @@ namespace UnrealGameSync
 		{
 		}
 
-		public virtual void OnMouseDown(Point Location)
+		public virtual void OnMouseDown(Point location)
 		{
 		}
 
-		public virtual void OnMouseUp(Point Location)
+		public virtual void OnMouseUp(Point location)
 		{
 		}
 
-		public abstract void Layout(Graphics Graphics, Rectangle Bounds);
-		public abstract void Render(Graphics Graphics);
+		public abstract void Layout(Graphics graphics, Rectangle bounds);
+		public abstract void Render(Graphics graphics);
 	}
 
 	class StatusLineListViewWidget : CustomListViewWidget
@@ -87,14 +87,14 @@ namespace UnrealGameSync
 			get; set;
 		}
 
-		Rectangle Bounds;
-		StatusElement? MouseDownElement;
-		StatusElement? MouseOverElement;
+		Rectangle _bounds;
+		StatusElement? _mouseDownElement;
+		StatusElement? _mouseOverElement;
 
-		public StatusLineListViewWidget(ListViewItem Item, StatusElementResources Resources)
-			: base(Item)
+		public StatusLineListViewWidget(ListViewItem item, StatusElementResources resources)
+			: base(item)
 		{
-			this.Resources = Resources;
+			this.Resources = resources;
 			this.HorizontalAlignment = HorizontalAlignment.Center;
 		}
 
@@ -103,55 +103,55 @@ namespace UnrealGameSync
 			return base.RequiresLayout() || Line.RequiresLayout();
 		}
 
-		public override void OnMouseDown(Point Location)
+		public override void OnMouseDown(Point location)
 		{
-			StatusElement? Element;
-			if(Line.HitTest(Location, out Element))
+			StatusElement? element;
+			if(Line.HitTest(location, out element))
 			{
-				MouseDownElement = Element;
-				MouseDownElement.bMouseDown = true;
+				_mouseDownElement = element;
+				_mouseDownElement.MouseDown = true;
 
 				Invalidate();
 			}
 		}
 
-		public override void OnMouseUp(Point Location)
+		public override void OnMouseUp(Point location)
 		{
-			if(MouseDownElement != null)
+			if(_mouseDownElement != null)
 			{
-				StatusElement ClickElement = MouseDownElement;
+				StatusElement clickElement = _mouseDownElement;
 
-				MouseDownElement.bMouseDown = false;
-				MouseDownElement = null;
+				_mouseDownElement.MouseDown = false;
+				_mouseDownElement = null;
 
-				if(ClickElement.Bounds.Contains(Location))
+				if(clickElement.Bounds.Contains(location))
 				{
-					ClickElement.OnClick(Location);
+					clickElement.OnClick(location);
 				}
 
 				Invalidate();
 			}
 		}
 
-		public override void OnMouseMove(Point Location)
+		public override void OnMouseMove(Point location)
 		{
-			StatusElement? NewMouseOverElement;
-			Line.HitTest(Location, out NewMouseOverElement);
+			StatusElement? newMouseOverElement;
+			Line.HitTest(location, out newMouseOverElement);
 
-			if(MouseOverElement != NewMouseOverElement)
+			if(_mouseOverElement != newMouseOverElement)
 			{
-				if(MouseOverElement != null)
+				if(_mouseOverElement != null)
 				{
 					Cursor = Cursors.Arrow;
-					MouseOverElement.bMouseOver = false;
+					_mouseOverElement.MouseOver = false;
 				}
 
-				MouseOverElement = NewMouseOverElement;
+				_mouseOverElement = newMouseOverElement;
 
-				if(MouseOverElement != null)
+				if(_mouseOverElement != null)
 				{
-					Cursor = MouseOverElement.Cursor;
-					MouseOverElement.bMouseOver = true;
+					Cursor = _mouseOverElement.Cursor;
+					_mouseOverElement.MouseOver = true;
 				}
 
 				Invalidate();
@@ -160,116 +160,116 @@ namespace UnrealGameSync
 
 		public override void OnMouseLeave()
 		{
-			if(MouseOverElement != null)
+			if(_mouseOverElement != null)
 			{
-				MouseOverElement.bMouseOver = false;
-				MouseOverElement = null;
+				_mouseOverElement.MouseOver = false;
+				_mouseOverElement = null;
 
 				Invalidate();
 			}
 		}
 
-		public override void Layout(Graphics Graphics, Rectangle Bounds)
+		public override void Layout(Graphics graphics, Rectangle bounds)
 		{
-			this.Bounds = Bounds;
+			this._bounds = bounds;
 
-			int OffsetY = Bounds.Y + Bounds.Height / 2;
-			Line.Layout(Graphics, new Point(Bounds.X, OffsetY), Resources);
+			int offsetY = bounds.Y + bounds.Height / 2;
+			Line.Layout(graphics, new Point(bounds.X, offsetY), Resources);
 
 			if(HorizontalAlignment == HorizontalAlignment.Center)
 			{
-				int OffsetX = Bounds.X + (Bounds.Width - Line.Bounds.Width) / 2;
-				Line.Layout(Graphics, new Point(OffsetX, OffsetY), Resources);
+				int offsetX = bounds.X + (bounds.Width - Line.Bounds.Width) / 2;
+				Line.Layout(graphics, new Point(offsetX, offsetY), Resources);
 			}
 			else if(HorizontalAlignment == HorizontalAlignment.Right)
 			{
-				int OffsetX = Bounds.Right - Line.Bounds.Width;
-				Line.Layout(Graphics, new Point(OffsetX, OffsetY), Resources);
+				int offsetX = bounds.Right - Line.Bounds.Width;
+				Line.Layout(graphics, new Point(offsetX, offsetY), Resources);
 			}
 		}
 
-		public override void Render(Graphics Graphics)
+		public override void Render(Graphics graphics)
 		{
-			Graphics.IntersectClip(Bounds);
+			graphics.IntersectClip(_bounds);
 
-			Line.Draw(Graphics, Resources);
+			Line.Draw(graphics, Resources);
 		}
 	}
 
 	partial class CustomListViewControl : ListView
 	{
-		VisualStyleRenderer? SelectedItemRenderer;
-		VisualStyleRenderer? TrackedItemRenderer;
+		VisualStyleRenderer? _selectedItemRenderer;
+		VisualStyleRenderer? _trackedItemRenderer;
 		public int HoverItem = -1;
 
-		CustomListViewWidget? MouseOverWidget;
-		CustomListViewWidget? MouseDownWidget;
+		CustomListViewWidget? _mouseOverWidget;
+		CustomListViewWidget? _mouseDownWidget;
 
 		public CustomListViewControl()
 		{
             if (Application.RenderWithVisualStyles) 
             { 
-				SelectedItemRenderer = new VisualStyleRenderer("Explorer::ListView", 1, 3);
-				TrackedItemRenderer = new VisualStyleRenderer("Explorer::ListView", 1, 2); 
+				_selectedItemRenderer = new VisualStyleRenderer("Explorer::ListView", 1, 3);
+				_trackedItemRenderer = new VisualStyleRenderer("Explorer::ListView", 1, 2); 
 			}
 		}
 
-		protected void ConditionalLayoutWidget(CustomListViewWidget Widget)
+		protected void ConditionalLayoutWidget(CustomListViewWidget widget)
 		{
-			if(Widget.Item.Tag == Widget)
+			if(widget.Item.Tag == widget)
 			{
-				Widget.ConditionalLayout(this, Widget.Item.Bounds);
+				widget.ConditionalLayout(this, widget.Item.Bounds);
 			}
 			else
 			{
-				for(int Idx = 0; Idx < Widget.Item.SubItems.Count; Idx++)
+				for(int idx = 0; idx < widget.Item.SubItems.Count; idx++)
 				{
-					ListViewItem.ListViewSubItem SubItem = Widget.Item.SubItems[Idx];
-					if(SubItem.Tag == Widget)
+					ListViewItem.ListViewSubItem subItem = widget.Item.SubItems[idx];
+					if(subItem.Tag == widget)
 					{
-						Rectangle Bounds = SubItem.Bounds;
-						for(int EndIdx = Idx + 1; EndIdx < Widget.Item.SubItems.Count && Widget.Item.SubItems[EndIdx].Tag == Widget; EndIdx++)
+						Rectangle bounds = subItem.Bounds;
+						for(int endIdx = idx + 1; endIdx < widget.Item.SubItems.Count && widget.Item.SubItems[endIdx].Tag == widget; endIdx++)
 						{
-							Rectangle EndBounds = Widget.Item.SubItems[EndIdx].Bounds;
-							Bounds = new Rectangle(Bounds.X, Bounds.Y, EndBounds.Right - Bounds.X, EndBounds.Bottom - Bounds.Y);
+							Rectangle endBounds = widget.Item.SubItems[endIdx].Bounds;
+							bounds = new Rectangle(bounds.X, bounds.Y, endBounds.Right - bounds.X, endBounds.Bottom - bounds.Y);
 						}
-						Widget.ConditionalLayout(this, Bounds);
+						widget.ConditionalLayout(this, bounds);
 						break;
 					}
 				}
 			}
 		}
 
-		protected void ConditionalRedrawWidget(CustomListViewWidget Widget)
+		protected void ConditionalRedrawWidget(CustomListViewWidget widget)
 		{
-			if(Widget.Item.Index != -1 && Widget.RequiresLayout())
+			if(widget.Item.Index != -1 && widget.RequiresLayout())
 			{
-				ConditionalLayoutWidget(Widget);
-				RedrawItems(Widget.Item.Index, Widget.Item.Index, true);
+				ConditionalLayoutWidget(widget);
+				RedrawItems(widget.Item.Index, widget.Item.Index, true);
 			}
 		}
 
-		protected CustomListViewWidget? FindWidget(Point Location)
+		protected CustomListViewWidget? FindWidget(Point location)
 		{
-			return FindWidget(HitTest(Location));
+			return FindWidget(HitTest(location));
 		}
 
-		protected CustomListViewWidget? FindWidget(ListViewHitTestInfo HitTest)
+		protected CustomListViewWidget? FindWidget(ListViewHitTestInfo hitTest)
 		{
-			if(HitTest.Item != null)
+			if(hitTest.Item != null)
 			{
-				CustomListViewWidget? Widget = HitTest.Item.Tag as CustomListViewWidget;
-				if(Widget != null)
+				CustomListViewWidget? widget = hitTest.Item.Tag as CustomListViewWidget;
+				if(widget != null)
 				{
-					return Widget;
+					return widget;
 				}
 			}
-			if(HitTest.SubItem != null)
+			if(hitTest.SubItem != null)
 			{
-				CustomListViewWidget? Widget = HitTest.SubItem.Tag as CustomListViewWidget;
-				if(Widget != null)
+				CustomListViewWidget? widget = hitTest.SubItem.Tag as CustomListViewWidget;
+				if(widget != null)
 				{
-					return Widget;
+					return widget;
 				}
 			}
 			return null;
@@ -279,11 +279,11 @@ namespace UnrealGameSync
 		{
 			base.OnMouseLeave(e);
 
-			if(MouseOverWidget != null)
+			if(_mouseOverWidget != null)
 			{
-				MouseOverWidget.OnMouseLeave();
-				ConditionalRedrawWidget(MouseOverWidget);
-				MouseOverWidget = null;
+				_mouseOverWidget.OnMouseLeave();
+				ConditionalRedrawWidget(_mouseOverWidget);
+				_mouseOverWidget = null;
 				Invalidate();
 			}
 
@@ -298,25 +298,25 @@ namespace UnrealGameSync
 		{
 			base.OnMouseDown(e);
 
-			if(MouseDownWidget != null)
+			if(_mouseDownWidget != null)
 			{
-				MouseDownWidget.OnMouseUp(e.Location);
-				ConditionalRedrawWidget(MouseDownWidget);
+				_mouseDownWidget.OnMouseUp(e.Location);
+				ConditionalRedrawWidget(_mouseDownWidget);
 			}
 
 			if((e.Button & MouseButtons.Left) != 0)
 			{
-				MouseDownWidget = FindWidget(e.Location);
+				_mouseDownWidget = FindWidget(e.Location);
 			}
 			else
 			{
-				MouseDownWidget = null;
+				_mouseDownWidget = null;
 			}
 
-			if(MouseDownWidget != null)
+			if(_mouseDownWidget != null)
 			{
-				MouseDownWidget.OnMouseDown(e.Location);
-				ConditionalRedrawWidget(MouseDownWidget);
+				_mouseDownWidget.OnMouseDown(e.Location);
+				ConditionalRedrawWidget(_mouseDownWidget);
 			}
 		}
 
@@ -324,78 +324,78 @@ namespace UnrealGameSync
 		{
 			base.OnMouseUp(e);
 
-			if(MouseDownWidget != null)
+			if(_mouseDownWidget != null)
 			{
-				CustomListViewWidget Widget = MouseDownWidget;
-				MouseDownWidget = null;
+				CustomListViewWidget widget = _mouseDownWidget;
+				_mouseDownWidget = null;
 
-				Widget.OnMouseUp(e.Location);
-				ConditionalRedrawWidget(Widget);
+				widget.OnMouseUp(e.Location);
+				ConditionalRedrawWidget(widget);
 			}
 		}
 
 		protected override void OnMouseMove(MouseEventArgs e)
 		{
-			ListViewHitTestInfo HitTest = this.HitTest(e.Location);
+			ListViewHitTestInfo hitTest = this.HitTest(e.Location);
 
-			CustomListViewWidget? NewMouseOverWidget = FindWidget(HitTest);
-			if(MouseOverWidget != null && MouseOverWidget != NewMouseOverWidget)
+			CustomListViewWidget? newMouseOverWidget = FindWidget(hitTest);
+			if(_mouseOverWidget != null && _mouseOverWidget != newMouseOverWidget)
 			{
 				Cursor = Cursors.Arrow;
-				MouseOverWidget.OnMouseLeave();
-				ConditionalRedrawWidget(MouseOverWidget);
+				_mouseOverWidget.OnMouseLeave();
+				ConditionalRedrawWidget(_mouseOverWidget);
 			}
-			MouseOverWidget = NewMouseOverWidget;
-			if(MouseOverWidget != null)
+			_mouseOverWidget = newMouseOverWidget;
+			if(_mouseOverWidget != null)
 			{
-				Cursor = MouseOverWidget.Cursor;
-				MouseOverWidget.OnMouseMove(e.Location);
-				ConditionalRedrawWidget(MouseOverWidget);
+				Cursor = _mouseOverWidget.Cursor;
+				_mouseOverWidget.OnMouseMove(e.Location);
+				ConditionalRedrawWidget(_mouseOverWidget);
 			}
 
-			int PrevHoverItem = HoverItem;
-			HoverItem = (HitTest.Item == null)? -1 : HitTest.Item.Index;
-			if(HoverItem != PrevHoverItem)
+			int prevHoverItem = HoverItem;
+			HoverItem = (hitTest.Item == null)? -1 : hitTest.Item.Index;
+			if(HoverItem != prevHoverItem)
 			{
 				if(HoverItem != -1)
 				{
 					RedrawItems(HoverItem, HoverItem, true);
 				}
-				if(PrevHoverItem != -1 && PrevHoverItem < Items.Count)
+				if(prevHoverItem != -1 && prevHoverItem < Items.Count)
 				{
-					RedrawItems(PrevHoverItem, PrevHoverItem, true);
+					RedrawItems(prevHoverItem, prevHoverItem, true);
 				}
 			}
 
 			base.OnMouseMove(e);
 		}
-		public void DrawText(Graphics Graphics, Rectangle Bounds, HorizontalAlignment TextAlign, Color TextColor, string Text)
+		public void DrawText(Graphics graphics, Rectangle bounds, HorizontalAlignment textAlign, Color textColor, string text)
 		{
-			TextFormatFlags Flags = TextFormatFlags.EndEllipsis | TextFormatFlags.SingleLine | TextFormatFlags.VerticalCenter | TextFormatFlags.NoPrefix;
-			if (TextAlign == HorizontalAlignment.Left)
+			TextFormatFlags flags = TextFormatFlags.EndEllipsis | TextFormatFlags.SingleLine | TextFormatFlags.VerticalCenter | TextFormatFlags.NoPrefix;
+			if (textAlign == HorizontalAlignment.Left)
 			{
-				Flags |= TextFormatFlags.Left;
+				flags |= TextFormatFlags.Left;
 			}
-			else if (TextAlign == HorizontalAlignment.Center)
+			else if (textAlign == HorizontalAlignment.Center)
 			{
-				Flags |= TextFormatFlags.HorizontalCenter;
+				flags |= TextFormatFlags.HorizontalCenter;
 			}
 			else
 			{
-				Flags |= TextFormatFlags.Right;
+				flags |= TextFormatFlags.Right;
 			}
-			TextRenderer.DrawText(Graphics, Text, Font, Bounds, TextColor, Flags);
+			TextRenderer.DrawText(graphics, text, Font, bounds, textColor, flags);
 		}
 
-		public void DrawIcon(Graphics Graphics, Rectangle Bounds, Rectangle Icon)
+		public void DrawIcon(Graphics graphics, Rectangle bounds, Rectangle icon)
 		{
-			float DpiScaleX = Graphics.DpiX / 96.0f;
-			float DpiScaleY = Graphics.DpiY / 96.0f;
+			float dpiScaleX = graphics.DpiX / 96.0f;
+			float dpiScaleY = graphics.DpiY / 96.0f;
 
-			float IconX = Bounds.Left + (Bounds.Width - 16 * DpiScaleX) / 2;
-			float IconY = Bounds.Top + (Bounds.Height - 16 * DpiScaleY) / 2;
+			float iconX = bounds.Left + (bounds.Width - 16 * dpiScaleX) / 2;
+			float iconY = bounds.Top + (bounds.Height - 16 * dpiScaleY) / 2;
 
-			Graphics.DrawImage(Properties.Resources.Icons, IconX, IconY, Icon, GraphicsUnit.Pixel);
+			graphics.DrawImage(Properties.Resources.Icons, iconX, iconY, icon, GraphicsUnit.Pixel);
 		}
 
 		public void DrawNormalSubItem(DrawListViewSubItemEventArgs e)
@@ -403,19 +403,19 @@ namespace UnrealGameSync
 			DrawText(e.Graphics, e.SubItem.Bounds, Columns[e.ColumnIndex].TextAlign, SystemColors.WindowText, e.SubItem.Text);
 		}
 
-		public void DrawCustomSubItem(Graphics Graphics, ListViewItem.ListViewSubItem SubItem)
+		public void DrawCustomSubItem(Graphics graphics, ListViewItem.ListViewSubItem subItem)
 		{
-			CustomListViewWidget? Widget = SubItem.Tag as CustomListViewWidget;
-			if(Widget != null)
+			CustomListViewWidget? widget = subItem.Tag as CustomListViewWidget;
+			if(widget != null)
 			{
-				foreach(ListViewItem.ListViewSubItem? OtherSubItem in Widget.Item.SubItems)
+				foreach(ListViewItem.ListViewSubItem? otherSubItem in widget.Item.SubItems)
 				{
-					if(OtherSubItem != null && OtherSubItem.Tag == Widget)
+					if(otherSubItem != null && otherSubItem.Tag == widget)
 					{
-						if(OtherSubItem == SubItem)
+						if(otherSubItem == subItem)
 						{
-							ConditionalLayoutWidget(Widget);
-							Widget.Render(Graphics);
+							ConditionalLayoutWidget(widget);
+							widget.Render(graphics);
 						}
 						break;
 					}
@@ -423,48 +423,48 @@ namespace UnrealGameSync
 			}
 		}
 
-		public void DrawBackground(Graphics Graphics, ListViewItem Item)
+		public void DrawBackground(Graphics graphics, ListViewItem item)
 		{
-			if(Item.Selected)
+			if(item.Selected)
 			{
-				DrawSelectedBackground(Graphics, Item.Bounds);
+				DrawSelectedBackground(graphics, item.Bounds);
 			}
-			else if(Item.Index == HoverItem)
+			else if(item.Index == HoverItem)
 			{
-				DrawTrackedBackground(Graphics, Item.Bounds);
+				DrawTrackedBackground(graphics, item.Bounds);
 			}
 			else
 			{
-				DrawDefaultBackground(Graphics, Item.Bounds);
+				DrawDefaultBackground(graphics, item.Bounds);
 			}
 		}
 
-		public void DrawDefaultBackground(Graphics Graphics, Rectangle Bounds)
+		public void DrawDefaultBackground(Graphics graphics, Rectangle bounds)
 		{
-			Graphics.FillRectangle(SystemBrushes.Window, Bounds);
+			graphics.FillRectangle(SystemBrushes.Window, bounds);
 		}
 
-		public void DrawSelectedBackground(Graphics Graphics, Rectangle Bounds)
+		public void DrawSelectedBackground(Graphics graphics, Rectangle bounds)
 		{
-			if (SelectedItemRenderer != null)
+			if (_selectedItemRenderer != null)
 			{
-				SelectedItemRenderer.DrawBackground(Graphics, Bounds);
+				_selectedItemRenderer.DrawBackground(graphics, bounds);
 			}
 			else
 			{
-				Graphics.FillRectangle(SystemBrushes.ButtonFace, Bounds);
+				graphics.FillRectangle(SystemBrushes.ButtonFace, bounds);
 			}
 		}
 
-		public void DrawTrackedBackground(Graphics Graphics, Rectangle Bounds)
+		public void DrawTrackedBackground(Graphics graphics, Rectangle bounds)
 		{
-			if (TrackedItemRenderer != null)
+			if (_trackedItemRenderer != null)
 			{
-				TrackedItemRenderer.DrawBackground(Graphics, Bounds);
+				_trackedItemRenderer.DrawBackground(graphics, bounds);
 			}
 			else
 			{
-				Graphics.FillRectangle(SystemBrushes.Window, Bounds);
+				graphics.FillRectangle(SystemBrushes.Window, bounds);
 			}
 		}
 	}
diff --git a/Engine/Source/Programs/UnrealGameSync/UnrealGameSync/Controls/ErrorPanel.cs b/Engine/Source/Programs/UnrealGameSync/UnrealGameSync/Controls/ErrorPanel.cs
--- a/Engine/Source/Programs/UnrealGameSync/UnrealGameSync/Controls/ErrorPanel.cs
+++ b/Engine/Source/Programs/UnrealGameSync/UnrealGameSync/Controls/ErrorPanel.cs
@@ -11,9 +11,9 @@ namespace UnrealGameSync
 {
 	class ErrorPanel : StatusPanel, IMainWindowTabPanel
 	{
-		public ErrorPanel(UserSelectedProjectSettings InSelectedProject)
+		public ErrorPanel(UserSelectedProjectSettings inSelectedProject)
 		{
-			SelectedProject = InSelectedProject;
+			SelectedProject = inSelectedProject;
 			SetProjectLogo(Properties.Resources.DefaultErrorLogo, false);
 		}
 
