UGS: Reformat to standard C# style conventions.

diff --git a/Engine/Source/Programs/UnrealGameSync/UnrealGameSync/Controls/LogControl.cs b/Engine/Source/Programs/UnrealGameSync/UnrealGameSync/Controls/LogControl.cs
--- a/Engine/Source/Programs/UnrealGameSync/UnrealGameSync/Controls/LogControl.cs
+++ b/Engine/Source/Programs/UnrealGameSync/UnrealGameSync/Controls/LogControl.cs
@@ -23,16 +23,16 @@ namespace UnrealGameSync
 		[Flags]
 		public enum ScrollInfoMask : uint
 		{
-			SIF_RANGE = 0x1,
-			SIF_PAGE = 0x2,
-			SIF_POS = 0x4,
-			SIF_DISABLENOSCROLL = 0x8,
-			SIF_TRACKPOS = 0x10,
-			SIF_ALL = SIF_RANGE | SIF_PAGE | SIF_POS | SIF_TRACKPOS,
+			SifRange = 0x1,
+			SifPage = 0x2,
+			SifPos = 0x4,
+			SifDisablenoscroll = 0x8,
+			SifTrackpos = 0x10,
+			SifAll = SifRange | SifPage | SifPos | SifTrackpos,
 		}
 
 		[StructLayout(LayoutKind.Sequential)]
-		class SCROLLINFO 
+		class Scrollinfo 
 		{
 			public int cbSize;
 			public ScrollInfoMask fMask;
@@ -45,38 +45,38 @@ namespace UnrealGameSync
 
 		enum ScrollBarType : int
 		{
-			SB_HORZ = 0,
-			SB_VERT = 1,
-			SB_CONTROL = 2,
-			SB_BOTH = 3,
+			SbHorz = 0,
+			SbVert = 1,
+			SbControl = 2,
+			SbBoth = 3,
 		}
 
 		enum ScrollBarArrows : uint
 		{
-			ESB_ENABLE_BOTH = 0,
+			EsbEnableBoth = 0,
 		}
 
 		[DllImport("user32.dll")]
 		static extern bool EnableScrollBar(IntPtr hWnd, ScrollBarType wSBflags, ScrollBarArrows wArrows);
 
 		[DllImport("user32.dll")]
-		static extern int SetScrollInfo(IntPtr hwnd, ScrollBarType fnBar, SCROLLINFO lpsi, bool fRedraw);
+		static extern int SetScrollInfo(IntPtr hwnd, ScrollBarType fnBar, Scrollinfo lpsi, bool fRedraw);
 
 		[DllImport("user32.dll")]
-		static extern int GetScrollInfo(IntPtr hwnd, ScrollBarType fnBar, SCROLLINFO lpsi);
+		static extern int GetScrollInfo(IntPtr hwnd, ScrollBarType fnBar, Scrollinfo lpsi);
 
 		[DllImport("user32.dll", CharSet = CharSet.Auto)]
-		static extern IntPtr SendMessage(IntPtr hWnd, UInt32 Msg, IntPtr wParam, IntPtr lParam);
+		static extern IntPtr SendMessage(IntPtr hWnd, UInt32 msg, IntPtr wParam, IntPtr lParam);
 
 		struct TextLocation
 		{
 			public readonly int LineIdx;
 			public readonly int ColumnIdx;
 
-			public TextLocation(int InLineIdx, int InColumnIdx)
+			public TextLocation(int inLineIdx, int inColumnIdx)
 			{
-				LineIdx = InLineIdx;
-				ColumnIdx = InColumnIdx;
+				LineIdx = inLineIdx;
+				ColumnIdx = inColumnIdx;
 			}
 
 			public override string ToString()
@@ -90,10 +90,10 @@ namespace UnrealGameSync
 			public TextLocation Start;
 			public TextLocation End;
 
-			public TextSelection(TextLocation InStart, TextLocation InEnd)
+			public TextSelection(TextLocation inStart, TextLocation inEnd)
 			{
-				Start = InStart;
-				End = InEnd;
+				Start = inStart;
+				End = inEnd;
 			}
 
 			public bool IsEmpty()
@@ -107,27 +107,27 @@ namespace UnrealGameSync
 			}
 		}
 
-		List<string> Lines = new List<string>();
-		int MaxLineLength;
+		List<string> _lines = new List<string>();
+		int _maxLineLength;
 
-		int ScrollLine;
-		int ScrollLinesPerPage;
-		bool bTrackingScroll;
+		int _scrollLine;
+		int _scrollLinesPerPage;
+		bool _trackingScroll;
 
-		int ScrollColumn;
-		int ScrollColumnsPerPage;
+		int _scrollColumn;
+		int _scrollColumnsPerPage;
 
-		bool bIsSelecting;
-		TextSelection? Selection;
+		bool _isSelecting;
+		TextSelection? _selection;
 
-		Size FontSize;
+		Size _fontSize;
 
-		Timer? SelectionScrollTimer;
-		int AutoScrollRate = 0;
+		Timer? _selectionScrollTimer;
+		int _autoScrollRate = 0;
 
-		Timer? UpdateTimer;
-		ConcurrentQueue<string> QueuedLines = new ConcurrentQueue<string>();
-		FileStream? LogFileStream;
+		Timer? _updateTimer;
+		ConcurrentQueue<string> _queuedLines = new ConcurrentQueue<string>();
+		FileStream? _logFileStream;
 
 		public LogControl()
 		{
@@ -148,18 +148,18 @@ namespace UnrealGameSync
 		{
  			base.OnCreateControl();
 
-			EnableScrollBar(Handle, ScrollBarType.SB_BOTH, ScrollBarArrows.ESB_ENABLE_BOTH);
+			EnableScrollBar(Handle, ScrollBarType.SbBoth, ScrollBarArrows.EsbEnableBoth);
 
 			Cursor = Cursors.IBeam;
 
-			UpdateTimer = new Timer();
-			UpdateTimer.Interval = 200;
-			UpdateTimer.Tick += (a, b) => Tick();
-			UpdateTimer.Enabled = true;
+			_updateTimer = new Timer();
+			_updateTimer.Interval = 200;
+			_updateTimer.Tick += (a, b) => Tick();
+			_updateTimer.Enabled = true;
 
-			SelectionScrollTimer = new Timer();
-			SelectionScrollTimer.Interval = 200;
-			SelectionScrollTimer.Tick += new EventHandler(SelectionScrollTimer_TimerElapsed);
+			_selectionScrollTimer = new Timer();
+			_selectionScrollTimer.Interval = 200;
+			_selectionScrollTimer.Tick += new EventHandler(SelectionScrollTimer_TimerElapsed);
 
 			Clear();
 
@@ -171,34 +171,34 @@ namespace UnrealGameSync
 		{
 			base.Dispose(disposing);
 
-			if(UpdateTimer != null)
+			if(_updateTimer != null)
 			{
-				UpdateTimer.Dispose();
-				UpdateTimer = null;
+				_updateTimer.Dispose();
+				_updateTimer = null;
 			}
 
 			CloseFile();
 		}
 
-		public bool OpenFile(FileReference NewLogFileName)
+		public bool OpenFile(FileReference newLogFileName)
 		{
 			CloseFile();
 			Clear();
 			try
 			{
-				LogFileStream = FileReference.Open(NewLogFileName, FileMode.OpenOrCreate, FileAccess.ReadWrite, FileShare.Read);
+				_logFileStream = FileReference.Open(newLogFileName, FileMode.OpenOrCreate, FileAccess.ReadWrite, FileShare.Read);
 			}
 			catch(Exception)
 			{
 				return false;
 			}
 
-			LogFileStream.Seek(0, SeekOrigin.Begin);
+			_logFileStream.Seek(0, SeekOrigin.Begin);
 
-			string Text = new StreamReader(LogFileStream).ReadToEnd().TrimEnd('\r', '\n');
-			if(Text.Length > 0)
+			string text = new StreamReader(_logFileStream).ReadToEnd().TrimEnd('\r', '\n');
+			if(text.Length > 0)
 			{
-				AddLinesInternal(Text.Split('\n').Select(x => x + "\n").ToList());
+				AddLinesInternal(text.Split('\n').Select(x => x + "\n").ToList());
 			}
 
 			Invalidate();
@@ -207,184 +207,184 @@ namespace UnrealGameSync
 
 		public void CloseFile()
 		{
-			if(LogFileStream != null)
+			if(_logFileStream != null)
 			{
-				LogFileStream.Dispose();
-				LogFileStream = null;
+				_logFileStream.Dispose();
+				_logFileStream = null;
 			}
 		}
 
 		public void Clear()
 		{
-			Lines.Clear();
-			MaxLineLength = 0;
+			_lines.Clear();
+			_maxLineLength = 0;
 
-			ScrollLine = 0;
-			ScrollColumn = 0;
-			bIsSelecting = false;
-			Selection = null;
+			_scrollLine = 0;
+			_scrollColumn = 0;
+			_isSelecting = false;
+			_selection = null;
 
 			if(IsHandleCreated)
 			{
-				SCROLLINFO ScrollInfo = new SCROLLINFO();
-				ScrollInfo.cbSize = Marshal.SizeOf(ScrollInfo);
-				ScrollInfo.fMask = ScrollInfoMask.SIF_RANGE | ScrollInfoMask.SIF_PAGE | ScrollInfoMask.SIF_POS;
-				SetScrollInfo(Handle, ScrollBarType.SB_HORZ, ScrollInfo, true);
-				SetScrollInfo(Handle, ScrollBarType.SB_VERT, ScrollInfo, true);
+				Scrollinfo scrollInfo = new Scrollinfo();
+				scrollInfo.cbSize = Marshal.SizeOf(scrollInfo);
+				scrollInfo.fMask = ScrollInfoMask.SifRange | ScrollInfoMask.SifPage | ScrollInfoMask.SifPos;
+				SetScrollInfo(Handle, ScrollBarType.SbHorz, scrollInfo, true);
+				SetScrollInfo(Handle, ScrollBarType.SbVert, scrollInfo, true);
 			}
 
-			QueuedLines = new ConcurrentQueue<string>();
-			if(LogFileStream != null)
+			_queuedLines = new ConcurrentQueue<string>();
+			if(_logFileStream != null)
 			{
-				LogFileStream.SetLength(0);
+				_logFileStream.SetLength(0);
 			}
 		}
 
 		public void ScrollToEnd()
 		{
-			ScrollWindow(Lines.Count);
+			ScrollWindow(_lines.Count);
 			Invalidate();
 		}
 
-		public void AppendLine(string Line)
+		public void AppendLine(string line)
 		{
-			QueuedLines.Enqueue(Line);
+			_queuedLines.Enqueue(line);
 		}
 
 		private void Tick()
 		{
 			if(!Focused || !Capture)
 			{
-				List<string> NewLines = new List<string>();
+				List<string> newLines = new List<string>();
 				for(;;)
 				{
-					string? NextLine;
-					if(!QueuedLines.TryDequeue(out NextLine))
+					string? nextLine;
+					if(!_queuedLines.TryDequeue(out nextLine))
 					{
 						break;
 					}
-					NewLines.Add(NextLine.TrimEnd('\r', '\n') + "\n");
+					newLines.Add(nextLine.TrimEnd('\r', '\n') + "\n");
 				}
-				if(NewLines.Count > 0)
+				if(newLines.Count > 0)
 				{
-					StringBuilder TextToAppendBuilder = new StringBuilder();
-					foreach(string NewLine in NewLines)
+					StringBuilder textToAppendBuilder = new StringBuilder();
+					foreach(string newLine in newLines)
 					{
-						TextToAppendBuilder.AppendLine(NewLine.TrimEnd('\n'));
+						textToAppendBuilder.AppendLine(newLine.TrimEnd('\n'));
 					}
 
-					string TextToAppend = TextToAppendBuilder.ToString();
-					if(LogFileStream != null)
+					string textToAppend = textToAppendBuilder.ToString();
+					if(_logFileStream != null)
 					{
-						byte[] Data = Encoding.UTF8.GetBytes(TextToAppend);
+						byte[] data = Encoding.UTF8.GetBytes(textToAppend);
 						try
 						{
-							LogFileStream.Write(Data, 0, Data.Length);
-							LogFileStream.Flush();
+							_logFileStream.Write(data, 0, data.Length);
+							_logFileStream.Flush();
 						}
-						catch(Exception Ex)
+						catch(Exception ex)
 						{
-							TextToAppend += String.Format("Failed to write to log file ({0}): {1}\n", LogFileStream.Name, Ex.ToString());
+							textToAppend += String.Format("Failed to write to log file ({0}): {1}\n", _logFileStream.Name, ex.ToString());
 						}
 					}
 
-					AddLinesInternal(NewLines);
+					AddLinesInternal(newLines);
 				}
 			}
 		}
 
-		private void AddLinesInternal(List<string> NewLines)
+		private void AddLinesInternal(List<string> newLines)
 		{
-			Lines.AddRange(NewLines);
+			_lines.AddRange(newLines);
 
 			if(IsHandleCreated)
 			{
 				// Figure out if we're tracking the last line
-				bool bIsTrackingLastLine = IsTrackingLastLine();
+				bool isTrackingLastLine = IsTrackingLastLine();
 
-				int NewMaxLineLength = Math.Max(MaxLineLength, NewLines.Max(x => x.Length));
-				if(NewMaxLineLength > MaxLineLength)
+				int newMaxLineLength = Math.Max(_maxLineLength, newLines.Max(x => x.Length));
+				if(newMaxLineLength > _maxLineLength)
 				{
-					MaxLineLength = NewMaxLineLength;
+					_maxLineLength = newMaxLineLength;
 
-					SCROLLINFO HorizontalScroll = new SCROLLINFO();
-					HorizontalScroll.cbSize = Marshal.SizeOf(HorizontalScroll);
-					HorizontalScroll.fMask = ScrollInfoMask.SIF_PAGE | ScrollInfoMask.SIF_RANGE;
-					HorizontalScroll.nMin = 0;
-					HorizontalScroll.nMax = NewMaxLineLength;
-					HorizontalScroll.nPage = ScrollColumnsPerPage;
-					SetScrollInfo(Handle, ScrollBarType.SB_HORZ, HorizontalScroll, true);
+					Scrollinfo horizontalScroll = new Scrollinfo();
+					horizontalScroll.cbSize = Marshal.SizeOf(horizontalScroll);
+					horizontalScroll.fMask = ScrollInfoMask.SifPage | ScrollInfoMask.SifRange;
+					horizontalScroll.nMin = 0;
+					horizontalScroll.nMax = newMaxLineLength;
+					horizontalScroll.nPage = _scrollColumnsPerPage;
+					SetScrollInfo(Handle, ScrollBarType.SbHorz, horizontalScroll, true);
 				}
 
-				SCROLLINFO VerticalScroll = new SCROLLINFO();
-				VerticalScroll.cbSize = Marshal.SizeOf(VerticalScroll);
-				VerticalScroll.fMask = ScrollInfoMask.SIF_POS | ScrollInfoMask.SIF_RANGE | ScrollInfoMask.SIF_PAGE | ScrollInfoMask.SIF_TRACKPOS;
-				GetScrollInfo(Handle, ScrollBarType.SB_VERT, VerticalScroll);
+				Scrollinfo verticalScroll = new Scrollinfo();
+				verticalScroll.cbSize = Marshal.SizeOf(verticalScroll);
+				verticalScroll.fMask = ScrollInfoMask.SifPos | ScrollInfoMask.SifRange | ScrollInfoMask.SifPage | ScrollInfoMask.SifTrackpos;
+				GetScrollInfo(Handle, ScrollBarType.SbVert, verticalScroll);
 
-				if(bTrackingScroll)
+				if(_trackingScroll)
 				{
-					UpdateVerticalScrollPosition(VerticalScroll.nTrackPos, ref VerticalScroll);
+					UpdateVerticalScrollPosition(verticalScroll.nTrackPos, ref verticalScroll);
 				}
 				else
 				{
-					VerticalScroll.fMask = ScrollInfoMask.SIF_RANGE | ScrollInfoMask.SIF_PAGE;
-					VerticalScroll.nMin = 0;
-					if(bIsTrackingLastLine)
+					verticalScroll.fMask = ScrollInfoMask.SifRange | ScrollInfoMask.SifPage;
+					verticalScroll.nMin = 0;
+					if(isTrackingLastLine)
 					{
-						VerticalScroll.fMask |= ScrollInfoMask.SIF_POS;
-						VerticalScroll.nPos = Math.Max(Lines.Count - 1 - ScrollLinesPerPage + 1, 0);
-						ScrollLine = Math.Max(Lines.Count - 1 - ScrollLinesPerPage + 1, 0);
+						verticalScroll.fMask |= ScrollInfoMask.SifPos;
+						verticalScroll.nPos = Math.Max(_lines.Count - 1 - _scrollLinesPerPage + 1, 0);
+						_scrollLine = Math.Max(_lines.Count - 1 - _scrollLinesPerPage + 1, 0);
 					}
-					VerticalScroll.nMax = Math.Max(Lines.Count - 1, 1);
-					VerticalScroll.nPage = ScrollLinesPerPage;
-					SetScrollInfo(Handle, ScrollBarType.SB_VERT, VerticalScroll, true);
+					verticalScroll.nMax = Math.Max(_lines.Count - 1, 1);
+					verticalScroll.nPage = _scrollLinesPerPage;
+					SetScrollInfo(Handle, ScrollBarType.SbVert, verticalScroll, true);
 				}
 			}
 
 			Invalidate();
 		}
 
-		private SCROLLINFO GetScrollInfo(ScrollBarType Type, ScrollInfoMask Mask)
+		private Scrollinfo GetScrollInfo(ScrollBarType type, ScrollInfoMask mask)
 		{
-			SCROLLINFO ScrollInfo = new SCROLLINFO();
-			ScrollInfo.cbSize = Marshal.SizeOf(VerticalScroll);
-			ScrollInfo.fMask = Mask;
-			GetScrollInfo(Handle, ScrollBarType.SB_VERT, ScrollInfo);
-			return ScrollInfo;
+			Scrollinfo scrollInfo = new Scrollinfo();
+			scrollInfo.cbSize = Marshal.SizeOf(VerticalScroll);
+			scrollInfo.fMask = mask;
+			GetScrollInfo(Handle, ScrollBarType.SbVert, scrollInfo);
+			return scrollInfo;
 		}
 
 		private bool IsTrackingLastLine()
 		{
-			SCROLLINFO VerticalScroll = new SCROLLINFO();
-			VerticalScroll.cbSize = Marshal.SizeOf(VerticalScroll);
-			VerticalScroll.fMask = ScrollInfoMask.SIF_POS | ScrollInfoMask.SIF_RANGE | ScrollInfoMask.SIF_PAGE | ScrollInfoMask.SIF_TRACKPOS;
-			GetScrollInfo(Handle, ScrollBarType.SB_VERT, VerticalScroll);
+			Scrollinfo verticalScroll = new Scrollinfo();
+			verticalScroll.cbSize = Marshal.SizeOf(verticalScroll);
+			verticalScroll.fMask = ScrollInfoMask.SifPos | ScrollInfoMask.SifRange | ScrollInfoMask.SifPage | ScrollInfoMask.SifTrackpos;
+			GetScrollInfo(Handle, ScrollBarType.SbVert, verticalScroll);
 
-			return (VerticalScroll.nPos >= VerticalScroll.nMax - ScrollLinesPerPage + 1);
+			return (verticalScroll.nPos >= verticalScroll.nMax - _scrollLinesPerPage + 1);
 		}
 
 		private void CopySelection()
 		{
-			if(Lines.Count > 0)
+			if(_lines.Count > 0)
 			{
-				StringBuilder SelectedText = new StringBuilder();
-				if(Selection == null || Selection.IsEmpty())
+				StringBuilder selectedText = new StringBuilder();
+				if(_selection == null || _selection.IsEmpty())
 				{
-					foreach(string Line in Lines)
+					foreach(string line in _lines)
 					{
-						SelectedText.Append(Line);
+						selectedText.Append(line);
 					}
 				}
 				else
 				{
-					for(int LineIdx = Math.Min(Selection.Start.LineIdx, Selection.End.LineIdx); LineIdx <= Math.Max(Selection.Start.LineIdx, Selection.End.LineIdx); LineIdx++)
+					for(int lineIdx = Math.Min(_selection.Start.LineIdx, _selection.End.LineIdx); lineIdx <= Math.Max(_selection.Start.LineIdx, _selection.End.LineIdx); lineIdx++)
 					{
-						int MinIdx, MaxIdx;
-						ClipSelectionToLine(LineIdx, out MinIdx, out MaxIdx);
-						SelectedText.Append(Lines[LineIdx], MinIdx, MaxIdx - MinIdx);
+						int minIdx, maxIdx;
+						ClipSelectionToLine(lineIdx, out minIdx, out maxIdx);
+						selectedText.Append(_lines[lineIdx], minIdx, maxIdx - minIdx);
 					}
 				}
-				Clipboard.SetText(Regex.Replace(SelectedText.ToString(), "(?<!\r)\n", "\r\n"));
+				Clipboard.SetText(Regex.Replace(selectedText.ToString(), "(?<!\r)\n", "\r\n"));
 			}
 		}
 
@@ -400,33 +400,33 @@ namespace UnrealGameSync
 
 		public void SelectAll()
 		{
-			if(Lines.Count > 0)
+			if(_lines.Count > 0)
 			{
-				Selection = new TextSelection(new TextLocation(0, 0), new TextLocation(Lines.Count - 1, Lines[Lines.Count - 1].Length));
+				_selection = new TextSelection(new TextLocation(0, 0), new TextLocation(_lines.Count - 1, _lines[_lines.Count - 1].Length));
 				Invalidate();
 			}
 		}
 
-		TextLocation PointToTextLocation(Point ClientPoint)
+		TextLocation PointToTextLocation(Point clientPoint)
 		{
-			if(Lines.Count == 0)
+			if(_lines.Count == 0)
 			{
 				return new TextLocation(0, 0);
 			}
 			else
 			{
-				int UnclippedLineIdx = (ClientPoint.Y / FontSize.Height) + ScrollLine;
-				if(UnclippedLineIdx < 0)
+				int unclippedLineIdx = (clientPoint.Y / _fontSize.Height) + _scrollLine;
+				if(unclippedLineIdx < 0)
 				{
 					return new TextLocation(0, 0);
 				}
-				else if(UnclippedLineIdx >= Lines.Count)
+				else if(unclippedLineIdx >= _lines.Count)
 				{
-					return new TextLocation(Lines.Count - 1, Lines[Lines.Count - 1].Length);
+					return new TextLocation(_lines.Count - 1, _lines[_lines.Count - 1].Length);
 				}
 				else
 				{
-					return new TextLocation(UnclippedLineIdx, Math.Max(0, Math.Min((ClientPoint.X / FontSize.Width) + ScrollColumn, Lines[UnclippedLineIdx].Length)));
+					return new TextLocation(unclippedLineIdx, Math.Max(0, Math.Min((clientPoint.X / _fontSize.Width) + _scrollColumn, _lines[unclippedLineIdx].Length)));
 				}
 			}
 		}
@@ -437,13 +437,13 @@ namespace UnrealGameSync
 
 			UpdateScrollBarPageSize();
 
-			int NewScrollLine = ScrollLine;
-			NewScrollLine = Math.Min(NewScrollLine, Lines.Count - ScrollLinesPerPage);
-			NewScrollLine = Math.Max(NewScrollLine, 0);
+			int newScrollLine = _scrollLine;
+			newScrollLine = Math.Min(newScrollLine, _lines.Count - _scrollLinesPerPage);
+			newScrollLine = Math.Max(newScrollLine, 0);
 
-			if(ScrollLine != NewScrollLine)
+			if(_scrollLine != newScrollLine)
 			{
-				ScrollLine = NewScrollLine;
+				_scrollLine = newScrollLine;
 				Invalidate();
 			}
 		}
@@ -451,14 +451,14 @@ namespace UnrealGameSync
 		void UpdateVerticalScrollBarPageSize()
 		{
 			// Calculate the number of lines per page
-			ScrollLinesPerPage = ClientRectangle.Height / FontSize.Height;
+			_scrollLinesPerPage = ClientRectangle.Height / _fontSize.Height;
 
 			// Update the vertical scroll bar
-			SCROLLINFO VerticalScroll = new SCROLLINFO();
-			VerticalScroll.cbSize = Marshal.SizeOf(VerticalScroll);
-			VerticalScroll.fMask = ScrollInfoMask.SIF_PAGE;
-			VerticalScroll.nPage = ScrollLinesPerPage;
-			SetScrollInfo(Handle, ScrollBarType.SB_VERT, VerticalScroll, true);
+			Scrollinfo verticalScroll = new Scrollinfo();
+			verticalScroll.cbSize = Marshal.SizeOf(verticalScroll);
+			verticalScroll.fMask = ScrollInfoMask.SifPage;
+			verticalScroll.nPage = _scrollLinesPerPage;
+			SetScrollInfo(Handle, ScrollBarType.SbVert, verticalScroll, true);
 		}
 
 		void UpdateScrollBarPageSize()
@@ -467,14 +467,14 @@ namespace UnrealGameSync
 			UpdateVerticalScrollBarPageSize();
 
 			// This may have modified the client area. Now calculate the number of columns per page.
-			ScrollColumnsPerPage = ClientRectangle.Width / FontSize.Width;
+			_scrollColumnsPerPage = ClientRectangle.Width / _fontSize.Width;
 
 			// And update the horizontal scroll bar
-			SCROLLINFO HorizontalScroll = new SCROLLINFO();
-			HorizontalScroll.cbSize = Marshal.SizeOf(HorizontalScroll);
-			HorizontalScroll.fMask = ScrollInfoMask.SIF_PAGE;
-			HorizontalScroll.nPage = ScrollColumnsPerPage;
-			SetScrollInfo(Handle, ScrollBarType.SB_HORZ, HorizontalScroll, true);
+			Scrollinfo horizontalScroll = new Scrollinfo();
+			horizontalScroll.cbSize = Marshal.SizeOf(horizontalScroll);
+			horizontalScroll.fMask = ScrollInfoMask.SifPage;
+			horizontalScroll.nPage = _scrollColumnsPerPage;
+			SetScrollInfo(Handle, ScrollBarType.SbHorz, horizontalScroll, true);
 
 			// Now that we know whether we have a horizontal scroll bar or not, calculate the vertical scroll again
 			UpdateVerticalScrollBarPageSize();
@@ -489,9 +489,9 @@ namespace UnrealGameSync
 
 		private void UpdateFontMetrics()
 		{
-			using(Graphics Graphics = CreateGraphics())
+			using(Graphics graphics = CreateGraphics())
 			{
-				FontSize = TextRenderer.MeasureText(Graphics, "A", Font, new Size(int.MaxValue, int.MaxValue), TextFormatFlags.NoPadding);
+				_fontSize = TextRenderer.MeasureText(graphics, "A", Font, new Size(int.MaxValue, int.MaxValue), TextFormatFlags.NoPadding);
 			}
 		}
 
@@ -501,13 +501,13 @@ namespace UnrealGameSync
 
 			if(e.Button == MouseButtons.Left)
 			{
-				bIsSelecting = true;
+				_isSelecting = true;
 
-				TextLocation Location = PointToTextLocation(e.Location);
-				Selection = new TextSelection(Location, Location);
+				TextLocation location = PointToTextLocation(e.Location);
+				_selection = new TextSelection(location, location);
 
-				SelectionScrollTimer!.Start();
-				AutoScrollRate = 0;
+				_selectionScrollTimer!.Start();
+				_autoScrollRate = 0;
 
 				Capture = true;
 				Invalidate();
@@ -525,16 +525,16 @@ namespace UnrealGameSync
 		{
 			base.OnMouseMove(e);
 
-			if(bIsSelecting && Selection != null)
+			if(_isSelecting && _selection != null)
 			{
-				TextLocation NewSelectionEnd = PointToTextLocation(e.Location);
-				if(NewSelectionEnd.LineIdx != Selection.End.LineIdx || NewSelectionEnd.ColumnIdx != Selection.End.ColumnIdx)
+				TextLocation newSelectionEnd = PointToTextLocation(e.Location);
+				if(newSelectionEnd.LineIdx != _selection.End.LineIdx || newSelectionEnd.ColumnIdx != _selection.End.ColumnIdx)
 				{
-					Selection.End = NewSelectionEnd;
+					_selection.End = newSelectionEnd;
 					Invalidate();
 				}
 
-				AutoScrollRate = ((e.Y < 0)? e.Y : Math.Max(e.Y - ClientSize.Height, 0)) / FontSize.Height;
+				_autoScrollRate = ((e.Y < 0)? e.Y : Math.Max(e.Y - ClientSize.Height, 0)) / _fontSize.Height;
 			}
 		}
 
@@ -560,15 +560,15 @@ namespace UnrealGameSync
 			Invalidate();
 		}
 
-		protected override bool IsInputKey(Keys Key)
+		protected override bool IsInputKey(Keys key)
 		{
-			switch(Key & ~Keys.Modifiers)
+			switch(key & ~Keys.Modifiers)
 			{
 				case Keys.Up:
 				case Keys.Down:
 					return true;
 				default:
-					return base.IsInputKey(Key);
+					return base.IsInputKey(key);
 			}
 		}
 
@@ -585,19 +585,19 @@ namespace UnrealGameSync
 					Invalidate();
 					break;
 				case Keys.PageUp:
-					ScrollWindow(-ScrollLinesPerPage);
+					ScrollWindow(-_scrollLinesPerPage);
 					Invalidate();
 					break;
 				case Keys.PageDown:
-					ScrollWindow(+ScrollLinesPerPage);
+					ScrollWindow(+_scrollLinesPerPage);
 					Invalidate();
 					break;
 				case Keys.Home:
-					ScrollWindow(-Lines.Count);
+					ScrollWindow(-_lines.Count);
 					Invalidate();
 					break;
 				case Keys.End:
-					ScrollWindow(+Lines.Count);
+					ScrollWindow(+_lines.Count);
 					Invalidate();
 					break;
 				case Keys.A:
@@ -617,14 +617,14 @@ namespace UnrealGameSync
 
 		void EndSelection()
 		{
-			if(bIsSelecting)
+			if(_isSelecting)
 			{
-				if(Selection != null && Selection.IsEmpty())
+				if(_selection != null && _selection.IsEmpty())
 				{
-					Selection = null;
+					_selection = null;
 				}
-				SelectionScrollTimer!.Stop();
-				bIsSelecting = false;
+				_selectionScrollTimer!.Stop();
+				_isSelecting = false;
 			}
 		}
 
@@ -635,42 +635,42 @@ namespace UnrealGameSync
 			if(se.ScrollOrientation == ScrollOrientation.HorizontalScroll)
 			{
 				// Get the current scroll position
-				SCROLLINFO ScrollInfo = new SCROLLINFO();
-				ScrollInfo.cbSize = Marshal.SizeOf(ScrollInfo);
-				ScrollInfo.fMask = ScrollInfoMask.SIF_ALL;
-				GetScrollInfo(Handle, ScrollBarType.SB_HORZ, ScrollInfo);
+				Scrollinfo scrollInfo = new Scrollinfo();
+				scrollInfo.cbSize = Marshal.SizeOf(scrollInfo);
+				scrollInfo.fMask = ScrollInfoMask.SifAll;
+				GetScrollInfo(Handle, ScrollBarType.SbHorz, scrollInfo);
 
 				// Get the new scroll position
-				int TargetScrollPos = ScrollInfo.nPos;
+				int targetScrollPos = scrollInfo.nPos;
 				switch(se.Type)
 				{
 					case ScrollEventType.SmallDecrement:
-						TargetScrollPos = ScrollInfo.nPos - 1;
+						targetScrollPos = scrollInfo.nPos - 1;
 						break;
 					case ScrollEventType.SmallIncrement:
-						TargetScrollPos = ScrollInfo.nPos + 1;
+						targetScrollPos = scrollInfo.nPos + 1;
 						break;
 					case ScrollEventType.LargeDecrement:
-						TargetScrollPos = ScrollInfo.nPos - ScrollInfo.nPage;
+						targetScrollPos = scrollInfo.nPos - scrollInfo.nPage;
 						break;
 					case ScrollEventType.LargeIncrement:
-						TargetScrollPos = ScrollInfo.nPos + ScrollInfo.nPage;
+						targetScrollPos = scrollInfo.nPos + scrollInfo.nPage;
 						break;
 					case ScrollEventType.ThumbPosition:
-						TargetScrollPos = ScrollInfo.nTrackPos;
+						targetScrollPos = scrollInfo.nTrackPos;
 						break;
 					case ScrollEventType.ThumbTrack:
-						TargetScrollPos = ScrollInfo.nTrackPos;
+						targetScrollPos = scrollInfo.nTrackPos;
 						break;
 				}
-				ScrollColumn = Math.Max(0, Math.Min(TargetScrollPos, ScrollInfo.nMax - ScrollInfo.nPage + 1));
+				_scrollColumn = Math.Max(0, Math.Min(targetScrollPos, scrollInfo.nMax - scrollInfo.nPage + 1));
 
 				// Update the scroll bar if we're not tracking
 				if(se.Type != ScrollEventType.ThumbTrack)
 				{
-					ScrollInfo.fMask = ScrollInfoMask.SIF_POS;
-					ScrollInfo.nPos = ScrollColumn;
-					SetScrollInfo(Handle, ScrollBarType.SB_HORZ, ScrollInfo, true);
+					scrollInfo.fMask = ScrollInfoMask.SifPos;
+					scrollInfo.nPos = _scrollColumn;
+					SetScrollInfo(Handle, ScrollBarType.SbHorz, scrollInfo, true);
 				}
 
 				Invalidate();
@@ -678,152 +678,152 @@ namespace UnrealGameSync
 			else if(se.ScrollOrientation == ScrollOrientation.VerticalScroll)
 			{
 				// Get the current scroll position
-				SCROLLINFO ScrollInfo = new SCROLLINFO();
-				ScrollInfo.cbSize = Marshal.SizeOf(ScrollInfo);
-				ScrollInfo.fMask = ScrollInfoMask.SIF_ALL;
-				GetScrollInfo(Handle, ScrollBarType.SB_VERT, ScrollInfo);
+				Scrollinfo scrollInfo = new Scrollinfo();
+				scrollInfo.cbSize = Marshal.SizeOf(scrollInfo);
+				scrollInfo.fMask = ScrollInfoMask.SifAll;
+				GetScrollInfo(Handle, ScrollBarType.SbVert, scrollInfo);
 
 				// Get the new scroll position
-				int TargetScrollPos = ScrollInfo.nPos;
+				int targetScrollPos = scrollInfo.nPos;
 				switch(se.Type)
 				{
 					case ScrollEventType.SmallDecrement:
-						TargetScrollPos = ScrollInfo.nPos - 1;
+						targetScrollPos = scrollInfo.nPos - 1;
 						break;
 					case ScrollEventType.SmallIncrement:
-						TargetScrollPos = ScrollInfo.nPos + 1;
+						targetScrollPos = scrollInfo.nPos + 1;
 						break;
 					case ScrollEventType.LargeDecrement:
-						TargetScrollPos = ScrollInfo.nPos - ScrollLinesPerPage;
+						targetScrollPos = scrollInfo.nPos - _scrollLinesPerPage;
 						break;
 					case ScrollEventType.LargeIncrement:
-						TargetScrollPos = ScrollInfo.nPos + ScrollLinesPerPage;
+						targetScrollPos = scrollInfo.nPos + _scrollLinesPerPage;
 						break;
 					case ScrollEventType.ThumbPosition:
-						TargetScrollPos = ScrollInfo.nTrackPos;
+						targetScrollPos = scrollInfo.nTrackPos;
 						break;
 					case ScrollEventType.ThumbTrack:
-						TargetScrollPos = ScrollInfo.nTrackPos;
+						targetScrollPos = scrollInfo.nTrackPos;
 						break;
 				}
 
 				// Try to move to the new scroll position
-				UpdateVerticalScrollPosition(TargetScrollPos, ref ScrollInfo);
+				UpdateVerticalScrollPosition(targetScrollPos, ref scrollInfo);
 
 				// Update the new range as well
 				if(se.Type == ScrollEventType.ThumbTrack)
 				{
-					bTrackingScroll = true;
+					_trackingScroll = true;
 				}
 				else
 				{
 					// If we've just finished tracking, allow updating the range
-					ScrollInfo.nPos = ScrollLine;
-					if(bTrackingScroll)
+					scrollInfo.nPos = _scrollLine;
+					if(_trackingScroll)
 					{
-						ScrollInfo.fMask |= ScrollInfoMask.SIF_RANGE;
-						ScrollInfo.nMax = Math.Max(Lines.Count, 1) - 1;
-						bTrackingScroll = false;
+						scrollInfo.fMask |= ScrollInfoMask.SifRange;
+						scrollInfo.nMax = Math.Max(_lines.Count, 1) - 1;
+						_trackingScroll = false;
 					}
-					SetScrollInfo(Handle, ScrollBarType.SB_VERT, ScrollInfo, true);
+					SetScrollInfo(Handle, ScrollBarType.SbVert, scrollInfo, true);
 				}
 
 				Invalidate();
 			}
 		}
 
-		void UpdateVerticalScrollPosition(int TargetScrollPos, ref SCROLLINFO ScrollInfo)
+		void UpdateVerticalScrollPosition(int targetScrollPos, ref Scrollinfo scrollInfo)
 		{
 			// Scale it up to the number of lines. We don't adjust the scroll max while tracking so we can easily scroll to the end without more stuff being added.
-			float Ratio = Math.Max(Math.Min((float)TargetScrollPos / (float)(ScrollInfo.nMax - ScrollLinesPerPage + 1), 1.0f), 0.0f);
+			float ratio = Math.Max(Math.Min((float)targetScrollPos / (float)(scrollInfo.nMax - _scrollLinesPerPage + 1), 1.0f), 0.0f);
 
 			// Calculate the new scroll line, rounding to the nearest
-			ScrollLine = (int)((((Math.Max(Lines.Count, 1) - 1) - ScrollLinesPerPage + 1) * Ratio) + 0.5f);
+			_scrollLine = (int)((((Math.Max(_lines.Count, 1) - 1) - _scrollLinesPerPage + 1) * ratio) + 0.5f);
 		}
 
 		protected override void OnPaint(PaintEventArgs e)
 		{
 			base.OnPaint(e);
 
-			int TextX = -ScrollColumn * FontSize.Width;
-			for(int Idx = ScrollLine; Idx < ScrollLine + ScrollLinesPerPage + 1 && Idx < Lines.Count; Idx++)
+			int textX = -_scrollColumn * _fontSize.Width;
+			for(int idx = _scrollLine; idx < _scrollLine + _scrollLinesPerPage + 1 && idx < _lines.Count; idx++)
 			{
-				int SelectMinIdx;
-				int SelectMaxIdx;
-				ClipSelectionToLine(Idx, out SelectMinIdx, out SelectMaxIdx);
+				int selectMinIdx;
+				int selectMaxIdx;
+				ClipSelectionToLine(idx, out selectMinIdx, out selectMaxIdx);
 
-				Color TextColor;
-				if(Regex.IsMatch(Lines[Idx], "(?<!\\w)error[: ]", RegexOptions.IgnoreCase))
+				Color textColor;
+				if(Regex.IsMatch(_lines[idx], "(?<!\\w)error[: ]", RegexOptions.IgnoreCase))
 				{
-					TextColor = Color.FromArgb(189, 54, 47);
+					textColor = Color.FromArgb(189, 54, 47);
 				}
-				else if(Regex.IsMatch(Lines[Idx], "(?<!\\w)warning[: ]", RegexOptions.IgnoreCase))
+				else if(Regex.IsMatch(_lines[idx], "(?<!\\w)warning[: ]", RegexOptions.IgnoreCase))
 				{
-					TextColor = Color.FromArgb(128, 128, 0);
+					textColor = Color.FromArgb(128, 128, 0);
 				}
 				else
 				{
-					TextColor = Color.Black;
+					textColor = Color.Black;
 				}
 
-				int TextY = (Idx - ScrollLine) * FontSize.Height;
-				if(SelectMinIdx > 0)
+				int textY = (idx - _scrollLine) * _fontSize.Height;
+				if(selectMinIdx > 0)
 				{
-					TextRenderer.DrawText(e.Graphics, Lines[Idx].Substring(0, SelectMinIdx), Font, new Point(TextX, TextY), TextColor, TextFormatFlags.NoPadding);
+					TextRenderer.DrawText(e.Graphics, _lines[idx].Substring(0, selectMinIdx), Font, new Point(textX, textY), textColor, TextFormatFlags.NoPadding);
 				}
-				if(SelectMinIdx < SelectMaxIdx)
+				if(selectMinIdx < selectMaxIdx)
 				{
-					e.Graphics.FillRectangle(SystemBrushes.Highlight, TextX + (SelectMinIdx * FontSize.Width), TextY, (SelectMaxIdx - SelectMinIdx) * FontSize.Width, FontSize.Height);
-					TextRenderer.DrawText(e.Graphics, Lines[Idx].Substring(SelectMinIdx, SelectMaxIdx - SelectMinIdx), Font, new Point(TextX + (SelectMinIdx * FontSize.Width), TextY), SystemColors.HighlightText, TextFormatFlags.NoPadding);
+					e.Graphics.FillRectangle(SystemBrushes.Highlight, textX + (selectMinIdx * _fontSize.Width), textY, (selectMaxIdx - selectMinIdx) * _fontSize.Width, _fontSize.Height);
+					TextRenderer.DrawText(e.Graphics, _lines[idx].Substring(selectMinIdx, selectMaxIdx - selectMinIdx), Font, new Point(textX + (selectMinIdx * _fontSize.Width), textY), SystemColors.HighlightText, TextFormatFlags.NoPadding);
 				}
-				if(SelectMaxIdx < Lines[Idx].Length)
+				if(selectMaxIdx < _lines[idx].Length)
 				{
-					TextRenderer.DrawText(e.Graphics, Lines[Idx].Substring(SelectMaxIdx), Font, new Point(TextX + (SelectMaxIdx * FontSize.Width), TextY), TextColor, TextFormatFlags.NoPadding);
+					TextRenderer.DrawText(e.Graphics, _lines[idx].Substring(selectMaxIdx), Font, new Point(textX + (selectMaxIdx * _fontSize.Width), textY), textColor, TextFormatFlags.NoPadding);
 				}
 			}
 		}
 
-		void ScrollWindow(int ScrollDelta)
+		void ScrollWindow(int scrollDelta)
 		{
-			SCROLLINFO VerticalScroll = new SCROLLINFO();
-			VerticalScroll.cbSize = Marshal.SizeOf(VerticalScroll);
-			VerticalScroll.fMask = ScrollInfoMask.SIF_ALL;
-			GetScrollInfo(Handle, ScrollBarType.SB_VERT, VerticalScroll);
+			Scrollinfo verticalScroll = new Scrollinfo();
+			verticalScroll.cbSize = Marshal.SizeOf(verticalScroll);
+			verticalScroll.fMask = ScrollInfoMask.SifAll;
+			GetScrollInfo(Handle, ScrollBarType.SbVert, verticalScroll);
 
-			VerticalScroll.nPos = Math.Min(Math.Max(VerticalScroll.nPos + ScrollDelta, 0), VerticalScroll.nMax - VerticalScroll.nPage + 1);
-			VerticalScroll.fMask = ScrollInfoMask.SIF_POS;
-			SetScrollInfo(Handle, ScrollBarType.SB_VERT, VerticalScroll, true);
+			verticalScroll.nPos = Math.Min(Math.Max(verticalScroll.nPos + scrollDelta, 0), verticalScroll.nMax - verticalScroll.nPage + 1);
+			verticalScroll.fMask = ScrollInfoMask.SifPos;
+			SetScrollInfo(Handle, ScrollBarType.SbVert, verticalScroll, true);
 
-			ScrollLine = VerticalScroll.nPos;
+			_scrollLine = verticalScroll.nPos;
 		}
 
-		protected void SelectionScrollTimer_TimerElapsed(object? Sender, EventArgs Args)
+		protected void SelectionScrollTimer_TimerElapsed(object? sender, EventArgs args)
 		{
-			if(AutoScrollRate != 0 && Selection != null)
+			if(_autoScrollRate != 0 && _selection != null)
 			{
-				ScrollWindow(AutoScrollRate);
+				ScrollWindow(_autoScrollRate);
 
-				SCROLLINFO VerticalScroll = new SCROLLINFO();
-				VerticalScroll.cbSize = Marshal.SizeOf(VerticalScroll);
-				VerticalScroll.fMask = ScrollInfoMask.SIF_ALL;
-				GetScrollInfo(Handle, ScrollBarType.SB_VERT, VerticalScroll);
+				Scrollinfo verticalScroll = new Scrollinfo();
+				verticalScroll.cbSize = Marshal.SizeOf(verticalScroll);
+				verticalScroll.fMask = ScrollInfoMask.SifAll;
+				GetScrollInfo(Handle, ScrollBarType.SbVert, verticalScroll);
 
-				if(AutoScrollRate < 0)
+				if(_autoScrollRate < 0)
 				{
-					if(Selection.End.LineIdx > VerticalScroll.nPos)
+					if(_selection.End.LineIdx > verticalScroll.nPos)
 					{
-						Selection.End = new TextLocation(VerticalScroll.nPos, 0);
+						_selection.End = new TextLocation(verticalScroll.nPos, 0);
 					}
 				}
 				else
 				{
-					if(Selection.End.LineIdx < VerticalScroll.nPos + VerticalScroll.nPage)
+					if(_selection.End.LineIdx < verticalScroll.nPos + verticalScroll.nPage)
 					{
-						int LineIdx = Math.Min(VerticalScroll.nPos + VerticalScroll.nPage, Lines.Count - 1);
+						int lineIdx = Math.Min(verticalScroll.nPos + verticalScroll.nPage, _lines.Count - 1);
 
-						if (LineIdx >= 0)
+						if (lineIdx >= 0)
 						{
-							Selection.End = new TextLocation(LineIdx, Lines[LineIdx].Length);
+							_selection.End = new TextLocation(lineIdx, _lines[lineIdx].Length);
 						}
 					}
 				}
@@ -832,80 +832,80 @@ namespace UnrealGameSync
 			}
 		}
 
-		protected void ClipSelectionToLine(int LineIdx, out int SelectMinIdx, out int SelectMaxIdx)
+		protected void ClipSelectionToLine(int lineIdx, out int selectMinIdx, out int selectMaxIdx)
 		{
-			if(Selection == null)
+			if(_selection == null)
 			{
-				SelectMinIdx = 0;
-				SelectMaxIdx = 0;
+				selectMinIdx = 0;
+				selectMaxIdx = 0;
 			}
-			else if(Selection.Start.LineIdx < Selection.End.LineIdx)
+			else if(_selection.Start.LineIdx < _selection.End.LineIdx)
 			{
-				if(LineIdx < Selection.Start.LineIdx)
+				if(lineIdx < _selection.Start.LineIdx)
 				{
-					SelectMinIdx = Lines[LineIdx].Length;
-					SelectMaxIdx = Lines[LineIdx].Length;
+					selectMinIdx = _lines[lineIdx].Length;
+					selectMaxIdx = _lines[lineIdx].Length;
 				}
-				else if(LineIdx == Selection.Start.LineIdx)
+				else if(lineIdx == _selection.Start.LineIdx)
 				{
-					SelectMinIdx = Selection.Start.ColumnIdx;
-					SelectMaxIdx = Lines[LineIdx].Length;
+					selectMinIdx = _selection.Start.ColumnIdx;
+					selectMaxIdx = _lines[lineIdx].Length;
 				}
-				else if(LineIdx < Selection.End.LineIdx)
+				else if(lineIdx < _selection.End.LineIdx)
 				{
-					SelectMinIdx = 0;
-					SelectMaxIdx = Lines[LineIdx].Length;
+					selectMinIdx = 0;
+					selectMaxIdx = _lines[lineIdx].Length;
 				}
-				else if(LineIdx == Selection.End.LineIdx)
+				else if(lineIdx == _selection.End.LineIdx)
 				{
-					SelectMinIdx = 0;
-					SelectMaxIdx = Selection.End.ColumnIdx;
+					selectMinIdx = 0;
+					selectMaxIdx = _selection.End.ColumnIdx;
 				}
 				else
 				{
-					SelectMinIdx = 0;
-					SelectMaxIdx = 0;
+					selectMinIdx = 0;
+					selectMaxIdx = 0;
 				}
 			}
-			else if(Selection.Start.LineIdx > Selection.End.LineIdx)
+			else if(_selection.Start.LineIdx > _selection.End.LineIdx)
 			{
-				if(LineIdx < Selection.End.LineIdx)
+				if(lineIdx < _selection.End.LineIdx)
 				{
-					SelectMinIdx = Lines[LineIdx].Length;
-					SelectMaxIdx = Lines[LineIdx].Length;
+					selectMinIdx = _lines[lineIdx].Length;
+					selectMaxIdx = _lines[lineIdx].Length;
 				}
-				else if(LineIdx == Selection.End.LineIdx)
+				else if(lineIdx == _selection.End.LineIdx)
 				{
-					SelectMinIdx = Selection.End.ColumnIdx;
-					SelectMaxIdx = Lines[LineIdx].Length;
+					selectMinIdx = _selection.End.ColumnIdx;
+					selectMaxIdx = _lines[lineIdx].Length;
 				}
-				else if(LineIdx < Selection.Start.LineIdx)
+				else if(lineIdx < _selection.Start.LineIdx)
 				{
-					SelectMinIdx = 0;
-					SelectMaxIdx = Lines[LineIdx].Length;
+					selectMinIdx = 0;
+					selectMaxIdx = _lines[lineIdx].Length;
 				}
-				else if(LineIdx == Selection.Start.LineIdx)
+				else if(lineIdx == _selection.Start.LineIdx)
 				{
-					SelectMinIdx = 0;
-					SelectMaxIdx = Selection.Start.ColumnIdx;
+					selectMinIdx = 0;
+					selectMaxIdx = _selection.Start.ColumnIdx;
 				}
 				else
 				{
-					SelectMinIdx = 0;
-					SelectMaxIdx = 0;
+					selectMinIdx = 0;
+					selectMaxIdx = 0;
 				}
 			}
 			else
 			{
-				if(LineIdx == Selection.Start.LineIdx)
+				if(lineIdx == _selection.Start.LineIdx)
 				{
-					SelectMinIdx = Math.Min(Selection.Start.ColumnIdx, Selection.End.ColumnIdx);
-					SelectMaxIdx = Math.Max(Selection.Start.ColumnIdx, Selection.End.ColumnIdx);
+					selectMinIdx = Math.Min(_selection.Start.ColumnIdx, _selection.End.ColumnIdx);
+					selectMaxIdx = Math.Max(_selection.Start.ColumnIdx, _selection.End.ColumnIdx);
 				}
 				else
 				{
-					SelectMinIdx = 0;
-					SelectMaxIdx = 0;
+					selectMinIdx = 0;
+					selectMaxIdx = 0;
 				}
 			}
 		}
@@ -918,23 +918,23 @@ namespace UnrealGameSync
 			public void Dispose() { }
 		}
 
-		LogControl LogControl;
+		LogControl _logControl;
 
-		public LogControlTextWriter(LogControl InLogControl)
+		public LogControlTextWriter(LogControl inLogControl)
 		{
-			LogControl = InLogControl;
+			_logControl = inLogControl;
 		}
 
-		public void Log<TState>(LogLevel LogLevel, EventId EventId, TState State, Exception Exception, Func<TState, Exception, string> Formatter)
+		public void Log<TState>(LogLevel logLevel, EventId eventId, TState state, Exception exception, Func<TState, Exception, string> formatter)
 		{
-			string Message = Formatter(State, Exception);
-			LogControl.AppendLine(Message);
+			string message = formatter(state, exception);
+			_logControl.AppendLine(message);
 
-			if (Exception != null)
+			if (exception != null)
 			{
-				foreach (string Line in Exception.ToString().Trim().Split('\n'))
+				foreach (string line in exception.ToString().Trim().Split('\n'))
 				{
-					LogControl.AppendLine(Line);
+					_logControl.AppendLine(line);
 				}
 			}
 		}
diff --git a/Engine/Source/Programs/UnrealGameSync/UnrealGameSync/Controls/LogSplitContainer.cs b/Engine/Source/Programs/UnrealGameSync/UnrealGameSync/Controls/LogSplitContainer.cs
--- a/Engine/Source/Programs/UnrealGameSync/UnrealGameSync/Controls/LogSplitContainer.cs
+++ b/Engine/Source/Programs/UnrealGameSync/UnrealGameSync/Controls/LogSplitContainer.cs
@@ -16,17 +16,17 @@ namespace UnrealGameSync
 {
 	public class LogSplitContainer : SplitContainer
 	{
-		const int WM_SETCURSOR = 0x0020;
+		const int WmSetcursor = 0x0020;
 
-		bool bHoverOverClose;
-		bool bMouseDownOverClose;
-		bool bLogVisible = true;
-		bool bSplitterMoved = false;
-		bool bChangingLogVisibility = false;
+		bool _hoverOverClose;
+		bool _mouseDownOverClose;
+		bool _logVisible = true;
+		bool _splitterMoved = false;
+		bool _changingLogVisibility = false;
 
 		const int CaptionPadding = 4;
 
-		int LogHeight = 200;
+		int _logHeight = 200;
 
 		public event Action<bool>? OnVisibilityChanged;
 
@@ -68,50 +68,50 @@ namespace UnrealGameSync
 			Invalidate();
 		}
 
-		public void OnSplitterMoved(object Obj, EventArgs Args)
+		public void OnSplitterMoved(object obj, EventArgs args)
 		{
-			bSplitterMoved = true;
+			_splitterMoved = true;
 
-			if(bLogVisible && !bChangingLogVisibility)
+			if(_logVisible && !_changingLogVisibility)
 			{
-				LogHeight = Height - (SplitterDistance + SplitterWidth);
+				_logHeight = Height - (SplitterDistance + SplitterWidth);
 			}
 
 			Invalidate();
 		}
 
-		public void SetLogVisibility(bool bVisible)
+		public void SetLogVisibility(bool visible)
 		{
-			if(bVisible != bLogVisible)
+			if(visible != _logVisible)
 			{
 				SuspendLayout();
 
-				bLogVisible = bVisible;
-				bChangingLogVisibility = true;
+				_logVisible = visible;
+				_changingLogVisibility = true;
 				UpdateMetrics();
-				if(bLogVisible)
+				if(_logVisible)
 				{
-					SplitterDistance = Math.Max(Height - LogHeight - SplitterWidth, Panel1MinSize);
+					SplitterDistance = Math.Max(Height - _logHeight - SplitterWidth, Panel1MinSize);
 				}
 				else
 				{
 					SplitterDistance = Height - SplitterWidth;
 				}
 
-				bChangingLogVisibility = false;
+				_changingLogVisibility = false;
 
 				ResumeLayout();
 
 				if(OnVisibilityChanged != null)
 				{
-					OnVisibilityChanged(bVisible);
+					OnVisibilityChanged(visible);
 				}
 			}
 		}
 
 		private void UpdateMetrics()
 		{
-			if(bLogVisible)
+			if(_logVisible)
 			{
 				IsSplitterFixed = false;
 				SplitterWidth = CaptionFont.Height + (CaptionPadding * 2) + 4;
@@ -136,7 +136,7 @@ namespace UnrealGameSync
 			{
 				UpdateMetrics();
 
-				if(!bLogVisible)
+				if(!_logVisible)
 				{
 					SplitterDistance = Height - SplitterWidth;
 				}
@@ -147,15 +147,15 @@ namespace UnrealGameSync
 		{
 			base.OnLayout(e);
 
-			if(bLogVisible && !bChangingLogVisibility)
+			if(_logVisible && !_changingLogVisibility)
 			{
-				LogHeight = Height - (SplitterDistance + SplitterWidth);
+				_logHeight = Height - (SplitterDistance + SplitterWidth);
 			}
 		}
 
 		public bool IsLogVisible()
 		{
-			return bLogVisible;
+			return _logVisible;
 		}
 
 		protected override void OnFontChanged(EventArgs e)
@@ -169,10 +169,10 @@ namespace UnrealGameSync
 		{
 			base.OnMouseMove(e);
 
-			bool bNewHoverOverClose = CloseButtonRectangle.Contains(e.Location);
-			if(bHoverOverClose != bNewHoverOverClose)
+			bool newHoverOverClose = CloseButtonRectangle.Contains(e.Location);
+			if(_hoverOverClose != newHoverOverClose)
 			{
-				bHoverOverClose = bNewHoverOverClose;
+				_hoverOverClose = newHoverOverClose;
 				Invalidate();
 			}
 		}
@@ -181,21 +181,21 @@ namespace UnrealGameSync
 		{
 			base.OnMouseLeave(e);
 
-			if(bHoverOverClose)
+			if(_hoverOverClose)
 			{
-				bHoverOverClose = false;
-				bMouseDownOverClose = false;
+				_hoverOverClose = false;
+				_mouseDownOverClose = false;
 				Invalidate();
 			}
 		}
 
 		protected override void OnMouseDown(MouseEventArgs e)
 		{
-			bSplitterMoved = false;
+			_splitterMoved = false;
 
 			if(CloseButtonRectangle.Contains(e.Location))
 			{
-				bMouseDownOverClose = bHoverOverClose;
+				_mouseDownOverClose = _hoverOverClose;
 			}
 			else
 			{
@@ -205,11 +205,11 @@ namespace UnrealGameSync
 
 		protected override void OnMouseUp(MouseEventArgs e)
 		{
-			bMouseDownOverClose = false;
+			_mouseDownOverClose = false;
 
 			base.OnMouseUp(e);
 
-			if(e.Button.HasFlag(MouseButtons.Left) && !bSplitterMoved)
+			if(e.Button.HasFlag(MouseButtons.Left) && !_splitterMoved)
 			{
 				SetLogVisibility(!IsLogVisible());
 			}
@@ -235,67 +235,67 @@ namespace UnrealGameSync
 		{
 			base.OnPaint(e);
 
-			int CaptionMinY = SplitterDistance + CaptionPadding;
-			int CaptionMaxY = SplitterDistance + SplitterWidth;
-			int ButtonSize = SplitterWidth - CaptionPadding - 2 - 4;
+			int captionMinY = SplitterDistance + CaptionPadding;
+			int captionMaxY = SplitterDistance + SplitterWidth;
+			int buttonSize = SplitterWidth - CaptionPadding - 2 - 4;
 			if(IsLogVisible())
 			{
-				CaptionMaxY -= CaptionPadding;
-				ButtonSize -= CaptionPadding;
+				captionMaxY -= CaptionPadding;
+				buttonSize -= CaptionPadding;
 			}
 
-			using(Pen CaptionBorderPen = new Pen(SystemColors.ControlDark, 1.0f))
+			using(Pen captionBorderPen = new Pen(SystemColors.ControlDark, 1.0f))
 			{
-				e.Graphics.DrawRectangle(CaptionBorderPen, 0, CaptionMinY, ClientRectangle.Width - 1, CaptionMaxY - CaptionMinY - 1);
+				e.Graphics.DrawRectangle(captionBorderPen, 0, captionMinY, ClientRectangle.Width - 1, captionMaxY - captionMinY - 1);
 			}
 
-			using(Brush BackgroundBrush = new SolidBrush(BackColor))
+			using(Brush backgroundBrush = new SolidBrush(BackColor))
 			{
-				e.Graphics.FillRectangle(BackgroundBrush, 0, 0, ClientRectangle.Width, CaptionMinY);
-				e.Graphics.FillRectangle(BackgroundBrush, 0, CaptionMaxY, ClientRectangle.Width, Height - CaptionMaxY);
-				e.Graphics.FillRectangle(BackgroundBrush, 1, CaptionMinY + 1, ClientRectangle.Width - 2, CaptionMaxY - CaptionMinY - 2);
+				e.Graphics.FillRectangle(backgroundBrush, 0, 0, ClientRectangle.Width, captionMinY);
+				e.Graphics.FillRectangle(backgroundBrush, 0, captionMaxY, ClientRectangle.Width, Height - captionMaxY);
+				e.Graphics.FillRectangle(backgroundBrush, 1, captionMinY + 1, ClientRectangle.Width - 2, captionMaxY - captionMinY - 2);
 			}
 
-			int CrossX = ClientRectangle.Right - (ButtonSize / 2) - 4;
-			int CrossY = (CaptionMinY + CaptionMaxY) / 2;
-			if(bHoverOverClose)
+			int crossX = ClientRectangle.Right - (buttonSize / 2) - 4;
+			int crossY = (captionMinY + captionMaxY) / 2;
+			if(_hoverOverClose)
 			{
-				e.Graphics.FillRectangle(SystemBrushes.ActiveCaption, CrossX - ButtonSize / 2, CrossY - ButtonSize / 2, ButtonSize, ButtonSize);
+				e.Graphics.FillRectangle(SystemBrushes.ActiveCaption, crossX - buttonSize / 2, crossY - buttonSize / 2, buttonSize, buttonSize);
 			}
-			else if(bMouseDownOverClose)
+			else if(_mouseDownOverClose)
 			{
-				e.Graphics.FillRectangle(SystemBrushes.InactiveCaption, CrossX - ButtonSize / 2, CrossY - ButtonSize / 2, ButtonSize - 2, ButtonSize - 2);
+				e.Graphics.FillRectangle(SystemBrushes.InactiveCaption, crossX - buttonSize / 2, crossY - buttonSize / 2, buttonSize - 2, buttonSize - 2);
 			}
-			if(bHoverOverClose || bMouseDownOverClose)
+			if(_hoverOverClose || _mouseDownOverClose)
 			{
-				using(Pen BorderPen = new Pen(SystemColors.InactiveCaptionText, 1.0f))
+				using(Pen borderPen = new Pen(SystemColors.InactiveCaptionText, 1.0f))
 				{
-					e.Graphics.DrawRectangle(BorderPen, CrossX - ButtonSize / 2, CrossY - ButtonSize / 2, ButtonSize, ButtonSize);
+					e.Graphics.DrawRectangle(borderPen, crossX - buttonSize / 2, crossY - buttonSize / 2, buttonSize, buttonSize);
 				}
 			}
 			if(SplitterDistance >= Height - SplitterWidth)
 			{
-				e.Graphics.DrawImage(Properties.Resources.Log, new Rectangle(CrossX - (ButtonSize / 2) - 1, CrossY - (ButtonSize / 2) - 1, ButtonSize + 2, ButtonSize + 2), new Rectangle(16, 0, 16, 16), GraphicsUnit.Pixel);
+				e.Graphics.DrawImage(Properties.Resources.Log, new Rectangle(crossX - (buttonSize / 2) - 1, crossY - (buttonSize / 2) - 1, buttonSize + 2, buttonSize + 2), new Rectangle(16, 0, 16, 16), GraphicsUnit.Pixel);
 			}
 			else
 			{
-				e.Graphics.DrawImage(Properties.Resources.Log, new Rectangle(CrossX - (ButtonSize / 2) - 1, CrossY - (ButtonSize / 2) - 1, ButtonSize + 2, ButtonSize + 2), new Rectangle(0, 0, 16, 16), GraphicsUnit.Pixel);
+				e.Graphics.DrawImage(Properties.Resources.Log, new Rectangle(crossX - (buttonSize / 2) - 1, crossY - (buttonSize / 2) - 1, buttonSize + 2, buttonSize + 2), new Rectangle(0, 0, 16, 16), GraphicsUnit.Pixel);
 			}
 
-			TextRenderer.DrawText(e.Graphics, Caption, CaptionFont, new Rectangle(2, CaptionMinY, ClientRectangle.Width - 20, CaptionMaxY - CaptionMinY), SystemColors.ControlText, TextFormatFlags.Left | TextFormatFlags.VerticalCenter);
+			TextRenderer.DrawText(e.Graphics, Caption, CaptionFont, new Rectangle(2, captionMinY, ClientRectangle.Width - 20, captionMaxY - captionMinY), SystemColors.ControlText, TextFormatFlags.Left | TextFormatFlags.VerticalCenter);
 		}
 
-		protected override void WndProc(ref Message Message)
+		protected override void WndProc(ref Message message)
 		{
-			base.WndProc(ref Message);
+			base.WndProc(ref message);
 
-			switch (Message.Msg)
+			switch (message.Msg)
 			{
-				case WM_SETCURSOR:
+				case WmSetcursor:
 					if(RectangleToScreen(CloseButtonRectangle).Contains(Cursor.Position))
 					{
 						Cursor.Current = Cursors.Default;
-						Message.Result = new IntPtr(1);
+						message.Result = new IntPtr(1);
 					}
 					return;
 			}
diff --git a/Engine/Source/Programs/UnrealGameSync/UnrealGameSync/Controls/OIDCControl.cs b/Engine/Source/Programs/UnrealGameSync/UnrealGameSync/Controls/OIDCControl.cs
--- a/Engine/Source/Programs/UnrealGameSync/UnrealGameSync/Controls/OIDCControl.cs
+++ b/Engine/Source/Programs/UnrealGameSync/UnrealGameSync/Controls/OIDCControl.cs
@@ -7,81 +7,81 @@ using System.Windows.Forms;
 
 namespace UnrealGameSync.Controls
 {
-	public partial class OIDCControl : UserControl
+	public partial class OidcControl : UserControl
 	{
-		private Font BadgeFont;
-		private readonly OIDCTokenManager OidcManager;
-		private readonly string ProviderIdentifier;
+		private Font _badgeFont;
+		private readonly OidcTokenManager _oidcManager;
+		private readonly string _providerIdentifier;
 
-		public OIDCControl(OIDCTokenManager InOidcManager, string InProviderIdentifier, string ServiceName)
+		public OidcControl(OidcTokenManager inOidcManager, string inProviderIdentifier, string serviceName)
 		{
 			InitializeComponent();
 
-			OIDCControlGroupBox.Text = ServiceName;
-			OidcManager = InOidcManager;
-			ProviderIdentifier = InProviderIdentifier;
+			OIDCControlGroupBox.Text = serviceName;
+			_oidcManager = inOidcManager;
+			_providerIdentifier = inProviderIdentifier;
 
-			BadgeFont = new Font(this.Font.FontFamily, this.Font.SizeInPoints - 2, FontStyle.Bold);
+			_badgeFont = new Font(this.Font.FontFamily, this.Font.SizeInPoints - 2, FontStyle.Bold);
 
-			OIDCStatus ServiceStatus = OidcManager.GetStatusForProvider(ProviderIdentifier);
-			LoginButton.Enabled = ServiceStatus == OIDCStatus.NotLoggedIn;
+			OidcStatus serviceStatus = _oidcManager.GetStatusForProvider(_providerIdentifier);
+			LoginButton.Enabled = serviceStatus == OidcStatus.NotLoggedIn;
 		}
 
-		private void DrawBadge(Graphics Graphics, Rectangle BadgeRect, string BadgeText, Color BadgeColor, bool bMergeLeft, bool bMergeRight)
+		private void DrawBadge(Graphics graphics, Rectangle badgeRect, string badgeText, Color badgeColor, bool mergeLeft, bool mergeRight)
 		{
-			if (BadgeColor.A != 0)
+			if (badgeColor.A != 0)
 			{
-				using (GraphicsPath Path = new GraphicsPath())
+				using (GraphicsPath path = new GraphicsPath())
 				{
-					Path.StartFigure();
-					Path.AddLine(BadgeRect.Left + (bMergeLeft ? 1 : 0), BadgeRect.Top, BadgeRect.Left - (bMergeLeft ? 1 : 0), BadgeRect.Bottom);
-					Path.AddLine(BadgeRect.Left - (bMergeLeft ? 1 : 0), BadgeRect.Bottom, BadgeRect.Right - 1 - (bMergeRight ? 1 : 0), BadgeRect.Bottom);
-					Path.AddLine(BadgeRect.Right - 1 - (bMergeRight ? 1 : 0), BadgeRect.Bottom, BadgeRect.Right - 1 + (bMergeRight ? 1 : 0), BadgeRect.Top);
-					Path.AddLine(BadgeRect.Right - 1 + (bMergeRight ? 1 : 0), BadgeRect.Top, BadgeRect.Left + (bMergeLeft ? 1 : 0), BadgeRect.Top);
-					Path.CloseFigure();
-
-					using (SolidBrush Brush = new SolidBrush(BadgeColor))
+					path.StartFigure();
+					path.AddLine(badgeRect.Left + (mergeLeft ? 1 : 0), badgeRect.Top, badgeRect.Left - (mergeLeft ? 1 : 0), badgeRect.Bottom);
+					path.AddLine(badgeRect.Left - (mergeLeft ? 1 : 0), badgeRect.Bottom, badgeRect.Right - 1 - (mergeRight ? 1 : 0), badgeRect.Bottom);
+					path.AddLine(badgeRect.Right - 1 - (mergeRight ? 1 : 0), badgeRect.Bottom, badgeRect.Right - 1 + (mergeRight ? 1 : 0), badgeRect.Top);
+					path.AddLine(badgeRect.Right - 1 + (mergeRight ? 1 : 0), badgeRect.Top, badgeRect.Left + (mergeLeft ? 1 : 0), badgeRect.Top);
+					path.CloseFigure();
+
+					using (SolidBrush brush = new SolidBrush(badgeColor))
 					{
-						Graphics.FillPath(Brush, Path);
+						graphics.FillPath(brush, path);
 					}
 				}
 
-				TextRenderer.DrawText(Graphics, BadgeText, BadgeFont, BadgeRect, Color.White, TextFormatFlags.HorizontalCenter | TextFormatFlags.VerticalCenter | TextFormatFlags.SingleLine | TextFormatFlags.NoPrefix | TextFormatFlags.PreserveGraphicsClipping);
+				TextRenderer.DrawText(graphics, badgeText, _badgeFont, badgeRect, Color.White, TextFormatFlags.HorizontalCenter | TextFormatFlags.VerticalCenter | TextFormatFlags.SingleLine | TextFormatFlags.NoPrefix | TextFormatFlags.PreserveGraphicsClipping);
 			}
 		}
 
 		private void StatusPanel_Paint(object sender, PaintEventArgs e)
 		{
-			OIDCStatus ServiceStatus = OidcManager.GetStatusForProvider(ProviderIdentifier);
-			DrawBadge(e.Graphics, e.ClipRectangle, ToServiceDisplayText(ServiceStatus), ToServiceDisplayColor(ServiceStatus), false, false);
+			OidcStatus serviceStatus = _oidcManager.GetStatusForProvider(_providerIdentifier);
+			DrawBadge(e.Graphics, e.ClipRectangle, ToServiceDisplayText(serviceStatus), ToServiceDisplayColor(serviceStatus), false, false);
 		}
 
-		private Color ToServiceDisplayColor(OIDCStatus OidcStatus)
+		private Color ToServiceDisplayColor(OidcStatus oidcStatus)
 		{
-			switch (OidcStatus)
+			switch (oidcStatus)
 			{
-				case OIDCStatus.Connected:
-				case OIDCStatus.TokenRefreshRequired:
+				case OidcStatus.Connected:
+				case OidcStatus.TokenRefreshRequired:
 					return Color.Green;
-				case OIDCStatus.NotLoggedIn:
+				case OidcStatus.NotLoggedIn:
 					return Color.Red;
 				default:
-					throw new ArgumentOutOfRangeException(nameof(OidcStatus), OidcStatus, null);
+					throw new ArgumentOutOfRangeException(nameof(oidcStatus), oidcStatus, null);
 			}
 		}
 
-		private string ToServiceDisplayText(OIDCStatus OidcStatus)
+		private string ToServiceDisplayText(OidcStatus oidcStatus)
 		{
-			switch (OidcStatus)
+			switch (oidcStatus)
 			{
-				case OIDCStatus.Connected:
+				case OidcStatus.Connected:
 					return "Connected";
-				case OIDCStatus.NotLoggedIn:
+				case OidcStatus.NotLoggedIn:
 					return "Not Logged In";
-				case OIDCStatus.TokenRefreshRequired:
+				case OidcStatus.TokenRefreshRequired:
 					return "Refresh Pending";
 				default:
-					throw new ArgumentOutOfRangeException(nameof(OidcStatus), OidcStatus, null);
+					throw new ArgumentOutOfRangeException(nameof(oidcStatus), oidcStatus, null);
 			}
 		}
 
@@ -89,18 +89,18 @@ namespace UnrealGameSync.Controls
 		{
 			try
 			{
-				await OidcManager.Login(ProviderIdentifier);
+				await _oidcManager.Login(_providerIdentifier);
 
 				Focus();
 
-				OIDCStatus ServiceStatus = OidcManager.GetStatusForProvider(ProviderIdentifier);
-				LoginButton.Enabled = ServiceStatus == OIDCStatus.NotLoggedIn;
+				OidcStatus serviceStatus = _oidcManager.GetStatusForProvider(_providerIdentifier);
+				LoginButton.Enabled = serviceStatus == OidcStatus.NotLoggedIn;
 
 				Refresh();
 			}
-			catch (LoginFailedException Exception)
+			catch (LoginFailedException exception)
 			{
-				MessageBox.Show(Exception.Message, "Login Failed", MessageBoxButtons.OK);
+				MessageBox.Show(exception.Message, "Login Failed", MessageBoxButtons.OK);
 			}
 		}
 	}
diff --git a/Engine/Source/Programs/UnrealGameSync/UnrealGameSync/Controls/RichTextBox41.cs b/Engine/Source/Programs/UnrealGameSync/UnrealGameSync/Controls/RichTextBox41.cs
--- a/Engine/Source/Programs/UnrealGameSync/UnrealGameSync/Controls/RichTextBox41.cs
+++ b/Engine/Source/Programs/UnrealGameSync/UnrealGameSync/Controls/RichTextBox41.cs
@@ -19,12 +19,12 @@ namespace UnrealGameSync
 		{
 			get
 			{
-				CreateParams Params = base.CreateParams;
+				CreateParams p = base.CreateParams;
 				if (LoadLibrary("msftedit.dll") != IntPtr.Zero)
 				{
-					Params.ClassName = "RICHEDIT50W";
+					p.ClassName = "RICHEDIT50W";
 				}
-				return Params;
+				return p;
 			}
 		}
 	}
diff --git a/Engine/Source/Programs/UnrealGameSync/UnrealGameSync/Controls/StatusPanel.cs b/Engine/Source/Programs/UnrealGameSync/UnrealGameSync/Controls/StatusPanel.cs
--- a/Engine/Source/Programs/UnrealGameSync/UnrealGameSync/Controls/StatusPanel.cs
+++ b/Engine/Source/Programs/UnrealGameSync/UnrealGameSync/Controls/StatusPanel.cs
@@ -20,186 +20,186 @@ namespace UnrealGameSync
 		[DllImport("user32.dll")]
 		public static extern IntPtr LoadCursor(IntPtr hInstance, IntPtr lpCursorName);
 
-		const int IDC_HAND = 32649;
+		const int IdcHand = 32649;
 
-		public static readonly Cursor Hand = new Cursor(LoadCursor(IntPtr.Zero, new IntPtr(IDC_HAND)));
+		public static readonly Cursor Hand = new Cursor(LoadCursor(IntPtr.Zero, new IntPtr(IdcHand)));
 	}
 
 	class StatusElementResources
 	{
-		Dictionary<FontStyle, Font> FontCache = new Dictionary<FontStyle, Font>();
+		Dictionary<FontStyle, Font> _fontCache = new Dictionary<FontStyle, Font>();
 		public readonly Font BadgeFont;
 
-		public StatusElementResources(Font BaseFont)
+		public StatusElementResources(Font baseFont)
 		{
-			FontCache.Add(FontStyle.Regular, BaseFont);
+			_fontCache.Add(FontStyle.Regular, baseFont);
 
-			using(Graphics Graphics = Graphics.FromHwnd(IntPtr.Zero))
+			using(Graphics graphics = Graphics.FromHwnd(IntPtr.Zero))
 			{
-				BadgeFont = new Font(BaseFont.FontFamily, 7.0f/* * Graphics.DpiY / 96.0f*/, FontStyle.Bold);
+				BadgeFont = new Font(baseFont.FontFamily, 7.0f/* * Graphics.DpiY / 96.0f*/, FontStyle.Bold);
 			}
 		}
 
 		public void Dispose()
 		{
-			foreach(KeyValuePair<FontStyle, Font> FontPair in FontCache)
+			foreach(KeyValuePair<FontStyle, Font> fontPair in _fontCache)
 			{
-				if(FontPair.Key != FontStyle.Regular)
+				if(fontPair.Key != FontStyle.Regular)
 				{
-					FontPair.Value.Dispose();
+					fontPair.Value.Dispose();
 				}
 			}
 
 			BadgeFont.Dispose();
 		}
 
-		public Font FindOrAddFont(FontStyle Style)
+		public Font FindOrAddFont(FontStyle style)
 		{
-			Font? Result;
-			if(!FontCache.TryGetValue(Style, out Result))
+			Font? result;
+			if(!_fontCache.TryGetValue(style, out result))
 			{
-				Result = new Font(FontCache[FontStyle.Regular], Style);
-				FontCache[Style] = Result;
+				result = new Font(_fontCache[FontStyle.Regular], style);
+				_fontCache[style] = result;
 			}
-			return Result;
+			return result;
 		}
 	}
 
 	abstract class StatusElement
 	{
 		public Cursor Cursor = Cursors.Arrow;
-		public bool bMouseOver;
-		public bool bMouseDown;
+		public bool MouseOver;
+		public bool MouseDown;
 		public Rectangle Bounds;
 
-		public Point Layout(Graphics Graphics, Point Location, StatusElementResources Resources)
+		public Point Layout(Graphics graphics, Point location, StatusElementResources resources)
 		{
-			Size Size = Measure(Graphics, Resources);
-			Bounds = new Rectangle(Location.X, Location.Y - (Size.Height + 1) / 2, Size.Width, Size.Height);
-			return new Point(Location.X + Size.Width, Location.Y);
+			Size size = Measure(graphics, resources);
+			Bounds = new Rectangle(location.X, location.Y - (size.Height + 1) / 2, size.Width, size.Height);
+			return new Point(location.X + size.Width, location.Y);
 		}
 
-		public abstract Size Measure(Graphics Graphics, StatusElementResources Resources);
-		public abstract void Draw(Graphics Grahpics, StatusElementResources Resources);
+		public abstract Size Measure(Graphics graphics, StatusElementResources resources);
+		public abstract void Draw(Graphics grahpics, StatusElementResources resources);
 
-		public virtual void OnClick(Point Location)
+		public virtual void OnClick(Point location)
 		{
 		}
 	}
 
 	class IconStatusElement : StatusElement
 	{
-		Image Icon;
+		Image _icon;
 
-		public IconStatusElement(Image InIcon)
+		public IconStatusElement(Image inIcon)
 		{
-			Icon = InIcon;
+			_icon = inIcon;
 		}
 
-		public override Size Measure(Graphics Graphics, StatusElementResources Resources)
+		public override Size Measure(Graphics graphics, StatusElementResources resources)
 		{
-			return new Size((int)(Icon.Width * Graphics.DpiX / 96.0f), (int)(Icon.Height * Graphics.DpiY / 96.0f));
+			return new Size((int)(_icon.Width * graphics.DpiX / 96.0f), (int)(_icon.Height * graphics.DpiY / 96.0f));
 		}
 
-		public override void Draw(Graphics Graphics, StatusElementResources Resources)
+		public override void Draw(Graphics graphics, StatusElementResources resources)
 		{
-			Graphics.DrawImage(Icon, Bounds.Location);
+			graphics.DrawImage(_icon, Bounds.Location);
 		}
 	}
 
 	class IconStripStatusElement : StatusElement
 	{
-		Image Strip;
-		Size IconSize;
-		int Index;
+		Image _strip;
+		Size _iconSize;
+		int _index;
 
-		public IconStripStatusElement(Image InStrip, Size InIconSize, int InIndex)
+		public IconStripStatusElement(Image inStrip, Size inIconSize, int inIndex)
 		{
-			Strip = InStrip;
-			IconSize = InIconSize;
-			Index = InIndex;
+			_strip = inStrip;
+			_iconSize = inIconSize;
+			_index = inIndex;
 		}
 
-		public override Size Measure(Graphics Graphics, StatusElementResources Resources)
+		public override Size Measure(Graphics graphics, StatusElementResources resources)
 		{
-			return new Size((int)(IconSize.Width * Graphics.DpiX / 96.0f), (int)(IconSize.Height * Graphics.DpiY / 96.0f));
+			return new Size((int)(_iconSize.Width * graphics.DpiX / 96.0f), (int)(_iconSize.Height * graphics.DpiY / 96.0f));
 		}
 
-		public override void Draw(Graphics Graphics, StatusElementResources Resources)
+		public override void Draw(Graphics graphics, StatusElementResources resources)
 		{
-			Graphics.DrawImage(Strip, Bounds, new Rectangle(IconSize.Width * Index, 0, IconSize.Width, IconSize.Height), GraphicsUnit.Pixel);
+			graphics.DrawImage(_strip, Bounds, new Rectangle(_iconSize.Width * _index, 0, _iconSize.Width, _iconSize.Height), GraphicsUnit.Pixel);
 		}
 	}
 
 	class TextStatusElement : StatusElement
 	{
-		string Text;
-		Color Color;
-		FontStyle Style;
+		string _text;
+		Color _color;
+		FontStyle _style;
 
-		public TextStatusElement(string InText, Color InColor, FontStyle InStyle)
+		public TextStatusElement(string inText, Color inColor, FontStyle inStyle)
 		{
-			Text = InText;
-			Color = InColor;
-			Style = InStyle;
+			_text = inText;
+			_color = inColor;
+			_style = inStyle;
 		}
 
-		public override Size Measure(Graphics Graphics, StatusElementResources Resources)
+		public override Size Measure(Graphics graphics, StatusElementResources resources)
 		{
-			return TextRenderer.MeasureText(Graphics, Text, Resources.FindOrAddFont(Style), new Size(int.MaxValue, int.MaxValue), TextFormatFlags.NoPadding);
+			return TextRenderer.MeasureText(graphics, _text, resources.FindOrAddFont(_style), new Size(int.MaxValue, int.MaxValue), TextFormatFlags.NoPadding);
 		}
 
-		public override void Draw(Graphics Graphics, StatusElementResources Resources)
+		public override void Draw(Graphics graphics, StatusElementResources resources)
 		{
-			TextRenderer.DrawText(Graphics, Text, Resources.FindOrAddFont(Style), Bounds.Location, Color, TextFormatFlags.NoPadding);
+			TextRenderer.DrawText(graphics, _text, resources.FindOrAddFont(_style), Bounds.Location, _color, TextFormatFlags.NoPadding);
 		}
 	}
 
 	class LinkStatusElement : StatusElement
 	{
-		string Text;
-		FontStyle Style;
-		Action<Point, Rectangle> LinkAction;
+		string _text;
+		FontStyle _style;
+		Action<Point, Rectangle> _linkAction;
 
-		public LinkStatusElement(string InText, FontStyle InStyle, Action<Point, Rectangle> InLinkAction)
+		public LinkStatusElement(string inText, FontStyle inStyle, Action<Point, Rectangle> inLinkAction)
 		{
-			Text = InText;
-			Style = InStyle;
-			LinkAction = InLinkAction;
+			_text = inText;
+			_style = inStyle;
+			_linkAction = inLinkAction;
 			Cursor = NativeCursors.Hand;
 		}
 
-		public override void OnClick(Point Location)
+		public override void OnClick(Point location)
 		{
-			LinkAction(Location, Bounds);
+			_linkAction(location, Bounds);
 		}
 
-		public override Size Measure(Graphics Graphics, StatusElementResources Resources)
+		public override Size Measure(Graphics graphics, StatusElementResources resources)
 		{
-			return TextRenderer.MeasureText(Graphics, Text, Resources.FindOrAddFont(Style), new Size(int.MaxValue, int.MaxValue), TextFormatFlags.NoPadding | TextFormatFlags.NoPrefix);
+			return TextRenderer.MeasureText(graphics, _text, resources.FindOrAddFont(_style), new Size(int.MaxValue, int.MaxValue), TextFormatFlags.NoPadding | TextFormatFlags.NoPrefix);
 		}
 
-		public override void Draw(Graphics Graphics, StatusElementResources Resources)
+		public override void Draw(Graphics graphics, StatusElementResources resources)
 		{
-			Color TextColor = SystemColors.HotTrack;
-			if(bMouseDown)
+			Color textColor = SystemColors.HotTrack;
+			if(MouseDown)
 			{
-				TextColor = Color.FromArgb(TextColor.B / 2, TextColor.G / 2, TextColor.R);
+				textColor = Color.FromArgb(textColor.B / 2, textColor.G / 2, textColor.R);
 			}
-			else if(bMouseOver)
+			else if(MouseOver)
 			{
-				TextColor = Color.FromArgb(TextColor.B, TextColor.G, TextColor.R);
+				textColor = Color.FromArgb(textColor.B, textColor.G, textColor.R);
 			}
-			TextRenderer.DrawText(Graphics, Text, Resources.FindOrAddFont(Style), Bounds.Location, TextColor, TextFormatFlags.NoPadding | TextFormatFlags.NoPrefix);
+			TextRenderer.DrawText(graphics, _text, resources.FindOrAddFont(_style), Bounds.Location, textColor, TextFormatFlags.NoPadding | TextFormatFlags.NoPrefix);
 		}
 	}
 
 	class BadgeStatusElement : StatusElement
 	{
-		string Name;
-		Color BackgroundColor;
-		Color HoverBackgroundColor;
-		Action<Point, Rectangle>? ClickAction;
+		string _name;
+		Color _backgroundColor;
+		Color _hoverBackgroundColor;
+		Action<Point, Rectangle>? _clickAction;
 
 		public bool MergeLeft
 		{
@@ -211,100 +211,100 @@ namespace UnrealGameSync
 			get; set;
 		}
 
-		public BadgeStatusElement(string InName, Color InBackgroundColor, Action<Point, Rectangle>? InClickAction)
+		public BadgeStatusElement(string inName, Color inBackgroundColor, Action<Point, Rectangle>? inClickAction)
 		{
-			Name = InName;
-			BackgroundColor = InBackgroundColor;
-			ClickAction = InClickAction;
+			_name = inName;
+			_backgroundColor = inBackgroundColor;
+			_clickAction = inClickAction;
 
-			if (ClickAction == null)
+			if (_clickAction == null)
 			{
-				HoverBackgroundColor = BackgroundColor;
+				_hoverBackgroundColor = _backgroundColor;
 			}
 			else
 			{
-				HoverBackgroundColor = Color.FromArgb(Math.Min(BackgroundColor.R + 32, 255), Math.Min(BackgroundColor.G + 32, 255), Math.Min(BackgroundColor.B + 32, 255));
+				_hoverBackgroundColor = Color.FromArgb(Math.Min(_backgroundColor.R + 32, 255), Math.Min(_backgroundColor.G + 32, 255), Math.Min(_backgroundColor.B + 32, 255));
 			}
-			if(ClickAction != null)
+			if(_clickAction != null)
 			{
 				Cursor = NativeCursors.Hand;
 			}
 		}
 
-		public override void OnClick(Point Location)
+		public override void OnClick(Point location)
 		{
-			if(ClickAction != null)
+			if(_clickAction != null)
 			{
-				ClickAction(Location, Bounds);
+				_clickAction(location, Bounds);
 			}
 		}
 
-		public override Size Measure(Graphics Graphics, StatusElementResources Resources)
+		public override Size Measure(Graphics graphics, StatusElementResources resources)
 		{
-			Size LabelSize = TextRenderer.MeasureText(Name, Resources.BadgeFont);
-			int BadgeHeight = Resources.BadgeFont.Height + 1;
-			return new Size(LabelSize.Width + 1 + BadgeHeight - 4, BadgeHeight);
+			Size labelSize = TextRenderer.MeasureText(_name, resources.BadgeFont);
+			int badgeHeight = resources.BadgeFont.Height + 1;
+			return new Size(labelSize.Width + 1 + badgeHeight - 4, badgeHeight);
 		}
 
-		public override void Draw(Graphics Graphics, StatusElementResources Resources)
+		public override void Draw(Graphics graphics, StatusElementResources resources)
 		{
-			SmoothingMode PrevSmoothingMode = Graphics.SmoothingMode;
-			Graphics.SmoothingMode = SmoothingMode.HighQuality;
+			SmoothingMode prevSmoothingMode = graphics.SmoothingMode;
+			graphics.SmoothingMode = SmoothingMode.HighQuality;
 
-			using (GraphicsPath Path = new GraphicsPath())
+			using (GraphicsPath path = new GraphicsPath())
 			{
-				Path.StartFigure();
-				Path.AddLine(Bounds.Left + (MergeLeft? 1 : 0), Bounds.Top, Bounds.Left - (MergeLeft? 1 : 0), Bounds.Bottom);
-				Path.AddLine(Bounds.Left - (MergeLeft? 1 : 0), Bounds.Bottom, Bounds.Right - 2 - (MergeRight? 1 : 0), Bounds.Bottom);
-				Path.AddLine(Bounds.Right - 2 - (MergeRight? 1 : 0), Bounds.Bottom, Bounds.Right - 2 + (MergeRight? 1 : 0), Bounds.Top);
-				Path.AddLine(Bounds.Right - 2 + (MergeRight? 1 : 0), Bounds.Top, Bounds.Left + (MergeLeft? 1 : 0), Bounds.Top);
-				Path.CloseFigure();
+				path.StartFigure();
+				path.AddLine(Bounds.Left + (MergeLeft? 1 : 0), Bounds.Top, Bounds.Left - (MergeLeft? 1 : 0), Bounds.Bottom);
+				path.AddLine(Bounds.Left - (MergeLeft? 1 : 0), Bounds.Bottom, Bounds.Right - 2 - (MergeRight? 1 : 0), Bounds.Bottom);
+				path.AddLine(Bounds.Right - 2 - (MergeRight? 1 : 0), Bounds.Bottom, Bounds.Right - 2 + (MergeRight? 1 : 0), Bounds.Top);
+				path.AddLine(Bounds.Right - 2 + (MergeRight? 1 : 0), Bounds.Top, Bounds.Left + (MergeLeft? 1 : 0), Bounds.Top);
+				path.CloseFigure();
 
-				using(SolidBrush Brush = new SolidBrush(bMouseOver? HoverBackgroundColor : BackgroundColor))
+				using(SolidBrush brush = new SolidBrush(MouseOver? _hoverBackgroundColor : _backgroundColor))
 				{
-					Graphics.FillPath(Brush, Path);
+					graphics.FillPath(brush, path);
 				}
 			}
 
-			Graphics.SmoothingMode = PrevSmoothingMode;
-			TextRenderer.DrawText(Graphics, Name, Resources.BadgeFont, Bounds, Color.White, TextFormatFlags.HorizontalCenter | TextFormatFlags.VerticalCenter | TextFormatFlags.SingleLine | TextFormatFlags.NoPrefix | TextFormatFlags.PreserveGraphicsClipping);
+			graphics.SmoothingMode = prevSmoothingMode;
+			TextRenderer.DrawText(graphics, _name, resources.BadgeFont, Bounds, Color.White, TextFormatFlags.HorizontalCenter | TextFormatFlags.VerticalCenter | TextFormatFlags.SingleLine | TextFormatFlags.NoPrefix | TextFormatFlags.PreserveGraphicsClipping);
 		}
 	}
 
 	class ProgressBarStatusElement : StatusElement
 	{
-		float Progress;
+		float _progress;
 
-		public ProgressBarStatusElement(float InProgress)
+		public ProgressBarStatusElement(float inProgress)
 		{
-			Progress = InProgress;
+			_progress = inProgress;
 		}
 
-		public override Size Measure(Graphics Graphics, StatusElementResources Resources)
+		public override Size Measure(Graphics graphics, StatusElementResources resources)
 		{
-			int Height = (int)(Resources.FindOrAddFont(FontStyle.Regular).Height * 0.9f);
-			return new Size(Height * 16, Height);
+			int height = (int)(resources.FindOrAddFont(FontStyle.Regular).Height * 0.9f);
+			return new Size(height * 16, height);
 		}
 
-		public override void Draw(Graphics Graphics, StatusElementResources Resources)
+		public override void Draw(Graphics graphics, StatusElementResources resources)
 		{
-			Graphics.SmoothingMode = System.Drawing.Drawing2D.SmoothingMode.None;
+			graphics.SmoothingMode = System.Drawing.Drawing2D.SmoothingMode.None;
 
-			Graphics.DrawRectangle(Pens.Black, Bounds.Left, Bounds.Top, Bounds.Width - 1, Bounds.Height - 1);
-			Graphics.FillRectangle(Brushes.White, Bounds.Left + 1, Bounds.Top + 1, Bounds.Width - 2, Bounds.Height - 2);
+			graphics.DrawRectangle(Pens.Black, Bounds.Left, Bounds.Top, Bounds.Width - 1, Bounds.Height - 1);
+			graphics.FillRectangle(Brushes.White, Bounds.Left + 1, Bounds.Top + 1, Bounds.Width - 2, Bounds.Height - 2);
 
-			int ProgressX = Bounds.Left + 2 + (int)((Bounds.Width - 4) * Progress);
-			using(Brush ProgressBarBrush = new SolidBrush(Color.FromArgb(112, 146, 190)))
+			int progressX = Bounds.Left + 2 + (int)((Bounds.Width - 4) * _progress);
+			using(Brush progressBarBrush = new SolidBrush(Color.FromArgb(112, 146, 190)))
 			{
-				Graphics.FillRectangle(ProgressBarBrush, Bounds.Left + 2, Bounds.Y + 2, ProgressX - (Bounds.Left + 2), Bounds.Height - 4);
+				graphics.FillRectangle(progressBarBrush, Bounds.Left + 2, Bounds.Y + 2, progressX - (Bounds.Left + 2), Bounds.Height - 4);
 			}
 		}
 	}
 
 	class StatusLine
 	{
-		bool bModified;
-		List<StatusElement> Elements = new List<StatusElement>();
+		bool _modified;
+		List<StatusElement> _elements = new List<StatusElement>();
 
 		public StatusLine()
 		{
@@ -325,89 +325,89 @@ namespace UnrealGameSync
 
 		public bool RequiresLayout()
 		{
-			return bModified;
+			return _modified;
 		}
 
 		public void Clear()
 		{
-			Elements.Clear();
-			bModified = true;
+			_elements.Clear();
+			_modified = true;
 		}
 
-		public void Add(StatusElement Element)
+		public void Add(StatusElement element)
 		{
-			Elements.Add(Element);
-			bModified = true;
+			_elements.Add(element);
+			_modified = true;
 		}
 
-		public void AddIcon(Image InIcon)
+		public void AddIcon(Image inIcon)
 		{
-			Elements.Add(new IconStatusElement(InIcon));
-			bModified = true;
+			_elements.Add(new IconStatusElement(inIcon));
+			_modified = true;
 		}
 
-		public void AddIcon(Image InStrip, Size InIconSize, int InIndex)
+		public void AddIcon(Image inStrip, Size inIconSize, int inIndex)
 		{
-			Elements.Add(new IconStripStatusElement(InStrip, InIconSize, InIndex));
-			bModified = true;
+			_elements.Add(new IconStripStatusElement(inStrip, inIconSize, inIndex));
+			_modified = true;
 		}
 
-		public void AddText(string InText, FontStyle InStyle = FontStyle.Regular)
+		public void AddText(string inText, FontStyle inStyle = FontStyle.Regular)
 		{
-			Elements.Add(new TextStatusElement(InText, SystemColors.ControlText, InStyle));
-			bModified = true;
+			_elements.Add(new TextStatusElement(inText, SystemColors.ControlText, inStyle));
+			_modified = true;
 		}
 
-		public void AddText(string InText, Color InColor, FontStyle InStyle = FontStyle.Regular)
+		public void AddText(string inText, Color inColor, FontStyle inStyle = FontStyle.Regular)
 		{
-			Elements.Add(new TextStatusElement(InText, InColor, InStyle));
-			bModified = true;
+			_elements.Add(new TextStatusElement(inText, inColor, inStyle));
+			_modified = true;
 		}
 
-		public void AddLink(string InText, FontStyle InStyle, Action InLinkAction)
+		public void AddLink(string inText, FontStyle inStyle, Action inLinkAction)
 		{
-			Elements.Add(new LinkStatusElement(InText, InStyle, (P, R) => { InLinkAction(); }));
-			bModified = true;
+			_elements.Add(new LinkStatusElement(inText, inStyle, (p, r) => { inLinkAction(); }));
+			_modified = true;
 		}
 
-		public void AddLink(string InText, FontStyle InStyle, Action<Point, Rectangle> InLinkAction)
+		public void AddLink(string inText, FontStyle inStyle, Action<Point, Rectangle> inLinkAction)
 		{
-			Elements.Add(new LinkStatusElement(InText, InStyle, InLinkAction));
-			bModified = true;
+			_elements.Add(new LinkStatusElement(inText, inStyle, inLinkAction));
+			_modified = true;
 		}
 
-		public void AddBadge(string InText, Color InBackgroundColor, Action<Point, Rectangle>? InClickAction)
+		public void AddBadge(string inText, Color inBackgroundColor, Action<Point, Rectangle>? inClickAction)
 		{
-			Elements.Add(new BadgeStatusElement(InText, InBackgroundColor, InClickAction));
-			if(Elements.Count >= 2)
+			_elements.Add(new BadgeStatusElement(inText, inBackgroundColor, inClickAction));
+			if(_elements.Count >= 2)
 			{
-				BadgeStatusElement? PrevBadge = Elements[Elements.Count - 2] as BadgeStatusElement;
-				BadgeStatusElement? NextBadge = Elements[Elements.Count - 1] as BadgeStatusElement;
-				if(PrevBadge != null && NextBadge != null)
+				BadgeStatusElement? prevBadge = _elements[_elements.Count - 2] as BadgeStatusElement;
+				BadgeStatusElement? nextBadge = _elements[_elements.Count - 1] as BadgeStatusElement;
+				if(prevBadge != null && nextBadge != null)
 				{
-					PrevBadge.MergeRight = true;
-					NextBadge.MergeLeft = true;
+					prevBadge.MergeRight = true;
+					nextBadge.MergeLeft = true;
 				}
 			}
-			bModified = true;
+			_modified = true;
 		}
 
-		public void AddProgressBar(float Progress)
+		public void AddProgressBar(float progress)
 		{
-			Elements.Add(new ProgressBarStatusElement(Progress));
-			bModified = true;
+			_elements.Add(new ProgressBarStatusElement(progress));
+			_modified = true;
 		}
 
-		public bool HitTest(Point Location, [NotNullWhen(true)] out StatusElement? OutElement)
+		public bool HitTest(Point location, [NotNullWhen(true)] out StatusElement? outElement)
 		{
-			OutElement = null;
-			if(Bounds.Contains(Location))
+			outElement = null;
+			if(Bounds.Contains(location))
 			{
-				foreach(StatusElement Element in Elements)
+				foreach(StatusElement element in _elements)
 				{
-					if(Element.Bounds.Contains(Location))
+					if(element.Bounds.Contains(location))
 					{
-						OutElement = Element;
+						outElement = element;
 						return true;
 					}
 				}
@@ -415,25 +415,25 @@ namespace UnrealGameSync
 			return false;
 		}
 
-		public void Layout(Graphics Graphics, Point Location, StatusElementResources Resources)
+		public void Layout(Graphics graphics, Point location, StatusElementResources resources)
 		{
-			Bounds = new Rectangle(Location, new Size(0, 0));
+			Bounds = new Rectangle(location, new Size(0, 0));
 
-			Point NextLocation = Location;
-			foreach(StatusElement Element in Elements)
+			Point nextLocation = location;
+			foreach(StatusElement element in _elements)
 			{
-				NextLocation = Element.Layout(Graphics, NextLocation, Resources);
-				Bounds = Rectangle.Union(Bounds, Element.Bounds);
+				nextLocation = element.Layout(graphics, nextLocation, resources);
+				Bounds = Rectangle.Union(Bounds, element.Bounds);
 			}
 
-			bModified = false;
+			_modified = false;
 		}
 
-		public void Draw(Graphics Graphics, StatusElementResources Resources)
+		public void Draw(Graphics graphics, StatusElementResources resources)
 		{
-			foreach(StatusElement Element in Elements)
+			foreach(StatusElement element in _elements)
 			{
-				Element.Draw(Graphics, Resources);
+				element.Draw(graphics, resources);
 			}
 		}
 	}
@@ -442,46 +442,46 @@ namespace UnrealGameSync
 	{
 		const float LineSpacing = 1.35f;
 
-		Image? ProjectLogo;
-		bool bDisposeProjectLogo;
-		Rectangle ProjectLogoBounds;
-		StatusElementResources? Resources;
-		List<StatusLine> Lines = new List<StatusLine>();
-		StatusLine? Caption;
-		Pen? AlertDividerPen;
-		int AlertDividerY;
-		StatusLine? Alert;
-		Color? TintColor;
-		Point? MouseOverLocation;
-		StatusElement? MouseOverElement;
-		Point? MouseDownLocation;
-		StatusElement? MouseDownElement;
-		int ContentWidth = 400;
-		int SuspendDisplayCount;
+		Image? _projectLogo;
+		bool _disposeProjectLogo;
+		Rectangle _projectLogoBounds;
+		StatusElementResources? _resources;
+		List<StatusLine> _lines = new List<StatusLine>();
+		StatusLine? _caption;
+		Pen? _alertDividerPen;
+		int _alertDividerY;
+		StatusLine? _alert;
+		Color? _tintColor;
+		Point? _mouseOverLocation;
+		StatusElement? _mouseOverElement;
+		Point? _mouseDownLocation;
+		StatusElement? _mouseDownElement;
+		int _contentWidth = 400;
+		int _suspendDisplayCount;
 
 		public StatusPanel()
 		{
 			DoubleBuffered = true;
 
-			AlertDividerPen = new Pen(Color.Black);
-			AlertDividerPen.DashPattern = new float[] { 2, 2 };
+			_alertDividerPen = new Pen(Color.Black);
+			_alertDividerPen.DashPattern = new float[] { 2, 2 };
 		}
 
 		public void SuspendDisplay()
 		{
-			SuspendDisplayCount++;
+			_suspendDisplayCount++;
 		}
 
 		public void ResumeDisplay()
 		{
-			SuspendDisplayCount--;
+			_suspendDisplayCount--;
 		}
 
-		public void SetContentWidth(int NewContentWidth)
+		public void SetContentWidth(int newContentWidth)
 		{
-			if(ContentWidth != NewContentWidth)
+			if(_contentWidth != newContentWidth)
 			{
-				ContentWidth = NewContentWidth;
+				_contentWidth = newContentWidth;
 				LayoutElements();
 				Invalidate();
 			}
@@ -493,18 +493,18 @@ namespace UnrealGameSync
 
 			if(disposing)
 			{
-				if(AlertDividerPen != null)
+				if(_alertDividerPen != null)
 				{
-					AlertDividerPen.Dispose();
-					AlertDividerPen = null;
+					_alertDividerPen.Dispose();
+					_alertDividerPen = null;
 				}
-				if(ProjectLogo != null)
+				if(_projectLogo != null)
 				{
-					if(bDisposeProjectLogo)
+					if(_disposeProjectLogo)
 					{
-						ProjectLogo.Dispose();
+						_projectLogo.Dispose();
 					}
-					ProjectLogo = null;
+					_projectLogo = null;
 				}
 				ResetFontCache();
 			}
@@ -512,100 +512,100 @@ namespace UnrealGameSync
 
 		private void ResetFontCache()
 		{
-			if(Resources != null)
+			if(_resources != null)
 			{
-				Resources.Dispose();
-				Resources = null;
+				_resources.Dispose();
+				_resources = null;
 			}
-			Resources = new StatusElementResources(Font);
+			_resources = new StatusElementResources(Font);
 		}
 
-		public void SetProjectLogo(Image NewProjectLogo, bool bDispose)
+		public void SetProjectLogo(Image newProjectLogo, bool dispose)
 		{
-			if(ProjectLogo != null)
+			if(_projectLogo != null)
 			{
-				if(bDisposeProjectLogo)
+				if(_disposeProjectLogo)
 				{
-					ProjectLogo.Dispose();
+					_projectLogo.Dispose();
 				}
 			}
-			ProjectLogo = NewProjectLogo;
-			bDisposeProjectLogo = bDispose;
+			_projectLogo = newProjectLogo;
+			_disposeProjectLogo = dispose;
 			Invalidate();
 		}
 
 		public void Clear()
 		{
 			InvalidateElements();
-			Lines.Clear();
-			Caption = null;
+			_lines.Clear();
+			_caption = null;
 		}
 
-		public void Set(IEnumerable<StatusLine> NewLines, StatusLine? NewCaption, StatusLine? NewAlert, Color? NewTintColor)
+		public void Set(IEnumerable<StatusLine> newLines, StatusLine? newCaption, StatusLine? newAlert, Color? newTintColor)
 		{
-			if(Resources == null)
+			if(_resources == null)
 			{
-				Resources = new StatusElementResources(Font);
+				_resources = new StatusElementResources(Font);
 			}
 
-			if(TintColor != NewTintColor)
+			if(_tintColor != newTintColor)
 			{
 				Invalidate();
 			}
 
 			InvalidateElements();
-			Lines.Clear();
-			Lines.AddRange(NewLines);
-			Caption = NewCaption;
-			Alert = NewAlert;
-			TintColor = NewTintColor;
+			_lines.Clear();
+			_lines.AddRange(newLines);
+			_caption = newCaption;
+			_alert = newAlert;
+			_tintColor = newTintColor;
 			LayoutElements();
 			InvalidateElements();
 
-			MouseOverElement = null;
-			MouseDownElement = null;
-			SetMouseOverLocation(MouseOverLocation);
-			SetMouseDownLocation(MouseDownLocation);
+			_mouseOverElement = null;
+			_mouseDownElement = null;
+			SetMouseOverLocation(_mouseOverLocation);
+			SetMouseDownLocation(_mouseDownLocation);
 		}
 
 		protected void InvalidateElements()
 		{
-			Invalidate(ProjectLogoBounds);
+			Invalidate(_projectLogoBounds);
 
-			foreach(StatusLine Line in Lines)
+			foreach(StatusLine line in _lines)
 			{
-				Invalidate(Line.Bounds);
+				Invalidate(line.Bounds);
 			}
-			if(Caption != null)
+			if(_caption != null)
 			{
-				Invalidate(Caption.Bounds);
+				Invalidate(_caption.Bounds);
 			}
-			if(Alert != null)
+			if(_alert != null)
 			{
-				Invalidate(Alert.Bounds);
+				Invalidate(_alert.Bounds);
 			}
 		}
 
-		protected bool HitTest(Point Location, [NotNullWhen(true)] out StatusElement? OutElement)
+		protected bool HitTest(Point location, [NotNullWhen(true)] out StatusElement? outElement)
 		{
-			OutElement = null;
-			foreach(StatusLine Line in Lines)
+			outElement = null;
+			foreach(StatusLine line in _lines)
 			{
-				if(Line.HitTest(Location, out OutElement))
+				if(line.HitTest(location, out outElement))
 				{
 					return true;
 				}
 			}
-			if(Caption != null)
+			if(_caption != null)
 			{
-				if(Caption.HitTest(Location, out OutElement))
+				if(_caption.HitTest(location, out outElement))
 				{
 					return true;
 				}
 			}
-			if(Alert != null)
+			if(_alert != null)
 			{
-				if(Alert.HitTest(Location, out OutElement))
+				if(_alert.HitTest(location, out outElement))
 				{
 					return true;
 				}
@@ -623,59 +623,59 @@ namespace UnrealGameSync
 
 		protected void LayoutElements()
 		{
-			using(Graphics Graphics = CreateGraphics())
+			using(Graphics graphics = CreateGraphics())
 			{
-				LayoutElements(Graphics);
+				LayoutElements(graphics);
 			}
 		}
 
-		protected void LayoutElements(Graphics Graphics)
+		protected void LayoutElements(Graphics graphics)
 		{
 			// Layout the alert message
-			int BodyHeight = Height;
-			if(Alert != null)
+			int bodyHeight = Height;
+			if(_alert != null)
 			{
-				AlertDividerY = Height - (int)(Font.Height * 2);
-				Alert.Layout(Graphics, Point.Empty, Resources!);
-				Alert.Layout(Graphics, new Point((Width - Alert.Bounds.Width) / 2, (Height + AlertDividerY) / 2), Resources!);
-				BodyHeight = AlertDividerY;
+				_alertDividerY = Height - (int)(Font.Height * 2);
+				_alert.Layout(graphics, Point.Empty, _resources!);
+				_alert.Layout(graphics, new Point((Width - _alert.Bounds.Width) / 2, (Height + _alertDividerY) / 2), _resources!);
+				bodyHeight = _alertDividerY;
 			}
 
 			// Get the logo size
-			Image DrawProjectLogo = ProjectLogo ?? Properties.Resources.DefaultProjectLogo;
-			float LogoScale = Math.Min((float)BodyHeight - ((Caption != null)? Font.Height : 0) / DrawProjectLogo.Height, Graphics.DpiY / 96.0f);
-			int LogoWidth = (int)(DrawProjectLogo.Width * LogoScale);
-			int LogoHeight = (int)(DrawProjectLogo.Height * LogoScale);
+			Image drawProjectLogo = _projectLogo ?? Properties.Resources.DefaultProjectLogo;
+			float logoScale = Math.Min((float)bodyHeight - ((_caption != null)? Font.Height : 0) / drawProjectLogo.Height, graphics.DpiY / 96.0f);
+			int logoWidth = (int)(drawProjectLogo.Width * logoScale);
+			int logoHeight = (int)(drawProjectLogo.Height * logoScale);
 
 			// Figure out where the split between content and the logo is going to be
-			int DividerX = ((Width - LogoWidth - ContentWidth) / 2) + LogoWidth;
+			int dividerX = ((Width - logoWidth - _contentWidth) / 2) + logoWidth;
 
 			// Get the logo position
-			int LogoX = DividerX - LogoWidth;
-			int LogoY = (BodyHeight - LogoHeight) / 2;
+			int logoX = dividerX - logoWidth;
+			int logoY = (bodyHeight - logoHeight) / 2;
 
 			// Layout the caption. We may move the logo to make room for this.
-			LogoY -= Font.Height / 2;
-			if(Caption != null)
+			logoY -= Font.Height / 2;
+			if(_caption != null)
 			{
-				Caption.Layout(Graphics, Point.Empty, Resources!);
-				int CaptionWidth = Caption.Bounds.Width;
-				Caption.Layout(Graphics, new Point(Math.Min(LogoX + (LogoWidth / 2) - (CaptionWidth / 2), DividerX - CaptionWidth), LogoY + LogoHeight), Resources!);
+				_caption.Layout(graphics, Point.Empty, _resources!);
+				int captionWidth = _caption.Bounds.Width;
+				_caption.Layout(graphics, new Point(Math.Min(logoX + (logoWidth / 2) - (captionWidth / 2), dividerX - captionWidth), logoY + logoHeight), _resources!);
 			}
 
 			// Set the logo rectangle
-			ProjectLogoBounds = new Rectangle(LogoX, LogoY, LogoWidth, LogoHeight);
+			_projectLogoBounds = new Rectangle(logoX, logoY, logoWidth, logoHeight);
 
 			// Measure up all the line height
-			float TotalLineHeight = Lines.Sum(x => x.LineHeight);
+			float totalLineHeight = _lines.Sum(x => x.LineHeight);
 
 			// Space out all the lines
-			float LineY = (BodyHeight - TotalLineHeight * (int)(Font.Height * LineSpacing)) / 2;
-			foreach(StatusLine Line in Lines)
+			float lineY = (bodyHeight - totalLineHeight * (int)(Font.Height * LineSpacing)) / 2;
+			foreach(StatusLine line in _lines)
 			{
-				LineY += (int)(Font.Height * LineSpacing * Line.LineHeight * 0.5f);
-				Line.Layout(Graphics, new Point(DividerX + 5, (int)LineY), Resources!);
-				LineY += (int)(Font.Height * LineSpacing * Line.LineHeight * 0.5f);
+				lineY += (int)(Font.Height * LineSpacing * line.LineHeight * 0.5f);
+				line.Layout(graphics, new Point(dividerX + 5, (int)lineY), _resources!);
+				lineY += (int)(Font.Height * LineSpacing * line.LineHeight * 0.5f);
 			}
 		}
 
@@ -700,9 +700,9 @@ namespace UnrealGameSync
 		{
 			base.OnMouseUp(e);
 
-			if(MouseDownElement != null && MouseOverElement == MouseDownElement)
+			if(_mouseDownElement != null && _mouseOverElement == _mouseDownElement)
 			{
-				MouseDownElement.OnClick(e.Location);
+				_mouseDownElement.OnClick(e.Location);
 			}
 
 			SetMouseDownLocation(null);
@@ -728,20 +728,20 @@ namespace UnrealGameSync
 		{
 			base.OnPaintBackground(e);
 
-			if(TintColor.HasValue)
+			if(_tintColor.HasValue)
 			{
-				int TintSize = Width / 2;
-				using(LinearGradientBrush BackgroundBrush = new LinearGradientBrush(new Point(Width, 0), new Point(Width - TintSize, TintSize), TintColor.Value, BackColor))
+				int tintSize = Width / 2;
+				using(LinearGradientBrush backgroundBrush = new LinearGradientBrush(new Point(Width, 0), new Point(Width - tintSize, tintSize), _tintColor.Value, BackColor))
 				{
-					BackgroundBrush.WrapMode = WrapMode.TileFlipXY;
-					using (GraphicsPath Path = new GraphicsPath())
+					backgroundBrush.WrapMode = WrapMode.TileFlipXY;
+					using (GraphicsPath path = new GraphicsPath())
 					{
-						Path.StartFigure();
-						Path.AddLine(Width, 0, Width - TintSize * 2, 0);
-						Path.AddLine(Width, TintSize * 2, Width, 0);
-						Path.CloseFigure();
+						path.StartFigure();
+						path.AddLine(Width, 0, Width - tintSize * 2, 0);
+						path.AddLine(Width, tintSize * 2, Width, 0);
+						path.CloseFigure();
 
-						e.Graphics.FillPath(BackgroundBrush, Path);
+						e.Graphics.FillPath(backgroundBrush, path);
 					}
 				}
 			}
@@ -751,82 +751,82 @@ namespace UnrealGameSync
 		{
 			base.OnPaint(e);
 
-			if(SuspendDisplayCount == 0)
+			if(_suspendDisplayCount == 0)
 			{
-				e.Graphics.DrawImage(ProjectLogo ?? Properties.Resources.DefaultProjectLogo, ProjectLogoBounds);
+				e.Graphics.DrawImage(_projectLogo ?? Properties.Resources.DefaultProjectLogo, _projectLogoBounds);
 
 				e.Graphics.SmoothingMode = System.Drawing.Drawing2D.SmoothingMode.None;
 
-				foreach(StatusLine Line in Lines)
+				foreach(StatusLine line in _lines)
 				{
-					Line.Draw(e.Graphics, Resources!);
+					line.Draw(e.Graphics, _resources!);
 				}
-				if(Caption != null)
+				if(_caption != null)
 				{
-					Caption.Draw(e.Graphics, Resources!);
+					_caption.Draw(e.Graphics, _resources!);
 				}
-				if(Alert != null)
+				if(_alert != null)
 				{
-					e.Graphics.DrawLine(AlertDividerPen, 0, AlertDividerY, Width, AlertDividerY);
-					Alert.Draw(e.Graphics, Resources!);
+					e.Graphics.DrawLine(_alertDividerPen, 0, _alertDividerY, Width, _alertDividerY);
+					_alert.Draw(e.Graphics, _resources!);
 				}
 			}
 		}
 
-		protected void SetMouseOverLocation(Point? NewMouseOverLocation)
+		protected void SetMouseOverLocation(Point? newMouseOverLocation)
 		{
-			MouseOverLocation = NewMouseOverLocation;
+			_mouseOverLocation = newMouseOverLocation;
 
-			StatusElement? NewMouseOverElement = null;
-			if(MouseOverLocation.HasValue)
+			StatusElement? newMouseOverElement = null;
+			if(_mouseOverLocation.HasValue)
 			{
-				HitTest(MouseOverLocation.Value, out NewMouseOverElement);
+				HitTest(_mouseOverLocation.Value, out newMouseOverElement);
 			}
 
-			if(NewMouseOverElement != MouseOverElement)
+			if(newMouseOverElement != _mouseOverElement)
 			{
-				if(MouseOverElement != null)
+				if(_mouseOverElement != null)
 				{
-					MouseOverElement.bMouseOver = false;
+					_mouseOverElement.MouseOver = false;
 					Cursor = Cursors.Arrow;
-					Invalidate(MouseOverElement.Bounds);
+					Invalidate(_mouseOverElement.Bounds);
 				}
 
-				MouseOverElement = NewMouseOverElement;
+				_mouseOverElement = newMouseOverElement;
 
-				if(MouseOverElement != null)
+				if(_mouseOverElement != null)
 				{
-					MouseOverElement.bMouseOver = true;
-					Cursor = MouseOverElement.Cursor;
-					Invalidate(MouseOverElement.Bounds);
+					_mouseOverElement.MouseOver = true;
+					Cursor = _mouseOverElement.Cursor;
+					Invalidate(_mouseOverElement.Bounds);
 				}
 			}
 		}
 
-		protected void SetMouseDownLocation(Point? NewMouseDownLocation)
+		protected void SetMouseDownLocation(Point? newMouseDownLocation)
 		{
-			MouseDownLocation = NewMouseDownLocation;
+			_mouseDownLocation = newMouseDownLocation;
 
-			StatusElement? NewMouseDownElement = null;
-			if(MouseDownLocation.HasValue)
+			StatusElement? newMouseDownElement = null;
+			if(_mouseDownLocation.HasValue)
 			{
-				HitTest(MouseDownLocation.Value, out NewMouseDownElement);
+				HitTest(_mouseDownLocation.Value, out newMouseDownElement);
 			}
 
-			if(NewMouseDownElement != MouseDownElement)
+			if(newMouseDownElement != _mouseDownElement)
 			{
-				if(MouseDownElement != null)
+				if(_mouseDownElement != null)
 				{
-					MouseDownElement.bMouseDown = false;
-					Invalidate(MouseDownElement.Bounds);
+					_mouseDownElement.MouseDown = false;
+					Invalidate(_mouseDownElement.Bounds);
 				}
 
-				MouseDownElement = NewMouseDownElement;
+				_mouseDownElement = newMouseDownElement;
 
-				if(MouseDownElement != null)
+				if(_mouseDownElement != null)
 				{
-					MouseDownElement.bMouseDown = true;
-					Invalidate(MouseDownElement.Bounds);
+					_mouseDownElement.MouseDown = true;
+					Invalidate(_mouseDownElement.Bounds);
 				}
 			}
 		}
diff --git a/Engine/Source/Programs/UnrealGameSync/UnrealGameSync/Controls/SyncFilterControl.cs b/Engine/Source/Programs/UnrealGameSync/UnrealGameSync/Controls/SyncFilterControl.cs
--- a/Engine/Source/Programs/UnrealGameSync/UnrealGameSync/Controls/SyncFilterControl.cs
+++ b/Engine/Source/Programs/UnrealGameSync/UnrealGameSync/Controls/SyncFilterControl.cs
@@ -15,7 +15,7 @@ namespace UnrealGameSync.Controls
 {
 	public partial class SyncFilterControl : UserControl
 	{
-		HashSet<string> StripLines = new HashSet<string>()
+		HashSet<string> _stripLines = new HashSet<string>()
 		{
 			"; Rules are specified one per line, and may use any standard Perforce wildcards:",
 			";    ?    Matches one character.",
@@ -33,13 +33,13 @@ namespace UnrealGameSync.Controls
 
 		private void SyntaxButton_LinkClicked(object sender, LinkLabelLinkClickedEventArgs e)
 		{
-			SyncFilterSyntax Dialog = new SyncFilterSyntax();
-			Dialog.ShowDialog(ParentForm);
+			SyncFilterSyntax dialog = new SyncFilterSyntax();
+			dialog.ShowDialog(ParentForm);
 		}
 
-		public void SetView(string[] View)
+		public void SetView(string[] view)
 		{
-			ViewTextBox.Lines = View.Where(x => !StripLines.Contains(x.Trim())).SkipWhile(x => x.Trim().Length == 0 || x.Trim() == ";").ToArray();
+			ViewTextBox.Lines = view.Where(x => !_stripLines.Contains(x.Trim())).SkipWhile(x => x.Trim().Length == 0 || x.Trim() == ";").ToArray();
 			ViewTextBox.Select(ViewTextBox.Text.Length, 0);
 		}
 
