UGS: Reformat to standard C# style conventions.

diff --git a/Engine/Source/Programs/UnrealGameSync/UnrealGameSync/Controls/TabControl.cs b/Engine/Source/Programs/UnrealGameSync/UnrealGameSync/Controls/TabControl.cs
--- a/Engine/Source/Programs/UnrealGameSync/UnrealGameSync/Controls/TabControl.cs
+++ b/Engine/Source/Programs/UnrealGameSync/UnrealGameSync/Controls/TabControl.cs
@@ -27,10 +27,10 @@ namespace UnrealGameSync
 			public Tuple<Color, float>? Highlight;
 			public Color? TintColor;
 
-			public TabData(string Name, object Data)
+			public TabData(string name, object data)
 			{
-				this.Name = Name;
-				this.Data = Data;
+				this.Name = name;
+				this.Data = data;
 			}
 		}
 
@@ -41,216 +41,216 @@ namespace UnrealGameSync
 			public int MouseX;
 			public int RelativeMouseX;
 
-			public TabDragData(TabData Tab, int InitialIdx, int MouseX, int RelativeMouseX)
+			public TabDragData(TabData tab, int initialIdx, int mouseX, int relativeMouseX)
 			{
-				this.Tab = Tab;
-				this.InitialIdx = InitialIdx;
-				this.MouseX = MouseX;
-				this.RelativeMouseX = RelativeMouseX;
+				this.Tab = tab;
+				this.InitialIdx = initialIdx;
+				this.MouseX = mouseX;
+				this.RelativeMouseX = relativeMouseX;
 			}
 		}
 
-		List<TabData> Tabs = new List<TabData>();
-		int SelectedTabIdx;
-		int HoverTabIdx;
-		int HighlightTabIdx;
-		TabDragData? DragState;
+		List<TabData> _tabs = new List<TabData>();
+		int _selectedTabIdx;
+		int _hoverTabIdx;
+		int _highlightTabIdx;
+		TabDragData? _dragState;
 
-		public delegate void OnTabChangedDelegate(object? NewTabData);
+		public delegate void OnTabChangedDelegate(object? newTabData);
 		public event OnTabChangedDelegate? OnTabChanged;
 
-		public delegate void OnNewTabClickDelegate(Point Location, MouseButtons Buttons);
+		public delegate void OnNewTabClickDelegate(Point location, MouseButtons buttons);
 		public event OnNewTabClickDelegate? OnNewTabClick;
 
-		public delegate void OnTabClickedDelegate(object? TabData, Point Location, MouseButtons Buttons);
+		public delegate void OnTabClickedDelegate(object? tabData, Point location, MouseButtons buttons);
 		public event OnTabClickedDelegate? OnTabClicked;
 
-		public delegate bool OnTabClosingDelegate(object TabData);
+		public delegate bool OnTabClosingDelegate(object tabData);
 		public event OnTabClosingDelegate? OnTabClosing;
 
-		public delegate void OnTabClosedDelegate(object TabData);
+		public delegate void OnTabClosedDelegate(object tabData);
 		public event OnTabClosedDelegate? OnTabClosed;
 
 		public delegate void OnTabReorderDelegate();
 		public event OnTabReorderDelegate? OnTabReorder;
 
-		public delegate void OnButtonClickDelegate(int ButtonIdx, Point Location, MouseButtons Buttons);
+		public delegate void OnButtonClickDelegate(int buttonIdx, Point location, MouseButtons buttons);
 		public event OnButtonClickDelegate? OnButtonClick;
 
 		public TabControl()
 		{
 			DoubleBuffered = true;
-			Tabs.Add(new TabData("+", "+"));
-			SelectedTabIdx = -1;
-			HoverTabIdx = -1;
-			HighlightTabIdx = -1;
+			_tabs.Add(new TabData("+", "+"));
+			_selectedTabIdx = -1;
+			_hoverTabIdx = -1;
+			_highlightTabIdx = -1;
 		}
 
 		public void LockHover()
 		{
-			HighlightTabIdx = HoverTabIdx;
+			_highlightTabIdx = _hoverTabIdx;
 		}
 
 		public void UnlockHover()
 		{
-			HighlightTabIdx = -1;
+			_highlightTabIdx = -1;
 			Invalidate();
 		}
 
-		public int FindTabIndex(object Data)
+		public int FindTabIndex(object data)
 		{
-			int Result = -1;
-			for(int TabIdx = 0; TabIdx < Tabs.Count - 1; TabIdx++)
+			int result = -1;
+			for(int tabIdx = 0; tabIdx < _tabs.Count - 1; tabIdx++)
 			{
-				if(Tabs[TabIdx].Data == Data)
+				if(_tabs[tabIdx].Data == data)
 				{
-					Result = TabIdx;
+					result = tabIdx;
 					break;
 				}
 			}
-			return Result;
+			return result;
 		}
 
-		public void SetHighlight(int TabIdx, Tuple<Color, float>? Highlight)
+		public void SetHighlight(int tabIdx, Tuple<Color, float>? highlight)
 		{
-			Tuple<Color, float>? CurrentHighlight = Tabs[TabIdx].Highlight;
-			if(Highlight == null || CurrentHighlight == null)
+			Tuple<Color, float>? currentHighlight = _tabs[tabIdx].Highlight;
+			if(highlight == null || currentHighlight == null)
 			{
-				if(Highlight != CurrentHighlight)
+				if(highlight != currentHighlight)
 				{
-					Tabs[TabIdx].Highlight = Highlight;
+					_tabs[tabIdx].Highlight = highlight;
 					Invalidate();
 				}
 			}
 			else
 			{
-				if(Highlight.Item1 != CurrentHighlight.Item1 || Highlight.Item2 != CurrentHighlight.Item2)
+				if(highlight.Item1 != currentHighlight.Item1 || highlight.Item2 != currentHighlight.Item2)
 				{
-					Tabs[TabIdx].Highlight = Highlight;
+					_tabs[tabIdx].Highlight = highlight;
 					Invalidate();
 				}
 			}
 		}
 
-		public void SetTint(int TabIdx, Color? TintColor)
+		public void SetTint(int tabIdx, Color? tintColor)
 		{
-			if(Tabs[TabIdx].TintColor != TintColor)
+			if(_tabs[tabIdx].TintColor != tintColor)
 			{
-				Tabs[TabIdx].TintColor = TintColor;
+				_tabs[tabIdx].TintColor = tintColor;
 				Invalidate();
 			}
 		}
 
-		public int InsertTab(int InsertIdx, string Name, object Data, Color? TintColor)
+		public int InsertTab(int insertIdx, string name, object data, Color? tintColor)
 		{
-			int Idx = Tabs.Count - 1;
-			if(InsertIdx >= 0 && InsertIdx < Tabs.Count - 1)
+			int idx = _tabs.Count - 1;
+			if(insertIdx >= 0 && insertIdx < _tabs.Count - 1)
 			{
-				Idx = InsertIdx;
+				idx = insertIdx;
 			}
-			if(SelectedTabIdx >= Idx)
+			if(_selectedTabIdx >= idx)
 			{
-				SelectedTabIdx++;
+				_selectedTabIdx++;
 			}
 
-			Tabs.Insert(Idx, new TabData(Name, Data){ TintColor = TintColor });
+			_tabs.Insert(idx, new TabData(name, data){ TintColor = tintColor });
 			LayoutTabs();
 			Invalidate();
-			return Idx;
+			return idx;
 		}
 
 		public int GetTabCount()
 		{
-			return Tabs.Count - 1;
+			return _tabs.Count - 1;
 		}
 
-		public void SetTabName(int TabIdx, string Name)
+		public void SetTabName(int tabIdx, string name)
 		{
-			Tabs[TabIdx].Name = Name;
+			_tabs[tabIdx].Name = name;
 			LayoutTabs();
 			Invalidate();
 		}
 
-		public object GetTabData(int TabIdx)
+		public object GetTabData(int tabIdx)
 		{
-			return Tabs[TabIdx].Data;
+			return _tabs[tabIdx].Data;
 		}
 
-		public void SetTabData(int TabIdx, object Data)
+		public void SetTabData(int tabIdx, object data)
 		{
-			Tabs[TabIdx].Data = Data;
+			_tabs[tabIdx].Data = data;
 		}
 
 		public int GetSelectedTabIndex()
 		{
-			return SelectedTabIdx;
+			return _selectedTabIdx;
 		}
 
 		public object GetSelectedTabData()
 		{
-			return Tabs[SelectedTabIdx].Data;
+			return _tabs[_selectedTabIdx].Data;
 		}
 
-		public void SelectTab(int TabIdx)
+		public void SelectTab(int tabIdx)
 		{
-			if(TabIdx != SelectedTabIdx)
+			if(tabIdx != _selectedTabIdx)
 			{
-				ForceSelectTab(TabIdx);
+				ForceSelectTab(tabIdx);
 			}
 		}
 
-		private void ForceSelectTab(int TabIdx)
+		private void ForceSelectTab(int tabIdx)
 		{
-			SelectedTabIdx = TabIdx;
+			_selectedTabIdx = tabIdx;
 			LayoutTabs();
 			Invalidate();
 
 			if(OnTabChanged != null)
 			{
-				if(SelectedTabIdx == -1)
+				if(_selectedTabIdx == -1)
 				{
 					OnTabChanged(null);
 				}
 				else
 				{
-					OnTabChanged(Tabs[SelectedTabIdx].Data);
+					OnTabChanged(_tabs[_selectedTabIdx].Data);
 				}
 			}
 		}
 
-		public void RemoveTab(int TabIdx)
+		public void RemoveTab(int tabIdx)
 		{
-			object TabData = Tabs[TabIdx].Data;
-			if(OnTabClosing == null || OnTabClosing(TabData))
+			object tabData = _tabs[tabIdx].Data;
+			if(OnTabClosing == null || OnTabClosing(tabData))
 			{
-				Tabs.RemoveAt(TabIdx);
+				_tabs.RemoveAt(tabIdx);
 
-				if(HoverTabIdx == TabIdx)
+				if(_hoverTabIdx == tabIdx)
 				{
-					HoverTabIdx = -1;
+					_hoverTabIdx = -1;
 				}
 
-				if(HighlightTabIdx == TabIdx)
+				if(_highlightTabIdx == tabIdx)
 				{
-					HighlightTabIdx = -1;
+					_highlightTabIdx = -1;
 				}
 
-				if(SelectedTabIdx == TabIdx)
+				if(_selectedTabIdx == tabIdx)
 				{
-					if(SelectedTabIdx < Tabs.Count - 1)
+					if(_selectedTabIdx < _tabs.Count - 1)
 					{
-						ForceSelectTab(SelectedTabIdx);
+						ForceSelectTab(_selectedTabIdx);
 					}
 					else
 					{
-						ForceSelectTab(SelectedTabIdx - 1);
+						ForceSelectTab(_selectedTabIdx - 1);
 					}
 				}
 				else
 				{
-					if(SelectedTabIdx > TabIdx)
+					if(_selectedTabIdx > tabIdx)
 					{
-						SelectedTabIdx--;
+						_selectedTabIdx--;
 					}
 				}
 
@@ -259,73 +259,73 @@ namespace UnrealGameSync
 
 				if(OnTabClosed != null)
 				{
-					OnTabClosed(TabData);
+					OnTabClosed(tabData);
 				}
 			}
 		}
 
 		void LayoutTabs()
 		{
-			using(Graphics Graphics = CreateGraphics())
+			using(Graphics graphics = CreateGraphics())
 			{
-				float DpiScaleX = Graphics.DpiX / 96.0f;
-				float DpiScaleY = Graphics.DpiY / 96.0f;
+				float dpiScaleX = graphics.DpiX / 96.0f;
+				float dpiScaleY = graphics.DpiY / 96.0f;
 
-				for(int Idx = 0; Idx < Tabs.Count; Idx++)
+				for(int idx = 0; idx < _tabs.Count; idx++)
 				{
-					TabData Tab = Tabs[Idx];
-					Tab.TextSize = TextRenderer.MeasureText(Graphics, Tab.Name, Font, new Size(int.MaxValue, int.MaxValue), TextFormatFlags.NoPadding);
-					Tab.Width = TabPadding + Tab.TextSize.Width + (int)(TabPadding * DpiScaleX);
-					if(Idx == SelectedTabIdx)
+					TabData tab = _tabs[idx];
+					tab.TextSize = TextRenderer.MeasureText(graphics, tab.Name, Font, new Size(int.MaxValue, int.MaxValue), TextFormatFlags.NoPadding);
+					tab.Width = TabPadding + tab.TextSize.Width + (int)(TabPadding * dpiScaleX);
+					if(idx == _selectedTabIdx)
 					{
-						Tab.CloseButtonWidth = (int)(TabCloseButtonWidth * DpiScaleX);
-						Tab.Width += Tab.CloseButtonWidth;
+						tab.CloseButtonWidth = (int)(TabCloseButtonWidth * dpiScaleX);
+						tab.Width += tab.CloseButtonWidth;
 					}
 				}
 
-				int LeftX = 0;
-				for(int Idx = 0; Idx < Tabs.Count; Idx++)
+				int leftX = 0;
+				for(int idx = 0; idx < _tabs.Count; idx++)
 				{
-					TabData Tab = Tabs[Idx];
-					Tab.MinX = LeftX;
-					LeftX += Tab.Width;
+					TabData tab = _tabs[idx];
+					tab.MinX = leftX;
+					leftX += tab.Width;
 				}
 
-				int RightX = Width - 1;
-				if(LeftX > RightX)
+				int rightX = Width - 1;
+				if(leftX > rightX)
 				{
-					int UsedWidth = Tabs.Take(Tabs.Count - 1).Sum(x => x.Width);
-					int RemainingWidth = RightX + 1 - Tabs[Tabs.Count - 1].Width;
-					if(SelectedTabIdx != -1)
+					int usedWidth = _tabs.Take(_tabs.Count - 1).Sum(x => x.Width);
+					int remainingWidth = rightX + 1 - _tabs[_tabs.Count - 1].Width;
+					if(_selectedTabIdx != -1)
 					{
-						UsedWidth -= Tabs[SelectedTabIdx].Width;
-						RemainingWidth -= Tabs[SelectedTabIdx].Width;
+						usedWidth -= _tabs[_selectedTabIdx].Width;
+						remainingWidth -= _tabs[_selectedTabIdx].Width;
 					}
 
-					int NewX = 0;
-					for(int Idx = 0; Idx < Tabs.Count; Idx++)
+					int newX = 0;
+					for(int idx = 0; idx < _tabs.Count; idx++)
 					{
-						Tabs[Idx].MinX = NewX;
-						int PrevWidth = Tabs[Idx].Width;
-						if(Idx != SelectedTabIdx && Idx != Tabs.Count - 1)
+						_tabs[idx].MinX = newX;
+						int prevWidth = _tabs[idx].Width;
+						if(idx != _selectedTabIdx && idx != _tabs.Count - 1)
 						{
-							Tabs[Idx].Width = Math.Max((Tabs[Idx].Width * RemainingWidth) / UsedWidth, TabPadding * 3);
+							_tabs[idx].Width = Math.Max((_tabs[idx].Width * remainingWidth) / usedWidth, TabPadding * 3);
 						}
-						Tabs[Idx].TextSize.Width -= PrevWidth - Tabs[Idx].Width;
-						NewX += Tabs[Idx].Width;
+						_tabs[idx].TextSize.Width -= prevWidth - _tabs[idx].Width;
+						newX += _tabs[idx].Width;
 					}
 				}
 
-				if(DragState != null)
+				if(_dragState != null)
 				{
-					int MinOffset = -DragState.Tab.MinX;
-					int MaxOffset = Tabs[Tabs.Count - 1].MinX - DragState.Tab.Width - DragState.Tab.MinX;
+					int minOffset = -_dragState.Tab.MinX;
+					int maxOffset = _tabs[_tabs.Count - 1].MinX - _dragState.Tab.Width - _dragState.Tab.MinX;
 
-					int Offset = (DragState.MouseX - DragState.RelativeMouseX) - DragState.Tab.MinX;
-					Offset = Math.Max(Offset, MinOffset);
-					Offset = Math.Min(Offset, MaxOffset);
+					int offset = (_dragState.MouseX - _dragState.RelativeMouseX) - _dragState.Tab.MinX;
+					offset = Math.Max(offset, minOffset);
+					offset = Math.Min(offset, maxOffset);
 
-					DragState.Tab.MinX += Offset;
+					_dragState.Tab.MinX += offset;
 				}
 			}
 		}
@@ -344,50 +344,50 @@ namespace UnrealGameSync
 
 			OnMouseMove(e);
 
-			if(HoverTabIdx != -1)
+			if(_hoverTabIdx != -1)
 			{
 				if(e.Button == MouseButtons.Left)
 				{
-					if(HoverTabIdx == SelectedTabIdx)
+					if(_hoverTabIdx == _selectedTabIdx)
 					{
-						if(e.Location.X > Tabs[HoverTabIdx].MinX + Tabs[HoverTabIdx].Width - Tabs[HoverTabIdx].CloseButtonWidth)
+						if(e.Location.X > _tabs[_hoverTabIdx].MinX + _tabs[_hoverTabIdx].Width - _tabs[_hoverTabIdx].CloseButtonWidth)
 						{
-							RemoveTab(HoverTabIdx);
+							RemoveTab(_hoverTabIdx);
 						}
 						else
 						{
-							DragState = new TabDragData(Tabs[SelectedTabIdx], SelectedTabIdx, e.Location.X, e.Location.X - Tabs[SelectedTabIdx].MinX);
+							_dragState = new TabDragData(_tabs[_selectedTabIdx], _selectedTabIdx, e.Location.X, e.Location.X - _tabs[_selectedTabIdx].MinX);
 						}
 					}
 					else
 					{
-						if(HoverTabIdx > Tabs.Count - 1)
+						if(_hoverTabIdx > _tabs.Count - 1)
 						{
-							OnButtonClick?.Invoke(HoverTabIdx - Tabs.Count, e.Location, e.Button);
+							OnButtonClick?.Invoke(_hoverTabIdx - _tabs.Count, e.Location, e.Button);
 						}
-						else if(HoverTabIdx == Tabs.Count - 1)
+						else if(_hoverTabIdx == _tabs.Count - 1)
 						{
 							OnNewTabClick?.Invoke(e.Location, e.Button);
 						}
 						else
 						{
-							SelectTab(HoverTabIdx);
+							SelectTab(_hoverTabIdx);
 						}
 					}
 				}
 				else if(e.Button == MouseButtons.Middle)
 				{
-					if(HoverTabIdx < Tabs.Count - 1)
+					if(_hoverTabIdx < _tabs.Count - 1)
 					{
-						RemoveTab(HoverTabIdx);
+						RemoveTab(_hoverTabIdx);
 					}
 				}
 			}
 
 			if(OnTabClicked != null)
 			{
-				object? TabData = (HoverTabIdx == -1)? null : Tabs[HoverTabIdx].Data;
-				OnTabClicked(TabData, e.Location, e.Button);
+				object? tabData = (_hoverTabIdx == -1)? null : _tabs[_hoverTabIdx].Data;
+				OnTabClicked(tabData, e.Location, e.Button);
 			}
 		}
 
@@ -395,14 +395,14 @@ namespace UnrealGameSync
 		{
 			base.OnMouseUp(e);
 
-			if(DragState != null)
+			if(_dragState != null)
 			{
 				if(OnTabReorder != null)
 				{
 					OnTabReorder();
 				}
 
-				DragState = null;
+				_dragState = null;
 
 				LayoutTabs();
 				Invalidate();
@@ -413,22 +413,22 @@ namespace UnrealGameSync
 		{
 			base.OnMouseMove(e);
 
-			if(DragState != null)
+			if(_dragState != null)
 			{
-				DragState.MouseX = e.Location.X;
+				_dragState.MouseX = e.Location.X;
 
-				Tabs.Remove(DragState.Tab);
+				_tabs.Remove(_dragState.Tab);
 
-				int TotalWidth = 0;
-				for(int InsertIdx = 0;;InsertIdx++)
+				int totalWidth = 0;
+				for(int insertIdx = 0;;insertIdx++)
 				{
-					if(InsertIdx == Tabs.Count - 1 || DragState.MouseX - DragState.RelativeMouseX < TotalWidth + Tabs[InsertIdx].Width / 2)
+					if(insertIdx == _tabs.Count - 1 || _dragState.MouseX - _dragState.RelativeMouseX < totalWidth + _tabs[insertIdx].Width / 2)
 					{
-						HoverTabIdx = SelectedTabIdx = InsertIdx;
-						Tabs.Insert(InsertIdx, DragState.Tab);
+						_hoverTabIdx = _selectedTabIdx = insertIdx;
+						_tabs.Insert(insertIdx, _dragState.Tab);
 						break;
 					}
-					TotalWidth += Tabs[InsertIdx].Width;
+					totalWidth += _tabs[insertIdx].Width;
 				}
 
 				LayoutTabs();
@@ -436,26 +436,26 @@ namespace UnrealGameSync
 			}
 			else
 			{
-				int NewHoverTabIdx = -1;
+				int newHoverTabIdx = -1;
 				if(ClientRectangle.Contains(e.Location))
 				{
-					for(int Idx = 0; Idx < Tabs.Count; Idx++)
+					for(int idx = 0; idx < _tabs.Count; idx++)
 					{
-						if(e.Location.X > Tabs[Idx].MinX && e.Location.X < Tabs[Idx].MinX + Tabs[Idx].Width)
+						if(e.Location.X > _tabs[idx].MinX && e.Location.X < _tabs[idx].MinX + _tabs[idx].Width)
 						{
-							NewHoverTabIdx = Idx;
+							newHoverTabIdx = idx;
 							break;
 						}
 					}
 				}
 
-				if(HoverTabIdx != NewHoverTabIdx)
+				if(_hoverTabIdx != newHoverTabIdx)
 				{
-					HoverTabIdx = NewHoverTabIdx;
+					_hoverTabIdx = newHoverTabIdx;
 					LayoutTabs();
 					Invalidate();
 
-					if(HoverTabIdx != -1)
+					if(_hoverTabIdx != -1)
 					{
 						Capture = true;
 					}
@@ -471,9 +471,9 @@ namespace UnrealGameSync
 		{
 			base.OnMouseCaptureChanged(e);
 
-			if(HoverTabIdx != -1)
+			if(_hoverTabIdx != -1)
 			{
-				HoverTabIdx = -1;
+				_hoverTabIdx = -1;
 				Invalidate();
 			}
 		}
@@ -484,126 +484,126 @@ namespace UnrealGameSync
 
 			LayoutTabs();
 
-			using(SolidBrush HoverBrush = new SolidBrush(Color.FromArgb(192, 192, 192)))
+			using(SolidBrush hoverBrush = new SolidBrush(Color.FromArgb(192, 192, 192)))
 			{
-				for(int Idx = 0; Idx < Tabs.Count; Idx++)
+				for(int idx = 0; idx < _tabs.Count; idx++)
 				{
-					TabData Tab = Tabs[Idx];
-					if(Idx == HoverTabIdx || Idx == HighlightTabIdx)
+					TabData tab = _tabs[idx];
+					if(idx == _hoverTabIdx || idx == _highlightTabIdx)
 					{
-						DrawBackground(e.Graphics, Tab, SystemBrushes.Window, HoverBrush, Tab.Highlight);
+						DrawBackground(e.Graphics, tab, SystemBrushes.Window, hoverBrush, tab.Highlight);
 					}
 					else
 					{
-						DrawBackground(e.Graphics, Tab, SystemBrushes.Control, SystemBrushes.Control, Tab.Highlight);
+						DrawBackground(e.Graphics, tab, SystemBrushes.Control, SystemBrushes.Control, tab.Highlight);
 					}
 				}
 			}
 
-			using(Pen SeparatorPen = new Pen(Color.FromArgb(192, SystemColors.ControlDarkDark)))
+			using(Pen separatorPen = new Pen(Color.FromArgb(192, SystemColors.ControlDarkDark)))
 			{
-				for(int Idx = 0; Idx < Tabs.Count; Idx++)
+				for(int idx = 0; idx < _tabs.Count; idx++)
 				{
-					TabData Tab = Tabs[Idx];
-					if((Idx > 0 && Idx < Tabs.Count - 1) || Idx >= Tabs.Count || Idx == SelectedTabIdx || Idx == HoverTabIdx || Idx == HighlightTabIdx || Tabs.Count == 1)
+					TabData tab = _tabs[idx];
+					if((idx > 0 && idx < _tabs.Count - 1) || idx >= _tabs.Count || idx == _selectedTabIdx || idx == _hoverTabIdx || idx == _highlightTabIdx || _tabs.Count == 1)
 					{
-						e.Graphics.DrawLine(SeparatorPen, Tab.MinX, 0, Tab.MinX, ClientSize.Height - 2);
+						e.Graphics.DrawLine(separatorPen, tab.MinX, 0, tab.MinX, ClientSize.Height - 2);
 					}
-					if(Idx < Tabs.Count - 1 || Idx >= Tabs.Count || Idx == HoverTabIdx || Idx == HighlightTabIdx || Tabs.Count == 1)
+					if(idx < _tabs.Count - 1 || idx >= _tabs.Count || idx == _hoverTabIdx || idx == _highlightTabIdx || _tabs.Count == 1)
 					{
-						e.Graphics.DrawLine(SeparatorPen, Tab.MinX + Tab.Width, 0, Tab.MinX + Tab.Width, ClientSize.Height - 2);
+						e.Graphics.DrawLine(separatorPen, tab.MinX + tab.Width, 0, tab.MinX + tab.Width, ClientSize.Height - 2);
 					}
-					if(Idx == HoverTabIdx || Idx == HighlightTabIdx || Idx >= Tabs.Count)
+					if(idx == _hoverTabIdx || idx == _highlightTabIdx || idx >= _tabs.Count)
 					{
-						e.Graphics.DrawLine(SeparatorPen, Tab.MinX, 0, Tab.MinX + Tab.Width, 0);
+						e.Graphics.DrawLine(separatorPen, tab.MinX, 0, tab.MinX + tab.Width, 0);
 					}
 				}
 			}
 
-			for(int Idx = 0; Idx < Tabs.Count; Idx++)
+			for(int idx = 0; idx < _tabs.Count; idx++)
 			{
-				if(Idx != SelectedTabIdx)
+				if(idx != _selectedTabIdx)
 				{
-					DrawText(e.Graphics, Tabs[Idx]);
+					DrawText(e.Graphics, _tabs[idx]);
 				}
 			}
 
 
-			if(SelectedTabIdx != -1)
+			if(_selectedTabIdx != -1)
 			{
-				TabData SelectedTab = Tabs[SelectedTabIdx];
+				TabData selectedTab = _tabs[_selectedTabIdx];
 
-				using(SolidBrush SelectedBrush = new SolidBrush(Color.FromArgb(0, 136, 204)))
+				using(SolidBrush selectedBrush = new SolidBrush(Color.FromArgb(0, 136, 204)))
 				{
-					DrawBackground(e.Graphics, SelectedTab, SystemBrushes.Window, SelectedBrush, SelectedTab.Highlight);
+					DrawBackground(e.Graphics, selectedTab, SystemBrushes.Window, selectedBrush, selectedTab.Highlight);
 				}
 
-				DrawText(e.Graphics, SelectedTab);
+				DrawText(e.Graphics, selectedTab);
 	
 				// Draw the close button
-				SmoothingMode PrevSmoothingMode = e.Graphics.SmoothingMode;
+				SmoothingMode prevSmoothingMode = e.Graphics.SmoothingMode;
 				e.Graphics.SmoothingMode = System.Drawing.Drawing2D.SmoothingMode.AntiAlias;
 
-				float DpiScaleX = e.Graphics.DpiX / 96.0f;
-				float DpiScaleY = e.Graphics.DpiY / 96.0f;
+				float dpiScaleX = e.Graphics.DpiX / 96.0f;
+				float dpiScaleY = e.Graphics.DpiY / 96.0f;
 
-				int CloseMidX = SelectedTab.MinX + SelectedTab.Width - (int)((TabPadding + TabCloseButtonWidth) * DpiScaleX / 2);// + (13 / 2);
-				int CloseMidY = (ClientSize.Height - (int)(4 * DpiScaleY)) / 2;
+				int closeMidX = selectedTab.MinX + selectedTab.Width - (int)((TabPadding + TabCloseButtonWidth) * dpiScaleX / 2);// + (13 / 2);
+				int closeMidY = (ClientSize.Height - (int)(4 * dpiScaleY)) / 2;
 
-				Rectangle CloseButton = new Rectangle(CloseMidX - (int)((13 / 2) * DpiScaleX), CloseMidY - (int)((13 / 2) * DpiScaleY), (int)(13 * DpiScaleX), (int)(13 * DpiScaleY));
-				e.Graphics.FillEllipse(SystemBrushes.ControlDark, CloseButton);
+				Rectangle closeButton = new Rectangle(closeMidX - (int)((13 / 2) * dpiScaleX), closeMidY - (int)((13 / 2) * dpiScaleY), (int)(13 * dpiScaleX), (int)(13 * dpiScaleY));
+				e.Graphics.FillEllipse(SystemBrushes.ControlDark, closeButton);
 
-				using(Pen CrossPen = new Pen(SystemBrushes.Window, 2.0f * DpiScaleX))
+				using(Pen crossPen = new Pen(SystemBrushes.Window, 2.0f * dpiScaleX))
 				{
-					float IndentX = 3.5f * DpiScaleX;
-					float IndentY = 3.5f * DpiScaleY;
-					e.Graphics.DrawLine(CrossPen, CloseButton.Left + IndentX, CloseButton.Top + IndentY, CloseButton.Right - IndentX, CloseButton.Bottom - IndentY);
-					e.Graphics.DrawLine(CrossPen, CloseButton.Left + IndentX, CloseButton.Bottom - IndentY, CloseButton.Right - IndentX, CloseButton.Top + IndentY);
+					float indentX = 3.5f * dpiScaleX;
+					float indentY = 3.5f * dpiScaleY;
+					e.Graphics.DrawLine(crossPen, closeButton.Left + indentX, closeButton.Top + indentY, closeButton.Right - indentX, closeButton.Bottom - indentY);
+					e.Graphics.DrawLine(crossPen, closeButton.Left + indentX, closeButton.Bottom - indentY, closeButton.Right - indentX, closeButton.Top + indentY);
 				}
-				e.Graphics.SmoothingMode = PrevSmoothingMode;
+				e.Graphics.SmoothingMode = prevSmoothingMode;
 
 				// Border
-				e.Graphics.DrawLine(SystemPens.ControlDark, SelectedTab.MinX, 0, SelectedTab.MinX + SelectedTab.Width, 0);
-				e.Graphics.DrawLine(SystemPens.ControlDark, SelectedTab.MinX, 0, SelectedTab.MinX, ClientSize.Height - 2);
-				e.Graphics.DrawLine(SystemPens.ControlDark, SelectedTab.MinX + SelectedTab.Width, 0, SelectedTab.MinX + SelectedTab.Width, ClientSize.Height - 2);
+				e.Graphics.DrawLine(SystemPens.ControlDark, selectedTab.MinX, 0, selectedTab.MinX + selectedTab.Width, 0);
+				e.Graphics.DrawLine(SystemPens.ControlDark, selectedTab.MinX, 0, selectedTab.MinX, ClientSize.Height - 2);
+				e.Graphics.DrawLine(SystemPens.ControlDark, selectedTab.MinX + selectedTab.Width, 0, selectedTab.MinX + selectedTab.Width, ClientSize.Height - 2);
 			}
 
 			e.Graphics.DrawLine(SystemPens.ControlDarkDark, 0, ClientSize.Height - 2,  ClientSize.Width, ClientSize.Height - 2);
 			e.Graphics.DrawLine(SystemPens.ControlLightLight, 0, ClientSize.Height - 1, ClientSize.Width, ClientSize.Height - 1);
 		}
 
-		void DrawBackground(Graphics Graphics, TabData Tab, Brush BackgroundBrush, Brush StripeBrush, Tuple<Color, float>? Highlight)
+		void DrawBackground(Graphics graphics, TabData tab, Brush backgroundBrush, Brush stripeBrush, Tuple<Color, float>? highlight)
 		{
-			float DpiScaleY = Graphics.DpiY / 96.0f;
+			float dpiScaleY = graphics.DpiY / 96.0f;
 
-			Graphics.SmoothingMode = System.Drawing.Drawing2D.SmoothingMode.None;
+			graphics.SmoothingMode = System.Drawing.Drawing2D.SmoothingMode.None;
 
-			Graphics.FillRectangle(BackgroundBrush, Tab.MinX + 1, 1, Tab.Width - 1, ClientSize.Height - 3);
-			Graphics.FillRectangle(StripeBrush, Tab.MinX + 1, ClientSize.Height - (int)(5 * DpiScaleY), Tab.Width - 1, (int)(5 * DpiScaleY) - 2);
+			graphics.FillRectangle(backgroundBrush, tab.MinX + 1, 1, tab.Width - 1, ClientSize.Height - 3);
+			graphics.FillRectangle(stripeBrush, tab.MinX + 1, ClientSize.Height - (int)(5 * dpiScaleY), tab.Width - 1, (int)(5 * dpiScaleY) - 2);
 
-			if (Tab.TintColor.HasValue)
+			if (tab.TintColor.HasValue)
 			{
-				int MidX = Tab.MinX + Tab.Width / 2;
-				using (LinearGradientBrush Brush = new LinearGradientBrush(new Rectangle(MidX, 1, 1, ClientSize.Height), Color.FromArgb(64, Tab.TintColor.Value), Color.FromArgb(32, Tab.TintColor.Value), LinearGradientMode.Vertical))
+				int midX = tab.MinX + tab.Width / 2;
+				using (LinearGradientBrush brush = new LinearGradientBrush(new Rectangle(midX, 1, 1, ClientSize.Height), Color.FromArgb(64, tab.TintColor.Value), Color.FromArgb(32, tab.TintColor.Value), LinearGradientMode.Vertical))
 				{
-					Graphics.FillRectangle(Brush, Tab.MinX + 1, 1, Tab.Width - 1, ClientSize.Height - 3);
+					graphics.FillRectangle(brush, tab.MinX + 1, 1, tab.Width - 1, ClientSize.Height - 3);
 				}
 			}
 
-			if (Highlight != null && Highlight.Item2 > 0.0f)
+			if (highlight != null && highlight.Item2 > 0.0f)
 			{
-				using(SolidBrush Brush = new SolidBrush(Highlight.Item1))
+				using(SolidBrush brush = new SolidBrush(highlight.Item1))
 				{
-					Graphics.FillRectangle(Brush, Tab.MinX + 1, ClientSize.Height - (int)(5 * DpiScaleY), (int)((Tab.Width - 2) * Highlight.Item2), (int)(5 * DpiScaleY) - 2);
+					graphics.FillRectangle(brush, tab.MinX + 1, ClientSize.Height - (int)(5 * dpiScaleY), (int)((tab.Width - 2) * highlight.Item2), (int)(5 * dpiScaleY) - 2);
 				}
 			}
 		}
 
-		void DrawText(Graphics Graphics, TabData Tab)
+		void DrawText(Graphics graphics, TabData tab)
 		{
-			float DpiScaleX = Graphics.DpiX / 96.0f;
-			float DpiScaleY = Graphics.DpiY / 96.0f;
-			TextRenderer.DrawText(Graphics, Tab.Name, Font, new Rectangle(Tab.MinX + (int)(TabPadding * DpiScaleX), 0, Tab.TextSize.Width, ClientSize.Height - (int)(3 * DpiScaleY)), Color.Black, TextFormatFlags.NoPadding | TextFormatFlags.EndEllipsis | TextFormatFlags.VerticalCenter);
+			float dpiScaleX = graphics.DpiX / 96.0f;
+			float dpiScaleY = graphics.DpiY / 96.0f;
+			TextRenderer.DrawText(graphics, tab.Name, Font, new Rectangle(tab.MinX + (int)(TabPadding * dpiScaleX), 0, tab.TextSize.Width, ClientSize.Height - (int)(3 * dpiScaleY)), Color.Black, TextFormatFlags.NoPadding | TextFormatFlags.EndEllipsis | TextFormatFlags.VerticalCenter);
 		}
 	}
 }
diff --git a/Engine/Source/Programs/UnrealGameSync/UnrealGameSync/Controls/TextBoxWithCueBanner.cs b/Engine/Source/Programs/UnrealGameSync/UnrealGameSync/Controls/TextBoxWithCueBanner.cs
--- a/Engine/Source/Programs/UnrealGameSync/UnrealGameSync/Controls/TextBoxWithCueBanner.cs
+++ b/Engine/Source/Programs/UnrealGameSync/UnrealGameSync/Controls/TextBoxWithCueBanner.cs
@@ -12,24 +12,24 @@ namespace UnrealGameSync
 {
 	class TextBoxWithCueBanner : TextBox
 	{
-		private const int EM_SETCUEBANNER = 0x1501;
+		private const int EmSetcuebanner = 0x1501;
 
 		[DllImport("user32.dll", CharSet = CharSet.Unicode)]
-		private static extern IntPtr SendMessage(IntPtr hWnd, int Msg, IntPtr Wparam, string Lparam);
+		private static extern IntPtr SendMessage(IntPtr hWnd, int msg, IntPtr wparam, string lparam);
 
-		private string CueBannerValue = String.Empty;
+		private string _cueBannerValue = String.Empty;
 
 		public string CueBanner
 		{
-			get { return CueBannerValue; }
-			set { CueBannerValue = value; UpdateCueBanner(); }
+			get { return _cueBannerValue; }
+			set { _cueBannerValue = value; UpdateCueBanner(); }
 		}
 
 		private void UpdateCueBanner()
 		{
 			if(IsHandleCreated)
 			{
-				SendMessage(Handle, EM_SETCUEBANNER, (IntPtr)1, CueBannerValue);
+				SendMessage(Handle, EmSetcuebanner, (IntPtr)1, _cueBannerValue);
 			}
 		}
 
