UGS: Reformat to standard C# style conventions.

diff --git a/Engine/Source/Programs/UnrealGameSync/UnrealGameSync/Controls/WorkspaceControl.cs b/Engine/Source/Programs/UnrealGameSync/UnrealGameSync/Controls/WorkspaceControl.cs
--- a/Engine/Source/Programs/UnrealGameSync/UnrealGameSync/Controls/WorkspaceControl.cs
+++ b/Engine/Source/Programs/UnrealGameSync/UnrealGameSync/Controls/WorkspaceControl.cs
@@ -26,7 +26,7 @@ namespace UnrealGameSync
 {
 	public class UncontrolledChangelist
 	{
-		public string GUID { get; set; } = String.Empty;
+		public string Guid { get; set; } = String.Empty;
 		public string Name { get; set; } = String.Empty;
 		public List<string> Files { get; set; } = new List<string>();
 	}
@@ -41,23 +41,23 @@ namespace UnrealGameSync
 	{
 		ToolUpdateMonitor ToolUpdateMonitor { get; }
 
-		void EditSelectedProject(WorkspaceControl Workspace);
-		void RequestProjectChange(WorkspaceControl Workspace, UserSelectedProjectSettings Project, bool bModal);
+		void EditSelectedProject(WorkspaceControl workspace);
+		void RequestProjectChange(WorkspaceControl workspace, UserSelectedProjectSettings project, bool modal);
 		void ShowAndActivate();
-		void StreamChanged(WorkspaceControl Workspace);
-		void SetTabNames(TabLabels TabNames);
+		void StreamChanged(WorkspaceControl workspace);
+		void SetTabNames(TabLabels tabNames);
 		void SetupScheduledSync();
 		void UpdateProgress();
 		void ModifyApplicationSettings();
 		void UpdateAlertWindows();
 		void UpdateTintColors();
 
-		IssueMonitor CreateIssueMonitor(string? ApiUrl, string UserName);
-		void ReleaseIssueMonitor(IssueMonitor IssueMonitor);
+		IssueMonitor CreateIssueMonitor(string? apiUrl, string userName);
+		void ReleaseIssueMonitor(IssueMonitor issueMonitor);
 	}
 
-	delegate void WorkspaceStartupCallback(WorkspaceControl Workspace, bool bCancel);
-	delegate void WorkspaceUpdateCallback(WorkspaceUpdateResult Result);
+	delegate void WorkspaceStartupCallback(WorkspaceControl workspace, bool cancel);
+	delegate void WorkspaceUpdateCallback(WorkspaceUpdateResult result);
 
 	partial class WorkspaceControl : UserControl, IMainWindowTabPanel
 	{
@@ -88,32 +88,32 @@ namespace UnrealGameSync
 			public Action? ClickHandler;
 			public string? ToolTip;
 
-			public BadgeInfo(string Label, string Group, Color BadgeColor)
-				: this(Label, Group, null, BadgeColor, BadgeColor, null)
+			public BadgeInfo(string label, string group, Color badgeColor)
+				: this(label, group, null, badgeColor, badgeColor, null)
 			{
 			}
 
-			public BadgeInfo(string Label, string? Group, string? UniqueId, Color BackgroundColor, Color HoverBackgroundColor, Action? ClickHandler)
+			public BadgeInfo(string label, string? group, string? uniqueId, Color backgroundColor, Color hoverBackgroundColor, Action? clickHandler)
 			{
-				this.Label = Label;
-				this.Group = Group;
-				this.UniqueId = UniqueId;
-				this.BackgroundColor = BackgroundColor;
-				this.HoverBackgroundColor = HoverBackgroundColor;
-				this.ClickHandler = ClickHandler;
+				this.Label = label;
+				this.Group = group;
+				this.UniqueId = uniqueId;
+				this.BackgroundColor = backgroundColor;
+				this.HoverBackgroundColor = hoverBackgroundColor;
+				this.ClickHandler = clickHandler;
 			}
 
-			public BadgeInfo(BadgeInfo Other)
-				: this(Other.Label, Other.Group, Other.UniqueId, Other.BackgroundColor, Other.HoverBackgroundColor, Other.ClickHandler)
+			public BadgeInfo(BadgeInfo other)
+				: this(other.Label, other.Group, other.UniqueId, other.BackgroundColor, other.HoverBackgroundColor, other.ClickHandler)
 			{
-				this.Offset = Other.Offset;
-				this.Width = Other.Width;
-				this.Height = Other.Height;
+				this.Offset = other.Offset;
+				this.Width = other.Width;
+				this.Height = other.Height;
 			}
 
-			public Rectangle GetBounds(Point ListLocation)
+			public Rectangle GetBounds(Point listLocation)
 			{
-				return new Rectangle(ListLocation.X + Offset, ListLocation.Y, Width, Height);
+				return new Rectangle(listLocation.X + Offset, listLocation.Y, Width, Height);
 			}
 		}
 
@@ -123,21 +123,21 @@ namespace UnrealGameSync
 			public string? Url;
 			public BadgeResult Result;
 
-			public ServiceBadgeInfo(string Name, string? Url, BadgeResult Result)
+			public ServiceBadgeInfo(string name, string? url, BadgeResult result)
 			{
-				this.Name = Name;
-				this.Url = Url;
-				this.Result = Result;
+				this.Name = name;
+				this.Url = url;
+				this.Result = result;
 			}
 
-			public bool Equals(ServiceBadgeInfo? OtherBadge)
+			public bool Equals(ServiceBadgeInfo? otherBadge)
 			{
-				return OtherBadge != null && Name == OtherBadge.Name && Url == OtherBadge.Url && Result == OtherBadge.Result;
+				return otherBadge != null && Name == otherBadge.Name && Url == otherBadge.Url && Result == otherBadge.Result;
 			}
 
-			public override bool Equals(object? Other)
+			public override bool Equals(object? other)
 			{
-				return Equals(Other as ServiceBadgeInfo);
+				return Equals(other as ServiceBadgeInfo);
 			}
 
 			public override int GetHashCode()
@@ -156,21 +156,21 @@ namespace UnrealGameSync
 
 		class BuildListItemComparer : System.Collections.IComparer
 		{
-			public int Compare(object? A, object? B)
+			public int Compare(object? a, object? b)
 			{
-				ChangesRecord? ChangeA = ((ListViewItem?)A)?.Tag as ChangesRecord;
-				ChangesRecord? ChangeB = ((ListViewItem?)B)?.Tag as ChangesRecord;
-				if (ChangeA == null)
+				ChangesRecord? changeA = ((ListViewItem?)a)?.Tag as ChangesRecord;
+				ChangesRecord? changeB = ((ListViewItem?)b)?.Tag as ChangesRecord;
+				if (changeA == null)
 				{
 					return +1;
 				}
-				else if (ChangeB == null)
+				else if (changeB == null)
 				{
 					return -1;
 				}
 				else
 				{
-					return ChangeB.Number - ChangeA.Number;
+					return changeB.Number - changeA.Number;
 				}
 			}
 		}
@@ -179,42 +179,42 @@ namespace UnrealGameSync
 		public static Rectangle MixedBuildIcon = new Rectangle(16, 0, 16, 16);
 		public static Rectangle BadBuildIcon = new Rectangle(32, 0, 16, 16);
 		public static Rectangle DefaultBuildIcon = new Rectangle(48, 0, 16, 16);
-		static Rectangle PromotedBuildIcon = new Rectangle(64, 0, 16, 16);
-		static Rectangle DetailsIcon = new Rectangle(80, 0, 16, 16);
-		static Rectangle InfoIcon = new Rectangle(96, 0, 16, 16);
-		static Rectangle CancelIcon = new Rectangle(112, 0, 16, 16);
-		static Rectangle SyncIcon = new Rectangle(128, 0, 32, 16);
-		static Rectangle HappyIcon = new Rectangle(160, 0, 16, 16);
-		static Rectangle DisabledHappyIcon = new Rectangle(176, 0, 16, 16);
-		static Rectangle FrownIcon = new Rectangle(192, 0, 16, 16);
-		static Rectangle DisabledFrownIcon = new Rectangle(208, 0, 16, 16);
-		static Rectangle PreviousSyncIcon = new Rectangle(224, 0, 16, 16);
-		static Rectangle AdditionalSyncIcon = new Rectangle(240, 0, 16, 16);
-		static Rectangle BisectPassIcon = new Rectangle(256, 0, 16, 16);
-		static Rectangle BisectFailIcon = new Rectangle(273, 0, 16, 16);
-		static Rectangle BisectImplicitPassIcon = new Rectangle(290, 0, 16, 16);
-		static Rectangle BisectImplicitFailIcon = new Rectangle(307, 0, 16, 16);
+		static Rectangle _promotedBuildIcon = new Rectangle(64, 0, 16, 16);
+		static Rectangle _detailsIcon = new Rectangle(80, 0, 16, 16);
+		static Rectangle _infoIcon = new Rectangle(96, 0, 16, 16);
+		static Rectangle _cancelIcon = new Rectangle(112, 0, 16, 16);
+		static Rectangle _syncIcon = new Rectangle(128, 0, 32, 16);
+		static Rectangle _happyIcon = new Rectangle(160, 0, 16, 16);
+		static Rectangle _disabledHappyIcon = new Rectangle(176, 0, 16, 16);
+		static Rectangle _frownIcon = new Rectangle(192, 0, 16, 16);
+		static Rectangle _disabledFrownIcon = new Rectangle(208, 0, 16, 16);
+		static Rectangle _previousSyncIcon = new Rectangle(224, 0, 16, 16);
+		static Rectangle _additionalSyncIcon = new Rectangle(240, 0, 16, 16);
+		static Rectangle _bisectPassIcon = new Rectangle(256, 0, 16, 16);
+		static Rectangle _bisectFailIcon = new Rectangle(273, 0, 16, 16);
+		static Rectangle _bisectImplicitPassIcon = new Rectangle(290, 0, 16, 16);
+		static Rectangle _bisectImplicitFailIcon = new Rectangle(307, 0, 16, 16);
 
 		[DllImport("uxtheme.dll", CharSet = CharSet.Unicode)]
 		static extern int SetWindowTheme(IntPtr hWnd, string pszSubAppName, string pszSubIdList);
 
 		const int BuildListExpandCount = 250;
 
-		IWorkspaceControlOwner Owner;
-		DirectoryReference AppDataFolder;
-		string? ApiUrl;
-		DirectoryReference WorkspaceDataFolder;
-		IServiceProvider ServiceProvider;
-		ILogger Logger;
-		IPerforceSettings PerforceSettings;
+		IWorkspaceControlOwner _owner;
+		DirectoryReference _appDataFolder;
+		string? _apiUrl;
+		DirectoryReference _workspaceDataFolder;
+		IServiceProvider _serviceProvider;
+		ILogger _logger;
+		IPerforceSettings _perforceSettings;
 		ProjectInfo ProjectInfo { get; }
 
-		UserSettings Settings;
-		UserWorkspaceState WorkspaceState;
-		UserWorkspaceSettings WorkspaceSettings;
-		UserProjectSettings ProjectSettings;
+		UserSettings _settings;
+		UserWorkspaceState _workspaceState;
+		UserWorkspaceSettings _workspaceSettings;
+		UserProjectSettings _projectSettings;
 
-		bool bUserHasOpenIssues = false;
+		bool _userHasOpenIssues = false;
 
 		public UserSelectedProjectSettings SelectedProject { get; }
 
@@ -224,41 +224,41 @@ namespace UnrealGameSync
 		public string? StreamName => ProjectInfo.StreamName;
 		public string ClientName => ProjectInfo.ClientName;
 
-		SynchronizationContext MainThreadSynchronizationContext;
-		bool bIsDisposing;
-
-		JupiterMonitor JupiterMonitor;
-		PerforceMonitor PerforceMonitor;
-		Workspace Workspace;
-		IssueMonitor IssueMonitor;
-		EventMonitor EventMonitor;
-		System.Windows.Forms.Timer UpdateTimer;
-		HashSet<int> PromotedChangeNumbers = new HashSet<int>();
-		List<int> ListIndexToChangeIndex = new List<int>();
-		List<int> SortedChangeNumbers = new List<int>();
-		Dictionary<string, Dictionary<int, string?>> ArchiveToChangeNumberToArchiveKey = new Dictionary<string, Dictionary<int, string?>>();
-		Dictionary<int, ChangeLayoutInfo> ChangeNumberToLayoutInfo = new Dictionary<int, ChangeLayoutInfo>();
-		List<ContextMenuStrip> CustomStatusPanelMenus = new List<ContextMenuStrip>();
-		List<(string, Action<Point, Rectangle>)> CustomStatusPanelLinks = new List<(string, Action<Point, Rectangle>)>();
-		int NumChanges;
-		int PendingSelectedChangeNumber = -1;
-		bool bHasBuildSteps = false;
-
-		Dictionary<string, int> NotifiedBuildTypeToChangeNumber = new Dictionary<string, int>();
-
-		bool bMouseOverExpandLink;
-
-		string? HoverBadgeUniqueId = null;
-		bool bHoverSync;
-		ChangesRecord? ContextMenuChange;
-		Font? BuildFont;
-		Font? SelectedBuildFont;
-		Font? BadgeFont;
-		List<KeyValuePair<string, string>> BadgeNameAndGroupPairs = new List<KeyValuePair<string, string>>();
-		Dictionary<string, Size> BadgeLabelToSize = new Dictionary<string, Size>();
-		List<ServiceBadgeInfo> ServiceBadges = new List<ServiceBadgeInfo>();
-
-		NotificationWindow NotificationWindow;
+		SynchronizationContext _mainThreadSynchronizationContext;
+		bool _isDisposing;
+
+		JupiterMonitor _jupiterMonitor;
+		PerforceMonitor _perforceMonitor;
+		Workspace _workspace;
+		IssueMonitor _issueMonitor;
+		EventMonitor _eventMonitor;
+		System.Windows.Forms.Timer _updateTimer;
+		HashSet<int> _promotedChangeNumbers = new HashSet<int>();
+		List<int> _listIndexToChangeIndex = new List<int>();
+		List<int> _sortedChangeNumbers = new List<int>();
+		Dictionary<string, Dictionary<int, string?>> _archiveToChangeNumberToArchiveKey = new Dictionary<string, Dictionary<int, string?>>();
+		Dictionary<int, ChangeLayoutInfo> _changeNumberToLayoutInfo = new Dictionary<int, ChangeLayoutInfo>();
+		List<ContextMenuStrip> _customStatusPanelMenus = new List<ContextMenuStrip>();
+		List<(string, Action<Point, Rectangle>)> _customStatusPanelLinks = new List<(string, Action<Point, Rectangle>)>();
+		int _numChanges;
+		int _pendingSelectedChangeNumber = -1;
+		bool _hasBuildSteps = false;
+
+		Dictionary<string, int> _notifiedBuildTypeToChangeNumber = new Dictionary<string, int>();
+
+		bool _mouseOverExpandLink;
+
+		string? _hoverBadgeUniqueId = null;
+		bool _hoverSync;
+		ChangesRecord? _contextMenuChange;
+		Font? _buildFont;
+		Font? _selectedBuildFont;
+		Font? _badgeFont;
+		List<KeyValuePair<string, string>> _badgeNameAndGroupPairs = new List<KeyValuePair<string, string>>();
+		Dictionary<string, Size> _badgeLabelToSize = new Dictionary<string, Size>();
+		List<ServiceBadgeInfo> _serviceBadges = new List<ServiceBadgeInfo>();
+
+		NotificationWindow _notificationWindow;
 
 		public Tuple<TaskbarState, float> DesiredTaskbarState
 		{
@@ -266,54 +266,54 @@ namespace UnrealGameSync
 			private set;
 		}
 
-		int BuildListWidth;
-		float[] ColumnWidths = Array.Empty<float>();
-		float[] ColumnWeights = Array.Empty<float>();
-		int[] MinColumnWidths = Array.Empty<int>();
-		int[] DesiredColumnWidths = Array.Empty<int>();
-		string? LastColumnSettings;
-		List<ColumnHeader> CustomColumns = new List<ColumnHeader>();
-		int MaxBuildBadgeChars;
-		ListViewItem? ExpandItem;
+		int _buildListWidth;
+		float[] _columnWidths = Array.Empty<float>();
+		float[] _columnWeights = Array.Empty<float>();
+		int[] _minColumnWidths = Array.Empty<int>();
+		int[] _desiredColumnWidths = Array.Empty<int>();
+		string? _lastColumnSettings;
+		List<ColumnHeader> _customColumns = new List<ColumnHeader>();
+		int _maxBuildBadgeChars;
+		ListViewItem? _expandItem;
 
-		bool bUpdateBuildListPosted;
-		bool bUpdateBuildMetadataPosted;
-		bool bUpdateReviewsPosted;
+		bool _updateBuildListPosted;
+		bool _updateBuildMetadataPosted;
+		bool _updateReviewsPosted;
 
-		WorkspaceUpdateCallback? UpdateCallback;
+		WorkspaceUpdateCallback? _updateCallback;
 
-		System.Threading.Timer? StartupTimer;
-		List<WorkspaceStartupCallback>? StartupCallbacks;
+		System.Threading.Timer? _startupTimer;
+		List<WorkspaceStartupCallback>? _startupCallbacks;
 
 		// When an author filter is applied this will be non-empty and != AuthorFilterPlaceholderText
-		string AuthorFilterText = "";
+		string _authorFilterText = "";
 		
 		// Placeholder text that is in the control and cleared when the user starts editing.
-		static string AuthorFilterPlaceholderText = "<username>";
+		static string _authorFilterPlaceholderText = "<username>";
 
-		public WorkspaceControl(IWorkspaceControlOwner InOwner, DirectoryReference InAppDataFolder, string? InApiUrl, OpenProjectInfo OpenProjectInfo, IServiceProvider InServiceProvider, UserSettings InSettings, OIDCTokenManager? InOidcTokenManager)
+		public WorkspaceControl(IWorkspaceControlOwner inOwner, DirectoryReference inAppDataFolder, string? inApiUrl, OpenProjectInfo openProjectInfo, IServiceProvider inServiceProvider, UserSettings inSettings, OidcTokenManager? inOidcTokenManager)
 		{
 			InitializeComponent();
 
-			MainThreadSynchronizationContext = SynchronizationContext.Current!;
-
-			Owner = InOwner;
-			AppDataFolder = InAppDataFolder;
-			ApiUrl = InApiUrl;
-			WorkspaceDataFolder = OpenProjectInfo.ProjectInfo.DataFolder;
-			ServiceProvider = InServiceProvider;
-			Logger = InServiceProvider.GetRequiredService<ILogger<WorkspaceControl>>();
-			PerforceSettings = OpenProjectInfo.PerforceSettings;
-			ProjectInfo = OpenProjectInfo.ProjectInfo;
-			Settings = InSettings;
-			this.WorkspaceSettings = OpenProjectInfo.WorkspaceSettings;
-			this.WorkspaceState = OpenProjectInfo.WorkspaceState;
-			ProjectSettings = InSettings.FindOrAddProjectSettings(OpenProjectInfo.ProjectInfo, OpenProjectInfo.WorkspaceSettings, Logger);
+			_mainThreadSynchronizationContext = SynchronizationContext.Current!;
+
+			_owner = inOwner;
+			_appDataFolder = inAppDataFolder;
+			_apiUrl = inApiUrl;
+			_workspaceDataFolder = openProjectInfo.ProjectInfo.DataFolder;
+			_serviceProvider = inServiceProvider;
+			_logger = inServiceProvider.GetRequiredService<ILogger<WorkspaceControl>>();
+			_perforceSettings = openProjectInfo.PerforceSettings;
+			ProjectInfo = openProjectInfo.ProjectInfo;
+			_settings = inSettings;
+			this._workspaceSettings = openProjectInfo.WorkspaceSettings;
+			this._workspaceState = openProjectInfo.WorkspaceState;
+			_projectSettings = inSettings.FindOrAddProjectSettings(openProjectInfo.ProjectInfo, openProjectInfo.WorkspaceSettings, _logger);
 
 			DesiredTaskbarState = Tuple.Create(TaskbarState.NoProgress, 0.0f);
 
-			System.Reflection.PropertyInfo DoubleBufferedProperty = typeof(Control).GetProperty("DoubleBuffered", System.Reflection.BindingFlags.NonPublic | System.Reflection.BindingFlags.Instance)!;
-			DoubleBufferedProperty.SetValue(BuildList, true, null);
+			System.Reflection.PropertyInfo doubleBufferedProperty = typeof(Control).GetProperty("DoubleBuffered", System.Reflection.BindingFlags.NonPublic | System.Reflection.BindingFlags.Instance)!;
+			doubleBufferedProperty.SetValue(BuildList, true, null);
 
 			// force the height of the rows
 			BuildList.SmallImageList = new ImageList() { ImageSize = new Size(1, 20) };
@@ -323,27 +323,27 @@ namespace UnrealGameSync
 
 			Splitter.OnVisibilityChanged += Splitter_OnVisibilityChanged;
 
-			UpdateTimer = new System.Windows.Forms.Timer();
-			UpdateTimer.Interval = 30;
-			UpdateTimer.Tick += TimerCallback;
+			_updateTimer = new System.Windows.Forms.Timer();
+			_updateTimer.Interval = 30;
+			_updateTimer.Tick += TimerCallback;
 
 			UpdateCheckedBuildConfig();
 
 			UpdateSyncActionCheckboxes();
 
 			// Set the project logo on the status panel and notification window
-			NotificationWindow = new NotificationWindow(Properties.Resources.DefaultNotificationLogo);
+			_notificationWindow = new NotificationWindow(Properties.Resources.DefaultNotificationLogo);
 			if (ProjectInfo.LocalFileName.HasExtension(".uproject"))
 			{
-				FileReference LogoFileName = FileReference.Combine(ProjectInfo.LocalFileName.Directory, "Build", "UnrealGameSync.png");
-				if (FileReference.Exists(LogoFileName))
+				FileReference logoFileName = FileReference.Combine(ProjectInfo.LocalFileName.Directory, "Build", "UnrealGameSync.png");
+				if (FileReference.Exists(logoFileName))
 				{
 					try
 					{
 						// Duplicate the image, otherwise we'll leave the file locked
-						using (Image Image = Image.FromFile(LogoFileName.FullName))
+						using (Image image = Image.FromFile(logoFileName.FullName))
 						{
-							StatusPanel.SetProjectLogo(new Bitmap(Image), true);
+							StatusPanel.SetProjectLogo(new Bitmap(image), true);
 						}
 					}
 					catch
@@ -353,42 +353,42 @@ namespace UnrealGameSync
 			}
 
 			// Commit all the new project info
-			IPerforceSettings PerforceClientSettings = OpenProjectInfo.PerforceSettings;
-			SelectedProject = OpenProjectInfo.SelectedProject;
+			IPerforceSettings perforceClientSettings = openProjectInfo.PerforceSettings;
+			SelectedProject = openProjectInfo.SelectedProject;
 
 			// Figure out which API server to use
-			string? NewApiUrl;
-			if (TryGetProjectSetting(OpenProjectInfo.LatestProjectConfigFile, "ApiUrl", out NewApiUrl))
+			string? newApiUrl;
+			if (TryGetProjectSetting(openProjectInfo.LatestProjectConfigFile, "ApiUrl", out newApiUrl))
 			{
-				ApiUrl = NewApiUrl;
+				_apiUrl = newApiUrl;
 			}
 
-			ProjectInfo Project = OpenProjectInfo.ProjectInfo;
+			ProjectInfo project = openProjectInfo.ProjectInfo;
 
-			Workspace = new Workspace(PerforceClientSettings, Project, WorkspaceState, OpenProjectInfo.WorkspaceProjectConfigFile, OpenProjectInfo.WorkspaceProjectStreamFilter, new LogControlTextWriter(SyncLog), ServiceProvider);
-			Workspace.OnUpdateComplete += UpdateCompleteCallback;
+			_workspace = new Workspace(perforceClientSettings, project, _workspaceState, openProjectInfo.WorkspaceProjectConfigFile, openProjectInfo.WorkspaceProjectStreamFilter, new LogControlTextWriter(SyncLog), _serviceProvider);
+			_workspace.OnUpdateComplete += UpdateCompleteCallback;
 
-			FileReference ProjectLogBaseName = FileReference.Combine(WorkspaceDataFolder, "sync.log");
+			FileReference projectLogBaseName = FileReference.Combine(_workspaceDataFolder, "sync.log");
 
-			ILogger PerforceLogger = ServiceProvider.GetRequiredService<ILogger<PerforceMonitor>>();
-			PerforceMonitor = new PerforceMonitor(PerforceClientSettings, OpenProjectInfo.ProjectInfo, OpenProjectInfo.LatestProjectConfigFile, OpenProjectInfo.ProjectInfo.CacheFolder, OpenProjectInfo.LocalConfigFiles, ServiceProvider);
-			PerforceMonitor.OnUpdate += UpdateBuildListCallback;
-			PerforceMonitor.OnUpdateMetadata += UpdateBuildMetadataCallback;
-			PerforceMonitor.OnStreamChange += StreamChanged;
-			PerforceMonitor.OnLoginExpired += LoginExpired;
+			ILogger perforceLogger = _serviceProvider.GetRequiredService<ILogger<PerforceMonitor>>();
+			_perforceMonitor = new PerforceMonitor(perforceClientSettings, openProjectInfo.ProjectInfo, openProjectInfo.LatestProjectConfigFile, openProjectInfo.ProjectInfo.CacheFolder, openProjectInfo.LocalConfigFiles, _serviceProvider);
+			_perforceMonitor.OnUpdate += UpdateBuildListCallback;
+			_perforceMonitor.OnUpdateMetadata += UpdateBuildMetadataCallback;
+			_perforceMonitor.OnStreamChange += StreamChanged;
+			_perforceMonitor.OnLoginExpired += LoginExpired;
 
-			ILogger EventLogger = ServiceProvider.GetRequiredService < ILogger<EventMonitor>>();
-			EventMonitor = new EventMonitor(ApiUrl, PerforceUtils.GetClientOrDepotDirectoryName(SelectedProjectIdentifier), OpenProjectInfo.PerforceSettings.UserName, ServiceProvider);
-			EventMonitor.OnUpdatesReady += UpdateReviewsCallback;
+			ILogger eventLogger = _serviceProvider.GetRequiredService < ILogger<EventMonitor>>();
+			_eventMonitor = new EventMonitor(_apiUrl, PerforceUtils.GetClientOrDepotDirectoryName(SelectedProjectIdentifier), openProjectInfo.PerforceSettings.UserName, _serviceProvider);
+			_eventMonitor.OnUpdatesReady += UpdateReviewsCallback;
 
-			ILogger<JupiterMonitor> JupiterLogger = ServiceProvider.GetRequiredService<ILogger<JupiterMonitor>>();
-			JupiterMonitor = JupiterMonitor.CreateFromConfigFile(InOidcTokenManager, JupiterLogger, OpenProjectInfo.LatestProjectConfigFile, SelectedProjectIdentifier);
+			ILogger<JupiterMonitor> jupiterLogger = _serviceProvider.GetRequiredService<ILogger<JupiterMonitor>>();
+			_jupiterMonitor = JupiterMonitor.CreateFromConfigFile(inOidcTokenManager, jupiterLogger, openProjectInfo.LatestProjectConfigFile, SelectedProjectIdentifier);
 
 			UpdateColumnSettings(true);
 
-			SyncLog.OpenFile(ProjectLogBaseName);
+			SyncLog.OpenFile(projectLogBaseName);
 
-			Splitter.SetLogVisibility(Settings.bShowLogWindow);
+			Splitter.SetLogVisibility(_settings.ShowLogWindow);
 
 			BuildList.Items.Clear();
 			UpdateBuildList();
@@ -397,23 +397,23 @@ namespace UnrealGameSync
 			UpdateStatusPanel();
 			UpdateServiceBadges();
 
-			PerforceMonitor.Start();
-			EventMonitor.Start();
+			_perforceMonitor.Start();
+			_eventMonitor.Start();
 
-			StartupTimer = new System.Threading.Timer(x => MainThreadSynchronizationContext.Post((o) => { if (!IsDisposed) { StartupTimerElapsed(false); } }, null), null, TimeSpan.FromSeconds(20.0), TimeSpan.FromMilliseconds(-1.0));
-			StartupCallbacks = new List<WorkspaceStartupCallback>();
+			_startupTimer = new System.Threading.Timer(x => _mainThreadSynchronizationContext.Post((o) => { if (!IsDisposed) { StartupTimerElapsed(false); } }, null), null, TimeSpan.FromSeconds(20.0), TimeSpan.FromMilliseconds(-1.0));
+			_startupCallbacks = new List<WorkspaceStartupCallback>();
 
-			string? IssuesApiUrl = GetIssuesApiUrl();
-			IssueMonitor = InOwner.CreateIssueMonitor(IssuesApiUrl, OpenProjectInfo.PerforceSettings.UserName);
-			IssueMonitor.OnIssuesChanged += IssueMonitor_OnIssuesChangedAsync;
+			string? issuesApiUrl = GetIssuesApiUrl();
+			_issueMonitor = inOwner.CreateIssueMonitor(issuesApiUrl, openProjectInfo.PerforceSettings.UserName);
+			_issueMonitor.OnIssuesChanged += IssueMonitor_OnIssuesChangedAsync;
 
 			if (SelectedFileName.HasExtension(".uproject"))
 			{
-				DirectoryReference ConfigDir = DirectoryReference.Combine(SelectedFileName.Directory, "Saved", "Config");
+				DirectoryReference configDir = DirectoryReference.Combine(SelectedFileName.Directory, "Saved", "Config");
 
-				DirectoryReference.CreateDirectory(ConfigDir);
+				DirectoryReference.CreateDirectory(configDir);
 
-				EditorConfigWatcher.Path = ConfigDir.FullName;
+				EditorConfigWatcher.Path = configDir.FullName;
 				EditorConfigWatcher.NotifyFilter = NotifyFilters.LastWrite | NotifyFilters.FileName;
 				EditorConfigWatcher.Filter = "*.ini";
 				EditorConfigWatcher.IncludeSubdirectories = true;
@@ -422,24 +422,24 @@ namespace UnrealGameSync
 				TintColor = GetTintColor();
 			}
 
-			Owner.ToolUpdateMonitor.OnChange += UpdateStatusPanel_CrossThread;
+			_owner.ToolUpdateMonitor.OnChange += UpdateStatusPanel_CrossThread;
 		}
 
 		public void IssueMonitor_OnIssuesChangedAsync()
 		{
-			MainThreadSynchronizationContext.Post((o) => { if (IssueMonitor != null) { IssueMonitor_OnIssuesChanged(); } }, null);
+			_mainThreadSynchronizationContext.Post((o) => { if (_issueMonitor != null) { IssueMonitor_OnIssuesChanged(); } }, null);
 		}
 
 		public List<IssueData> GetOpenIssuesForUser()
 		{
-			return IssueMonitor.GetIssues().Where(x => x.FixChange <= 0 && String.Compare(x.Owner, PerforceSettings.UserName, StringComparison.OrdinalIgnoreCase) == 0).ToList();
+			return _issueMonitor.GetIssues().Where(x => x.FixChange <= 0 && String.Compare(x.Owner, _perforceSettings.UserName, StringComparison.OrdinalIgnoreCase) == 0).ToList();
 		}
 
 		public void IssueMonitor_OnIssuesChanged()
 		{
-			bool bPrevUserHasOpenIssues = bUserHasOpenIssues;
-			bUserHasOpenIssues = GetOpenIssuesForUser().Count > 0;
-			if (bUserHasOpenIssues != bPrevUserHasOpenIssues)
+			bool prevUserHasOpenIssues = _userHasOpenIssues;
+			_userHasOpenIssues = GetOpenIssuesForUser().Count > 0;
+			if (_userHasOpenIssues != prevUserHasOpenIssues)
 			{
 				UpdateStatusPanel();
 				StatusPanel.Invalidate();
@@ -448,226 +448,226 @@ namespace UnrealGameSync
 
 		public IssueMonitor GetIssueMonitor()
 		{
-			return IssueMonitor;
+			return _issueMonitor;
 		}
 
 		private void CheckForStartupComplete()
 		{
-			if (StartupTimer != null)
+			if (_startupTimer != null)
 			{
-				int LatestChangeNumber;
-				if (FindChangeToSync(Settings.SyncTypeID, out LatestChangeNumber))
+				int latestChangeNumber;
+				if (FindChangeToSync(_settings.SyncTypeId, out latestChangeNumber))
 				{
 					StartupTimerElapsed(false);
 				}
 			}
 		}
 
-		private void StartupTimerElapsed(bool bCancel)
+		private void StartupTimerElapsed(bool cancel)
 		{
-			if (StartupTimer != null)
+			if (_startupTimer != null)
 			{
-				StartupTimer.Dispose();
-				StartupTimer = null;
+				_startupTimer.Dispose();
+				_startupTimer = null;
 			}
 
-			if (StartupCallbacks != null)
+			if (_startupCallbacks != null)
 			{
-				foreach (WorkspaceStartupCallback StartupCallback in StartupCallbacks)
+				foreach (WorkspaceStartupCallback startupCallback in _startupCallbacks)
 				{
-					StartupCallback(this, bCancel);
+					startupCallback(this, cancel);
 				}
-				StartupCallbacks = null;
+				_startupCallbacks = null;
 			}
 		}
 
-		public void AddStartupCallback(WorkspaceStartupCallback StartupCallback)
+		public void AddStartupCallback(WorkspaceStartupCallback startupCallback)
 		{
-			if (StartupTimer == null)
+			if (_startupTimer == null)
 			{
-				StartupCallback(this, false);
+				startupCallback(this, false);
 			}
 			else
 			{
-				StartupCallbacks!.Add(StartupCallback);
+				_startupCallbacks!.Add(startupCallback);
 			}
 		}
 
-		private void UpdateColumnSettings(bool bForceUpdate)
+		private void UpdateColumnSettings(bool forceUpdate)
 		{
-			string? NextColumnSettings;
-			TryGetProjectSetting(PerforceMonitor.LatestProjectConfigFile, "Columns", out NextColumnSettings);
+			string? nextColumnSettings;
+			TryGetProjectSetting(_perforceMonitor.LatestProjectConfigFile, "Columns", out nextColumnSettings);
 
-			if (bForceUpdate || NextColumnSettings != LastColumnSettings)
+			if (forceUpdate || nextColumnSettings != _lastColumnSettings)
 			{
-				LastColumnSettings = NextColumnSettings;
+				_lastColumnSettings = nextColumnSettings;
 
-				foreach (ColumnHeader CustomColumn in CustomColumns)
+				foreach (ColumnHeader customColumn in _customColumns)
 				{
-					BuildList.Columns.Remove(CustomColumn);
+					BuildList.Columns.Remove(customColumn);
 				}
 
-				Dictionary<string, ColumnHeader> NameToColumn = new Dictionary<string, ColumnHeader>();
-				foreach (ColumnHeader? Column in BuildList.Columns)
+				Dictionary<string, ColumnHeader> nameToColumn = new Dictionary<string, ColumnHeader>();
+				foreach (ColumnHeader? column in BuildList.Columns)
 				{
-					if (Column != null)
+					if (column != null)
 					{
-						NameToColumn[Column.Text] = Column;
-						Column.Tag = null;
+						nameToColumn[column.Text] = column;
+						column.Tag = null;
 					}
 				}
 
-				CustomColumns = new List<ColumnHeader>();
-				if (NextColumnSettings != null)
+				_customColumns = new List<ColumnHeader>();
+				if (nextColumnSettings != null)
 				{
-					foreach (string CustomColumn in NextColumnSettings.Split('\n'))
+					foreach (string customColumn in nextColumnSettings.Split('\n'))
 					{
-						ConfigObject ColumnConfig = new ConfigObject(CustomColumn);
+						ConfigObject columnConfig = new ConfigObject(customColumn);
 
-						string? Name = ColumnConfig.GetValue("Name", null);
-						if (Name != null)
+						string? name = columnConfig.GetValue("Name", null);
+						if (name != null)
 						{
-							ColumnHeader? Column;
-							if (NameToColumn.TryGetValue(Name, out Column))
+							ColumnHeader? column;
+							if (nameToColumn.TryGetValue(name, out column))
 							{
-								Column.Tag = ColumnConfig;
+								column.Tag = columnConfig;
 							}
 							else
 							{
-								int Index = ColumnConfig.GetValue("Index", -1);
-								if (Index < 0 || Index > BuildList.Columns.Count)
+								int index = columnConfig.GetValue("Index", -1);
+								if (index < 0 || index > BuildList.Columns.Count)
 								{
-									Index = ((CustomColumns.Count > 0) ? CustomColumns[CustomColumns.Count - 1].Index : CISColumn.Index) + 1;
+									index = ((_customColumns.Count > 0) ? _customColumns[_customColumns.Count - 1].Index : CISColumn.Index) + 1;
 								}
 
-								Column = new ColumnHeader();
-								Column.Text = Name;
-								Column.Tag = ColumnConfig;
-								BuildList.Columns.Insert(Index, Column);
+								column = new ColumnHeader();
+								column.Text = name;
+								column.Tag = columnConfig;
+								BuildList.Columns.Insert(index, column);
 
-								CustomColumns.Add(Column);
+								_customColumns.Add(column);
 							}
 						}
 					}
 				}
 
-				ColumnWidths = new float[BuildList.Columns.Count];
-				for (int Idx = 0; Idx < BuildList.Columns.Count; Idx++)
+				_columnWidths = new float[BuildList.Columns.Count];
+				for (int idx = 0; idx < BuildList.Columns.Count; idx++)
 				{
-					ColumnWidths[Idx] = BuildList.Columns[Idx].Width;
+					_columnWidths[idx] = BuildList.Columns[idx].Width;
 				}
 
-				using (Graphics Graphics = Graphics.FromHwnd(IntPtr.Zero))
+				using (Graphics graphics = Graphics.FromHwnd(IntPtr.Zero))
 				{
-					float DpiScaleX = Graphics.DpiX / 96.0f;
+					float dpiScaleX = graphics.DpiX / 96.0f;
 
-					MinColumnWidths = Enumerable.Repeat(32, BuildList.Columns.Count).ToArray();
-					MinColumnWidths[IconColumn.Index] = (int)(50 * DpiScaleX);
-					MinColumnWidths[TypeColumn.Index] = (int)(100 * DpiScaleX);
-					MinColumnWidths[TimeColumn.Index] = (int)(75 * DpiScaleX);
-					MinColumnWidths[ChangeColumn.Index] = (int)(75 * DpiScaleX);
-					MinColumnWidths[CISColumn.Index] = (int)(200 * DpiScaleX);
+					_minColumnWidths = Enumerable.Repeat(32, BuildList.Columns.Count).ToArray();
+					_minColumnWidths[IconColumn.Index] = (int)(50 * dpiScaleX);
+					_minColumnWidths[TypeColumn.Index] = (int)(100 * dpiScaleX);
+					_minColumnWidths[TimeColumn.Index] = (int)(75 * dpiScaleX);
+					_minColumnWidths[ChangeColumn.Index] = (int)(75 * dpiScaleX);
+					_minColumnWidths[CISColumn.Index] = (int)(200 * dpiScaleX);
 
-					DesiredColumnWidths = Enumerable.Repeat(65536, BuildList.Columns.Count).ToArray();
-					DesiredColumnWidths[IconColumn.Index] = MinColumnWidths[IconColumn.Index];
-					DesiredColumnWidths[TypeColumn.Index] = MinColumnWidths[TypeColumn.Index];
-					DesiredColumnWidths[TimeColumn.Index] = MinColumnWidths[TimeColumn.Index];
-					DesiredColumnWidths[ChangeColumn.Index] = MinColumnWidths[ChangeColumn.Index];
-					DesiredColumnWidths[AuthorColumn.Index] = (int)(120 * DpiScaleX);
-					DesiredColumnWidths[CISColumn.Index] = (int)(200 * DpiScaleX);
-					DesiredColumnWidths[StatusColumn.Index] = (int)(300 * DpiScaleX);
+					_desiredColumnWidths = Enumerable.Repeat(65536, BuildList.Columns.Count).ToArray();
+					_desiredColumnWidths[IconColumn.Index] = _minColumnWidths[IconColumn.Index];
+					_desiredColumnWidths[TypeColumn.Index] = _minColumnWidths[TypeColumn.Index];
+					_desiredColumnWidths[TimeColumn.Index] = _minColumnWidths[TimeColumn.Index];
+					_desiredColumnWidths[ChangeColumn.Index] = _minColumnWidths[ChangeColumn.Index];
+					_desiredColumnWidths[AuthorColumn.Index] = (int)(120 * dpiScaleX);
+					_desiredColumnWidths[CISColumn.Index] = (int)(200 * dpiScaleX);
+					_desiredColumnWidths[StatusColumn.Index] = (int)(300 * dpiScaleX);
 
-					ColumnWeights = Enumerable.Repeat(1.0f, BuildList.Columns.Count).ToArray();
-					ColumnWeights[IconColumn.Index] = 3.0f;
-					ColumnWeights[TypeColumn.Index] = 3.0f;
-					ColumnWeights[TimeColumn.Index] = 3.0f;
-					ColumnWeights[ChangeColumn.Index] = 3.0f;
-					ColumnWeights[DescriptionColumn.Index] = 1.25f;
-					ColumnWeights[CISColumn.Index] = 1.5f;
+					_columnWeights = Enumerable.Repeat(1.0f, BuildList.Columns.Count).ToArray();
+					_columnWeights[IconColumn.Index] = 3.0f;
+					_columnWeights[TypeColumn.Index] = 3.0f;
+					_columnWeights[TimeColumn.Index] = 3.0f;
+					_columnWeights[ChangeColumn.Index] = 3.0f;
+					_columnWeights[DescriptionColumn.Index] = 1.25f;
+					_columnWeights[CISColumn.Index] = 1.5f;
 
-					foreach (ColumnHeader? Column in BuildList.Columns)
+					foreach (ColumnHeader? column in BuildList.Columns)
 					{
-						if (Column != null)
+						if (column != null)
 						{
-							ConfigObject? ColumnConfig = (ConfigObject?)Column.Tag;
-							if (ColumnConfig != null)
+							ConfigObject? columnConfig = (ConfigObject?)column.Tag;
+							if (columnConfig != null)
 							{
-								MinColumnWidths[Column.Index] = (int)(ColumnConfig.GetValue("MinWidth", MinColumnWidths[Column.Index]) * DpiScaleX);
-								DesiredColumnWidths[Column.Index] = (int)(ColumnConfig.GetValue("DesiredWidth", DesiredColumnWidths[Column.Index]) * DpiScaleX);
-								ColumnWeights[Column.Index] = ColumnConfig.GetValue("Weight", MinColumnWidths[Column.Index]);
+								_minColumnWidths[column.Index] = (int)(columnConfig.GetValue("MinWidth", _minColumnWidths[column.Index]) * dpiScaleX);
+								_desiredColumnWidths[column.Index] = (int)(columnConfig.GetValue("DesiredWidth", _desiredColumnWidths[column.Index]) * dpiScaleX);
+								_columnWeights[column.Index] = columnConfig.GetValue("Weight", _minColumnWidths[column.Index]);
 							}
 						}
 					}
 
-					ConfigFile ProjectConfigFile = PerforceMonitor.LatestProjectConfigFile;
-					for (int Idx = 0; Idx < BuildList.Columns.Count; Idx++)
+					ConfigFile projectConfigFile = _perforceMonitor.LatestProjectConfigFile;
+					for (int idx = 0; idx < BuildList.Columns.Count; idx++)
 					{
-						if (!String.IsNullOrEmpty(BuildList.Columns[Idx].Text))
+						if (!String.IsNullOrEmpty(BuildList.Columns[idx].Text))
 						{
-							string? StringValue;
-							if (TryGetProjectSetting(ProjectConfigFile, String.Format("ColumnWidth_{0}", BuildList.Columns[Idx].Text), out StringValue))
+							string? stringValue;
+							if (TryGetProjectSetting(projectConfigFile, String.Format("ColumnWidth_{0}", BuildList.Columns[idx].Text), out stringValue))
 							{
-								int IntValue;
-								if (Int32.TryParse(StringValue, out IntValue))
+								int intValue;
+								if (Int32.TryParse(stringValue, out intValue))
 								{
-									DesiredColumnWidths[Idx] = (int)(IntValue * DpiScaleX);
+									_desiredColumnWidths[idx] = (int)(intValue * dpiScaleX);
 								}
 							}
 						}
 					}
 				}
 
-				if (ColumnWidths != null)
+				if (_columnWidths != null)
 				{
-					ResizeColumns(ColumnWidths.Sum());
+					ResizeColumns(_columnWidths.Sum());
 				}
 			}
 		}
 
 		private bool UpdateServiceBadges()
 		{
-			string[]? ServiceBadgeNames;
-			if (!TryGetProjectSetting(PerforceMonitor.LatestProjectConfigFile, "ServiceBadges", out ServiceBadgeNames))
+			string[]? serviceBadgeNames;
+			if (!TryGetProjectSetting(_perforceMonitor.LatestProjectConfigFile, "ServiceBadges", out serviceBadgeNames))
 			{
-				ServiceBadgeNames = Array.Empty<string>();
+				serviceBadgeNames = Array.Empty<string>();
 			}
 
-			List<ServiceBadgeInfo> PrevServiceBadges = ServiceBadges;
+			List<ServiceBadgeInfo> prevServiceBadges = _serviceBadges;
 
-			ServiceBadges = new List<ServiceBadgeInfo>();
-			foreach (string ServiceBadgeName in ServiceBadgeNames)
+			_serviceBadges = new List<ServiceBadgeInfo>();
+			foreach (string serviceBadgeName in serviceBadgeNames)
 			{
-				BadgeData? LatestBuild;
-				if (EventMonitor.TryGetLatestBadge(ServiceBadgeName, out LatestBuild))
+				BadgeData? latestBuild;
+				if (_eventMonitor.TryGetLatestBadge(serviceBadgeName, out latestBuild))
 				{
-					ServiceBadges.Add(new ServiceBadgeInfo(ServiceBadgeName, LatestBuild.Url, LatestBuild.Result));
+					_serviceBadges.Add(new ServiceBadgeInfo(serviceBadgeName, latestBuild.Url, latestBuild.Result));
 				}
 				else
 				{
-					ServiceBadges.Add(new ServiceBadgeInfo(ServiceBadgeName, null, BadgeResult.Skipped));
+					_serviceBadges.Add(new ServiceBadgeInfo(serviceBadgeName, null, BadgeResult.Skipped));
 				}
 			}
 
-			return !Enumerable.SequenceEqual(ServiceBadges, PrevServiceBadges);
+			return !Enumerable.SequenceEqual(_serviceBadges, prevServiceBadges);
 		}
 
 		protected override void OnLoad(EventArgs e)
 		{
 			base.OnLoad(e);
 
-			BuildListWidth = BuildList.Width;
+			_buildListWidth = BuildList.Width;
 
 			// Find the default widths 
-			ColumnWidths = new float[BuildList.Columns.Count];
-			for (int Idx = 0; Idx < BuildList.Columns.Count; Idx++)
+			_columnWidths = new float[BuildList.Columns.Count];
+			for (int idx = 0; idx < BuildList.Columns.Count; idx++)
 			{
-				if (DesiredColumnWidths[Idx] > 0)
+				if (_desiredColumnWidths[idx] > 0)
 				{
-					ColumnWidths[Idx] = DesiredColumnWidths[Idx];
+					_columnWidths[idx] = _desiredColumnWidths[idx];
 				}
 				else
 				{
-					ColumnWidths[Idx] = BuildList.Columns[Idx].Width;
+					_columnWidths[idx] = BuildList.Columns[idx].Width;
 				}
 			}
 
@@ -681,85 +681,85 @@ namespace UnrealGameSync
 		/// <param name="disposing">true if managed resources should be disposed; otherwise, false.</param>
 		protected override void Dispose(bool disposing)
 		{
-			if (Owner?.ToolUpdateMonitor?.OnChange != null)
+			if (_owner?.ToolUpdateMonitor?.OnChange != null)
 			{
-				Owner.ToolUpdateMonitor.OnChange -= UpdateStatusPanel_CrossThread;
+				_owner.ToolUpdateMonitor.OnChange -= UpdateStatusPanel_CrossThread;
 			}
 
-			bIsDisposing = true;
+			_isDisposing = true;
 
 			if (disposing && (components != null))
 			{
 				components.Dispose();
 			}
 
-			if (UpdateTimer != null)
+			if (_updateTimer != null)
 			{
-				UpdateTimer.Stop();
+				_updateTimer.Stop();
 			}
 
-			if (StartupCallbacks != null)
+			if (_startupCallbacks != null)
 			{
-				foreach (WorkspaceStartupCallback StartupCallback in StartupCallbacks)
+				foreach (WorkspaceStartupCallback startupCallback in _startupCallbacks)
 				{
-					StartupCallback(this, true);
+					startupCallback(this, true);
 				}
-				StartupCallbacks = null!;
+				_startupCallbacks = null!;
 			}
 
-			if (IssueMonitor != null)
+			if (_issueMonitor != null)
 			{
-				IssueMonitor.OnIssuesChanged -= IssueMonitor_OnIssuesChangedAsync;
-				if (Owner != null)
+				_issueMonitor.OnIssuesChanged -= IssueMonitor_OnIssuesChangedAsync;
+				if (_owner != null)
 				{
-					Owner.ReleaseIssueMonitor(IssueMonitor);
+					_owner.ReleaseIssueMonitor(_issueMonitor);
 				}
-				IssueMonitor = null!;
+				_issueMonitor = null!;
 			}
-			if (StartupTimer != null)
+			if (_startupTimer != null)
 			{
-				StartupTimer.Dispose();
-				StartupTimer = null!;
+				_startupTimer.Dispose();
+				_startupTimer = null!;
 			}
-			if (NotificationWindow != null)
+			if (_notificationWindow != null)
 			{
-				NotificationWindow.Dispose();
-				NotificationWindow = null!;
+				_notificationWindow.Dispose();
+				_notificationWindow = null!;
 			}
-			if (PerforceMonitor != null)
+			if (_perforceMonitor != null)
 			{
-				PerforceMonitor.Dispose();
-				PerforceMonitor = null!;
+				_perforceMonitor.Dispose();
+				_perforceMonitor = null!;
 			}
-			if (Workspace != null)
+			if (_workspace != null)
 			{
-				Workspace.Dispose();
-				Workspace = null!;
+				_workspace.Dispose();
+				_workspace = null!;
 			}
-			if (EventMonitor != null)
+			if (_eventMonitor != null)
 			{
-				EventMonitor.Dispose();
-				EventMonitor = null!;
+				_eventMonitor.Dispose();
+				_eventMonitor = null!;
 			}
-			if (BuildFont != null)
+			if (_buildFont != null)
 			{
-				BuildFont.Dispose();
-				BuildFont = null!;
+				_buildFont.Dispose();
+				_buildFont = null!;
 			}
-			if (SelectedBuildFont != null)
+			if (_selectedBuildFont != null)
 			{
-				SelectedBuildFont.Dispose();
-				SelectedBuildFont = null!;
+				_selectedBuildFont.Dispose();
+				_selectedBuildFont = null!;
 			}
-			if (BadgeFont != null)
+			if (_badgeFont != null)
 			{
-				BadgeFont.Dispose();
-				BadgeFont = null!;
+				_badgeFont.Dispose();
+				_badgeFont = null!;
 			}
-			if (JupiterMonitor != null)
+			if (_jupiterMonitor != null)
 			{
-				JupiterMonitor.Dispose();
-				JupiterMonitor = null!;
+				_jupiterMonitor.Dispose();
+				_jupiterMonitor = null!;
 			}
 
 			base.Dispose(disposing);
@@ -767,17 +767,17 @@ namespace UnrealGameSync
 
 		public bool IsBusy()
 		{
-			return Workspace.IsBusy();
+			return _workspace.IsBusy();
 		}
 
 		public bool CanSyncNow()
 		{
-			return Workspace != null && !Workspace.IsBusy();
+			return _workspace != null && !_workspace.IsBusy();
 		}
 
 		public bool CanLaunchEditor()
 		{
-			return Workspace != null && !Workspace.IsBusy() && Workspace.CurrentChangeNumber != -1;
+			return _workspace != null && !_workspace.IsBusy() && _workspace.CurrentChangeNumber != -1;
 		}
 
 		private void MainWindow_Load(object sender, EventArgs e)
@@ -785,22 +785,22 @@ namespace UnrealGameSync
 			UpdateStatusPanel();
 		}
 
-		private void ShowErrorDialog(string Format, params object[] Args)
+		private void ShowErrorDialog(string format, params object[] args)
 		{
-			string Message = String.Format(Format, Args);
-			Logger.LogError("{Message}", Message);
-			MessageBox.Show(Message);
+			string message = String.Format(format, args);
+			_logger.LogError("{Message}", message);
+			MessageBox.Show(message);
 		}
 
 		public bool CanClose()
 		{
 			CancelWorkspaceUpdate();
-			return !Workspace.IsBusy();
+			return !_workspace.IsBusy();
 		}
 
 		private void StreamChanged()
 		{
-			Owner.StreamChanged(this);
+			_owner.StreamChanged(this);
 			/*
 			StatusPanel.SuspendDisplay();
 
@@ -815,96 +815,96 @@ namespace UnrealGameSync
 
 		private void LoginExpired()
 		{
-			if (!bIsDisposing)
+			if (!_isDisposing)
 			{
-				Logger.LogInformation("Login has expired. Requesting project to be closed.");
-				Owner.RequestProjectChange(this, SelectedProject, false);
+				_logger.LogInformation("Login has expired. Requesting project to be closed.");
+				_owner.RequestProjectChange(this, SelectedProject, false);
 			}
 		}
 
 		private void BuildList_OnScroll()
 		{
-			PendingSelectedChangeNumber = -1;
+			_pendingSelectedChangeNumber = -1;
 		}
 
 		void ShrinkNumRequestedBuilds()
 		{
-			if (PerforceMonitor != null && BuildList.Items.Count > 0 && PendingSelectedChangeNumber == -1)
+			if (_perforceMonitor != null && BuildList.Items.Count > 0 && _pendingSelectedChangeNumber == -1)
 			{
 				// Find the number of visible items using a (slightly wasteful) binary search
-				int VisibleItemCount = 1;
-				for (int StepSize = BuildList.Items.Count / 2; StepSize >= 1;)
+				int visibleItemCount = 1;
+				for (int stepSize = BuildList.Items.Count / 2; stepSize >= 1;)
 				{
-					int TestIndex = VisibleItemCount + StepSize;
-					if (TestIndex < BuildList.Items.Count && BuildList.GetItemRect(TestIndex).Top < BuildList.Height)
+					int testIndex = visibleItemCount + stepSize;
+					if (testIndex < BuildList.Items.Count && BuildList.GetItemRect(testIndex).Top < BuildList.Height)
 					{
-						VisibleItemCount += StepSize;
+						visibleItemCount += stepSize;
 					}
 					else
 					{
-						StepSize /= 2;
+						stepSize /= 2;
 					}
 				}
 
 				// Figure out the last index to ensure is visible
-				int LastVisibleIndex = VisibleItemCount;
-				if (LastVisibleIndex >= ListIndexToChangeIndex.Count)
+				int lastVisibleIndex = visibleItemCount;
+				if (lastVisibleIndex >= _listIndexToChangeIndex.Count)
 				{
-					LastVisibleIndex = ListIndexToChangeIndex.Count - 1;
+					lastVisibleIndex = _listIndexToChangeIndex.Count - 1;
 				}
 
 				// Get the max number of changes to ensure this
-				int NewPendingMaxChanges = 0;
-				if (LastVisibleIndex >= 0)
+				int newPendingMaxChanges = 0;
+				if (lastVisibleIndex >= 0)
 				{
-					NewPendingMaxChanges = ListIndexToChangeIndex[LastVisibleIndex];
+					newPendingMaxChanges = _listIndexToChangeIndex[lastVisibleIndex];
 				}
-				NewPendingMaxChanges = PerforceMonitor.InitialMaxChangesValue + ((Math.Max(NewPendingMaxChanges - PerforceMonitor.InitialMaxChangesValue, 0) + BuildListExpandCount - 1) / BuildListExpandCount) * BuildListExpandCount;
+				newPendingMaxChanges = _perforceMonitor.InitialMaxChangesValue + ((Math.Max(newPendingMaxChanges - _perforceMonitor.InitialMaxChangesValue, 0) + BuildListExpandCount - 1) / BuildListExpandCount) * BuildListExpandCount;
 
 				// Shrink the number of changes retained by the PerforceMonitor class
-				if (PerforceMonitor.PendingMaxChanges > NewPendingMaxChanges)
+				if (_perforceMonitor.PendingMaxChanges > newPendingMaxChanges)
 				{
-					PerforceMonitor.PendingMaxChanges = NewPendingMaxChanges;
+					_perforceMonitor.PendingMaxChanges = newPendingMaxChanges;
 				}
 			}
 		}
 
-		void StartSync(int ChangeNumber)
+		void StartSync(int changeNumber)
 		{
-			StartSync(ChangeNumber, false, null);
+			StartSync(changeNumber, false, null);
 		}
 
-		void StartSync(int ChangeNumber, bool bSyncOnly, WorkspaceUpdateCallback? Callback)
+		void StartSync(int changeNumber, bool syncOnly, WorkspaceUpdateCallback? callback)
 		{
-			WorkspaceUpdateOptions Options = WorkspaceUpdateOptions.Sync | WorkspaceUpdateOptions.SyncArchives | WorkspaceUpdateOptions.GenerateProjectFiles;
-			if (!bSyncOnly)
+			WorkspaceUpdateOptions options = WorkspaceUpdateOptions.Sync | WorkspaceUpdateOptions.SyncArchives | WorkspaceUpdateOptions.GenerateProjectFiles;
+			if (!syncOnly)
 			{
-				if (Settings.bBuildAfterSync)
+				if (_settings.BuildAfterSync)
 				{
-					Options |= WorkspaceUpdateOptions.Build;
+					options |= WorkspaceUpdateOptions.Build;
 				}
-				if ((Settings.bBuildAfterSync || ShouldSyncPrecompiledEditor) && Settings.bRunAfterSync)
+				if ((_settings.BuildAfterSync || ShouldSyncPrecompiledEditor) && _settings.RunAfterSync)
 				{
-					Options |= WorkspaceUpdateOptions.RunAfterSync;
+					options |= WorkspaceUpdateOptions.RunAfterSync;
 				}
-				if (Settings.bOpenSolutionAfterSync)
+				if (_settings.OpenSolutionAfterSync)
 				{
-					Options |= WorkspaceUpdateOptions.OpenSolutionAfterSync;
+					options |= WorkspaceUpdateOptions.OpenSolutionAfterSync;
 				}
 			}
-			StartWorkspaceUpdate(ChangeNumber, Options, Callback);
+			StartWorkspaceUpdate(changeNumber, options, callback);
 		}
 
-		void StartWorkspaceUpdate(int ChangeNumber, WorkspaceUpdateOptions Options)
+		void StartWorkspaceUpdate(int changeNumber, WorkspaceUpdateOptions options)
 		{
-			StartWorkspaceUpdate(ChangeNumber, Options, null);
+			StartWorkspaceUpdate(changeNumber, options, null);
 		}
 
-		void StartWorkspaceUpdate(int ChangeNumber, WorkspaceUpdateOptions Options, WorkspaceUpdateCallback? Callback)
+		void StartWorkspaceUpdate(int changeNumber, WorkspaceUpdateOptions options, WorkspaceUpdateCallback? callback)
 		{
-			if ((Options & (WorkspaceUpdateOptions.Sync | WorkspaceUpdateOptions.Build)) != 0 && GetProcessesRunningInWorkspace().Length > 0)
+			if ((options & (WorkspaceUpdateOptions.Sync | WorkspaceUpdateOptions.Build)) != 0 && GetProcessesRunningInWorkspace().Length > 0)
 			{
-				if ((Options & WorkspaceUpdateOptions.ScheduledBuild) != 0)
+				if ((options & WorkspaceUpdateOptions.ScheduledBuild) != 0)
 				{
 					SyncLog.Clear();
 					SyncLog.AppendLine("Editor is open; scheduled sync has been aborted.");
@@ -919,220 +919,220 @@ namespace UnrealGameSync
 				}
 			}
 
-			string[] CombinedSyncFilter = UserSettings.GetCombinedSyncFilter(GetSyncCategories(), Settings.Global.Filter, WorkspaceSettings.Filter);
+			string[] combinedSyncFilter = UserSettings.GetCombinedSyncFilter(GetSyncCategories(), _settings.Global.Filter, _workspaceSettings.Filter);
 
-			ConfigSection PerforceSection = PerforceMonitor.LatestProjectConfigFile.FindSection("Perforce");
-			if (PerforceSection != null)
+			ConfigSection perforceSection = _perforceMonitor.LatestProjectConfigFile.FindSection("Perforce");
+			if (perforceSection != null)
 			{
-				IEnumerable<string> AdditionalPaths = PerforceSection.GetValues("AdditionalPathsToSync", new string[0]);
-				CombinedSyncFilter = AdditionalPaths.Union(CombinedSyncFilter).ToArray();
+				IEnumerable<string> additionalPaths = perforceSection.GetValues("AdditionalPathsToSync", new string[0]);
+				combinedSyncFilter = additionalPaths.Union(combinedSyncFilter).ToArray();
 			}
 
-			WorkspaceUpdateContext Context = new WorkspaceUpdateContext(ChangeNumber, Options, GetEditorBuildConfig(), CombinedSyncFilter, ProjectSettings.BuildSteps, null);
-			if (Options.HasFlag(WorkspaceUpdateOptions.SyncArchives))
+			WorkspaceUpdateContext context = new WorkspaceUpdateContext(changeNumber, options, GetEditorBuildConfig(), combinedSyncFilter, _projectSettings.BuildSteps, null);
+			if (options.HasFlag(WorkspaceUpdateOptions.SyncArchives))
 			{
-				IReadOnlyList<IArchiveInfo> Archives = GetArchives();
-				foreach (IArchiveInfo Archive in Archives)
+				IReadOnlyList<IArchiveInfo> archives = GetArchives();
+				foreach (IArchiveInfo archive in archives)
 				{
-					Context.ArchiveTypeToArchive[Archive.Type] = null;
+					context.ArchiveTypeToArchive[archive.Type] = null;
 				}
 
-				List<IArchiveInfo> SelectedArchives = GetSelectedArchives(Archives);
-				foreach (IArchiveInfo Archive in SelectedArchives)
+				List<IArchiveInfo> selectedArchives = GetSelectedArchives(archives);
+				foreach (IArchiveInfo archive in selectedArchives)
 				{
-					string? ArchivePath = GetArchiveKeyForChangeNumber(Archive, ChangeNumber);
-					if (ArchivePath == null)
+					string? archivePath = GetArchiveKeyForChangeNumber(archive, changeNumber);
+					if (archivePath == null)
 					{
-						MessageBox.Show(String.Format("There are no compiled {0} binaries for this change. To sync it, you must disable syncing of precompiled editor binaries.", Archive.Name));
+						MessageBox.Show(String.Format("There are no compiled {0} binaries for this change. To sync it, you must disable syncing of precompiled editor binaries.", archive.Name));
 						return;
 					}
 
-					if (Archive.Type == IArchiveInfo.EditorArchiveType)
+					if (archive.Type == IArchiveInfo.EditorArchiveType)
 					{
-						Context.Options &= ~(WorkspaceUpdateOptions.Build | WorkspaceUpdateOptions.GenerateProjectFiles | WorkspaceUpdateOptions.OpenSolutionAfterSync);
+						context.Options &= ~(WorkspaceUpdateOptions.Build | WorkspaceUpdateOptions.GenerateProjectFiles | WorkspaceUpdateOptions.OpenSolutionAfterSync);
 					}
 
-					string[]? ZippedBinariesSyncFilter;
-					if (TryGetProjectSetting(PerforceMonitor.LatestProjectConfigFile, "ZippedBinariesSyncFilter", out ZippedBinariesSyncFilter) && ZippedBinariesSyncFilter.Length > 0)
+					string[]? zippedBinariesSyncFilter;
+					if (TryGetProjectSetting(_perforceMonitor.LatestProjectConfigFile, "ZippedBinariesSyncFilter", out zippedBinariesSyncFilter) && zippedBinariesSyncFilter.Length > 0)
 					{
-						Context.SyncFilter = Enumerable.Concat(Context.SyncFilter, ZippedBinariesSyncFilter).ToArray();
+						context.SyncFilter = Enumerable.Concat(context.SyncFilter, zippedBinariesSyncFilter).ToArray();
 					}
 
-					Context.ArchiveTypeToArchive[Archive.Type] = new Tuple<IArchiveInfo, string>(Archive, ArchivePath);
+					context.ArchiveTypeToArchive[archive.Type] = new Tuple<IArchiveInfo, string>(archive, archivePath);
 				}
 			}
-			StartWorkspaceUpdate(Context, Callback);
+			StartWorkspaceUpdate(context, callback);
 		}
 
-		void StartWorkspaceUpdate(WorkspaceUpdateContext Context, WorkspaceUpdateCallback? Callback)
+		void StartWorkspaceUpdate(WorkspaceUpdateContext context, WorkspaceUpdateCallback? callback)
 		{
-			if (Settings.bAutoResolveConflicts)
+			if (_settings.AutoResolveConflicts)
 			{
-				Context.Options |= WorkspaceUpdateOptions.AutoResolveChanges;
+				context.Options |= WorkspaceUpdateOptions.AutoResolveChanges;
 			}
-			if (WorkspaceSettings.Filter.AllProjects ?? Settings.Global.Filter.AllProjects ?? false)
+			if (_workspaceSettings.Filter.AllProjects ?? _settings.Global.Filter.AllProjects ?? false)
 			{
-				Context.Options |= WorkspaceUpdateOptions.SyncAllProjects | WorkspaceUpdateOptions.IncludeAllProjectsInSolution;
+				context.Options |= WorkspaceUpdateOptions.SyncAllProjects | WorkspaceUpdateOptions.IncludeAllProjectsInSolution;
 			}
-			if (WorkspaceSettings.Filter.AllProjectsInSln ?? Settings.Global.Filter.AllProjectsInSln ?? false)
+			if (_workspaceSettings.Filter.AllProjectsInSln ?? _settings.Global.Filter.AllProjectsInSln ?? false)
 			{
-				Context.Options |= WorkspaceUpdateOptions.IncludeAllProjectsInSolution;
+				context.Options |= WorkspaceUpdateOptions.IncludeAllProjectsInSolution;
 			}
 
-			UpdateCallback = Callback;
+			_updateCallback = callback;
 
-			Context.StartTime = DateTime.UtcNow;
-			Context.PerforceSyncOptions = (PerforceSyncOptions)Settings.SyncOptions.Clone();
+			context.StartTime = DateTime.UtcNow;
+			context.PerforceSyncOptions = (PerforceSyncOptions)_settings.SyncOptions.Clone();
 
-			Logger.LogInformation("Updating workspace at {Time}...", Context.StartTime.ToLocalTime().ToString());
-			Logger.LogInformation("  ChangeNumber={Change}", Context.ChangeNumber);
-			Logger.LogInformation("  Options={Options}", Context.Options.ToString());
-			Logger.LogInformation("  Clobbering {NumFiles} files", Context.ClobberFiles.Count);
+			_logger.LogInformation("Updating workspace at {Time}...", context.StartTime.ToLocalTime().ToString());
+			_logger.LogInformation("  ChangeNumber={Change}", context.ChangeNumber);
+			_logger.LogInformation("  Options={Options}", context.Options.ToString());
+			_logger.LogInformation("  Clobbering {NumFiles} files", context.ClobberFiles.Count);
 
-			if (Context.Options.HasFlag(WorkspaceUpdateOptions.Sync))
+			if (context.Options.HasFlag(WorkspaceUpdateOptions.Sync))
 			{
-				EventMonitor.PostEvent(Context.ChangeNumber, EventType.Syncing);
+				_eventMonitor.PostEvent(context.ChangeNumber, EventType.Syncing);
 			}
 
-			if (Context.Options.HasFlag(WorkspaceUpdateOptions.Sync) || Context.Options.HasFlag(WorkspaceUpdateOptions.Build))
+			if (context.Options.HasFlag(WorkspaceUpdateOptions.Sync) || context.Options.HasFlag(WorkspaceUpdateOptions.Build))
 			{
-				if (!Context.Options.HasFlag(WorkspaceUpdateOptions.ContentOnly) && (Context.CustomBuildSteps == null || Context.CustomBuildSteps.Count == 0))
+				if (!context.Options.HasFlag(WorkspaceUpdateOptions.ContentOnly) && (context.CustomBuildSteps == null || context.CustomBuildSteps.Count == 0))
 				{
-					FileReference TargetFile = ConfigUtils.GetEditorTargetFile(Workspace.Project, Workspace.ProjectConfigFile);
-					foreach (BuildConfig Config in Enum.GetValues(typeof(BuildConfig)).OfType<BuildConfig>())
+					FileReference targetFile = ConfigUtils.GetEditorTargetFile(_workspace.Project, _workspace.ProjectConfigFile);
+					foreach (BuildConfig config in Enum.GetValues(typeof(BuildConfig)).OfType<BuildConfig>())
 					{
-						FileReference ReceiptFile = ConfigUtils.GetReceiptFile(Workspace.Project, TargetFile, Config.ToString());
-						if (FileReference.Exists(ReceiptFile))
+						FileReference receiptFile = ConfigUtils.GetReceiptFile(_workspace.Project, targetFile, config.ToString());
+						if (FileReference.Exists(receiptFile))
 						{
-							try { FileReference.Delete(ReceiptFile); } catch (Exception) { }
+							try { FileReference.Delete(receiptFile); } catch (Exception) { }
 						}
 					}
 				}
 			}
 
 			SyncLog.Clear();
-			Workspace.Update(Context);
+			_workspace.Update(context);
 			UpdateSyncActionCheckboxes();
 			Refresh();
-			UpdateTimer.Start();
+			_updateTimer.Start();
 		}
 
 		void CancelWorkspaceUpdate()
 		{
-			if (Workspace.IsBusy() && MessageBox.Show("Are you sure you want to cancel the current operation?", "Cancel operation", MessageBoxButtons.YesNo) == DialogResult.Yes)
+			if (_workspace.IsBusy() && MessageBox.Show("Are you sure you want to cancel the current operation?", "Cancel operation", MessageBoxButtons.YesNo) == DialogResult.Yes)
 			{
-				WorkspaceState.LastSyncChangeNumber = Workspace.PendingChangeNumber;
-				WorkspaceState.LastSyncResult = WorkspaceUpdateResult.Canceled;
-				WorkspaceState.LastSyncResultMessage = null;
-				WorkspaceState.LastSyncTime = null;
-				WorkspaceState.LastSyncDurationSeconds = 0;
-				WorkspaceState.Save(Logger);
+				_workspaceState.LastSyncChangeNumber = _workspace.PendingChangeNumber;
+				_workspaceState.LastSyncResult = WorkspaceUpdateResult.Canceled;
+				_workspaceState.LastSyncResultMessage = null;
+				_workspaceState.LastSyncTime = null;
+				_workspaceState.LastSyncDurationSeconds = 0;
+				_workspaceState.Save(_logger);
 
-				Workspace.CancelUpdate();
+				_workspace.CancelUpdate();
 
-				if (UpdateCallback != null)
+				if (_updateCallback != null)
 				{
-					UpdateCallback(WorkspaceUpdateResult.Canceled);
-					UpdateCallback = null;
+					_updateCallback(WorkspaceUpdateResult.Canceled);
+					_updateCallback = null;
 				}
 
-				UpdateTimer.Stop();
+				_updateTimer.Stop();
 
 				UpdateSyncActionCheckboxes();
 				Refresh();
 				UpdateStatusPanel();
 				DesiredTaskbarState = Tuple.Create(TaskbarState.NoProgress, 0.0f);
-				Owner.UpdateProgress();
+				_owner.UpdateProgress();
 			}
 		}
 
-		void UpdateCompleteCallback(WorkspaceUpdateContext Context, WorkspaceUpdateResult Result, string ResultMessage)
+		void UpdateCompleteCallback(WorkspaceUpdateContext context, WorkspaceUpdateResult result, string resultMessage)
 		{
-			MainThreadSynchronizationContext.Post((o) => { if (!bIsDisposing) { UpdateComplete(Context, Result, ResultMessage); } }, null);
+			_mainThreadSynchronizationContext.Post((o) => { if (!_isDisposing) { UpdateComplete(context, result, resultMessage); } }, null);
 		}
 
-		void UpdateComplete(WorkspaceUpdateContext Context, WorkspaceUpdateResult Result, string ResultMessage)
+		void UpdateComplete(WorkspaceUpdateContext context, WorkspaceUpdateResult result, string resultMessage)
 		{
-			if(bIsDisposing)
+			if(_isDisposing)
 			{
 				return;
 			}
 
-			UpdateTimer.Stop();
+			_updateTimer.Stop();
 
-			if (Result == WorkspaceUpdateResult.FilesToResolve)
+			if (result == WorkspaceUpdateResult.FilesToResolve)
 			{
 				MessageBox.Show("You have files to resolve after syncing your workspace. Please check P4.");
 			}
-			else if (Result == WorkspaceUpdateResult.FilesToDelete)
+			else if (result == WorkspaceUpdateResult.FilesToDelete)
 			{
 				DesiredTaskbarState = Tuple.Create(TaskbarState.Paused, 0.0f);
-				Owner.UpdateProgress();
+				_owner.UpdateProgress();
 
-				DeleteWindow Window = new DeleteWindow(Context.DeleteFiles);
-				if (Window.ShowDialog(this) == DialogResult.OK)
+				DeleteWindow window = new DeleteWindow(context.DeleteFiles);
+				if (window.ShowDialog(this) == DialogResult.OK)
 				{
-					StartWorkspaceUpdate(Context, UpdateCallback);
+					StartWorkspaceUpdate(context, _updateCallback);
 					return;
 				}
 			}
-			else if (Result == WorkspaceUpdateResult.FilesToClobber)
+			else if (result == WorkspaceUpdateResult.FilesToClobber)
 			{
 				DesiredTaskbarState = Tuple.Create(TaskbarState.Paused, 0.0f);
-				Owner.UpdateProgress();
+				_owner.UpdateProgress();
 
-				HashSet<string> UncontrolledFiles = new HashSet<string>();
+				HashSet<string> uncontrolledFiles = new HashSet<string>();
 				
 				if (SelectedProject.LocalPath != null && SelectedProject.LocalPath.EndsWith(".uprojectdirs", StringComparison.InvariantCultureIgnoreCase))
 				{
-					List<string> ProjectRoots = GetProjectRoots(SelectedProject.LocalPath);
+					List<string> projectRoots = GetProjectRoots(SelectedProject.LocalPath);
 
-					foreach (string ProjectRoot in ProjectRoots)
+					foreach (string projectRoot in projectRoots)
 					{
-						ParseUncontrolledChangelistsPersistencyFile(ProjectRoot, UncontrolledFiles);
+						ParseUncontrolledChangelistsPersistencyFile(projectRoot, uncontrolledFiles);
 					}
 				}
 				else
 				{
-					ParseUncontrolledChangelistsPersistencyFile(Path.GetDirectoryName(SelectedProject.LocalPath)!, UncontrolledFiles);
+					ParseUncontrolledChangelistsPersistencyFile(Path.GetDirectoryName(SelectedProject.LocalPath)!, uncontrolledFiles);
 				}
 
-				ClobberWindow Window = new ClobberWindow(Context.ClobberFiles, UncontrolledFiles);
+				ClobberWindow window = new ClobberWindow(context.ClobberFiles, uncontrolledFiles);
 				
-				if (Window.ShowDialog(this) == DialogResult.OK)
+				if (window.ShowDialog(this) == DialogResult.OK)
 				{
-					StartWorkspaceUpdate(Context, UpdateCallback);
+					StartWorkspaceUpdate(context, _updateCallback);
 					return;
 				}
 			}
-			else if (Result == WorkspaceUpdateResult.FailedToCompileWithCleanWorkspace)
+			else if (result == WorkspaceUpdateResult.FailedToCompileWithCleanWorkspace)
 			{
-				EventMonitor.PostEvent(Context.ChangeNumber, EventType.DoesNotCompile);
+				_eventMonitor.PostEvent(context.ChangeNumber, EventType.DoesNotCompile);
 			}
-			else if (Result == WorkspaceUpdateResult.Success)
+			else if (result == WorkspaceUpdateResult.Success)
 			{
-				if (Context.Options.HasFlag(WorkspaceUpdateOptions.Build))
+				if (context.Options.HasFlag(WorkspaceUpdateOptions.Build))
 				{
-					EventMonitor.PostEvent(Context.ChangeNumber, EventType.Compiles);
+					_eventMonitor.PostEvent(context.ChangeNumber, EventType.Compiles);
 				}
-				if (Context.Options.HasFlag(WorkspaceUpdateOptions.RunAfterSync))
+				if (context.Options.HasFlag(WorkspaceUpdateOptions.RunAfterSync))
 				{
 					LaunchEditor();
 				}
-				if (Context.Options.HasFlag(WorkspaceUpdateOptions.OpenSolutionAfterSync))
+				if (context.Options.HasFlag(WorkspaceUpdateOptions.OpenSolutionAfterSync))
 				{
 					OpenSolution();
 				}
 			}
 
-			if (UpdateCallback != null)
+			if (_updateCallback != null)
 			{
-				UpdateCallback(Result);
-				UpdateCallback = null;
+				_updateCallback(result);
+				_updateCallback = null;
 			}
 
-			DesiredTaskbarState = Tuple.Create((Result == WorkspaceUpdateResult.Success) ? TaskbarState.NoProgress : TaskbarState.Error, 0.0f);
-			Owner.UpdateProgress();
+			DesiredTaskbarState = Tuple.Create((result == WorkspaceUpdateResult.Success) ? TaskbarState.NoProgress : TaskbarState.Error, 0.0f);
+			_owner.UpdateProgress();
 
 			BuildList.Invalidate();
 			Refresh();
@@ -1140,7 +1140,7 @@ namespace UnrealGameSync
 			UpdateSyncActionCheckboxes();
 
 			// Do this last because it may result in the control being disposed
-			if (Result == WorkspaceUpdateResult.FailedToSyncLoginExpired)
+			if (result == WorkspaceUpdateResult.FailedToSyncLoginExpired)
 			{
 				LoginExpired();
 			}
@@ -1148,10 +1148,10 @@ namespace UnrealGameSync
 
 		void UpdateBuildListCallback()
 		{
-			if (!bUpdateBuildListPosted)
+			if (!_updateBuildListPosted)
 			{
-				bUpdateBuildListPosted = true;
-				MainThreadSynchronizationContext.Post((o) => { bUpdateBuildListPosted = false; if (!bIsDisposing) { UpdateBuildList(); } }, null);
+				_updateBuildListPosted = true;
+				_mainThreadSynchronizationContext.Post((o) => { _updateBuildListPosted = false; if (!_isDisposing) { UpdateBuildList(); } }, null);
 			}
 		}
 
@@ -1159,59 +1159,59 @@ namespace UnrealGameSync
 		{
 			if (SelectedFileName != null)
 			{
-				ArchiveToChangeNumberToArchiveKey.Clear();
-				ChangeNumberToLayoutInfo.Clear();
+				_archiveToChangeNumberToArchiveKey.Clear();
+				_changeNumberToLayoutInfo.Clear();
 
-				List<ChangesRecord> Changes = PerforceMonitor.GetChanges();
-				EventMonitor.FilterChanges(Changes.Select(x => x.Number));
+				List<ChangesRecord> changes = _perforceMonitor.GetChanges();
+				_eventMonitor.FilterChanges(changes.Select(x => x.Number));
 
-				PromotedChangeNumbers = PerforceMonitor.GetPromotedChangeNumbers();
+				_promotedChangeNumbers = _perforceMonitor.GetPromotedChangeNumbers();
 
-				string[] ExcludeChanges = new string[0];
-				if (Workspace != null)
+				string[] excludeChanges = new string[0];
+				if (_workspace != null)
 				{
-					ConfigFile ProjectConfigFile = PerforceMonitor.LatestProjectConfigFile;
-					if (ProjectConfigFile != null)
+					ConfigFile projectConfigFile = _perforceMonitor.LatestProjectConfigFile;
+					if (projectConfigFile != null)
 					{
-						ExcludeChanges = ProjectConfigFile.GetValues("Options.ExcludeChanges", ExcludeChanges);
+						excludeChanges = projectConfigFile.GetValues("Options.ExcludeChanges", excludeChanges);
 					}
 				}
 
-				bool bFirstChange = true;
-				bool bHideUnreviewed = !Settings.bShowUnreviewedChanges;
+				bool firstChange = true;
+				bool hideUnreviewed = !_settings.ShowUnreviewedChanges;
 
-				NumChanges = Changes.Count;
-				ListIndexToChangeIndex = new List<int>();
-				SortedChangeNumbers = new List<int>();
+				_numChanges = changes.Count;
+				_listIndexToChangeIndex = new List<int>();
+				_sortedChangeNumbers = new List<int>();
 
-				List<ChangesRecord> FilteredChanges = new List<ChangesRecord>();
-				for (int ChangeIdx = 0; ChangeIdx < Changes.Count; ChangeIdx++)
+				List<ChangesRecord> filteredChanges = new List<ChangesRecord>();
+				for (int changeIdx = 0; changeIdx < changes.Count; changeIdx++)
 				{
-					ChangesRecord Change = Changes[ChangeIdx];
-					if (ShouldShowChange(Change, ExcludeChanges) || PromotedChangeNumbers.Contains(Change.Number))
+					ChangesRecord change = changes[changeIdx];
+					if (ShouldShowChange(change, excludeChanges) || _promotedChangeNumbers.Contains(change.Number))
 					{
-						SortedChangeNumbers.Add(Change.Number);
+						_sortedChangeNumbers.Add(change.Number);
 
-						if (!bHideUnreviewed || (!EventMonitor.IsUnderInvestigation(Change.Number) && (ShouldIncludeInReviewedList(Change.Number) || bFirstChange)))
+						if (!hideUnreviewed || (!_eventMonitor.IsUnderInvestigation(change.Number) && (ShouldIncludeInReviewedList(change.Number) || firstChange)))
 						{
-							bFirstChange = false;
+							firstChange = false;
 
 							// Add the new change
-							FilteredChanges.Add(Change);
+							filteredChanges.Add(change);
 
 							// Store off the list index for this change
-							ListIndexToChangeIndex.Add(ChangeIdx);
+							_listIndexToChangeIndex.Add(changeIdx);
 						}
 					}
 				}
 
-				UpdateBuildListInternal(FilteredChanges, Changes.Count > 0);
+				UpdateBuildListInternal(filteredChanges, changes.Count > 0);
 
-				SortedChangeNumbers.Sort();
+				_sortedChangeNumbers.Sort();
 
-				if (PendingSelectedChangeNumber != -1)
+				if (_pendingSelectedChangeNumber != -1)
 				{
-					SelectChange(PendingSelectedChangeNumber);
+					SelectChange(_pendingSelectedChangeNumber);
 				}
 			}
 
@@ -1226,253 +1226,253 @@ namespace UnrealGameSync
 			UpdateSyncActionCheckboxes();
 		}
 
-		void UpdateBuildListInternal(List<ChangesRecord> Changes, bool bShowExpandItem)
+		void UpdateBuildListInternal(List<ChangesRecord> changes, bool showExpandItem)
 		{
 			BuildList.BeginUpdate();
 
 			// Remove any changes that no longer exist, and update the rest
-			Dictionary<int, ChangesRecord> ChangeNumberToSummary = Changes.ToDictionary(x => x.Number, x => x);
-			for (int Idx = BuildList.Items.Count - 1; Idx >= 0; Idx--)
+			Dictionary<int, ChangesRecord> changeNumberToSummary = changes.ToDictionary(x => x.Number, x => x);
+			for (int idx = BuildList.Items.Count - 1; idx >= 0; idx--)
 			{
-				ChangesRecord? Change = BuildList.Items[Idx].Tag as ChangesRecord;
-				if (Change != null)
+				ChangesRecord? change = BuildList.Items[idx].Tag as ChangesRecord;
+				if (change != null)
 				{
-					ChangesRecord? Summary;
-					if (ChangeNumberToSummary.TryGetValue(Change.Number, out Summary))
+					ChangesRecord? summary;
+					if (changeNumberToSummary.TryGetValue(change.Number, out summary))
 					{
 						// Update
-						BuildList.Items[Idx].SubItems[DescriptionColumn.Index].Text = Change.Description!.Replace('\n', ' ');
-						ChangeNumberToSummary.Remove(Change.Number);
+						BuildList.Items[idx].SubItems[DescriptionColumn.Index].Text = change.Description!.Replace('\n', ' ');
+						changeNumberToSummary.Remove(change.Number);
 					}
 					else
 					{
 						// Delete
-						BuildList.Items.RemoveAt(Idx);
+						BuildList.Items.RemoveAt(idx);
 					}
 				}
 			}
 
 			// Add everything left over
-			foreach (ChangesRecord Change in ChangeNumberToSummary.Values)
+			foreach (ChangesRecord change in changeNumberToSummary.Values)
 			{
-				BuildList_AddItem(Change);
+				BuildList_AddItem(change);
 			}
 
 			// Figure out which group the expand item should be in
-			ListViewGroup? NewExpandItemGroup = null;
-			for (int Idx = BuildList.Items.Count - 1; Idx >= 0; Idx--)
+			ListViewGroup? newExpandItemGroup = null;
+			for (int idx = BuildList.Items.Count - 1; idx >= 0; idx--)
 			{
-				ListViewItem Item = BuildList.Items[Idx];
-				if (Item != ExpandItem)
+				ListViewItem item = BuildList.Items[idx];
+				if (item != _expandItem)
 				{
-					NewExpandItemGroup = Item.Group;
+					newExpandItemGroup = item.Group;
 					break;
 				}
 			}
 
 			// Remove the expand row if it's in the wrong place
-			if (ExpandItem != null)
+			if (_expandItem != null)
 			{
-				if (!bShowExpandItem || ExpandItem.Group != NewExpandItemGroup)
+				if (!showExpandItem || _expandItem.Group != newExpandItemGroup)
 				{
-					BuildList.Items.Remove(ExpandItem);
-					ExpandItem = null;
+					BuildList.Items.Remove(_expandItem);
+					_expandItem = null;
 				}
 			}
 
 			// Remove any empty groups
-			for (int Idx = BuildList.Groups.Count - 1; Idx >= 0; Idx--)
+			for (int idx = BuildList.Groups.Count - 1; idx >= 0; idx--)
 			{
-				if (BuildList.Groups[Idx].Items.Count == 0)
+				if (BuildList.Groups[idx].Items.Count == 0)
 				{
-					BuildList.Groups.RemoveAt(Idx);
+					BuildList.Groups.RemoveAt(idx);
 				}
 			}
 
 			// Add the expand row back in
-			if (bShowExpandItem && ExpandItem == null)
+			if (showExpandItem && _expandItem == null)
 			{
-				ExpandItem = new ListViewItem(NewExpandItemGroup);
-				ExpandItem.Tag = null;
-				ExpandItem.Selected = false;
-				ExpandItem.Text = "";
-				for (int ColumnIdx = 1; ColumnIdx < BuildList.Columns.Count; ColumnIdx++)
+				_expandItem = new ListViewItem(newExpandItemGroup);
+				_expandItem.Tag = null;
+				_expandItem.Selected = false;
+				_expandItem.Text = "";
+				for (int columnIdx = 1; columnIdx < BuildList.Columns.Count; columnIdx++)
 				{
-					ExpandItem.SubItems.Add(new ListViewItem.ListViewSubItem(ExpandItem, ""));
+					_expandItem.SubItems.Add(new ListViewItem.ListViewSubItem(_expandItem, ""));
 				}
-				BuildList.Items.Add(ExpandItem);
+				BuildList.Items.Add(_expandItem);
 			}
 
 			BuildList.EndUpdate();
 		}
 
-		ListViewGroup BuildList_FindOrAddGroup(ChangesRecord Change, string GroupName)
+		ListViewGroup BuildList_FindOrAddGroup(ChangesRecord change, string groupName)
 		{
 			// Find or add the new group
-			int GroupIndex = 0;
-			for (; GroupIndex < BuildList.Groups.Count; GroupIndex++)
+			int groupIndex = 0;
+			for (; groupIndex < BuildList.Groups.Count; groupIndex++)
 			{
-				ListViewGroup NextGroup = BuildList.Groups[GroupIndex];
-				if (NextGroup.Header == GroupName)
+				ListViewGroup nextGroup = BuildList.Groups[groupIndex];
+				if (nextGroup.Header == groupName)
 				{
-					return NextGroup;
+					return nextGroup;
 				}
 
-				ChangesRecord? LastChange = null;
-				for (int Idx = NextGroup.Items.Count - 1; Idx >= 0 && LastChange == null; Idx--)
+				ChangesRecord? lastChange = null;
+				for (int idx = nextGroup.Items.Count - 1; idx >= 0 && lastChange == null; idx--)
 				{
-					LastChange = NextGroup.Items[Idx].Tag as ChangesRecord;
+					lastChange = nextGroup.Items[idx].Tag as ChangesRecord;
 				}
-				if (LastChange == null || Change.Number > LastChange.Number)
+				if (lastChange == null || change.Number > lastChange.Number)
 				{
 					break;
 				}
 			}
 
 			// Create the new group
-			ListViewGroup Group = new ListViewGroup(GroupName);
-			BuildList.Groups.Insert(GroupIndex, Group);
-			return Group;
+			ListViewGroup group = new ListViewGroup(groupName);
+			BuildList.Groups.Insert(groupIndex, group);
+			return group;
 		}
 
-		void BuildList_AddItem(ChangesRecord Change)
+		void BuildList_AddItem(ChangesRecord change)
 		{
 			// Get the display time for this item
-			DateTime DisplayTime;
-			if (Settings.bShowLocalTimes)
+			DateTime displayTime;
+			if (_settings.ShowLocalTimes)
 			{
-				DisplayTime = Change.Time.ToLocalTime();
+				displayTime = change.Time.ToLocalTime();
 			}
 			else
 			{
-				DisplayTime = new DateTime(Change.Time.Ticks + PerforceMonitor.ServerTimeZone.Ticks, DateTimeKind.Local);
+				displayTime = new DateTime(change.Time.Ticks + _perforceMonitor.ServerTimeZone.Ticks, DateTimeKind.Local);
 			}
 
 			// Find or add the new group
-			ListViewGroup Group = BuildList_FindOrAddGroup(Change, DisplayTime.ToString("D"));
+			ListViewGroup group = BuildList_FindOrAddGroup(change, displayTime.ToString("D"));
 
 			// Create the new item
-			ListViewItem Item = new ListViewItem();
-			Item.Tag = Change;
+			ListViewItem item = new ListViewItem();
+			item.Tag = change;
 
 			// Get the new text for each column
-			string[] Columns = new string[BuildList.Columns.Count];
-			Columns[ChangeColumn.Index] = String.Format("{0}", Change.Number);
+			string[] columns = new string[BuildList.Columns.Count];
+			columns[ChangeColumn.Index] = String.Format("{0}", change.Number);
 
-			Columns[TimeColumn.Index] = DisplayTime.ToShortTimeString();
-			string UserName = FormatUserName(Change.User!);
+			columns[TimeColumn.Index] = displayTime.ToShortTimeString();
+			string userName = FormatUserName(change.User!);
 
 			// If annotate robomerge is on, add a prefix to the user name
-			if (Settings.bAnnotateRobmergeChanges && IsRobomergeChange(Change))
+			if (_settings.AnnotateRobmergeChanges && IsRobomergeChange(change))
 			{
-				UserName += " (robo)";
+				userName += " (robo)";
 			}
 
-			Columns[AuthorColumn.Index] = UserName;
-			Columns[DescriptionColumn.Index] = Change.Description!.Replace('\n', ' ');
+			columns[AuthorColumn.Index] = userName;
+			columns[DescriptionColumn.Index] = change.Description!.Replace('\n', ' ');
 
-			for (int ColumnIdx = 1; ColumnIdx < BuildList.Columns.Count; ColumnIdx++)
+			for (int columnIdx = 1; columnIdx < BuildList.Columns.Count; columnIdx++)
 			{
-				Item.SubItems.Add(new ListViewItem.ListViewSubItem(Item, Columns[ColumnIdx] ?? ""));
+				item.SubItems.Add(new ListViewItem.ListViewSubItem(item, columns[columnIdx] ?? ""));
 			}
 
 			// Insert it at the right position within the group
-			int GroupInsertIdx = 0;
-			for (; GroupInsertIdx < Group.Items.Count; GroupInsertIdx++)
+			int groupInsertIdx = 0;
+			for (; groupInsertIdx < group.Items.Count; groupInsertIdx++)
 			{
-				ChangesRecord? OtherChange = Group.Items[GroupInsertIdx].Tag as ChangesRecord;
-				if (OtherChange == null || Change.Number >= OtherChange.Number)
+				ChangesRecord? otherChange = group.Items[groupInsertIdx].Tag as ChangesRecord;
+				if (otherChange == null || change.Number >= otherChange.Number)
 				{
 					break;
 				}
 			}
-			Group.Items.Insert(GroupInsertIdx, Item);
+			group.Items.Insert(groupInsertIdx, item);
 
 			// Insert it into the build list
-			BuildList.Items.Add(Item);
+			BuildList.Items.Add(item);
 		}
 
 		/// <summary>
 		/// Returns true if this change was submitted on behalf of someone by robomerge
 		/// </summary>
-		/// <param name="Change"></param>
+		/// <param name="change"></param>
 		/// <returns></returns>
-		bool IsRobomergeChange(ChangesRecord Change)
+		bool IsRobomergeChange(ChangesRecord change)
 		{
 			// If hiding robomerge, filter out based on workspace name. Note - don't look at the description because we
 			// *do* want to see conflicts that were manually merged by someone
-			return Change.Client!.IndexOf("ROBOMERGE", StringComparison.OrdinalIgnoreCase) >= 0;
+			return change.Client!.IndexOf("ROBOMERGE", StringComparison.OrdinalIgnoreCase) >= 0;
 		}
 
-		bool ShouldShowChange(ChangesRecord Change, string[] ExcludeChanges)
+		bool ShouldShowChange(ChangesRecord change, string[] excludeChanges)
 		{
-			if (ProjectSettings.FilterType != FilterType.None)
+			if (_projectSettings.FilterType != FilterType.None)
 			{
-				PerforceChangeDetails? Details;
-				if (!PerforceMonitor.TryGetChangeDetails(Change.Number, out Details))
+				PerforceChangeDetails? details;
+				if (!_perforceMonitor.TryGetChangeDetails(change.Number, out details))
 				{
 					return false;
 				}
-				if (ProjectSettings.FilterType == FilterType.Code && !Details.bContainsCode)
+				if (_projectSettings.FilterType == FilterType.Code && !details.ContainsCode)
 				{
 					return false;
 				}
-				if (ProjectSettings.FilterType == FilterType.Content && !Details.bContainsContent)
+				if (_projectSettings.FilterType == FilterType.Content && !details.ContainsContent)
 				{
 					return false;
 				}
 			}
 
 			// if filtering by user, only show changes where the author contains the filter text
-			if (!string.IsNullOrEmpty(this.AuthorFilterText)
-				&& this.AuthorFilterText != AuthorFilterPlaceholderText
-				&& Change.User!.IndexOf(this.AuthorFilterText, StringComparison.OrdinalIgnoreCase) < 0)
+			if (!string.IsNullOrEmpty(this._authorFilterText)
+				&& this._authorFilterText != _authorFilterPlaceholderText
+				&& change.User!.IndexOf(this._authorFilterText, StringComparison.OrdinalIgnoreCase) < 0)
 			{
 				return false;
 			}
 			
 			// If this is a robomerge change, check if any filters will cause it to be hidden
-			if (IsRobomergeChange(Change))
+			if (IsRobomergeChange(change))
 			{
-				if (Settings.ShowRobomerge == UserSettings.RobomergeShowChangesOption.None)
+				if (_settings.ShowRobomerge == UserSettings.RobomergeShowChangesOption.None)
 				{
 					return false;
 				}
-				else if (Settings.ShowRobomerge == UserSettings.RobomergeShowChangesOption.Badged)
+				else if (_settings.ShowRobomerge == UserSettings.RobomergeShowChangesOption.Badged)
 				{
 					// if this change has any badges, we'll show it unless it's later filtered out
-					EventSummary? Summary = EventMonitor.GetSummaryForChange(Change.Number);
-					if (Summary == null || !Summary.Badges.Any())
+					EventSummary? summary = _eventMonitor.GetSummaryForChange(change.Number);
+					if (summary == null || !summary.Badges.Any())
 					{
 						return false;
 					}
 				}
 			}
 
-			if (ProjectSettings.FilterBadges.Count > 0)
+			if (_projectSettings.FilterBadges.Count > 0)
 			{
-				EventSummary? Summary = EventMonitor.GetSummaryForChange(Change.Number);
-				if (Summary == null || !Summary.Badges.Any(x => ProjectSettings.FilterBadges.Contains(x.BadgeName)))
+				EventSummary? summary = _eventMonitor.GetSummaryForChange(change.Number);
+				if (summary == null || !summary.Badges.Any(x => _projectSettings.FilterBadges.Contains(x.BadgeName)))
 				{
 					return false;
 				}
 			}
-			if (!Settings.bShowAutomatedChanges)
+			if (!_settings.ShowAutomatedChanges)
 			{
-				foreach (string ExcludeChange in ExcludeChanges)
+				foreach (string excludeChange in excludeChanges)
 				{
-					if (Regex.IsMatch(Change.Description, ExcludeChange, RegexOptions.IgnoreCase))
+					if (Regex.IsMatch(change.Description, excludeChange, RegexOptions.IgnoreCase))
 					{
 						return false;
 					}
 				}
 
-				if (String.Compare(Change.User, "buildmachine", true) == 0 && Change.Description!.IndexOf("lightmaps", StringComparison.InvariantCultureIgnoreCase) == -1)
+				if (String.Compare(change.User, "buildmachine", true) == 0 && change.Description!.IndexOf("lightmaps", StringComparison.InvariantCultureIgnoreCase) == -1)
 				{
 					return false;
 				}
 			}
-			if (IsBisectModeEnabled() && !WorkspaceState.BisectChanges.Any(x => x.Change == Change.Number))
+			if (IsBisectModeEnabled() && !_workspaceState.BisectChanges.Any(x => x.Change == change.Number))
 			{
 				return false;
 			}
@@ -1481,73 +1481,73 @@ namespace UnrealGameSync
 
 		void UpdateBuildMetadataCallback()
 		{
-			if (!bUpdateBuildMetadataPosted)
+			if (!_updateBuildMetadataPosted)
 			{
-				bUpdateBuildMetadataPosted = true;
-				MainThreadSynchronizationContext.Post((o) => { bUpdateBuildMetadataPosted = false; if (!bIsDisposing) { UpdateBuildMetadata(); } }, null);
+				_updateBuildMetadataPosted = true;
+				_mainThreadSynchronizationContext.Post((o) => { _updateBuildMetadataPosted = false; if (!_isDisposing) { UpdateBuildMetadata(); } }, null);
 			}
 		}
 
 		string? GetIssuesApiUrl()
 		{
-			string? IssuesApiUrl;
-			if (!TryGetProjectSetting(PerforceMonitor.LatestProjectConfigFile, "IssuesApiUrl", out IssuesApiUrl))
+			string? issuesApiUrl;
+			if (!TryGetProjectSetting(_perforceMonitor.LatestProjectConfigFile, "IssuesApiUrl", out issuesApiUrl))
 			{
-				IssuesApiUrl = ApiUrl;
+				issuesApiUrl = _apiUrl;
 			}
-			return IssuesApiUrl;
+			return issuesApiUrl;
 		}
 
 		void UpdateBuildMetadata()
 		{
 			// Refresh the issue monitor if it's changed
-			string? IssuesApiUrl = GetIssuesApiUrl();
-			if (IssuesApiUrl != IssueMonitor.ApiUrl)
+			string? issuesApiUrl = GetIssuesApiUrl();
+			if (issuesApiUrl != _issueMonitor.ApiUrl)
 			{
-				Logger.LogInformation("Changing issues API url from {OldApiUrl} to {ApiUrl}", IssueMonitor.ApiUrl, IssuesApiUrl);
-				IssueMonitor NewIssueMonitor = Owner.CreateIssueMonitor(IssuesApiUrl, PerforceSettings.UserName);
-				Owner.ReleaseIssueMonitor(IssueMonitor);
-				IssueMonitor = NewIssueMonitor;
+				_logger.LogInformation("Changing issues API url from {OldApiUrl} to {ApiUrl}", _issueMonitor.ApiUrl, issuesApiUrl);
+				IssueMonitor newIssueMonitor = _owner.CreateIssueMonitor(issuesApiUrl, _perforceSettings.UserName);
+				_owner.ReleaseIssueMonitor(_issueMonitor);
+				_issueMonitor = newIssueMonitor;
 			}
 
 			// Update the column settings first, since this may affect the badges we hide
 			UpdateColumnSettings(false);
 
 			// Clear the badge size cache
-			BadgeLabelToSize.Clear();
+			_badgeLabelToSize.Clear();
 
 			// Reset the badge groups
-			Dictionary<string, string> BadgeNameToGroup = new Dictionary<string, string>();
+			Dictionary<string, string> badgeNameToGroup = new Dictionary<string, string>();
 
 			// Read the group mappings from the config file
-			string? GroupDefinitions;
-			if (TryGetProjectSetting(PerforceMonitor.LatestProjectConfigFile, "BadgeGroups", out GroupDefinitions))
+			string? groupDefinitions;
+			if (TryGetProjectSetting(_perforceMonitor.LatestProjectConfigFile, "BadgeGroups", out groupDefinitions))
 			{
-				string[] GroupDefinitionsArray = GroupDefinitions.Split('\n');
-				for (int Idx = 0; Idx < GroupDefinitionsArray.Length; Idx++)
+				string[] groupDefinitionsArray = groupDefinitions.Split('\n');
+				for (int idx = 0; idx < groupDefinitionsArray.Length; idx++)
 				{
-					string GroupName = String.Format("{0:0000}", Idx);
-					foreach (string BadgeName in GroupDefinitionsArray[Idx].Split(',').Select(x => x.Trim()))
+					string groupName = String.Format("{0:0000}", idx);
+					foreach (string badgeName in groupDefinitionsArray[idx].Split(',').Select(x => x.Trim()))
 					{
-						BadgeNameToGroup[BadgeName] = GroupName;
+						badgeNameToGroup[badgeName] = groupName;
 					}
 				}
 			}
 
 			// Add a dummy group for any other badges we have
-			foreach (ListViewItem? Item in BuildList.Items)
+			foreach (ListViewItem? item in BuildList.Items)
 			{
-				if (Item != null && Item.Tag != null)
+				if (item != null && item.Tag != null)
 				{
-					EventSummary? Summary = EventMonitor.GetSummaryForChange(((ChangesRecord)Item.Tag).Number);
-					if (Summary != null)
+					EventSummary? summary = _eventMonitor.GetSummaryForChange(((ChangesRecord)item.Tag).Number);
+					if (summary != null)
 					{
-						foreach (BadgeData Badge in Summary.Badges)
+						foreach (BadgeData badge in summary.Badges)
 						{
-							string BadgeSlot = Badge.BadgeName;
-							if (!BadgeNameToGroup.ContainsKey(BadgeSlot))
+							string badgeSlot = badge.BadgeName;
+							if (!badgeNameToGroup.ContainsKey(badgeSlot))
 							{
-								BadgeNameToGroup.Add(BadgeSlot, "XXXX");
+								badgeNameToGroup.Add(badgeSlot, "XXXX");
 							}
 						}
 					}
@@ -1555,30 +1555,30 @@ namespace UnrealGameSync
 			}
 
 			// Remove anything that's a service badge
-			foreach (ServiceBadgeInfo ServiceBadge in ServiceBadges)
+			foreach (ServiceBadgeInfo serviceBadge in _serviceBadges)
 			{
-				BadgeNameToGroup.Remove(ServiceBadge.Name);
+				badgeNameToGroup.Remove(serviceBadge.Name);
 			}
 
 			// Remove everything that's in a custom column
-			foreach (ColumnHeader CustomColumn in CustomColumns)
+			foreach (ColumnHeader customColumn in _customColumns)
 			{
-				ConfigObject ColumnConfig = (ConfigObject)CustomColumn.Tag;
-				foreach (string BadgeName in ColumnConfig.GetValue("Badges", "").Split(','))
+				ConfigObject columnConfig = (ConfigObject)customColumn.Tag;
+				foreach (string badgeName in columnConfig.GetValue("Badges", "").Split(','))
 				{
-					BadgeNameToGroup.Remove(BadgeName);
+					badgeNameToGroup.Remove(badgeName);
 				}
 			}
 
 			// Sort the list of groups for display
-			BadgeNameAndGroupPairs = BadgeNameToGroup.OrderBy(x => x.Value).ThenBy(x => x.Key).ToList();
+			_badgeNameAndGroupPairs = badgeNameToGroup.OrderBy(x => x.Value).ThenBy(x => x.Key).ToList();
 
 			// Figure out whether to show smaller badges due to limited space
 			UpdateMaxBuildBadgeChars();
 
 			// Update everything else
-			ArchiveToChangeNumberToArchiveKey.Clear();
-			ChangeNumberToLayoutInfo.Clear();
+			_archiveToChangeNumberToArchiveKey.Clear();
+			_changeNumberToLayoutInfo.Clear();
 			BuildList.Invalidate();
 			UpdateServiceBadges();
 			UpdateStatusPanel();
@@ -1586,7 +1586,7 @@ namespace UnrealGameSync
 			CheckForStartupComplete();
 
 			// If we are filtering by badges, we may also need to update the build list
-			if (ProjectSettings.FilterType != FilterType.None || ProjectSettings.FilterBadges.Count > 0)
+			if (_projectSettings.FilterType != FilterType.None || _projectSettings.FilterBadges.Count > 0)
 			{
 				UpdateBuildList();
 			}
@@ -1594,73 +1594,73 @@ namespace UnrealGameSync
 
 		void UpdateMaxBuildBadgeChars()
 		{
-			int NewMaxBuildBadgeChars;
+			int newMaxBuildBadgeChars;
 			if (GetBuildBadgeStripWidth(-1) < CISColumn.Width)
 			{
-				NewMaxBuildBadgeChars = -1;
+				newMaxBuildBadgeChars = -1;
 			}
 			else if (GetBuildBadgeStripWidth(3) < CISColumn.Width)
 			{
-				NewMaxBuildBadgeChars = 3;
+				newMaxBuildBadgeChars = 3;
 			}
 			else if (GetBuildBadgeStripWidth(2) < CISColumn.Width)
 			{
-				NewMaxBuildBadgeChars = 2;
+				newMaxBuildBadgeChars = 2;
 			}
 			else
 			{
-				NewMaxBuildBadgeChars = 1;
+				newMaxBuildBadgeChars = 1;
 			}
 
-			if (NewMaxBuildBadgeChars != MaxBuildBadgeChars)
+			if (newMaxBuildBadgeChars != _maxBuildBadgeChars)
 			{
-				ChangeNumberToLayoutInfo.Clear();
+				_changeNumberToLayoutInfo.Clear();
 				BuildList.Invalidate();
-				MaxBuildBadgeChars = NewMaxBuildBadgeChars;
+				_maxBuildBadgeChars = newMaxBuildBadgeChars;
 			}
 		}
 
-		int GetBuildBadgeStripWidth(int MaxNumChars)
+		int GetBuildBadgeStripWidth(int maxNumChars)
 		{
 			// Create dummy badges for each badge name
-			List<BadgeInfo> DummyBadges = new List<BadgeInfo>();
-			foreach (KeyValuePair<string, string> BadgeNameAndGroupPair in BadgeNameAndGroupPairs)
+			List<BadgeInfo> dummyBadges = new List<BadgeInfo>();
+			foreach (KeyValuePair<string, string> badgeNameAndGroupPair in _badgeNameAndGroupPairs)
 			{
-				string BadgeName = BadgeNameAndGroupPair.Key;
-				if (MaxNumChars != -1 && BadgeName.Length > MaxNumChars)
+				string badgeName = badgeNameAndGroupPair.Key;
+				if (maxNumChars != -1 && badgeName.Length > maxNumChars)
 				{
-					BadgeName = BadgeName.Substring(0, MaxNumChars);
+					badgeName = badgeName.Substring(0, maxNumChars);
 				}
 
-				BadgeInfo DummyBadge = CreateBadge(-1, BadgeName, BadgeNameAndGroupPair.Value, null);
-				DummyBadges.Add(DummyBadge);
+				BadgeInfo dummyBadge = CreateBadge(-1, badgeName, badgeNameAndGroupPair.Value, null);
+				dummyBadges.Add(dummyBadge);
 			}
 
 			// Check if they fit within the column
-			int Width = 0;
-			if (DummyBadges.Count > 0)
+			int width = 0;
+			if (dummyBadges.Count > 0)
 			{
-				LayoutBadges(DummyBadges);
-				Width = DummyBadges[DummyBadges.Count - 1].Offset + DummyBadges[DummyBadges.Count - 1].Width;
+				LayoutBadges(dummyBadges);
+				width = dummyBadges[dummyBadges.Count - 1].Offset + dummyBadges[dummyBadges.Count - 1].Width;
 			}
-			return Width;
+			return width;
 		}
 
-		bool ShouldIncludeInReviewedList(int ChangeNumber)
+		bool ShouldIncludeInReviewedList(int changeNumber)
 		{
-			if (PromotedChangeNumbers.Contains(ChangeNumber))
+			if (_promotedChangeNumbers.Contains(changeNumber))
 			{
 				return true;
 			}
 
-			EventSummary? Review = EventMonitor.GetSummaryForChange(ChangeNumber);
-			if (Review != null)
+			EventSummary? review = _eventMonitor.GetSummaryForChange(changeNumber);
+			if (review != null)
 			{
-				if (Review.LastStarReview != null && Review.LastStarReview.Type == EventType.Starred)
+				if (review.LastStarReview != null && review.LastStarReview.Type == EventType.Starred)
 				{
 					return true;
 				}
-				if (Review.Verdict == ReviewVerdict.Good || Review.Verdict == ReviewVerdict.Mixed)
+				if (review.Verdict == ReviewVerdict.Good || review.Verdict == ReviewVerdict.Mixed)
 				{
 					return true;
 				}
@@ -1670,93 +1670,93 @@ namespace UnrealGameSync
 
 		Dictionary<Guid, WorkspaceSyncCategory> GetSyncCategories()
 		{
-			return ConfigUtils.GetSyncCategories(PerforceMonitor.LatestProjectConfigFile);
+			return ConfigUtils.GetSyncCategories(_perforceMonitor.LatestProjectConfigFile);
 		}
 
-		List<string> GetProjectRoots(string InUProjectDirsPath)
+		List<string> GetProjectRoots(string inUProjectDirsPath)
 		{
-			List<string> ProjectRoots = new List<string>();
+			List<string> projectRoots = new List<string>();
 
-			if (!File.Exists(InUProjectDirsPath))
+			if (!File.Exists(inUProjectDirsPath))
 			{
-				return ProjectRoots;
+				return projectRoots;
 			}
 
-			string UProjectDirsRoot = Path.GetDirectoryName(InUProjectDirsPath)!;
+			string uProjectDirsRoot = Path.GetDirectoryName(inUProjectDirsPath)!;
 
-			foreach (string Line in File.ReadLines(InUProjectDirsPath))
+			foreach (string line in File.ReadLines(inUProjectDirsPath))
 			{
 				// Remove unnecessary whitespaces
-				string TrimmedLine = Line.Trim();
+				string trimmedLine = line.Trim();
 
 				// Remove every comments
-				int CommentIndex = TrimmedLine.IndexOf(';');
+				int commentIndex = trimmedLine.IndexOf(';');
 
-				if (CommentIndex >= 0)
+				if (commentIndex >= 0)
 				{
-					TrimmedLine = TrimmedLine.Substring(0, CommentIndex);
+					trimmedLine = trimmedLine.Substring(0, commentIndex);
 				}
 
-				if (TrimmedLine.Length <= 0)
+				if (trimmedLine.Length <= 0)
 				{
 					continue;
 				}
 
-				string UProjectDirsEntry = Path.Combine(UProjectDirsRoot, TrimmedLine);
+				string uProjectDirsEntry = Path.Combine(uProjectDirsRoot, trimmedLine);
 
-				if (!Directory.Exists(UProjectDirsEntry))
+				if (!Directory.Exists(uProjectDirsEntry))
 				{
 					continue;
 				}
 
-				List<string> EntryDirectories = new List<string>(Directory.EnumerateDirectories(UProjectDirsEntry));
+				List<string> entryDirectories = new List<string>(Directory.EnumerateDirectories(uProjectDirsEntry));
 
-				ProjectRoots = ProjectRoots.Concat(EntryDirectories).ToList();
+				projectRoots = projectRoots.Concat(entryDirectories).ToList();
 			}
 
-			return ProjectRoots;
+			return projectRoots;
 		}
 
-		void ParseUncontrolledChangelistsPersistencyFile(string InProjectRoot, HashSet<string> OutUncontrolledFiles)
+		void ParseUncontrolledChangelistsPersistencyFile(string inProjectRoot, HashSet<string> outUncontrolledFiles)
 		{
-			String UncontrolledChangelistPersistencyFilePath = Path.Combine(InProjectRoot, "Saved", "SourceControl", "UncontrolledChangelists.json");
+			String uncontrolledChangelistPersistencyFilePath = Path.Combine(inProjectRoot, "Saved", "SourceControl", "UncontrolledChangelists.json");
 
-			if (File.Exists(UncontrolledChangelistPersistencyFilePath))
+			if (File.Exists(uncontrolledChangelistPersistencyFilePath))
 			{
 				try
 				{
-					string JsonString = File.ReadAllText(UncontrolledChangelistPersistencyFilePath);
-					UncontrolledChangelistPersistency UCLPersistency = JsonSerializer.Deserialize<UncontrolledChangelistPersistency>(JsonString, Utility.DefaultJsonSerializerOptions);
+					string jsonString = File.ReadAllText(uncontrolledChangelistPersistencyFilePath);
+					UncontrolledChangelistPersistency uclPersistency = JsonSerializer.Deserialize<UncontrolledChangelistPersistency>(jsonString, Utility.DefaultJsonSerializerOptions);
 
-					foreach (UncontrolledChangelist UCL in UCLPersistency.Changelists)
+					foreach (UncontrolledChangelist ucl in uclPersistency.Changelists)
 					{
-						foreach (string UncontrolledFile in UCL.Files)
+						foreach (string uncontrolledFile in ucl.Files)
 						{
-							OutUncontrolledFiles.Add(UncontrolledFile);
+							outUncontrolledFiles.Add(uncontrolledFile);
 						}
 					}
 				}
-				catch (Exception Ex)
+				catch (Exception ex)
 				{
-					Logger.LogError(Ex, "Exception while parsing Uncontrolled Changelist persistency file");
+					_logger.LogError(ex, "Exception while parsing Uncontrolled Changelist persistency file");
 				}
 			}
 		}
 
 		void UpdateReviewsCallback()
 		{
-			if (!bUpdateReviewsPosted)
+			if (!_updateReviewsPosted)
 			{
-				bUpdateReviewsPosted = true;
-				MainThreadSynchronizationContext.Post((o) => { bUpdateReviewsPosted = false; if (!bIsDisposing) { UpdateReviews(); } }, null);
+				_updateReviewsPosted = true;
+				_mainThreadSynchronizationContext.Post((o) => { _updateReviewsPosted = false; if (!_isDisposing) { UpdateReviews(); } }, null);
 			}
 		}
 
 		void UpdateReviews()
 		{
-			ArchiveToChangeNumberToArchiveKey.Clear();
-			ChangeNumberToLayoutInfo.Clear();
-			EventMonitor.ApplyUpdates();
+			_archiveToChangeNumberToArchiveKey.Clear();
+			_changeNumberToLayoutInfo.Clear();
+			_eventMonitor.ApplyUpdates();
 
 			if (UpdateServiceBadges())
 			{
@@ -1781,68 +1781,68 @@ namespace UnrealGameSync
 				return;
 			}
 
-			int LastChangeByCurrentUser = PerforceMonitor.LastChangeByCurrentUser;
-			int LastCodeChangeByCurrentUser = PerforceMonitor.LastCodeChangeByCurrentUser;
+			int lastChangeByCurrentUser = _perforceMonitor.LastChangeByCurrentUser;
+			int lastCodeChangeByCurrentUser = _perforceMonitor.LastCodeChangeByCurrentUser;
 
 			// Find all the badges which should notify users due to content changes
-			HashSet<string> ContentBadges = new HashSet<string>();
-			ContentBadges.UnionWith(PerforceMonitor.LatestProjectConfigFile.GetValues("Notifications.ContentBadges", new string[0]));
+			HashSet<string> contentBadges = new HashSet<string>();
+			contentBadges.UnionWith(_perforceMonitor.LatestProjectConfigFile.GetValues("Notifications.ContentBadges", new string[0]));
 
 			// Find the most recent build of each type, and the last time it succeeded
-			Dictionary<string, BadgeData> TypeToLastBuild = new Dictionary<string, BadgeData>();
-			Dictionary<string, BadgeData> TypeToLastSucceededBuild = new Dictionary<string, BadgeData>();
-			for (int Idx = SortedChangeNumbers.Count - 1; Idx >= 0; Idx--)
+			Dictionary<string, BadgeData> typeToLastBuild = new Dictionary<string, BadgeData>();
+			Dictionary<string, BadgeData> typeToLastSucceededBuild = new Dictionary<string, BadgeData>();
+			for (int idx = _sortedChangeNumbers.Count - 1; idx >= 0; idx--)
 			{
-				EventSummary? Summary = EventMonitor.GetSummaryForChange(SortedChangeNumbers[Idx]);
-				if (Summary != null)
+				EventSummary? summary = _eventMonitor.GetSummaryForChange(_sortedChangeNumbers[idx]);
+				if (summary != null)
 				{
-					foreach (BadgeData Badge in Summary.Badges)
+					foreach (BadgeData badge in summary.Badges)
 					{
-						if (!TypeToLastBuild.ContainsKey(Badge.BuildType) && (Badge.Result == BadgeResult.Success || Badge.Result == BadgeResult.Warning || Badge.Result == BadgeResult.Failure))
+						if (!typeToLastBuild.ContainsKey(badge.BuildType) && (badge.Result == BadgeResult.Success || badge.Result == BadgeResult.Warning || badge.Result == BadgeResult.Failure))
 						{
-							TypeToLastBuild.Add(Badge.BuildType, Badge);
+							typeToLastBuild.Add(badge.BuildType, badge);
 						}
-						if (!TypeToLastSucceededBuild.ContainsKey(Badge.BuildType) && Badge.Result == BadgeResult.Success)
+						if (!typeToLastSucceededBuild.ContainsKey(badge.BuildType) && badge.Result == BadgeResult.Success)
 						{
-							TypeToLastSucceededBuild.Add(Badge.BuildType, Badge);
+							typeToLastSucceededBuild.Add(badge.BuildType, badge);
 						}
 					}
 				}
 			}
 
 			// Find all the build types that the user needs to be notified about.
-			int RequireNotificationForChange = -1;
-			List<BadgeData> NotifyBuilds = new List<BadgeData>();
-			foreach (BadgeData LastBuild in TypeToLastBuild.Values.OrderBy(x => x.BuildType))
+			int requireNotificationForChange = -1;
+			List<BadgeData> notifyBuilds = new List<BadgeData>();
+			foreach (BadgeData lastBuild in typeToLastBuild.Values.OrderBy(x => x.BuildType))
 			{
-				if (LastBuild.Result == BadgeResult.Failure || LastBuild.Result == BadgeResult.Warning)
+				if (lastBuild.Result == BadgeResult.Failure || lastBuild.Result == BadgeResult.Warning)
 				{
 					// Get the last submitted changelist by this user of the correct type
-					int LastChangeByCurrentUserOfType;
-					if (ContentBadges.Contains(LastBuild.BuildType))
+					int lastChangeByCurrentUserOfType;
+					if (contentBadges.Contains(lastBuild.BuildType))
 					{
-						LastChangeByCurrentUserOfType = LastChangeByCurrentUser;
+						lastChangeByCurrentUserOfType = lastChangeByCurrentUser;
 					}
 					else
 					{
-						LastChangeByCurrentUserOfType = LastCodeChangeByCurrentUser;
+						lastChangeByCurrentUserOfType = lastCodeChangeByCurrentUser;
 					}
 
 					// Check if the failed build was after we submitted
-					if (LastChangeByCurrentUserOfType > 0 && LastBuild.ChangeNumber >= LastChangeByCurrentUserOfType)
+					if (lastChangeByCurrentUserOfType > 0 && lastBuild.ChangeNumber >= lastChangeByCurrentUserOfType)
 					{
 						// And check that there wasn't a successful build after we submitted (if there was, we're in the clear)
-						BadgeData? LastSuccessfulBuild;
-						if (!TypeToLastSucceededBuild.TryGetValue(LastBuild.BuildType, out LastSuccessfulBuild) || LastSuccessfulBuild.ChangeNumber < LastChangeByCurrentUserOfType)
+						BadgeData? lastSuccessfulBuild;
+						if (!typeToLastSucceededBuild.TryGetValue(lastBuild.BuildType, out lastSuccessfulBuild) || lastSuccessfulBuild.ChangeNumber < lastChangeByCurrentUserOfType)
 						{
 							// Add it to the list of notifications
-							NotifyBuilds.Add(LastBuild);
+							notifyBuilds.Add(lastBuild);
 
 							// Check if this is a new notification, rather than one we've already dismissed
-							int NotifiedChangeNumber;
-							if (!NotifiedBuildTypeToChangeNumber.TryGetValue(LastBuild.BuildType, out NotifiedChangeNumber) || NotifiedChangeNumber < LastChangeByCurrentUserOfType)
+							int notifiedChangeNumber;
+							if (!_notifiedBuildTypeToChangeNumber.TryGetValue(lastBuild.BuildType, out notifiedChangeNumber) || notifiedChangeNumber < lastChangeByCurrentUserOfType)
 							{
-								RequireNotificationForChange = Math.Max(RequireNotificationForChange, LastChangeByCurrentUserOfType);
+								requireNotificationForChange = Math.Max(requireNotificationForChange, lastChangeByCurrentUserOfType);
 							}
 						}
 					}
@@ -1850,42 +1850,42 @@ namespace UnrealGameSync
 			}
 
 			// If there's anything we haven't already notified the user about, do so now
-			if (RequireNotificationForChange != -1)
+			if (requireNotificationForChange != -1)
 			{
 				// Format the platform list
-				StringBuilder PlatformList = new StringBuilder();
-				PlatformList.AppendFormat("{0}: {1}", StreamName, NotifyBuilds[0].BuildType);
-				for (int Idx = 1; Idx < NotifyBuilds.Count - 1; Idx++)
+				StringBuilder platformList = new StringBuilder();
+				platformList.AppendFormat("{0}: {1}", StreamName, notifyBuilds[0].BuildType);
+				for (int idx = 1; idx < notifyBuilds.Count - 1; idx++)
 				{
-					PlatformList.AppendFormat(", {0}", NotifyBuilds[Idx].BuildType);
+					platformList.AppendFormat(", {0}", notifyBuilds[idx].BuildType);
 				}
-				if (NotifyBuilds.Count > 1)
+				if (notifyBuilds.Count > 1)
 				{
-					PlatformList.AppendFormat(" and {0}", NotifyBuilds[NotifyBuilds.Count - 1].BuildType);
+					platformList.AppendFormat(" and {0}", notifyBuilds[notifyBuilds.Count - 1].BuildType);
 				}
 
 				// Show the balloon tooltip
-				if (NotifyBuilds.Any(x => x.Result == BadgeResult.Failure))
+				if (notifyBuilds.Any(x => x.Result == BadgeResult.Failure))
 				{
-					string Title = String.Format("{0} Errors", PlatformList.ToString());
-					string Message = String.Format("CIS failed after your last submitted changelist ({0}).", RequireNotificationForChange);
-					NotificationWindow.Show(NotificationType.Error, Title, Message);
+					string title = String.Format("{0} Errors", platformList.ToString());
+					string message = String.Format("CIS failed after your last submitted changelist ({0}).", requireNotificationForChange);
+					_notificationWindow.Show(NotificationType.Error, title, message);
 				}
 				else
 				{
-					string Title = String.Format("{0} Warnings", PlatformList.ToString());
-					string Message = String.Format("CIS completed with warnings after your last submitted changelist ({0}).", RequireNotificationForChange);
-					NotificationWindow.Show(NotificationType.Warning, Title, Message);
+					string title = String.Format("{0} Warnings", platformList.ToString());
+					string message = String.Format("CIS completed with warnings after your last submitted changelist ({0}).", requireNotificationForChange);
+					_notificationWindow.Show(NotificationType.Warning, title, message);
 				}
 
 				// Set the link to open the right build pages
-				int HighlightChange = NotifyBuilds.Max(x => x.ChangeNumber);
-				NotificationWindow.OnMoreInformation = () => { Owner.ShowAndActivate(); SelectChange(HighlightChange); };
+				int highlightChange = notifyBuilds.Max(x => x.ChangeNumber);
+				_notificationWindow.OnMoreInformation = () => { _owner.ShowAndActivate(); SelectChange(highlightChange); };
 
 				// Don't show messages for this change again
-				foreach (BadgeData NotifyBuild in NotifyBuilds)
+				foreach (BadgeData notifyBuild in notifyBuilds)
 				{
-					NotifiedBuildTypeToChangeNumber[NotifyBuild.BuildType] = RequireNotificationForChange;
+					_notifiedBuildTypeToChangeNumber[notifyBuild.BuildType] = requireNotificationForChange;
 				}
 			}
 		}
@@ -1903,77 +1903,77 @@ namespace UnrealGameSync
 			public Rectangle LinkRect;
 		}
 
-		ExpandRowLayout LayoutExpandRow(Graphics Graphics, Rectangle Bounds)
+		ExpandRowLayout LayoutExpandRow(Graphics graphics, Rectangle bounds)
 		{
-			ExpandRowLayout Layout = new ExpandRowLayout();
+			ExpandRowLayout layout = new ExpandRowLayout();
 
-			string ShowingChanges = String.Format("Showing {0}/{1} changes.", ListIndexToChangeIndex.Count, NumChanges);
+			string showingChanges = String.Format("Showing {0}/{1} changes.", _listIndexToChangeIndex.Count, _numChanges);
 
-			int CurrentMaxChanges = PerforceMonitor.CurrentMaxChanges;
-			if (PerforceMonitor.PendingMaxChanges > CurrentMaxChanges)
+			int currentMaxChanges = _perforceMonitor.CurrentMaxChanges;
+			if (_perforceMonitor.PendingMaxChanges > currentMaxChanges)
 			{
-				Layout.MainText = String.Format("{0}. Fetching {1} more from server...  ", ShowingChanges, PerforceMonitor.PendingMaxChanges - CurrentMaxChanges);
-				Layout.LinkText = "Cancel";
+				layout.MainText = String.Format("{0}. Fetching {1} more from server...  ", showingChanges, _perforceMonitor.PendingMaxChanges - currentMaxChanges);
+				layout.LinkText = "Cancel";
 			}
-			else if (PerforceMonitor.CurrentMaxChanges > NumChanges)
+			else if (_perforceMonitor.CurrentMaxChanges > _numChanges)
 			{
-				Layout.MainText = ShowingChanges;
-				Layout.LinkText = "";
+				layout.MainText = showingChanges;
+				layout.LinkText = "";
 			}
 			else
 			{
-				Layout.MainText = String.Format("{0}  ", ShowingChanges);
-				Layout.LinkText = "Show more...";
+				layout.MainText = String.Format("{0}  ", showingChanges);
+				layout.LinkText = "Show more...";
 			}
 
-			Size MainTextSize = TextRenderer.MeasureText(Graphics, Layout.MainText, BuildFont, new Size(1000, 1000), TextFormatFlags.NoPadding);
-			Size LinkTextSize = TextRenderer.MeasureText(Graphics, Layout.LinkText, BuildFont, new Size(1000, 1000), TextFormatFlags.NoPadding);
+			Size mainTextSize = TextRenderer.MeasureText(graphics, layout.MainText, _buildFont, new Size(1000, 1000), TextFormatFlags.NoPadding);
+			Size linkTextSize = TextRenderer.MeasureText(graphics, layout.LinkText, _buildFont, new Size(1000, 1000), TextFormatFlags.NoPadding);
 
-			int MinX = Bounds.Left + (Bounds.Width - MainTextSize.Width - LinkTextSize.Width) / 2;
-			int MinY = Bounds.Bottom - MainTextSize.Height - 1;
+			int minX = bounds.Left + (bounds.Width - mainTextSize.Width - linkTextSize.Width) / 2;
+			int minY = bounds.Bottom - mainTextSize.Height - 1;
 
-			Layout.MainRect = new Rectangle(new Point(MinX, MinY), MainTextSize);
-			Layout.LinkRect = new Rectangle(new Point(MinX + MainTextSize.Width, MinY), LinkTextSize);
+			layout.MainRect = new Rectangle(new Point(minX, minY), mainTextSize);
+			layout.LinkRect = new Rectangle(new Point(minX + mainTextSize.Width, minY), linkTextSize);
 
-			return Layout;
+			return layout;
 		}
 
-		private void DrawExpandRow(Graphics Graphics, ExpandRowLayout Layout)
+		private void DrawExpandRow(Graphics graphics, ExpandRowLayout layout)
 		{
-			TextRenderer.DrawText(Graphics, Layout.MainText, BuildFont, Layout.MainRect, SystemColors.WindowText, TextFormatFlags.SingleLine | TextFormatFlags.VerticalCenter | TextFormatFlags.NoPrefix | TextFormatFlags.NoPadding);
+			TextRenderer.DrawText(graphics, layout.MainText, _buildFont, layout.MainRect, SystemColors.WindowText, TextFormatFlags.SingleLine | TextFormatFlags.VerticalCenter | TextFormatFlags.NoPrefix | TextFormatFlags.NoPadding);
 
-			Color LinkColor = SystemColors.HotTrack;
-			if (bMouseOverExpandLink)
+			Color linkColor = SystemColors.HotTrack;
+			if (_mouseOverExpandLink)
 			{
-				LinkColor = Color.FromArgb(LinkColor.B, LinkColor.G, LinkColor.R);
+				linkColor = Color.FromArgb(linkColor.B, linkColor.G, linkColor.R);
 			}
 
-			TextRenderer.DrawText(Graphics, Layout.LinkText, BuildFont, Layout.LinkRect, LinkColor, TextFormatFlags.SingleLine | TextFormatFlags.VerticalCenter | TextFormatFlags.NoPrefix | TextFormatFlags.NoPadding);
+			TextRenderer.DrawText(graphics, layout.LinkText, _buildFont, layout.LinkRect, linkColor, TextFormatFlags.SingleLine | TextFormatFlags.VerticalCenter | TextFormatFlags.NoPrefix | TextFormatFlags.NoPadding);
 		}
 
-		private bool HitTestExpandLink(Point Location)
+		private bool HitTestExpandLink(Point location)
 		{
-			if (ExpandItem == null)
+			if (_expandItem == null)
 			{
 				return false;
 			}
-			using (Graphics Graphics = Graphics.FromHwnd(IntPtr.Zero))
+			using (Graphics graphics = Graphics.FromHwnd(IntPtr.Zero))
 			{
-				ExpandRowLayout Layout = LayoutExpandRow(Graphics, ExpandItem.Bounds);
-				return Layout.LinkRect.Contains(Location);
+				ExpandRowLayout layout = LayoutExpandRow(graphics, _expandItem.Bounds);
+				return layout.LinkRect.Contains(location);
 			}
 		}
 
 		private void BuildList_DrawItem(object sender, DrawListViewItemEventArgs e)
 		{
-			if (e.Item == ExpandItem)
+			if (e.Item == _expandItem)
 			{
 				BuildList.DrawDefaultBackground(e.Graphics, e.Bounds);
 
-				ExpandRowLayout ExpandLayout = LayoutExpandRow(e.Graphics, e.Bounds);
-				DrawExpandRow(e.Graphics, ExpandLayout);
+				ExpandRowLayout expandLayout = LayoutExpandRow(e.Graphics, e.Bounds);
+				DrawExpandRow(e.Graphics, expandLayout);
 			}
-			else if (Workspace != null)
+			else if (_workspace != null)
 			{
 				if (e.Item.Selected)
 				{
@@ -1983,7 +1983,7 @@ namespace UnrealGameSync
 				{
 					BuildList.DrawTrackedBackground(e.Graphics, e.Bounds);
 				}
-				else if (((ChangesRecord)e.Item.Tag).Number == Workspace.PendingChangeNumber)
+				else if (((ChangesRecord)e.Item.Tag).Number == _workspace.PendingChangeNumber)
 				{
 					BuildList.DrawTrackedBackground(e.Graphics, e.Bounds);
 				}
@@ -1994,53 +1994,53 @@ namespace UnrealGameSync
 			}
 		}
 
-		private string? GetArchiveKeyForChangeNumber(IArchiveInfo Archive, int ChangeNumber)
+		private string? GetArchiveKeyForChangeNumber(IArchiveInfo archive, int changeNumber)
 		{
-			string? ArchivePath;
+			string? archivePath;
 
-			Dictionary<int, string?>? ChangeNumberToArchivePath;
-			if (!ArchiveToChangeNumberToArchiveKey.TryGetValue(Archive.Name, out ChangeNumberToArchivePath))
+			Dictionary<int, string?>? changeNumberToArchivePath;
+			if (!_archiveToChangeNumberToArchiveKey.TryGetValue(archive.Name, out changeNumberToArchivePath))
 			{
-				ChangeNumberToArchivePath = new Dictionary<int, string?>();
-				ArchiveToChangeNumberToArchiveKey[Archive.Name] = ChangeNumberToArchivePath;
+				changeNumberToArchivePath = new Dictionary<int, string?>();
+				_archiveToChangeNumberToArchiveKey[archive.Name] = changeNumberToArchivePath;
 			}
 
-			if (!ChangeNumberToArchivePath.TryGetValue(ChangeNumber, out ArchivePath))
+			if (!changeNumberToArchivePath.TryGetValue(changeNumber, out archivePath))
 			{
-				PerforceChangeDetails? Details;
-				if (PerforceMonitor.TryGetChangeDetails(ChangeNumber, out Details))
+				PerforceChangeDetails? details;
+				if (_perforceMonitor.TryGetChangeDetails(changeNumber, out details))
 				{
 					// Try to get the archive for this CL
-					if (!Archive.TryGetArchiveKeyForChangeNumber(ChangeNumber, out ArchivePath) && !Details.bContainsCode)
+					if (!archive.TryGetArchiveKeyForChangeNumber(changeNumber, out archivePath) && !details.ContainsCode)
 					{
 						// Otherwise if it's a content-only change, find the previous build any use the archive path from that
-						int Index = SortedChangeNumbers.BinarySearch(ChangeNumber);
-						if (Index > 0)
+						int index = _sortedChangeNumbers.BinarySearch(changeNumber);
+						if (index > 0)
 						{
-							ArchivePath = GetArchiveKeyForChangeNumber(Archive, SortedChangeNumbers[Index - 1]);
+							archivePath = GetArchiveKeyForChangeNumber(archive, _sortedChangeNumbers[index - 1]);
 						}
 					}
 				}
-				ChangeNumberToArchivePath.Add(ChangeNumber, ArchivePath);
+				changeNumberToArchivePath.Add(changeNumber, archivePath);
 			}
 
-			return ArchivePath;
+			return archivePath;
 		}
 
-		private Color Blend(Color First, Color Second, float T)
+		private Color Blend(Color first, Color second, float T)
 		{
-			return Color.FromArgb((int)(First.R + (Second.R - First.R) * T), (int)(First.G + (Second.G - First.G) * T), (int)(First.B + (Second.B - First.B) * T));
+			return Color.FromArgb((int)(first.R + (second.R - first.R) * T), (int)(first.G + (second.G - first.G) * T), (int)(first.B + (second.B - first.B) * T));
 		}
 
-		private bool CanSyncChange(int ChangeNumber)
+		private bool CanSyncChange(int changeNumber)
 		{
-			if (PerforceMonitor == null)
+			if (_perforceMonitor == null)
 			{
 				return false;
 			}
 
-			List<IArchiveInfo> SelectedArchives = GetSelectedArchives(GetArchives());
-			return SelectedArchives.Count == 0 || SelectedArchives.All(x => GetArchiveKeyForChangeNumber(x, ChangeNumber) != null);
+			List<IArchiveInfo> selectedArchives = GetSelectedArchives(GetArchives());
+			return selectedArchives.Count == 0 || selectedArchives.All(x => GetArchiveKeyForChangeNumber(x, changeNumber) != null);
 		}
 
 		/// <summary>
@@ -2048,24 +2048,24 @@ namespace UnrealGameSync
 		/// Determines if the ini specified INI filter is a subset of the InBadges passed
 		/// BadgeGroupSyncFilter - is used as the ini 
 		/// </summary>
-		/// <param name="InBadges"></param>
+		/// <param name="inBadges"></param>
 		/// <returns></returns>
-		private bool DoRequiredBadgesExist(List<string> RequiredBadgeList, List<BadgeData> InBadges)
+		private bool DoRequiredBadgesExist(List<string> requiredBadgeList, List<BadgeData> inBadges)
 		{
-			Dictionary<string, BadgeData> InBadgeDictionary = new Dictionary<string, BadgeData>();
-			foreach (BadgeData Badge in InBadges)
+			Dictionary<string, BadgeData> inBadgeDictionary = new Dictionary<string, BadgeData>();
+			foreach (BadgeData badge in inBadges)
 			{
-				InBadgeDictionary.Add(Badge.BadgeName, Badge);
+				inBadgeDictionary.Add(badge.BadgeName, badge);
 			}
 
 			// make sure all required badges exist in the InBadgeSet
-			foreach (string BadgeName in RequiredBadgeList)
+			foreach (string badgeName in requiredBadgeList)
 			{
-				BadgeData? Badge;
-				if (InBadgeDictionary.TryGetValue(BadgeName, out Badge))
+				BadgeData? badge;
+				if (inBadgeDictionary.TryGetValue(badgeName, out badge))
 				{
 					// If any required badge is not successful then the filter isn't matched.
-					if (!Badge.IsSuccess)
+					if (!badge.IsSuccess)
 					{
 						return false;
 					}
@@ -2081,27 +2081,27 @@ namespace UnrealGameSync
 			return true;
 		}
 
-		private bool CanSyncChangeType(LatestChangeType ChangeType, ChangesRecord Change, EventSummary? Summary)
+		private bool CanSyncChangeType(LatestChangeType changeType, ChangesRecord change, EventSummary? summary)
 		{
-			if (ChangeType.bGood)
+			if (changeType.Good)
 			{
-				if (Summary == null || Summary.Verdict != ReviewVerdict.Good)
+				if (summary == null || summary.Verdict != ReviewVerdict.Good)
 				{
 					return false;
 				}
 			}
 
-			if (ChangeType.bStarred)
+			if (changeType.Starred)
 			{
-				if ((Summary == null || Summary.LastStarReview == null || Summary.LastStarReview.Type != EventType.Starred) && !PromotedChangeNumbers.Contains(Change.Number))
+				if ((summary == null || summary.LastStarReview == null || summary.LastStarReview.Type != EventType.Starred) && !_promotedChangeNumbers.Contains(change.Number))
 				{
 					return false;
 				}
 			}
 
-			if (ChangeType.RequiredBadges.Count > 0)
+			if (changeType.RequiredBadges.Count > 0)
 			{
-				if (Summary == null || !DoRequiredBadgesExist(ChangeType.RequiredBadges, Summary.Badges))
+				if (summary == null || !DoRequiredBadgesExist(changeType.RequiredBadges, summary.Badges))
 				{
 					return false;
 				}
@@ -2110,440 +2110,440 @@ namespace UnrealGameSync
 			return true;
 		}
 
-		private ChangeLayoutInfo GetChangeLayoutInfo(ChangesRecord Change)
+		private ChangeLayoutInfo GetChangeLayoutInfo(ChangesRecord change)
 		{
-			ChangeLayoutInfo? LayoutInfo;
-			if (!ChangeNumberToLayoutInfo.TryGetValue(Change.Number, out LayoutInfo))
+			ChangeLayoutInfo? layoutInfo;
+			if (!_changeNumberToLayoutInfo.TryGetValue(change.Number, out layoutInfo))
 			{
-				LayoutInfo = new ChangeLayoutInfo();
+				layoutInfo = new ChangeLayoutInfo();
 
-				LayoutInfo.DescriptionBadges = CreateDescriptionBadges(Change);
-				LayoutInfo.TypeBadges = CreateTypeBadges(Change.Number);
+				layoutInfo.DescriptionBadges = CreateDescriptionBadges(change);
+				layoutInfo.TypeBadges = CreateTypeBadges(change.Number);
 
-				EventSummary? Summary = EventMonitor.GetSummaryForChange(Change.Number);
-				LayoutInfo.BuildBadges = CreateBuildBadges(Change.Number, Summary);
-				LayoutInfo.CustomBadges = CreateCustomBadges(Change.Number, Summary);
+				EventSummary? summary = _eventMonitor.GetSummaryForChange(change.Number);
+				layoutInfo.BuildBadges = CreateBuildBadges(change.Number, summary);
+				layoutInfo.CustomBadges = CreateCustomBadges(change.Number, summary);
 
-				ChangeNumberToLayoutInfo[Change.Number] = LayoutInfo;
+				_changeNumberToLayoutInfo[change.Number] = layoutInfo;
 			}
-			return LayoutInfo;
+			return layoutInfo;
 		}
 
-		private void GetRemainingBisectRange(out int OutPassChangeNumber, out int OutFailChangeNumber)
+		private void GetRemainingBisectRange(out int outPassChangeNumber, out int outFailChangeNumber)
 		{
-			int PassChangeNumber = -1;
-			foreach (BisectEntry Entry in WorkspaceState.BisectChanges)
+			int passChangeNumber = -1;
+			foreach (BisectEntry entry in _workspaceState.BisectChanges)
 			{
-				if (Entry.State == BisectState.Pass && (Entry.Change > PassChangeNumber || PassChangeNumber == -1))
+				if (entry.State == BisectState.Pass && (entry.Change > passChangeNumber || passChangeNumber == -1))
 				{
-					PassChangeNumber = Entry.Change;
+					passChangeNumber = entry.Change;
 				}
 			}
 
-			int FailChangeNumber = -1;
-			foreach (BisectEntry Entry in WorkspaceState.BisectChanges)
+			int failChangeNumber = -1;
+			foreach (BisectEntry entry in _workspaceState.BisectChanges)
 			{
-				if (Entry.State == BisectState.Fail && Entry.Change > PassChangeNumber && (Entry.Change < FailChangeNumber || FailChangeNumber == -1))
+				if (entry.State == BisectState.Fail && entry.Change > passChangeNumber && (entry.Change < failChangeNumber || failChangeNumber == -1))
 				{
-					FailChangeNumber = Entry.Change;
+					failChangeNumber = entry.Change;
 				}
 			}
 
-			OutPassChangeNumber = PassChangeNumber;
-			OutFailChangeNumber = FailChangeNumber;
+			outPassChangeNumber = passChangeNumber;
+			outFailChangeNumber = failChangeNumber;
 		}
 
 		private void BuildList_DrawSubItem(object sender, DrawListViewSubItemEventArgs e)
 		{
-			ChangesRecord Change = (ChangesRecord)e.Item.Tag;
-			if (Change == null)
+			ChangesRecord change = (ChangesRecord)e.Item.Tag;
+			if (change == null)
 			{
 				return;
 			}
 
-			float DpiScaleX = e.Graphics.DpiX / 96.0f;
-			float DpiScaleY = e.Graphics.DpiY / 96.0f;
+			float dpiScaleX = e.Graphics.DpiX / 96.0f;
+			float dpiScaleY = e.Graphics.DpiY / 96.0f;
 
-			float IconY = e.Bounds.Top + (e.Bounds.Height - 16 * DpiScaleY) / 2;
+			float iconY = e.Bounds.Top + (e.Bounds.Height - 16 * dpiScaleY) / 2;
 
-			StringFormat Format = new StringFormat();
-			Format.LineAlignment = StringAlignment.Center;
-			Format.FormatFlags = StringFormatFlags.NoWrap;
-			Format.Trimming = StringTrimming.EllipsisCharacter;
+			StringFormat format = new StringFormat();
+			format.LineAlignment = StringAlignment.Center;
+			format.FormatFlags = StringFormatFlags.NoWrap;
+			format.Trimming = StringTrimming.EllipsisCharacter;
 
-			Font CurrentFont = (Change.Number == Workspace.PendingChangeNumber || Change.Number == Workspace.CurrentChangeNumber) ? SelectedBuildFont! : BuildFont!;
+			Font currentFont = (change.Number == _workspace.PendingChangeNumber || change.Number == _workspace.CurrentChangeNumber) ? _selectedBuildFont! : _buildFont!;
 
-			bool bAllowSync = CanSyncChange(Change.Number);
-			int BadgeAlpha = bAllowSync ? 255 : 128;
-			Color TextColor = (bAllowSync || Change.Number == Workspace.PendingChangeNumber || Change.Number == Workspace.CurrentChangeNumber || WorkspaceState.AdditionalChangeNumbers.Contains(Change.Number)) ? SystemColors.WindowText : Blend(SystemColors.Window, SystemColors.WindowText, 0.25f);
+			bool allowSync = CanSyncChange(change.Number);
+			int badgeAlpha = allowSync ? 255 : 128;
+			Color textColor = (allowSync || change.Number == _workspace.PendingChangeNumber || change.Number == _workspace.CurrentChangeNumber || _workspaceState.AdditionalChangeNumbers.Contains(change.Number)) ? SystemColors.WindowText : Blend(SystemColors.Window, SystemColors.WindowText, 0.25f);
 
-			const int FadeRange = 6;
-			if (e.ItemIndex >= BuildList.Items.Count - FadeRange && NumChanges >= PerforceMonitor.CurrentMaxChanges && !IsBisectModeEnabled())
+			const int fadeRange = 6;
+			if (e.ItemIndex >= BuildList.Items.Count - fadeRange && _numChanges >= _perforceMonitor.CurrentMaxChanges && !IsBisectModeEnabled())
 			{
-				float Opacity = (float)(BuildList.Items.Count - e.ItemIndex - 0.9f) / FadeRange;
-				BadgeAlpha = (int)(BadgeAlpha * Opacity);
-				TextColor = Blend(SystemColors.Window, TextColor, Opacity);
+				float opacity = (float)(BuildList.Items.Count - e.ItemIndex - 0.9f) / fadeRange;
+				badgeAlpha = (int)(badgeAlpha * opacity);
+				textColor = Blend(SystemColors.Window, textColor, opacity);
 			}
 
 			if (e.ColumnIndex == IconColumn.Index)
 			{
-				EventSummary? Summary = EventMonitor.GetSummaryForChange(Change.Number);
+				EventSummary? summary = _eventMonitor.GetSummaryForChange(change.Number);
 
-				float MinX = 4 * DpiScaleX;
-				if ((Summary != null && EventMonitor.WasSyncedByCurrentUser(Summary.ChangeNumber)) || (Workspace != null && Workspace.CurrentChangeNumber == Change.Number))
+				float minX = 4 * dpiScaleX;
+				if ((summary != null && _eventMonitor.WasSyncedByCurrentUser(summary.ChangeNumber)) || (_workspace != null && _workspace.CurrentChangeNumber == change.Number))
 				{
-					e.Graphics.DrawImage(Properties.Resources.Icons, MinX * DpiScaleX, IconY, PreviousSyncIcon, GraphicsUnit.Pixel);
+					e.Graphics.DrawImage(Properties.Resources.Icons, minX * dpiScaleX, iconY, _previousSyncIcon, GraphicsUnit.Pixel);
 				}
-				else if (WorkspaceSettings != null && WorkspaceState.AdditionalChangeNumbers.Contains(Change.Number))
+				else if (_workspaceSettings != null && _workspaceState.AdditionalChangeNumbers.Contains(change.Number))
 				{
-					e.Graphics.DrawImage(Properties.Resources.Icons, MinX * DpiScaleX, IconY, AdditionalSyncIcon, GraphicsUnit.Pixel);
+					e.Graphics.DrawImage(Properties.Resources.Icons, minX * dpiScaleX, iconY, _additionalSyncIcon, GraphicsUnit.Pixel);
 				}
-				else if (bAllowSync && ((Summary != null && Summary.LastStarReview != null && Summary.LastStarReview.Type == EventType.Starred) || PromotedChangeNumbers.Contains(Change.Number)))
+				else if (allowSync && ((summary != null && summary.LastStarReview != null && summary.LastStarReview.Type == EventType.Starred) || _promotedChangeNumbers.Contains(change.Number)))
 				{
-					e.Graphics.DrawImage(Properties.Resources.Icons, MinX * DpiScaleX, IconY, PromotedBuildIcon, GraphicsUnit.Pixel);
+					e.Graphics.DrawImage(Properties.Resources.Icons, minX * dpiScaleX, iconY, _promotedBuildIcon, GraphicsUnit.Pixel);
 				}
-				MinX += PromotedBuildIcon.Width * DpiScaleX;
+				minX += _promotedBuildIcon.Width * dpiScaleX;
 
-				if (bAllowSync)
+				if (allowSync)
 				{
-					Rectangle QualityIcon = DefaultBuildIcon;
+					Rectangle qualityIcon = DefaultBuildIcon;
 
 					if (IsBisectModeEnabled())
 					{
-						int FirstPass, FirstFail;
-						GetRemainingBisectRange(out FirstPass, out FirstFail);
+						int firstPass, firstFail;
+						GetRemainingBisectRange(out firstPass, out firstFail);
 
-						BisectEntry Entry = WorkspaceState.BisectChanges.FirstOrDefault(x => x.Change == Change.Number);
-						if (Entry == null || Entry.State == BisectState.Exclude)
+						BisectEntry entry = _workspaceState.BisectChanges.FirstOrDefault(x => x.Change == change.Number);
+						if (entry == null || entry.State == BisectState.Exclude)
 						{
-							QualityIcon = new Rectangle(0, 0, 0, 0);
+							qualityIcon = new Rectangle(0, 0, 0, 0);
 						}
-						else if (Entry.State == BisectState.Pass)
+						else if (entry.State == BisectState.Pass)
 						{
-							QualityIcon = BisectPassIcon;
+							qualityIcon = _bisectPassIcon;
 						}
-						else if (Entry.State == BisectState.Fail)
+						else if (entry.State == BisectState.Fail)
 						{
-							QualityIcon = BisectFailIcon;
+							qualityIcon = _bisectFailIcon;
 						}
-						else if (FirstFail != -1 && Change.Number > FirstFail)
+						else if (firstFail != -1 && change.Number > firstFail)
 						{
-							QualityIcon = BisectImplicitFailIcon;
+							qualityIcon = _bisectImplicitFailIcon;
 						}
-						else if (FirstPass != -1 && Change.Number < FirstPass)
+						else if (firstPass != -1 && change.Number < firstPass)
 						{
-							QualityIcon = BisectImplicitPassIcon;
+							qualityIcon = _bisectImplicitPassIcon;
 						}
 					}
 					else
 					{
-						if (EventMonitor.IsUnderInvestigation(Change.Number))
+						if (_eventMonitor.IsUnderInvestigation(change.Number))
 						{
-							QualityIcon = BadBuildIcon;
+							qualityIcon = BadBuildIcon;
 						}
-						else if (Summary != null)
+						else if (summary != null)
 						{
-							if (Summary.Verdict == ReviewVerdict.Good)
+							if (summary.Verdict == ReviewVerdict.Good)
 							{
-								QualityIcon = GoodBuildIcon;
+								qualityIcon = GoodBuildIcon;
 							}
-							else if (Summary.Verdict == ReviewVerdict.Bad)
+							else if (summary.Verdict == ReviewVerdict.Bad)
 							{
-								QualityIcon = BadBuildIcon;
+								qualityIcon = BadBuildIcon;
 							}
-							else if (Summary.Verdict == ReviewVerdict.Mixed)
+							else if (summary.Verdict == ReviewVerdict.Mixed)
 							{
-								QualityIcon = MixedBuildIcon;
+								qualityIcon = MixedBuildIcon;
 							}
 						}
 					}
-					e.Graphics.DrawImage(Properties.Resources.Icons, MinX, IconY, QualityIcon, GraphicsUnit.Pixel);
+					e.Graphics.DrawImage(Properties.Resources.Icons, minX, iconY, qualityIcon, GraphicsUnit.Pixel);
 
-					MinX += QualityIcon.Width * DpiScaleX;
+					minX += qualityIcon.Width * dpiScaleX;
 				}
 			}
 			else if (e.ColumnIndex == ChangeColumn.Index || e.ColumnIndex == TimeColumn.Index)
 			{
-				TextRenderer.DrawText(e.Graphics, e.SubItem.Text, CurrentFont, e.Bounds, TextColor, TextFormatFlags.EndEllipsis | TextFormatFlags.SingleLine | TextFormatFlags.HorizontalCenter | TextFormatFlags.VerticalCenter | TextFormatFlags.NoPrefix);
+				TextRenderer.DrawText(e.Graphics, e.SubItem.Text, currentFont, e.Bounds, textColor, TextFormatFlags.EndEllipsis | TextFormatFlags.SingleLine | TextFormatFlags.HorizontalCenter | TextFormatFlags.VerticalCenter | TextFormatFlags.NoPrefix);
 			}
 			else if (e.ColumnIndex == AuthorColumn.Index)
 			{
-				TextRenderer.DrawText(e.Graphics, e.SubItem.Text, CurrentFont, e.Bounds, TextColor, TextFormatFlags.EndEllipsis | TextFormatFlags.SingleLine | TextFormatFlags.VerticalCenter | TextFormatFlags.NoPrefix);
+				TextRenderer.DrawText(e.Graphics, e.SubItem.Text, currentFont, e.Bounds, textColor, TextFormatFlags.EndEllipsis | TextFormatFlags.SingleLine | TextFormatFlags.VerticalCenter | TextFormatFlags.NoPrefix);
 			}
 			else if (e.ColumnIndex == DescriptionColumn.Index)
 			{
-				ChangeLayoutInfo Layout = GetChangeLayoutInfo(Change);
+				ChangeLayoutInfo layout = GetChangeLayoutInfo(change);
 
-				Rectangle RemainingBounds = e.Bounds;
-				if (Layout.DescriptionBadges.Count > 0)
+				Rectangle remainingBounds = e.Bounds;
+				if (layout.DescriptionBadges.Count > 0)
 				{
 					e.Graphics.IntersectClip(e.Bounds);
 					e.Graphics.SmoothingMode = SmoothingMode.HighQuality;
 
-					RemainingBounds = new Rectangle(RemainingBounds.Left, RemainingBounds.Top, RemainingBounds.Width - (int)(2 * DpiScaleX), RemainingBounds.Height);
+					remainingBounds = new Rectangle(remainingBounds.Left, remainingBounds.Top, remainingBounds.Width - (int)(2 * dpiScaleX), remainingBounds.Height);
 
-					Point ListLocation = GetBadgeListLocation(Layout.DescriptionBadges, RemainingBounds, HorizontalAlign.Right, VerticalAlignment.Middle);
-					DrawBadges(e.Graphics, ListLocation, Layout.DescriptionBadges, BadgeAlpha);
+					Point listLocation = GetBadgeListLocation(layout.DescriptionBadges, remainingBounds, HorizontalAlign.Right, VerticalAlignment.Middle);
+					DrawBadges(e.Graphics, listLocation, layout.DescriptionBadges, badgeAlpha);
 
-					RemainingBounds = new Rectangle(RemainingBounds.Left, RemainingBounds.Top, ListLocation.X - RemainingBounds.Left - (int)(2 * DpiScaleX), RemainingBounds.Height);
+					remainingBounds = new Rectangle(remainingBounds.Left, remainingBounds.Top, listLocation.X - remainingBounds.Left - (int)(2 * dpiScaleX), remainingBounds.Height);
 				}
 
-				TextRenderer.DrawText(e.Graphics, e.SubItem.Text, CurrentFont, RemainingBounds, TextColor, TextFormatFlags.EndEllipsis | TextFormatFlags.SingleLine | TextFormatFlags.VerticalCenter | TextFormatFlags.NoPrefix);
+				TextRenderer.DrawText(e.Graphics, e.SubItem.Text, currentFont, remainingBounds, textColor, TextFormatFlags.EndEllipsis | TextFormatFlags.SingleLine | TextFormatFlags.VerticalCenter | TextFormatFlags.NoPrefix);
 			}
 			else if (e.ColumnIndex == TypeColumn.Index)
 			{
-				ChangeLayoutInfo Layout = GetChangeLayoutInfo(Change);
-				if (Layout.TypeBadges.Count > 0)
+				ChangeLayoutInfo layout = GetChangeLayoutInfo(change);
+				if (layout.TypeBadges.Count > 0)
 				{
 					e.Graphics.IntersectClip(e.Bounds);
 					e.Graphics.SmoothingMode = SmoothingMode.HighQuality;
-					Point TypeLocation = GetBadgeListLocation(Layout.TypeBadges, e.Bounds, HorizontalAlign.Center, VerticalAlignment.Middle);
-					DrawBadges(e.Graphics, TypeLocation, Layout.TypeBadges, BadgeAlpha);
+					Point typeLocation = GetBadgeListLocation(layout.TypeBadges, e.Bounds, HorizontalAlign.Center, VerticalAlignment.Middle);
+					DrawBadges(e.Graphics, typeLocation, layout.TypeBadges, badgeAlpha);
 				}
 			}
 			else if (e.ColumnIndex == CISColumn.Index)
 			{
-				ChangeLayoutInfo Layout = GetChangeLayoutInfo(Change);
-				if (Layout.BuildBadges.Count > 0)
+				ChangeLayoutInfo layout = GetChangeLayoutInfo(change);
+				if (layout.BuildBadges.Count > 0)
 				{
 					e.Graphics.IntersectClip(e.Bounds);
 					e.Graphics.SmoothingMode = SmoothingMode.HighQuality;
 
-					Point BuildsLocation = GetBadgeListLocation(Layout.BuildBadges, e.Bounds, HorizontalAlign.Center, VerticalAlignment.Middle);
-					BuildsLocation.X = Math.Max(BuildsLocation.X, e.Bounds.Left);
+					Point buildsLocation = GetBadgeListLocation(layout.BuildBadges, e.Bounds, HorizontalAlign.Center, VerticalAlignment.Middle);
+					buildsLocation.X = Math.Max(buildsLocation.X, e.Bounds.Left);
 
-					DrawBadges(e.Graphics, BuildsLocation, Layout.BuildBadges, BadgeAlpha);
+					DrawBadges(e.Graphics, buildsLocation, layout.BuildBadges, badgeAlpha);
 				}
 			}
 			else if (e.ColumnIndex == StatusColumn.Index)
 			{
-				float MaxX = e.SubItem.Bounds.Right;
+				float maxX = e.SubItem.Bounds.Right;
 
-				if (Change.Number == Workspace.PendingChangeNumber && Workspace.IsBusy())
+				if (change.Number == _workspace.PendingChangeNumber && _workspace.IsBusy())
 				{
-					Tuple<string, float> Progress = Workspace.CurrentProgress;
+					Tuple<string, float> progress = _workspace.CurrentProgress;
 
-					MaxX -= CancelIcon.Width;
-					e.Graphics.DrawImage(Properties.Resources.Icons, MaxX, IconY, CancelIcon, GraphicsUnit.Pixel);
+					maxX -= _cancelIcon.Width;
+					e.Graphics.DrawImage(Properties.Resources.Icons, maxX, iconY, _cancelIcon, GraphicsUnit.Pixel);
 
 					if (!Splitter.IsLogVisible())
 					{
-						MaxX -= InfoIcon.Width;
-						e.Graphics.DrawImage(Properties.Resources.Icons, MaxX, IconY, InfoIcon, GraphicsUnit.Pixel);
+						maxX -= _infoIcon.Width;
+						e.Graphics.DrawImage(Properties.Resources.Icons, maxX, iconY, _infoIcon, GraphicsUnit.Pixel);
 					}
 
-					float MinX = e.Bounds.Left;
+					float minX = e.Bounds.Left;
 
-					TextRenderer.DrawText(e.Graphics, Progress.Item1, CurrentFont, new Rectangle((int)MinX, e.Bounds.Top, (int)(MaxX - MinX), e.Bounds.Height), TextColor, TextFormatFlags.Left | TextFormatFlags.VerticalCenter | TextFormatFlags.EndEllipsis | TextFormatFlags.NoPrefix);
+					TextRenderer.DrawText(e.Graphics, progress.Item1, currentFont, new Rectangle((int)minX, e.Bounds.Top, (int)(maxX - minX), e.Bounds.Height), textColor, TextFormatFlags.Left | TextFormatFlags.VerticalCenter | TextFormatFlags.EndEllipsis | TextFormatFlags.NoPrefix);
 				}
 				else
 				{
 					e.Graphics.IntersectClip(e.Bounds);
 					e.Graphics.SmoothingMode = SmoothingMode.HighQuality;
 
-					if (Change.Number == Workspace.CurrentChangeNumber)
+					if (change.Number == _workspace.CurrentChangeNumber)
 					{
-						EventData? Review = EventMonitor.GetReviewByCurrentUser(Change.Number);
+						EventData? review = _eventMonitor.GetReviewByCurrentUser(change.Number);
 
-						MaxX -= FrownIcon.Width * DpiScaleX;
-						e.Graphics.DrawImage(Properties.Resources.Icons, MaxX, IconY, (Review == null || !EventMonitor.IsPositiveReview(Review.Type)) ? FrownIcon : DisabledFrownIcon, GraphicsUnit.Pixel);
+						maxX -= _frownIcon.Width * dpiScaleX;
+						e.Graphics.DrawImage(Properties.Resources.Icons, maxX, iconY, (review == null || !EventMonitor.IsPositiveReview(review.Type)) ? _frownIcon : _disabledFrownIcon, GraphicsUnit.Pixel);
 
-						MaxX -= HappyIcon.Width * DpiScaleX;
-						e.Graphics.DrawImage(Properties.Resources.Icons, MaxX, IconY, (Review == null || !EventMonitor.IsNegativeReview(Review.Type)) ? HappyIcon : DisabledHappyIcon, GraphicsUnit.Pixel);
+						maxX -= _happyIcon.Width * dpiScaleX;
+						e.Graphics.DrawImage(Properties.Resources.Icons, maxX, iconY, (review == null || !EventMonitor.IsNegativeReview(review.Type)) ? _happyIcon : _disabledHappyIcon, GraphicsUnit.Pixel);
 					}
-					else if (e.ItemIndex == BuildList.HoverItem && bAllowSync)
+					else if (e.ItemIndex == BuildList.HoverItem && allowSync)
 					{
-						Rectangle SyncBadgeRectangle = GetSyncBadgeRectangle(e.SubItem.Bounds);
-						DrawBadge(e.Graphics, SyncBadgeRectangle, "Sync", bHoverSync ? Color.FromArgb(140, 180, 230) : Color.FromArgb(112, 146, 190), true, true);
-						MaxX = SyncBadgeRectangle.Left - (int)(2 * DpiScaleX);
+						Rectangle syncBadgeRectangle = GetSyncBadgeRectangle(e.SubItem.Bounds);
+						DrawBadge(e.Graphics, syncBadgeRectangle, "Sync", _hoverSync ? Color.FromArgb(140, 180, 230) : Color.FromArgb(112, 146, 190), true, true);
+						maxX = syncBadgeRectangle.Left - (int)(2 * dpiScaleX);
 					}
 
-					string? SummaryText;
-					if (WorkspaceState.LastSyncChangeNumber == -1 || WorkspaceState.LastSyncChangeNumber != Change.Number || !GetLastUpdateMessage(WorkspaceState.LastSyncResult, WorkspaceState.LastSyncResultMessage, out SummaryText))
+					string? summaryText;
+					if (_workspaceState.LastSyncChangeNumber == -1 || _workspaceState.LastSyncChangeNumber != change.Number || !GetLastUpdateMessage(_workspaceState.LastSyncResult, _workspaceState.LastSyncResultMessage, out summaryText))
 					{
-						StringBuilder SummaryTextBuilder = new StringBuilder();
+						StringBuilder summaryTextBuilder = new StringBuilder();
 
-						EventSummary? Summary = EventMonitor.GetSummaryForChange(Change.Number);
+						EventSummary? summary = _eventMonitor.GetSummaryForChange(change.Number);
 
-						AppendItemList(SummaryTextBuilder, " ", "Under investigation by {0}.", EventMonitor.GetInvestigatingUsers(Change.Number).Select(x => FormatUserName(x)));
+						AppendItemList(summaryTextBuilder, " ", "Under investigation by {0}.", _eventMonitor.GetInvestigatingUsers(change.Number).Select(x => FormatUserName(x)));
 
-						if (Summary != null)
+						if (summary != null)
 						{
-							string Comments = String.Join(", ", Summary.Comments.Where(x => !String.IsNullOrWhiteSpace(x.Text)).Select(x => String.Format("{0}: \"{1}\"", FormatUserName(x.UserName), x.Text)));
-							if (Comments.Length > 0)
+							string comments = String.Join(", ", summary.Comments.Where(x => !String.IsNullOrWhiteSpace(x.Text)).Select(x => String.Format("{0}: \"{1}\"", FormatUserName(x.UserName), x.Text)));
+							if (comments.Length > 0)
 							{
-								SummaryTextBuilder.Append(((SummaryTextBuilder.Length == 0) ? "" : " ") + Comments);
+								summaryTextBuilder.Append(((summaryTextBuilder.Length == 0) ? "" : " ") + comments);
 							}
 							else
 							{
-								AppendItemList(SummaryTextBuilder, " ", "Used by {0}.", Summary.CurrentUsers.Select(x => FormatUserName(x)));
+								AppendItemList(summaryTextBuilder, " ", "Used by {0}.", summary.CurrentUsers.Select(x => FormatUserName(x)));
 							}
 						}
 
-						SummaryText = (SummaryTextBuilder.Length == 0) ? "No current users." : SummaryTextBuilder.ToString();
+						summaryText = (summaryTextBuilder.Length == 0) ? "No current users." : summaryTextBuilder.ToString();
 					}
 
-					if (SummaryText != null && SummaryText.Contains('\n'))
+					if (summaryText != null && summaryText.Contains('\n'))
 					{
-						SummaryText = SummaryText.Substring(0, SummaryText.IndexOf('\n')).TrimEnd() + "...";
+						summaryText = summaryText.Substring(0, summaryText.IndexOf('\n')).TrimEnd() + "...";
 					}
 
-					TextRenderer.DrawText(e.Graphics, SummaryText, CurrentFont, new Rectangle(e.Bounds.Left, e.Bounds.Top, (int)MaxX - e.Bounds.Left, e.Bounds.Height), TextColor, TextFormatFlags.Left | TextFormatFlags.VerticalCenter | TextFormatFlags.EndEllipsis | TextFormatFlags.NoPrefix);
+					TextRenderer.DrawText(e.Graphics, summaryText, currentFont, new Rectangle(e.Bounds.Left, e.Bounds.Top, (int)maxX - e.Bounds.Left, e.Bounds.Height), textColor, TextFormatFlags.Left | TextFormatFlags.VerticalCenter | TextFormatFlags.EndEllipsis | TextFormatFlags.NoPrefix);
 				}
 			}
 			else
 			{
-				ColumnHeader Column = BuildList.Columns[e.ColumnIndex];
-				if (CustomColumns.Contains(Column))
+				ColumnHeader column = BuildList.Columns[e.ColumnIndex];
+				if (_customColumns.Contains(column))
 				{
-					ChangeLayoutInfo Layout = GetChangeLayoutInfo(Change);
+					ChangeLayoutInfo layout = GetChangeLayoutInfo(change);
 
-					List<BadgeInfo>? Badges;
-					if (Layout.CustomBadges.TryGetValue(Column.Text, out Badges) && Badges.Count > 0)
+					List<BadgeInfo>? badges;
+					if (layout.CustomBadges.TryGetValue(column.Text, out badges) && badges.Count > 0)
 					{
 						e.Graphics.IntersectClip(e.Bounds);
 						e.Graphics.SmoothingMode = SmoothingMode.HighQuality;
 
-						Point BuildsLocation = GetBadgeListLocation(Badges, e.Bounds, HorizontalAlign.Center, VerticalAlignment.Middle);
-						DrawBadges(e.Graphics, BuildsLocation, Badges, BadgeAlpha);
+						Point buildsLocation = GetBadgeListLocation(badges, e.Bounds, HorizontalAlign.Center, VerticalAlignment.Middle);
+						DrawBadges(e.Graphics, buildsLocation, badges, badgeAlpha);
 					}
 				}
 			}
 		}
 
-		private void LayoutBadges(List<BadgeInfo> Badges)
+		private void LayoutBadges(List<BadgeInfo> badges)
 		{
-			int Offset = 0;
-			for (int Idx = 0; Idx < Badges.Count; Idx++)
+			int offset = 0;
+			for (int idx = 0; idx < badges.Count; idx++)
 			{
-				BadgeInfo Badge = Badges[Idx];
+				BadgeInfo badge = badges[idx];
 
-				if (Idx > 0 && Badge.Group != Badges[Idx - 1].Group)
+				if (idx > 0 && badge.Group != badges[idx - 1].Group)
 				{
-					Offset += 6;
+					offset += 6;
 				}
-				Badge.Offset = Offset;
+				badge.Offset = offset;
 
-				Size BadgeSize = GetBadgeSize(Badge.Label);
-				Badge.Width = BadgeSize.Width;
-				Badge.Height = BadgeSize.Height;
+				Size badgeSize = GetBadgeSize(badge.Label);
+				badge.Width = badgeSize.Width;
+				badge.Height = badgeSize.Height;
 
-				Offset += BadgeSize.Width + 1;
+				offset += badgeSize.Width + 1;
 			}
 		}
 
-		private Point GetBadgeListLocation(List<BadgeInfo> Badges, Rectangle Bounds, HorizontalAlign HorzAlign, VerticalAlignment VertAlign)
+		private Point GetBadgeListLocation(List<BadgeInfo> badges, Rectangle bounds, HorizontalAlign horzAlign, VerticalAlignment vertAlign)
 		{
-			Point Location = Bounds.Location;
-			if (Badges.Count == 0)
+			Point location = bounds.Location;
+			if (badges.Count == 0)
 			{
-				return Bounds.Location;
+				return bounds.Location;
 			}
 
-			BadgeInfo LastBadge = Badges[Badges.Count - 1];
+			BadgeInfo lastBadge = badges[badges.Count - 1];
 
-			int X = Bounds.X;
-			switch (HorzAlign)
+			int x = bounds.X;
+			switch (horzAlign)
 			{
 				case HorizontalAlign.Center:
-					X += (Bounds.Width - LastBadge.Width - LastBadge.Offset) / 2;
+					x += (bounds.Width - lastBadge.Width - lastBadge.Offset) / 2;
 					break;
 				case HorizontalAlign.Right:
-					X += (Bounds.Width - LastBadge.Width - LastBadge.Offset);
+					x += (bounds.Width - lastBadge.Width - lastBadge.Offset);
 					break;
 			}
 
-			int Y = Bounds.Y;
-			switch (VertAlign)
+			int y = bounds.Y;
+			switch (vertAlign)
 			{
 				case VerticalAlignment.Middle:
-					Y += (Bounds.Height - LastBadge.Height) / 2;
+					y += (bounds.Height - lastBadge.Height) / 2;
 					break;
 				case VerticalAlignment.Bottom:
-					Y += (Bounds.Height - LastBadge.Height);
+					y += (bounds.Height - lastBadge.Height);
 					break;
 			}
 
-			return new Point(X, Y);
+			return new Point(x, y);
 		}
 
-		private void DrawBadges(Graphics Graphics, Point ListLocation, List<BadgeInfo> Badges, int Alpha)
+		private void DrawBadges(Graphics graphics, Point listLocation, List<BadgeInfo> badges, int alpha)
 		{
-			for (int Idx = 0; Idx < Badges.Count; Idx++)
+			for (int idx = 0; idx < badges.Count; idx++)
 			{
-				BadgeInfo Badge = Badges[Idx];
-				bool bMergeLeft = (Idx > 0 && Badges[Idx - 1].Group == Badge.Group);
-				bool bMergeRight = (Idx < Badges.Count - 1 && Badges[Idx + 1].Group == Badge.Group);
+				BadgeInfo badge = badges[idx];
+				bool mergeLeft = (idx > 0 && badges[idx - 1].Group == badge.Group);
+				bool mergeRight = (idx < badges.Count - 1 && badges[idx + 1].Group == badge.Group);
 
-				Rectangle Bounds = new Rectangle(ListLocation.X + Badge.Offset, ListLocation.Y, Badge.Width, Badge.Height);
-				if (Badge.UniqueId != null && Badge.UniqueId == HoverBadgeUniqueId)
+				Rectangle bounds = new Rectangle(listLocation.X + badge.Offset, listLocation.Y, badge.Width, badge.Height);
+				if (badge.UniqueId != null && badge.UniqueId == _hoverBadgeUniqueId)
 				{
-					DrawBadge(Graphics, Bounds, Badge.Label, Color.FromArgb((Alpha * Badge.HoverBackgroundColor.A) / 255, Badge.HoverBackgroundColor), bMergeLeft, bMergeRight);
+					DrawBadge(graphics, bounds, badge.Label, Color.FromArgb((alpha * badge.HoverBackgroundColor.A) / 255, badge.HoverBackgroundColor), mergeLeft, mergeRight);
 				}
 				else
 				{
-					DrawBadge(Graphics, Bounds, Badge.Label, Color.FromArgb((Alpha * Badge.BackgroundColor.A) / 255, Badge.BackgroundColor), bMergeLeft, bMergeRight);
+					DrawBadge(graphics, bounds, badge.Label, Color.FromArgb((alpha * badge.BackgroundColor.A) / 255, badge.BackgroundColor), mergeLeft, mergeRight);
 				}
 			}
 		}
 
-		private List<BadgeInfo> CreateDescriptionBadges(ChangesRecord Change)
+		private List<BadgeInfo> CreateDescriptionBadges(ChangesRecord change)
 		{
-			string Description = Change.Description ?? String.Empty;
+			string description = change.Description ?? String.Empty;
 
-			PerforceChangeDetails? Details;
-			if (PerforceMonitor.TryGetChangeDetails(Change.Number, out Details))
+			PerforceChangeDetails? details;
+			if (_perforceMonitor.TryGetChangeDetails(change.Number, out details))
 			{
-				Description = Details.Description;
+				description = details.Description;
 			}
 
-			List<BadgeInfo> Badges = new List<BadgeInfo>();
+			List<BadgeInfo> badges = new List<BadgeInfo>();
 
 			try
 			{
-				ConfigFile ProjectConfigFile = PerforceMonitor.LatestProjectConfigFile;
-				if (ProjectConfigFile != null)
+				ConfigFile projectConfigFile = _perforceMonitor.LatestProjectConfigFile;
+				if (projectConfigFile != null)
 				{
-					string[] BadgeDefinitions = ProjectConfigFile.GetValues("Badges.DescriptionBadges", new string[0]);
-					foreach (string BadgeDefinition in BadgeDefinitions.Distinct())
+					string[] badgeDefinitions = projectConfigFile.GetValues("Badges.DescriptionBadges", new string[0]);
+					foreach (string badgeDefinition in badgeDefinitions.Distinct())
 					{
-						ConfigObject BadgeDefinitionObject = new ConfigObject(BadgeDefinition);
-						string? Pattern = BadgeDefinitionObject.GetValue("Pattern", null);
-						string? Name = BadgeDefinitionObject.GetValue("Name", null);
-						string? Group = BadgeDefinitionObject.GetValue("Group", null);
-						string Color = BadgeDefinitionObject.GetValue("Color", "#909090");
-						string HoverColor = BadgeDefinitionObject.GetValue("HoverColor", "#b0b0b0");
-						string? Url = BadgeDefinitionObject.GetValue("Url", null);
-						string? Arguments = BadgeDefinitionObject.GetValue("Arguments", null);
-						if (!String.IsNullOrEmpty(Name) && !String.IsNullOrEmpty(Pattern))
+						ConfigObject badgeDefinitionObject = new ConfigObject(badgeDefinition);
+						string? pattern = badgeDefinitionObject.GetValue("Pattern", null);
+						string? name = badgeDefinitionObject.GetValue("Name", null);
+						string? group = badgeDefinitionObject.GetValue("Group", null);
+						string color = badgeDefinitionObject.GetValue("Color", "#909090");
+						string hoverColor = badgeDefinitionObject.GetValue("HoverColor", "#b0b0b0");
+						string? url = badgeDefinitionObject.GetValue("Url", null);
+						string? arguments = badgeDefinitionObject.GetValue("Arguments", null);
+						if (!String.IsNullOrEmpty(name) && !String.IsNullOrEmpty(pattern))
 						{
-							foreach (Match? MatchResult in Regex.Matches(Description, Pattern, RegexOptions.Multiline))
+							foreach (Match? matchResult in Regex.Matches(description, pattern, RegexOptions.Multiline))
 							{
-								if (MatchResult != null)
+								if (matchResult != null)
 								{
-									Color BadgeColor = System.Drawing.ColorTranslator.FromHtml(Color);
-									Color HoverBadgeColor = System.Drawing.ColorTranslator.FromHtml(HoverColor);
+									Color badgeColor = System.Drawing.ColorTranslator.FromHtml(color);
+									Color hoverBadgeColor = System.Drawing.ColorTranslator.FromHtml(hoverColor);
 
-									string? UniqueId = String.IsNullOrEmpty(Url) ? null : String.Format("Description:{0}:{1}", Change.Number, Badges.Count);
+									string? uniqueId = String.IsNullOrEmpty(url) ? null : String.Format("Description:{0}:{1}", change.Number, badges.Count);
 
-									string? ExpandedUrl = ReplaceRegexMatches(Url, MatchResult);
-									string? ExpandedArguments = ReplaceRegexMatches(Arguments, MatchResult);
+									string? expandedUrl = ReplaceRegexMatches(url, matchResult);
+									string? expandedArguments = ReplaceRegexMatches(arguments, matchResult);
 
-									Action? ClickHandler;
-									if (String.IsNullOrEmpty(ExpandedUrl))
+									Action? clickHandler;
+									if (String.IsNullOrEmpty(expandedUrl))
 									{
-										ClickHandler = null;
+										clickHandler = null;
 									}
-									else if (String.IsNullOrEmpty(ExpandedArguments))
+									else if (String.IsNullOrEmpty(expandedArguments))
 									{
-										ClickHandler = () => SafeProcessStart(ExpandedUrl);
+										clickHandler = () => SafeProcessStart(expandedUrl);
 									}
 									else
 									{
-										ClickHandler = () => SafeProcessStart(ExpandedUrl, ExpandedArguments);
+										clickHandler = () => SafeProcessStart(expandedUrl, expandedArguments);
 									}
 
-									Badges.Add(new BadgeInfo(ReplaceRegexMatches(Name, MatchResult), Group, UniqueId, BadgeColor, HoverBadgeColor, ClickHandler));
+									badges.Add(new BadgeInfo(ReplaceRegexMatches(name, matchResult), group, uniqueId, badgeColor, hoverBadgeColor, clickHandler));
 								}
 							}
 						}
@@ -2554,17 +2554,17 @@ namespace UnrealGameSync
 			{
 			}
 
-			LayoutBadges(Badges);
+			LayoutBadges(badges);
 
-			return Badges;
+			return badges;
 		}
 
-		private void PrintLatestChangeTypeUsage(string ErroringDefinition)
+		private void PrintLatestChangeTypeUsage(string erroringDefinition)
 		{
-			string ErrorMessage = "Error: Name not set for config value in \"Sync.LatestChangeType\"";
-			string UsageMessage = "Expected Format under the [Sync] ini category: +LatestChangeType=(Name=[string], Description=[string], OrderIndex=[int], bGood=[bool], bStarred=[bool], RequiredBadges=\"List of Badges For Changelist to have\")";
+			string errorMessage = "Error: Name not set for config value in \"Sync.LatestChangeType\"";
+			string usageMessage = "Expected Format under the [Sync] ini category: +LatestChangeType=(Name=[string], Description=[string], OrderIndex=[int], bGood=[bool], bStarred=[bool], RequiredBadges=\"List of Badges For Changelist to have\")";
 
-			ShowErrorDialog("{0} Erroring String({1}) \n {2}", ErrorMessage, ErroringDefinition, UsageMessage);
+			ShowErrorDialog("{0} Erroring String({1}) \n {2}", errorMessage, erroringDefinition, usageMessage);
 		}
 
 
@@ -2574,89 +2574,89 @@ namespace UnrealGameSync
 		 */
 		public List<LatestChangeType> GetCustomLatestChangeTypes()
 		{
-			List<LatestChangeType> FoundLatestChangeTypes = new List<LatestChangeType>();
+			List<LatestChangeType> foundLatestChangeTypes = new List<LatestChangeType>();
 
 			// Add three default change types that can be synced to
-			FoundLatestChangeTypes.Add(LatestChangeType.LatestChange());
-			FoundLatestChangeTypes.Add(LatestChangeType.LatestGoodChange());
-			FoundLatestChangeTypes.Add(LatestChangeType.LatestStarredChange());
+			foundLatestChangeTypes.Add(LatestChangeType.LatestChange());
+			foundLatestChangeTypes.Add(LatestChangeType.LatestGoodChange());
+			foundLatestChangeTypes.Add(LatestChangeType.LatestStarredChange());
 
-			string[] LatestChangeTypeDefinitions = PerforceMonitor.LatestProjectConfigFile.GetValues("Sync.LatestChangeType", new string[0]);
-			foreach (string LatestChangeDefinition in LatestChangeTypeDefinitions.Distinct())
+			string[] latestChangeTypeDefinitions = _perforceMonitor.LatestProjectConfigFile.GetValues("Sync.LatestChangeType", new string[0]);
+			foreach (string latestChangeDefinition in latestChangeTypeDefinitions.Distinct())
 			{
-				LatestChangeType? NewType;
-				if (LatestChangeType.TryParseConfigEntry(LatestChangeDefinition, out NewType))
+				LatestChangeType? newType;
+				if (LatestChangeType.TryParseConfigEntry(latestChangeDefinition, out newType))
 				{
-					FoundLatestChangeTypes.Add(NewType);
+					foundLatestChangeTypes.Add(newType);
 				}
 				else
 				{
-					PrintLatestChangeTypeUsage(LatestChangeDefinition);
+					PrintLatestChangeTypeUsage(latestChangeDefinition);
 				}
 			}
 
-			FoundLatestChangeTypes = FoundLatestChangeTypes.OrderBy(X => X.OrderIndex).ToList();
+			foundLatestChangeTypes = foundLatestChangeTypes.OrderBy(x => x.OrderIndex).ToList();
 
-			return FoundLatestChangeTypes;
+			return foundLatestChangeTypes;
 		}
 
-		private static void SafeProcessStart(ProcessStartInfo StartInfo)
+		private static void SafeProcessStart(ProcessStartInfo startInfo)
 		{
 			try
 			{
-				Process.Start(StartInfo);
+				Process.Start(startInfo);
 			}
 			catch
 			{
-				MessageBox.Show(String.Format("Unable to open {0}", StartInfo.FileName));
+				MessageBox.Show(String.Format("Unable to open {0}", startInfo.FileName));
 			}
 		}
 
-		private static void SafeProcessStart(string Url)
+		private static void SafeProcessStart(string url)
 		{
 			try
 			{
-				ProcessStartInfo StartInfo = new ProcessStartInfo();
-				StartInfo.FileName = Url;
-				StartInfo.UseShellExecute = true;
-				using Process _ = Process.Start(StartInfo);
+				ProcessStartInfo startInfo = new ProcessStartInfo();
+				startInfo.FileName = url;
+				startInfo.UseShellExecute = true;
+				using Process _ = Process.Start(startInfo);
 			}
 			catch
 			{
-				MessageBox.Show(String.Format("Unable to open {0}", Url));
+				MessageBox.Show(String.Format("Unable to open {0}", url));
 			}
 		}
 
-		private static void SafeProcessStart(string Url, string Arguments)
+		private static void SafeProcessStart(string url, string arguments)
 		{
 			try
 			{
-				ProcessStartInfo StartInfo = new ProcessStartInfo();
-				StartInfo.FileName = Url;
-				StartInfo.Arguments = Arguments;
-				StartInfo.UseShellExecute = true;
-				using Process _ = Process.Start(StartInfo);
+				ProcessStartInfo startInfo = new ProcessStartInfo();
+				startInfo.FileName = url;
+				startInfo.Arguments = arguments;
+				startInfo.UseShellExecute = true;
+				using Process _ = Process.Start(startInfo);
 			}
 			catch
 			{
-				MessageBox.Show(String.Format("Unable to open {0} {1}", Url, Arguments));
+				MessageBox.Show(String.Format("Unable to open {0} {1}", url, arguments));
 			}
 		}
 
-		private static void SafeProcessStart(string Url, string Arguments, string WorkingDir)
+		private static void SafeProcessStart(string url, string arguments, string workingDir)
 		{
 			try
 			{
-				ProcessStartInfo StartInfo = new ProcessStartInfo();
-				StartInfo.FileName = Url;
-				StartInfo.Arguments = Arguments;
-				StartInfo.WorkingDirectory = WorkingDir;
-				StartInfo.UseShellExecute = true;
-				using Process _ = Process.Start(StartInfo);
+				ProcessStartInfo startInfo = new ProcessStartInfo();
+				startInfo.FileName = url;
+				startInfo.Arguments = arguments;
+				startInfo.WorkingDirectory = workingDir;
+				startInfo.UseShellExecute = true;
+				using Process _ = Process.Start(startInfo);
 			}
 			catch
 			{
-				MessageBox.Show(String.Format("Unable to open {0} {1}", Url, Arguments));
+				MessageBox.Show(String.Format("Unable to open {0} {1}", url, arguments));
 			}
 		}
 
@@ -2666,189 +2666,189 @@ namespace UnrealGameSync
 			UpdateStatusPanel();
 		}
 
-		[return: NotNullIfNotNull("Text")]
-		private string? ReplaceRegexMatches(string? Text, Match MatchResult)
+		[return: NotNullIfNotNull("text")]
+		private string? ReplaceRegexMatches(string? text, Match matchResult)
 		{
-			if (Text != null)
+			if (text != null)
 			{
-				for (int Idx = 1; Idx < MatchResult.Groups.Count; Idx++)
+				for (int idx = 1; idx < matchResult.Groups.Count; idx++)
 				{
-					string SourceText = String.Format("${0}", Idx);
-					Text = Text.Replace(SourceText, MatchResult.Groups[Idx].Value);
+					string sourceText = String.Format("${0}", idx);
+					text = text.Replace(sourceText, matchResult.Groups[idx].Value);
 				}
 			}
-			return Text;
+			return text;
 		}
 
-		private List<BadgeInfo> CreateTypeBadges(int ChangeNumber)
+		private List<BadgeInfo> CreateTypeBadges(int changeNumber)
 		{
-			List<BadgeInfo> Badges = new List<BadgeInfo>();
+			List<BadgeInfo> badges = new List<BadgeInfo>();
 
-			PerforceChangeDetails? Details;
-			if (PerforceMonitor.TryGetChangeDetails(ChangeNumber, out Details))
+			PerforceChangeDetails? details;
+			if (_perforceMonitor.TryGetChangeDetails(changeNumber, out details))
 			{
-				if (Details.bContainsCode)
+				if (details.ContainsCode)
 				{
-					Badges.Add(new BadgeInfo("Code", "ChangeType", Color.FromArgb(116, 185, 255)));
+					badges.Add(new BadgeInfo("Code", "ChangeType", Color.FromArgb(116, 185, 255)));
 				}
-				if (Details.bContainsContent)
+				if (details.ContainsContent)
 				{
-					Badges.Add(new BadgeInfo("Content", "ChangeType", Color.FromArgb(162, 155, 255)));
+					badges.Add(new BadgeInfo("Content", "ChangeType", Color.FromArgb(162, 155, 255)));
 				}
 			}
-			if (Badges.Count == 0)
+			if (badges.Count == 0)
 			{
-				Badges.Add(new BadgeInfo("Unknown", "ChangeType", Color.FromArgb(192, 192, 192)));
+				badges.Add(new BadgeInfo("Unknown", "ChangeType", Color.FromArgb(192, 192, 192)));
 			}
-			LayoutBadges(Badges);
+			LayoutBadges(badges);
 
-			return Badges;
+			return badges;
 		}
 
-		private bool TryGetProjectSetting(ConfigFile ProjectConfigFile, string Name, [NotNullWhen(true)] out string? Value)
+		private bool TryGetProjectSetting(ConfigFile projectConfigFile, string name, [NotNullWhen(true)] out string? value)
 		{
-			return ConfigUtils.TryGetProjectSetting(ProjectConfigFile, SelectedProjectIdentifier, Name, out Value);
+			return ConfigUtils.TryGetProjectSetting(projectConfigFile, SelectedProjectIdentifier, name, out value);
 		}
 
-		private bool TryGetProjectSetting(ConfigFile ProjectConfigFile, string Name, [NotNullWhen(true)] out string[]? Values)
+		private bool TryGetProjectSetting(ConfigFile projectConfigFile, string name, [NotNullWhen(true)] out string[]? values)
 		{
-			string? ValueList;
-			if (TryGetProjectSetting(ProjectConfigFile, Name, out ValueList))
+			string? valueList;
+			if (TryGetProjectSetting(projectConfigFile, name, out valueList))
 			{
-				Values = ValueList.Split('\n').Select(x => x.Trim()).Where(x => x.Length > 0).ToArray();
+				values = valueList.Split('\n').Select(x => x.Trim()).Where(x => x.Length > 0).ToArray();
 				return true;
 			}
 			else
 			{
-				Values = null;
+				values = null;
 				return false;
 			}
 		}
 
-		private bool TryGetProjectSetting(ConfigFile ProjectConfigFile, string Name, string LegacyName, [NotNullWhen(true)] out string? Value)
+		private bool TryGetProjectSetting(ConfigFile projectConfigFile, string name, string legacyName, [NotNullWhen(true)] out string? value)
 		{
-			string? NewValue;
-			if (TryGetProjectSetting(ProjectConfigFile, Name, out NewValue))
+			string? newValue;
+			if (TryGetProjectSetting(projectConfigFile, name, out newValue))
 			{
-				Value = NewValue;
+				value = newValue;
 				return true;
 			}
 
-			NewValue = ProjectConfigFile.GetValue(LegacyName, null);
-			if (NewValue != null)
+			newValue = projectConfigFile.GetValue(legacyName, null);
+			if (newValue != null)
 			{
-				Value = NewValue;
+				value = newValue;
 				return true;
 			}
 
-			Value = null;
+			value = null;
 			return false;
 		}
 
-		private List<BadgeInfo> CreateBuildBadges(int ChangeNumber, EventSummary? Summary)
+		private List<BadgeInfo> CreateBuildBadges(int changeNumber, EventSummary? summary)
 		{
-			List<BadgeInfo> Badges = new List<BadgeInfo>();
+			List<BadgeInfo> badges = new List<BadgeInfo>();
 
-			if (Summary != null && Summary.Badges.Count > 0)
+			if (summary != null && summary.Badges.Count > 0)
 			{
 				// Create a lookup for build data for each badge name
-				Dictionary<string, BadgeData> BadgeNameToBuildData = new Dictionary<string, BadgeData>();
-				foreach (BadgeData Badge in Summary.Badges)
+				Dictionary<string, BadgeData> badgeNameToBuildData = new Dictionary<string, BadgeData>();
+				foreach (BadgeData badge in summary.Badges)
 				{
-					BadgeNameToBuildData[Badge.BadgeName] = Badge;
+					badgeNameToBuildData[badge.BadgeName] = badge;
 				}
 
 				// Add all the badges, sorted by group
-				foreach (KeyValuePair<string, string> BadgeNameAndGroup in BadgeNameAndGroupPairs)
+				foreach (KeyValuePair<string, string> badgeNameAndGroup in _badgeNameAndGroupPairs)
 				{
-					BadgeData? BadgeData;
-					BadgeNameToBuildData.TryGetValue(BadgeNameAndGroup.Key, out BadgeData);
+					BadgeData? badgeData;
+					badgeNameToBuildData.TryGetValue(badgeNameAndGroup.Key, out badgeData);
 
-					BadgeInfo BadgeInfo = CreateBadge(ChangeNumber, BadgeNameAndGroup.Key, BadgeNameAndGroup.Value, BadgeData);
-					if (MaxBuildBadgeChars != -1 && BadgeInfo.Label.Length > MaxBuildBadgeChars)
+					BadgeInfo badgeInfo = CreateBadge(changeNumber, badgeNameAndGroup.Key, badgeNameAndGroup.Value, badgeData);
+					if (_maxBuildBadgeChars != -1 && badgeInfo.Label.Length > _maxBuildBadgeChars)
 					{
-						BadgeInfo.ToolTip = BadgeInfo.Label;
-						BadgeInfo.Label = BadgeInfo.Label.Substring(0, MaxBuildBadgeChars);
+						badgeInfo.ToolTip = badgeInfo.Label;
+						badgeInfo.Label = badgeInfo.Label.Substring(0, _maxBuildBadgeChars);
 					}
-					Badges.Add(BadgeInfo);
+					badges.Add(badgeInfo);
 				}
 			}
 
 			// Layout the badges
-			LayoutBadges(Badges);
-			return Badges;
+			LayoutBadges(badges);
+			return badges;
 		}
 
-		private BadgeInfo CreateBadge(int ChangeNumber, string BadgeName, string BadgeGroup, BadgeData? BadgeData)
+		private BadgeInfo CreateBadge(int changeNumber, string badgeName, string badgeGroup, BadgeData? badgeData)
 		{
-			string BadgeLabel = BadgeName;
-			Color BadgeColor = Color.FromArgb(0, Color.White);
+			string badgeLabel = badgeName;
+			Color badgeColor = Color.FromArgb(0, Color.White);
 
-			if (BadgeData != null)
+			if (badgeData != null)
 			{
-				BadgeLabel = BadgeData.BadgeLabel;
-				BadgeColor = GetBuildBadgeColor(BadgeData.Result);
+				badgeLabel = badgeData.BadgeLabel;
+				badgeColor = GetBuildBadgeColor(badgeData.Result);
 			}
 
-			Color HoverBadgeColor = Color.FromArgb(BadgeColor.A, Math.Min(BadgeColor.R + 32, 255), Math.Min(BadgeColor.G + 32, 255), Math.Min(BadgeColor.B + 32, 255));
+			Color hoverBadgeColor = Color.FromArgb(badgeColor.A, Math.Min(badgeColor.R + 32, 255), Math.Min(badgeColor.G + 32, 255), Math.Min(badgeColor.B + 32, 255));
 
-			Action? ClickHandler;
-			if (BadgeData == null || String.IsNullOrEmpty(BadgeData.Url))
+			Action? clickHandler;
+			if (badgeData == null || String.IsNullOrEmpty(badgeData.Url))
 			{
-				ClickHandler = null;
+				clickHandler = null;
 			}
 			else
 			{
-				ClickHandler = () => SafeProcessStart(BadgeData.Url);
+				clickHandler = () => SafeProcessStart(badgeData.Url);
 			}
 
-			string UniqueId = String.Format("{0}:{1}", ChangeNumber, BadgeName);
-			return new BadgeInfo(BadgeLabel, BadgeGroup, UniqueId, BadgeColor, HoverBadgeColor, ClickHandler);
+			string uniqueId = String.Format("{0}:{1}", changeNumber, badgeName);
+			return new BadgeInfo(badgeLabel, badgeGroup, uniqueId, badgeColor, hoverBadgeColor, clickHandler);
 		}
 
-		private Dictionary<string, List<BadgeInfo>> CreateCustomBadges(int ChangeNumber, EventSummary? Summary)
+		private Dictionary<string, List<BadgeInfo>> CreateCustomBadges(int changeNumber, EventSummary? summary)
 		{
-			Dictionary<string, List<BadgeInfo>> ColumnNameToBadges = new Dictionary<string, List<BadgeInfo>>();
-			if (Summary != null && Summary.Badges.Count > 0)
+			Dictionary<string, List<BadgeInfo>> columnNameToBadges = new Dictionary<string, List<BadgeInfo>>();
+			if (summary != null && summary.Badges.Count > 0)
 			{
-				foreach (ColumnHeader CustomColumn in CustomColumns)
+				foreach (ColumnHeader customColumn in _customColumns)
 				{
-					ConfigObject Config = (ConfigObject)CustomColumn.Tag;
-					if (Config != null)
+					ConfigObject config = (ConfigObject)customColumn.Tag;
+					if (config != null)
 					{
-						List<BadgeInfo> Badges = new List<BadgeInfo>();
+						List<BadgeInfo> badges = new List<BadgeInfo>();
 
-						string[] BadgeNames = Config.GetValue("Badges", "").Split(new char[] { ',' }, StringSplitOptions.RemoveEmptyEntries);
-						foreach (string BadgeName in BadgeNames)
+						string[] badgeNames = config.GetValue("Badges", "").Split(new char[] { ',' }, StringSplitOptions.RemoveEmptyEntries);
+						foreach (string badgeName in badgeNames)
 						{
-							BadgeInfo Badge = CreateBadge(ChangeNumber, BadgeName, "XXXX", Summary.Badges.FirstOrDefault(x => x.BadgeName == BadgeName));
-							Badges.Add(Badge);
+							BadgeInfo badge = CreateBadge(changeNumber, badgeName, "XXXX", summary.Badges.FirstOrDefault(x => x.BadgeName == badgeName));
+							badges.Add(badge);
 						}
 
-						LayoutBadges(Badges);
+						LayoutBadges(badges);
 
-						ColumnNameToBadges[CustomColumn.Text] = Badges;
+						columnNameToBadges[customColumn.Text] = badges;
 					}
 				}
 			}
-			return ColumnNameToBadges;
+			return columnNameToBadges;
 		}
 
-		private static Color GetBuildBadgeColor(BadgeResult Result)
+		private static Color GetBuildBadgeColor(BadgeResult result)
 		{
-			if (Result == BadgeResult.Starting)
+			if (result == BadgeResult.Starting)
 			{
 				return Color.FromArgb(128, 192, 255);
 			}
-			else if (Result == BadgeResult.Warning)
+			else if (result == BadgeResult.Warning)
 			{
 				return Color.FromArgb(255, 192, 0);
 			}
-			else if (Result == BadgeResult.Failure)
+			else if (result == BadgeResult.Failure)
 			{
 				return Color.FromArgb(192, 64, 0);
 			}
-			else if (Result == BadgeResult.Skipped)
+			else if (result == BadgeResult.Skipped)
 			{
 				return Color.FromArgb(192, 192, 192);
 			}
@@ -2858,170 +2858,170 @@ namespace UnrealGameSync
 			}
 		}
 
-		private Rectangle GetSyncBadgeRectangle(Rectangle Bounds)
+		private Rectangle GetSyncBadgeRectangle(Rectangle bounds)
 		{
-			Size BadgeSize = GetBadgeSize("Sync");
-			return new Rectangle(Bounds.Right - BadgeSize.Width - 2, Bounds.Top + (Bounds.Height - BadgeSize.Height) / 2, BadgeSize.Width, BadgeSize.Height);
+			Size badgeSize = GetBadgeSize("Sync");
+			return new Rectangle(bounds.Right - badgeSize.Width - 2, bounds.Top + (bounds.Height - badgeSize.Height) / 2, badgeSize.Width, badgeSize.Height);
 		}
 
-		private void DrawSingleBadge(Graphics Graphics, Rectangle DisplayRectangle, string BadgeText, Color BadgeColor)
+		private void DrawSingleBadge(Graphics graphics, Rectangle displayRectangle, string badgeText, Color badgeColor)
 		{
-			Size BadgeSize = GetBadgeSize(BadgeText);
+			Size badgeSize = GetBadgeSize(badgeText);
 
-			int X = DisplayRectangle.Left + (DisplayRectangle.Width - BadgeSize.Width) / 2;
-			int Y = DisplayRectangle.Top + (DisplayRectangle.Height - BadgeSize.Height) / 2;
+			int x = displayRectangle.Left + (displayRectangle.Width - badgeSize.Width) / 2;
+			int y = displayRectangle.Top + (displayRectangle.Height - badgeSize.Height) / 2;
 
-			DrawBadge(Graphics, new Rectangle(X, Y, BadgeSize.Width, BadgeSize.Height), BadgeText, BadgeColor, false, false);
+			DrawBadge(graphics, new Rectangle(x, y, badgeSize.Width, badgeSize.Height), badgeText, badgeColor, false, false);
 		}
 
-		private void DrawBadge(Graphics Graphics, Rectangle BadgeRect, string BadgeText, Color BadgeColor, bool bMergeLeft, bool bMergeRight)
+		private void DrawBadge(Graphics graphics, Rectangle badgeRect, string badgeText, Color badgeColor, bool mergeLeft, bool mergeRight)
 		{
-			if (BadgeColor.A != 0)
+			if (badgeColor.A != 0)
 			{
-				using (GraphicsPath Path = new GraphicsPath())
+				using (GraphicsPath path = new GraphicsPath())
 				{
-					Path.StartFigure();
-					Path.AddLine(BadgeRect.Left + (bMergeLeft ? 1 : 0), BadgeRect.Top, BadgeRect.Left - (bMergeLeft ? 1 : 0), BadgeRect.Bottom);
-					Path.AddLine(BadgeRect.Left - (bMergeLeft ? 1 : 0), BadgeRect.Bottom, BadgeRect.Right - 1 - (bMergeRight ? 1 : 0), BadgeRect.Bottom);
-					Path.AddLine(BadgeRect.Right - 1 - (bMergeRight ? 1 : 0), BadgeRect.Bottom, BadgeRect.Right - 1 + (bMergeRight ? 1 : 0), BadgeRect.Top);
-					Path.AddLine(BadgeRect.Right - 1 + (bMergeRight ? 1 : 0), BadgeRect.Top, BadgeRect.Left + (bMergeLeft ? 1 : 0), BadgeRect.Top);
-					Path.CloseFigure();
+					path.StartFigure();
+					path.AddLine(badgeRect.Left + (mergeLeft ? 1 : 0), badgeRect.Top, badgeRect.Left - (mergeLeft ? 1 : 0), badgeRect.Bottom);
+					path.AddLine(badgeRect.Left - (mergeLeft ? 1 : 0), badgeRect.Bottom, badgeRect.Right - 1 - (mergeRight ? 1 : 0), badgeRect.Bottom);
+					path.AddLine(badgeRect.Right - 1 - (mergeRight ? 1 : 0), badgeRect.Bottom, badgeRect.Right - 1 + (mergeRight ? 1 : 0), badgeRect.Top);
+					path.AddLine(badgeRect.Right - 1 + (mergeRight ? 1 : 0), badgeRect.Top, badgeRect.Left + (mergeLeft ? 1 : 0), badgeRect.Top);
+					path.CloseFigure();
 
-					using (SolidBrush Brush = new SolidBrush(BadgeColor))
+					using (SolidBrush brush = new SolidBrush(badgeColor))
 					{
-						Graphics.FillPath(Brush, Path);
+						graphics.FillPath(brush, path);
 					}
 				}
 
-				TextRenderer.DrawText(Graphics, BadgeText, BadgeFont, BadgeRect, Color.White, TextFormatFlags.HorizontalCenter | TextFormatFlags.VerticalCenter | TextFormatFlags.SingleLine | TextFormatFlags.NoPrefix | TextFormatFlags.PreserveGraphicsClipping);
+				TextRenderer.DrawText(graphics, badgeText, _badgeFont, badgeRect, Color.White, TextFormatFlags.HorizontalCenter | TextFormatFlags.VerticalCenter | TextFormatFlags.SingleLine | TextFormatFlags.NoPrefix | TextFormatFlags.PreserveGraphicsClipping);
 			}
 		}
 
-		private Size GetBadgeSize(string BadgeText)
+		private Size GetBadgeSize(string badgeText)
 		{
-			Size BadgeSize;
-			if (!BadgeLabelToSize.TryGetValue(BadgeText, out BadgeSize))
+			Size badgeSize;
+			if (!_badgeLabelToSize.TryGetValue(badgeText, out badgeSize))
 			{
-				Size LabelSize = TextRenderer.MeasureText(BadgeText, BadgeFont);
-				int BadgeHeight = BadgeFont!.Height + 1;
+				Size labelSize = TextRenderer.MeasureText(badgeText, _badgeFont);
+				int badgeHeight = _badgeFont!.Height + 1;
 
-				BadgeSize = new Size(LabelSize.Width + BadgeHeight - 4, BadgeHeight);
-				BadgeLabelToSize[BadgeText] = BadgeSize;
+				badgeSize = new Size(labelSize.Width + badgeHeight - 4, badgeHeight);
+				_badgeLabelToSize[badgeText] = badgeSize;
 			}
-			return BadgeSize;
+			return badgeSize;
 		}
 
-		private static bool GetLastUpdateMessage(WorkspaceUpdateResult Result, string? ResultMessage, [NotNullWhen(true)] out string? Message)
+		private static bool GetLastUpdateMessage(WorkspaceUpdateResult result, string? resultMessage, [NotNullWhen(true)] out string? message)
 		{
-			if (Result != WorkspaceUpdateResult.Success && ResultMessage != null)
+			if (result != WorkspaceUpdateResult.Success && resultMessage != null)
 			{
-				Message = ResultMessage;
+				message = resultMessage;
 				return true;
 			}
-			return GetGenericLastUpdateMessage(Result, out Message);
+			return GetGenericLastUpdateMessage(result, out message);
 		}
 
-		private static bool GetGenericLastUpdateMessage(WorkspaceUpdateResult Result, [NotNullWhen(true)] out string? Message)
+		private static bool GetGenericLastUpdateMessage(WorkspaceUpdateResult result, [NotNullWhen(true)] out string? message)
 		{
-			switch (Result)
+			switch (result)
 			{
 				case WorkspaceUpdateResult.Canceled:
-					Message = "Sync canceled.";
+					message = "Sync canceled.";
 					return true;
 				case WorkspaceUpdateResult.FailedToSync:
-					Message = "Failed to sync files.";
+					message = "Failed to sync files.";
 					return true;
 				case WorkspaceUpdateResult.FailedToSyncLoginExpired:
-					Message = "Failed to sync files (login expired).";
+					message = "Failed to sync files (login expired).";
 					return true;
 				case WorkspaceUpdateResult.FilesToResolve:
-					Message = "Sync finished with unresolved files.";
+					message = "Sync finished with unresolved files.";
 					return true;
 				case WorkspaceUpdateResult.FilesToClobber:
-					Message = "Sync failed due to modified files in workspace.";
+					message = "Sync failed due to modified files in workspace.";
 					return true;
 				case WorkspaceUpdateResult.FilesToDelete:
-					Message = "Sync aborted pending confirmation of files to delete.";
+					message = "Sync aborted pending confirmation of files to delete.";
 					return true;
 				case WorkspaceUpdateResult.FailedToCompile:
 				case WorkspaceUpdateResult.FailedToCompileWithCleanWorkspace:
-					Message = "Compile failed.";
+					message = "Compile failed.";
 					return true;
 				default:
-					Message = null;
+					message = null;
 					return false;
 			}
 		}
 
-		public static string FormatUserName(string UserName)
+		public static string FormatUserName(string userName)
 		{
-			if (UserName == null)
+			if (userName == null)
 			{
 				return "(Invalid Name)";
 			}
 			else
 			{
-				return Utility.FormatUserName(UserName);
+				return Utility.FormatUserName(userName);
 			}
 		}
 
-		private static void AppendUserList(StringBuilder Builder, string Separator, string Format, IEnumerable<EventData> Reviews)
+		private static void AppendUserList(StringBuilder builder, string separator, string format, IEnumerable<EventData> reviews)
 		{
-			AppendItemList(Builder, Separator, Format, Reviews.Select(x => FormatUserName(x.UserName)));
+			AppendItemList(builder, separator, format, reviews.Select(x => FormatUserName(x.UserName)));
 		}
 
-		private static void AppendItemList(StringBuilder Builder, string Separator, string Format, IEnumerable<string> Items)
+		private static void AppendItemList(StringBuilder builder, string separator, string format, IEnumerable<string> items)
 		{
-			string? ItemList = FormatItemList(Format, Items);
-			if (ItemList != null)
+			string? itemList = FormatItemList(format, items);
+			if (itemList != null)
 			{
-				if (Builder.Length > 0)
+				if (builder.Length > 0)
 				{
-					Builder.Append(Separator);
+					builder.Append(separator);
 				}
-				Builder.Append(ItemList);
+				builder.Append(itemList);
 			}
 		}
 
-		private static string? FormatItemList(string Format, IEnumerable<string> Items)
+		private static string? FormatItemList(string format, IEnumerable<string> items)
 		{
-			string[] ItemsArray = Items.Distinct().ToArray();
-			if (ItemsArray.Length == 0)
+			string[] itemsArray = items.Distinct().ToArray();
+			if (itemsArray.Length == 0)
 			{
 				return null;
 			}
 
-			StringBuilder Builder = new StringBuilder(ItemsArray[0]);
-			if (ItemsArray.Length > 1)
+			StringBuilder builder = new StringBuilder(itemsArray[0]);
+			if (itemsArray.Length > 1)
 			{
-				for (int Idx = 1; Idx < ItemsArray.Length - 1; Idx++)
+				for (int idx = 1; idx < itemsArray.Length - 1; idx++)
 				{
-					Builder.Append(", ");
-					Builder.Append(ItemsArray[Idx]);
+					builder.Append(", ");
+					builder.Append(itemsArray[idx]);
 				}
-				Builder.Append(" and ");
-				Builder.Append(ItemsArray.Last());
+				builder.Append(" and ");
+				builder.Append(itemsArray.Last());
 			}
-			return String.Format(Format, Builder.ToString());
+			return String.Format(format, builder.ToString());
 		}
 
-		private void BuildList_MouseDoubleClick(object Sender, MouseEventArgs Args)
+		private void BuildList_MouseDoubleClick(object sender, MouseEventArgs args)
 		{
-			if (Args.Button == MouseButtons.Left)
+			if (args.Button == MouseButtons.Left)
 			{
-				ListViewHitTestInfo HitTest = BuildList.HitTest(Args.Location);
-				if (HitTest.Item != null)
+				ListViewHitTestInfo hitTest = BuildList.HitTest(args.Location);
+				if (hitTest.Item != null)
 				{
-					ChangesRecord Change = (ChangesRecord)HitTest.Item.Tag;
-					if (Change != null)
+					ChangesRecord change = (ChangesRecord)hitTest.Item.Tag;
+					if (change != null)
 					{
-						if (Change.Number == Workspace.CurrentChangeNumber)
+						if (change.Number == _workspace.CurrentChangeNumber)
 						{
 							LaunchEditor();
 						}
 						else
 						{
-							StartSync(Change.Number, false, null);
+							StartSync(change.Number, false, null);
 						}
 					}
 				}
@@ -3030,49 +3030,49 @@ namespace UnrealGameSync
 
 		public void LaunchEditor()
 		{
-			if (!Workspace.IsBusy() && Workspace.CurrentChangeNumber != -1)
+			if (!_workspace.IsBusy() && _workspace.CurrentChangeNumber != -1)
 			{
-				BuildConfig EditorBuildConfig = GetEditorBuildConfig();
+				BuildConfig editorBuildConfig = GetEditorBuildConfig();
 
-				FileReference ReceiptFile = ConfigUtils.GetEditorReceiptFile(Workspace.Project, Workspace.ProjectConfigFile, EditorBuildConfig);
-				if (ConfigUtils.TryReadEditorReceipt(Workspace.Project, ReceiptFile, out TargetReceipt? Receipt) && Receipt.Launch != null && File.Exists(Receipt.Launch))
+				FileReference receiptFile = ConfigUtils.GetEditorReceiptFile(_workspace.Project, _workspace.ProjectConfigFile, editorBuildConfig);
+				if (ConfigUtils.TryReadEditorReceipt(_workspace.Project, receiptFile, out TargetReceipt? receipt) && receipt.Launch != null && File.Exists(receipt.Launch))
 				{
-					if (Settings.bEditorArgumentsPrompt && !ModifyEditorArguments())
+					if (_settings.EditorArgumentsPrompt && !ModifyEditorArguments())
 					{
 						return;
 					}
 
-					StringBuilder LaunchArguments = new StringBuilder();
+					StringBuilder launchArguments = new StringBuilder();
 					if (SelectedFileName.HasExtension(".uproject"))
 					{
-						LaunchArguments.AppendFormat("\"{0}\"", SelectedFileName);
+						launchArguments.AppendFormat("\"{0}\"", SelectedFileName);
 					}
-					foreach (Tuple<string, bool> EditorArgument in Settings.EditorArguments)
+					foreach (Tuple<string, bool> editorArgument in _settings.EditorArguments)
 					{
-						if (EditorArgument.Item2)
+						if (editorArgument.Item2)
 						{
-							LaunchArguments.AppendFormat(" {0}", EditorArgument.Item1);
+							launchArguments.AppendFormat(" {0}", editorArgument.Item1);
 						}
 					}
-					if (EditorBuildConfig == BuildConfig.Debug || EditorBuildConfig == BuildConfig.DebugGame)
+					if (editorBuildConfig == BuildConfig.Debug || editorBuildConfig == BuildConfig.DebugGame)
 					{
-						LaunchArguments.Append(" -debug");
+						launchArguments.Append(" -debug");
 					}
 
-					if (!Utility.SpawnProcess(Receipt.Launch, LaunchArguments.ToString()))
+					if (!Utility.SpawnProcess(receipt.Launch, launchArguments.ToString()))
 					{
-						ShowErrorDialog("Unable to spawn {0} {1}", Receipt.Launch, LaunchArguments.ToString());
+						ShowErrorDialog("Unable to spawn {0} {1}", receipt.Launch, launchArguments.ToString());
 					}
 				}
 				else
 				{
 					if (MessageBox.Show("The editor needs to be built before you can run it. Build it now?", "Editor out of date", MessageBoxButtons.YesNo) == System.Windows.Forms.DialogResult.Yes)
 					{
-						Owner.ShowAndActivate();
+						_owner.ShowAndActivate();
 
-						WorkspaceUpdateOptions Options = WorkspaceUpdateOptions.Build | WorkspaceUpdateOptions.RunAfterSync;
-						WorkspaceUpdateContext Context = new WorkspaceUpdateContext(Workspace.CurrentChangeNumber, Options, Settings.CompiledEditorBuildConfig, null, ProjectSettings.BuildSteps, null);
-						StartWorkspaceUpdate(Context, null);
+						WorkspaceUpdateOptions options = WorkspaceUpdateOptions.Build | WorkspaceUpdateOptions.RunAfterSync;
+						WorkspaceUpdateContext context = new WorkspaceUpdateContext(_workspace.CurrentChangeNumber, options, _settings.CompiledEditorBuildConfig, null, _projectSettings.BuildSteps, null);
+						StartWorkspaceUpdate(context, null);
 					}
 				}
 			}
@@ -3080,11 +3080,11 @@ namespace UnrealGameSync
 
 		private bool WaitForProgramsToFinish()
 		{
-			FileReference[] ProcessFileNames = GetProcessesRunningInWorkspace();
-			if (ProcessFileNames.Length > 0)
+			FileReference[] processFileNames = GetProcessesRunningInWorkspace();
+			if (processFileNames.Length > 0)
 			{
-				ProgramsRunningWindow ProgramsRunning = new ProgramsRunningWindow(GetProcessesRunningInWorkspace, ProcessFileNames);
-				if (ProgramsRunning.ShowDialog() != DialogResult.OK)
+				ProgramsRunningWindow programsRunning = new ProgramsRunningWindow(GetProcessesRunningInWorkspace, processFileNames);
+				if (programsRunning.ShowDialog() != DialogResult.OK)
 				{
 					return false;
 				}
@@ -3094,51 +3094,51 @@ namespace UnrealGameSync
 
 		private FileReference[] GetProcessesRunningInWorkspace()
 		{
-			HashSet<string> ProcessNames = new HashSet<string>(StringComparer.InvariantCultureIgnoreCase);
-			ProcessNames.Add("Win64\\UE4Editor.exe");
-			ProcessNames.Add("Win64\\UE4Editor-Cmd.exe");
-			ProcessNames.Add("Win64\\UE4Editor-Win64-Debug.exe");
-			ProcessNames.Add("Win64\\UE4Editor-Win64-Debug-Cmd.exe");
-			ProcessNames.Add("Win64\\UE4Editor-Win64-DebugGame.exe");
-			ProcessNames.Add("Win64\\UE4Editor-Win64-DebugGame-Cmd.exe");
-			ProcessNames.Add("Win64\\UnrealEditor.exe");
-			ProcessNames.Add("Win64\\UnrealEditor-Cmd.exe");
-			ProcessNames.Add("Win64\\UnrealEditor-Win64-Debug.exe");
-			ProcessNames.Add("Win64\\UnrealEditor-Win64-Debug-Cmd.exe");
-			ProcessNames.Add("Win64\\UnrealEditor-Win64-DebugGame.exe");
-			ProcessNames.Add("Win64\\UnrealEditor-Win64-DebugGame-Cmd.exe");
-			ProcessNames.Add("Win64\\CrashReportClient.exe");
-			ProcessNames.Add("Win64\\CrashReportClient-Win64-Development.exe");
-			ProcessNames.Add("Win64\\CrashReportClient-Win64-Debug.exe");
-			ProcessNames.Add("Win64\\CrashReportClientEditor.exe");
-			ProcessNames.Add("Win64\\CrashReportClientEditor-Win64-Development.exe");
-			ProcessNames.Add("Win64\\CrashReportClientEditor-Win64-Debug.exe");
-			ProcessNames.Add("DotNET\\UnrealBuildTool\\UnrealBuildTool.exe");
-			ProcessNames.Add("DotNET\\AutomationTool\\AutomationTool.exe");
-
-			List<FileReference> ProcessFileNames = new List<FileReference>();
+			HashSet<string> processNames = new HashSet<string>(StringComparer.InvariantCultureIgnoreCase);
+			processNames.Add("Win64\\UE4Editor.exe");
+			processNames.Add("Win64\\UE4Editor-Cmd.exe");
+			processNames.Add("Win64\\UE4Editor-Win64-Debug.exe");
+			processNames.Add("Win64\\UE4Editor-Win64-Debug-Cmd.exe");
+			processNames.Add("Win64\\UE4Editor-Win64-DebugGame.exe");
+			processNames.Add("Win64\\UE4Editor-Win64-DebugGame-Cmd.exe");
+			processNames.Add("Win64\\UnrealEditor.exe");
+			processNames.Add("Win64\\UnrealEditor-Cmd.exe");
+			processNames.Add("Win64\\UnrealEditor-Win64-Debug.exe");
+			processNames.Add("Win64\\UnrealEditor-Win64-Debug-Cmd.exe");
+			processNames.Add("Win64\\UnrealEditor-Win64-DebugGame.exe");
+			processNames.Add("Win64\\UnrealEditor-Win64-DebugGame-Cmd.exe");
+			processNames.Add("Win64\\CrashReportClient.exe");
+			processNames.Add("Win64\\CrashReportClient-Win64-Development.exe");
+			processNames.Add("Win64\\CrashReportClient-Win64-Debug.exe");
+			processNames.Add("Win64\\CrashReportClientEditor.exe");
+			processNames.Add("Win64\\CrashReportClientEditor-Win64-Development.exe");
+			processNames.Add("Win64\\CrashReportClientEditor-Win64-Debug.exe");
+			processNames.Add("DotNET\\UnrealBuildTool\\UnrealBuildTool.exe");
+			processNames.Add("DotNET\\AutomationTool\\AutomationTool.exe");
+
+			List<FileReference> processFileNames = new List<FileReference>();
 			try
 			{
-				DirectoryReference BinariesRootPath = DirectoryReference.Combine(Workspace.Project.LocalRootPath, "Engine\\Binaries");
+				DirectoryReference binariesRootPath = DirectoryReference.Combine(_workspace.Project.LocalRootPath, "Engine\\Binaries");
 
-				foreach (string ProcessName in ProcessNames)
+				foreach (string processName in processNames)
 				{
 					try
 					{
-						FileReference ProcessFilename = FileReference.Combine(BinariesRootPath, ProcessName);
+						FileReference processFilename = FileReference.Combine(binariesRootPath, processName);
 
-						if (FileReference.Exists(ProcessFilename))
+						if (FileReference.Exists(processFilename))
 						{
 							try
 							{
 								// Try to open the file to determine whether the executable is running.
 								// We do this so that we can also find processes running under other user sessions, without needing PROCESS_QUERY_INFORMATION access rights.
-								using FileStream TestStream = File.OpenWrite(ProcessFilename.FullName);
+								using FileStream testStream = File.OpenWrite(processFilename.FullName);
 							}
 							catch (IOException)
 							{
 								// This file is in use, so add it to the list of running processes
-								ProcessFileNames.Add(ProcessFilename);
+								processFileNames.Add(processFilename);
 							}
 						}
 					}
@@ -3150,22 +3150,22 @@ namespace UnrealGameSync
 			catch
 			{
 			}
-			return ProcessFileNames.ToArray();
+			return processFileNames.ToArray();
 		}
 
-		private void TimerCallback(object? Sender, EventArgs Args)
+		private void TimerCallback(object? sender, EventArgs args)
 		{
-			Tuple<string, float> Progress = Workspace.CurrentProgress;
-			if (Progress != null && Progress.Item2 > 0.0f)
+			Tuple<string, float> progress = _workspace.CurrentProgress;
+			if (progress != null && progress.Item2 > 0.0f)
 			{
-				DesiredTaskbarState = Tuple.Create(TaskbarState.Normal, Progress.Item2);
+				DesiredTaskbarState = Tuple.Create(TaskbarState.Normal, progress.Item2);
 			}
 			else
 			{
 				DesiredTaskbarState = Tuple.Create(TaskbarState.NoProgress, 0.0f);
 			}
 
-			Owner.UpdateProgress();
+			_owner.UpdateProgress();
 
 			UpdateStatusPanel();
 			BuildList.Refresh();
@@ -3173,16 +3173,16 @@ namespace UnrealGameSync
 
 		private string? GetDefaultIssueFilter()
 		{
-			string? BuildHealthProject;
-			if (TryGetProjectSetting(PerforceMonitor.LatestProjectConfigFile, "BuildHealthProject", out BuildHealthProject))
+			string? buildHealthProject;
+			if (TryGetProjectSetting(_perforceMonitor.LatestProjectConfigFile, "BuildHealthProject", out buildHealthProject))
 			{
-				return BuildHealthProject;
+				return buildHealthProject;
 			}
 
-			string? DefaultIssueFilter;
-			if (TryGetProjectSetting(PerforceMonitor.LatestProjectConfigFile, "DefaultIssueFilter", out DefaultIssueFilter))
+			string? defaultIssueFilter;
+			if (TryGetProjectSetting(_perforceMonitor.LatestProjectConfigFile, "DefaultIssueFilter", out defaultIssueFilter))
 			{
-				return DefaultIssueFilter;
+				return defaultIssueFilter;
 			}
 
 			return null;
@@ -3190,81 +3190,81 @@ namespace UnrealGameSync
 
 		private Dictionary<string, Func<IssueData, bool>> GetCustomIssueFilters()
 		{
-			Dictionary<string, Func<IssueData, bool>> CustomFilters = new Dictionary<string, Func<IssueData, bool>>();
+			Dictionary<string, Func<IssueData, bool>> customFilters = new Dictionary<string, Func<IssueData, bool>>();
 
-			string? BuildHealthProject;
-			if (TryGetProjectSetting(PerforceMonitor.LatestProjectConfigFile, "BuildHealthProject", out BuildHealthProject))
+			string? buildHealthProject;
+			if (TryGetProjectSetting(_perforceMonitor.LatestProjectConfigFile, "BuildHealthProject", out buildHealthProject))
 			{
-				CustomFilters[BuildHealthProject] = x => !String.IsNullOrEmpty(x.Project) && x.Project.Equals(BuildHealthProject, StringComparison.OrdinalIgnoreCase);
+				customFilters[buildHealthProject] = x => !String.IsNullOrEmpty(x.Project) && x.Project.Equals(buildHealthProject, StringComparison.OrdinalIgnoreCase);
 			}
 
-			string[]? Filters;
-			if (TryGetProjectSetting(PerforceMonitor.LatestProjectConfigFile, "IssueFilters", out Filters))
+			string[]? filters;
+			if (TryGetProjectSetting(_perforceMonitor.LatestProjectConfigFile, "IssueFilters", out filters))
 			{
-				foreach (string Filter in Filters)
+				foreach (string filter in filters)
 				{
 					try
 					{
-						ConfigObject Config = new ConfigObject(Filter);
+						ConfigObject config = new ConfigObject(filter);
 
-						string? Name = Config.GetValue("Name", null);
-						string? Pattern = Config.GetValue("Pattern", null);
+						string? name = config.GetValue("Name", null);
+						string? pattern = config.GetValue("Pattern", null);
 
-						if(Name != null && Pattern != null && StreamName != null)
+						if(name != null && pattern != null && StreamName != null)
 						{
-							string? DepotName;
-							if (PerforceUtils.TryGetDepotName(StreamName, out DepotName))
+							string? depotName;
+							if (PerforceUtils.TryGetDepotName(StreamName, out depotName))
 							{
-								Pattern = Regex.Replace(Pattern, @"\$\(Depot\)", DepotName, RegexOptions.IgnoreCase);
+								pattern = Regex.Replace(pattern, @"\$\(Depot\)", depotName, RegexOptions.IgnoreCase);
 							}
-							Pattern = Regex.Replace(Pattern, @"\$\(Stream\)", StreamName, RegexOptions.IgnoreCase);
+							pattern = Regex.Replace(pattern, @"\$\(Stream\)", StreamName, RegexOptions.IgnoreCase);
 
-							Regex PatternRegex = new Regex(String.Format("^(?:{0})$", Pattern), RegexOptions.IgnoreCase);
-							CustomFilters[Name] = x => x.Streams.Any(y => PatternRegex.IsMatch(y));
+							Regex patternRegex = new Regex(String.Format("^(?:{0})$", pattern), RegexOptions.IgnoreCase);
+							customFilters[name] = x => x.Streams.Any(y => patternRegex.IsMatch(y));
 						}
 					}
-					catch(Exception Ex)
+					catch(Exception ex)
 					{
-						Logger.LogError(Ex, "Unable to parse config filter '{Filter}'", Filter);
+						_logger.LogError(ex, "Unable to parse config filter '{Filter}'", filter);
 					}
 				}
 			}
 
-			return CustomFilters;
+			return customFilters;
 		}
 
 		private void OpenPerforce()
 		{
-			StringBuilder CommandLine = new StringBuilder();
-			if (Workspace != null)
+			StringBuilder commandLine = new StringBuilder();
+			if (_workspace != null)
 			{
-				CommandLine.Append(PerforceSettings.GetArgumentsForExternalProgram(true));
+				commandLine.Append(_perforceSettings.GetArgumentsForExternalProgram(true));
 			}
-			SafeProcessStart("p4v.exe", CommandLine.ToString());
+			SafeProcessStart("p4v.exe", commandLine.ToString());
 		}
 
 
-		void RunTool(ToolDefinition Tool, ToolLink Link)
+		void RunTool(ToolDefinition tool, ToolLink link)
 		{
-			DirectoryReference? ToolDir = Owner.ToolUpdateMonitor.GetToolPath(Tool.Name);
-			if (ToolDir != null)
+			DirectoryReference? toolDir = _owner.ToolUpdateMonitor.GetToolPath(tool.Name);
+			if (toolDir != null)
 			{
-				Dictionary<string, string> Variables = Workspace.GetVariables(GetEditorBuildConfig());
-				Variables["ToolDir"] = ToolDir.FullName;
+				Dictionary<string, string> variables = _workspace.GetVariables(GetEditorBuildConfig());
+				variables["ToolDir"] = toolDir.FullName;
 
-				string FileName = Utility.ExpandVariables(Link.FileName, Variables);
-				string Arguments = Utility.ExpandVariables(Link.Arguments ?? String.Empty, Variables);
-				string WorkingDir = Utility.ExpandVariables(Link.WorkingDir ?? ToolDir.FullName, Variables);
+				string fileName = Utility.ExpandVariables(link.FileName, variables);
+				string arguments = Utility.ExpandVariables(link.Arguments ?? String.Empty, variables);
+				string workingDir = Utility.ExpandVariables(link.WorkingDir ?? toolDir.FullName, variables);
 
-				SafeProcessStart(FileName, Arguments, WorkingDir);
+				SafeProcessStart(fileName, arguments, workingDir);
 			}
 		}
 
 		private void UpdateStatusPanel_CrossThread()
 		{
-			MainThreadSynchronizationContext.Post((o) => 
+			_mainThreadSynchronizationContext.Post((o) => 
 			{ 
-				if (!IsDisposed && !bIsDisposing)
+				if (!IsDisposed && !_isDisposing)
 				{ 
 					UpdateBuildSteps();
 					UpdateStatusPanel(); 
@@ -3274,533 +3274,533 @@ namespace UnrealGameSync
 
 		private void UpdateStatusPanel()
 		{
-			if(Workspace == null)
+			if(_workspace == null)
 			{
 				return;
 			}
 
-			int NewContentWidth = Math.Max(TextRenderer.MeasureText(String.Format("Opened {0}  |  Browse...  |  Connect...", SelectedFileName), StatusPanel.Font).Width, 400);
-			StatusPanel.SetContentWidth(NewContentWidth);
+			int newContentWidth = Math.Max(TextRenderer.MeasureText(String.Format("Opened {0}  |  Browse...  |  Connect...", SelectedFileName), StatusPanel.Font).Width, 400);
+			StatusPanel.SetContentWidth(newContentWidth);
 
-			List<StatusLine> Lines = new List<StatusLine>();
-			if (Workspace.IsBusy())
+			List<StatusLine> lines = new List<StatusLine>();
+			if (_workspace.IsBusy())
 			{
 				// Sync in progress
-				Tuple<string, float> Progress = Workspace.CurrentProgress;
+				Tuple<string, float> progress = _workspace.CurrentProgress;
 
-				StatusLine SummaryLine = new StatusLine();
-				if (Workspace.PendingChangeNumber == -1)
+				StatusLine summaryLine = new StatusLine();
+				if (_workspace.PendingChangeNumber == -1)
 				{
-					SummaryLine.AddText("Working... | ");
+					summaryLine.AddText("Working... | ");
 				}
 				else
 				{
-					SummaryLine.AddText("Syncing to changelist ");
-					SummaryLine.AddLink(Workspace.PendingChangeNumber.ToString(), FontStyle.Regular, () => { SelectChange(Workspace.PendingChangeNumber); });
-					SummaryLine.AddText("... | ");
+					summaryLine.AddText("Syncing to changelist ");
+					summaryLine.AddLink(_workspace.PendingChangeNumber.ToString(), FontStyle.Regular, () => { SelectChange(_workspace.PendingChangeNumber); });
+					summaryLine.AddText("... | ");
 				}
-				SummaryLine.AddLink(Splitter.IsLogVisible() ? "Hide Log" : "Show Log", FontStyle.Bold | FontStyle.Underline, () => { ToggleLogVisibility(); });
-				SummaryLine.AddText(" | ");
-				SummaryLine.AddLink("Cancel", FontStyle.Bold | FontStyle.Underline, () => { CancelWorkspaceUpdate(); });
-				Lines.Add(SummaryLine);
+				summaryLine.AddLink(Splitter.IsLogVisible() ? "Hide Log" : "Show Log", FontStyle.Bold | FontStyle.Underline, () => { ToggleLogVisibility(); });
+				summaryLine.AddText(" | ");
+				summaryLine.AddLink("Cancel", FontStyle.Bold | FontStyle.Underline, () => { CancelWorkspaceUpdate(); });
+				lines.Add(summaryLine);
 
-				StatusLine ProgressLine = new StatusLine();
-				ProgressLine.AddText(String.Format("{0}  ", Progress.Item1));
-				if (Progress.Item2 > 0.0f) ProgressLine.AddProgressBar(Progress.Item2);
-				Lines.Add(ProgressLine);
+				StatusLine progressLine = new StatusLine();
+				progressLine.AddText(String.Format("{0}  ", progress.Item1));
+				if (progress.Item2 > 0.0f) progressLine.AddProgressBar(progress.Item2);
+				lines.Add(progressLine);
 			}
 			else
 			{
 				// Project
-				StatusLine ProjectLine = new StatusLine();
-				ProjectLine.AddText(String.Format("Opened "));
-				ProjectLine.AddLink(SelectedFileName.FullName + " \u25BE", FontStyle.Regular, (P, R) => { SelectRecentProject(R); });
-				ProjectLine.AddText("  |  ");
-				ProjectLine.AddLink("Settings...", FontStyle.Regular, (P, R) => { Owner.EditSelectedProject(this); });
-				Lines.Add(ProjectLine);
+				StatusLine projectLine = new StatusLine();
+				projectLine.AddText(String.Format("Opened "));
+				projectLine.AddLink(SelectedFileName.FullName + " \u25BE", FontStyle.Regular, (p, r) => { SelectRecentProject(r); });
+				projectLine.AddText("  |  ");
+				projectLine.AddLink("Settings...", FontStyle.Regular, (p, r) => { _owner.EditSelectedProject(this); });
+				lines.Add(projectLine);
 
 				// Spacer
-				Lines.Add(new StatusLine() { LineHeight = 0.5f });
+				lines.Add(new StatusLine() { LineHeight = 0.5f });
 
 				// Sync status
-				StatusLine SummaryLine = new StatusLine();
+				StatusLine summaryLine = new StatusLine();
 				if (IsBisectModeEnabled())
 				{
-					int PassChangeNumber;
-					int FailChangeNumber;
-					GetRemainingBisectRange(out PassChangeNumber, out FailChangeNumber);
+					int passChangeNumber;
+					int failChangeNumber;
+					GetRemainingBisectRange(out passChangeNumber, out failChangeNumber);
 
-					SummaryLine.AddText("Bisecting changes between ");
-					SummaryLine.AddLink(String.Format("{0}", PassChangeNumber), FontStyle.Regular, () => { SelectChange(PassChangeNumber); });
-					SummaryLine.AddText(" and ");
-					SummaryLine.AddLink(String.Format("{0}", FailChangeNumber), FontStyle.Regular, () => { SelectChange(FailChangeNumber); });
-					SummaryLine.AddText(".  ");
+					summaryLine.AddText("Bisecting changes between ");
+					summaryLine.AddLink(String.Format("{0}", passChangeNumber), FontStyle.Regular, () => { SelectChange(passChangeNumber); });
+					summaryLine.AddText(" and ");
+					summaryLine.AddLink(String.Format("{0}", failChangeNumber), FontStyle.Regular, () => { SelectChange(failChangeNumber); });
+					summaryLine.AddText(".  ");
 
-					int BisectChangeNumber = GetBisectChangeNumber();
-					if (BisectChangeNumber != -1)
+					int bisectChangeNumber = GetBisectChangeNumber();
+					if (bisectChangeNumber != -1)
 					{
-						SummaryLine.AddLink("Sync Next", FontStyle.Bold | FontStyle.Underline, () => { SyncBisectChange(); });
-						SummaryLine.AddText(" | ");
+						summaryLine.AddLink("Sync Next", FontStyle.Bold | FontStyle.Underline, () => { SyncBisectChange(); });
+						summaryLine.AddText(" | ");
 					}
 
-					SummaryLine.AddLink("Cancel", FontStyle.Bold | FontStyle.Underline, () => { CancelBisectMode(); });
+					summaryLine.AddLink("Cancel", FontStyle.Bold | FontStyle.Underline, () => { CancelBisectMode(); });
 				}
 				else
 				{
-					if (Workspace.CurrentChangeNumber != -1)
+					if (_workspace.CurrentChangeNumber != -1)
 					{
 						if (StreamName == null)
 						{
-							SummaryLine.AddText("Last synced to changelist ");
+							summaryLine.AddText("Last synced to changelist ");
 						}
 						else
 						{
-							SummaryLine.AddText("Last synced to ");
-							if (Workspace.ProjectStreamFilter == null || Workspace.ProjectStreamFilter.Count == 0)
+							summaryLine.AddText("Last synced to ");
+							if (_workspace.ProjectStreamFilter == null || _workspace.ProjectStreamFilter.Count == 0)
 							{
-								SummaryLine.AddLink(StreamName, FontStyle.Regular, (P, R) => { SelectOtherStreamDialog(); });
+								summaryLine.AddLink(StreamName, FontStyle.Regular, (p, r) => { SelectOtherStreamDialog(); });
 							}
 							else
 							{
-								SummaryLine.AddLink(StreamName + "\u25BE", FontStyle.Regular, (P, R) => { SelectOtherStream(R); });
+								summaryLine.AddLink(StreamName + "\u25BE", FontStyle.Regular, (p, r) => { SelectOtherStream(r); });
 							}
-							SummaryLine.AddText(" at changelist ");
+							summaryLine.AddText(" at changelist ");
 						}
-						SummaryLine.AddLink(String.Format("{0}.", Workspace.CurrentChangeNumber), FontStyle.Regular, () => { SelectChange(Workspace.CurrentChangeNumber); });
+						summaryLine.AddLink(String.Format("{0}.", _workspace.CurrentChangeNumber), FontStyle.Regular, () => { SelectChange(_workspace.CurrentChangeNumber); });
 					}
 					else
 					{
-						SummaryLine.AddText("You are not currently synced to ");
+						summaryLine.AddText("You are not currently synced to ");
 						if (StreamName == null)
 						{
-							SummaryLine.AddText("this branch.");
+							summaryLine.AddText("this branch.");
 						}
 						else
 						{
-							SummaryLine.AddLink(StreamName + " \u25BE", FontStyle.Regular, (P, R) => { SelectOtherStream(R); });
+							summaryLine.AddLink(StreamName + " \u25BE", FontStyle.Regular, (p, r) => { SelectOtherStream(r); });
 						}
 					}
-					SummaryLine.AddText("  |  ");
-					SummaryLine.AddLink("Sync Now", FontStyle.Bold | FontStyle.Underline, () => { SyncLatestChange(); });
-					SummaryLine.AddText(" - ");
-					SummaryLine.AddLink(" To... \u25BE", 0, (P, R) => { ShowSyncMenu(R); });
+					summaryLine.AddText("  |  ");
+					summaryLine.AddLink("Sync Now", FontStyle.Bold | FontStyle.Underline, () => { SyncLatestChange(); });
+					summaryLine.AddText(" - ");
+					summaryLine.AddLink(" To... \u25BE", 0, (p, r) => { ShowSyncMenu(r); });
 				}
-				Lines.Add(SummaryLine);
+				lines.Add(summaryLine);
 
 				// Programs
-				StatusLine ProgramsLine = new StatusLine();
-				if (Workspace.CurrentChangeNumber != -1)
+				StatusLine programsLine = new StatusLine();
+				if (_workspace.CurrentChangeNumber != -1)
 				{
-					ProgramsLine.AddLink("Unreal Editor", FontStyle.Regular, () => { LaunchEditor(); });
-					ProgramsLine.AddText("  |  ");
+					programsLine.AddLink("Unreal Editor", FontStyle.Regular, () => { LaunchEditor(); });
+					programsLine.AddText("  |  ");
 				}
 
-				string[]? SdkInfoEntries;
-				if (TryGetProjectSetting(PerforceMonitor.LatestProjectConfigFile, "SdkInfo", out SdkInfoEntries))
+				string[]? sdkInfoEntries;
+				if (TryGetProjectSetting(_perforceMonitor.LatestProjectConfigFile, "SdkInfo", out sdkInfoEntries))
 				{
-					ProgramsLine.AddLink("SDK Info", FontStyle.Regular, () => { ShowRequiredSdkInfo(); });
-					ProgramsLine.AddText("  |  ");
+					programsLine.AddLink("SDK Info", FontStyle.Regular, () => { ShowRequiredSdkInfo(); });
+					programsLine.AddText("  |  ");
 				}
 
-				ProgramsLine.AddLink("Perforce", FontStyle.Regular, () => { OpenPerforce(); });
-				ProgramsLine.AddText("  |  ");
+				programsLine.AddLink("Perforce", FontStyle.Regular, () => { OpenPerforce(); });
+				programsLine.AddText("  |  ");
 				if (!ShouldSyncPrecompiledEditor)
 				{
-					ProgramsLine.AddLink("Visual Studio", FontStyle.Regular, () => { OpenSolution(); });
-					ProgramsLine.AddText("  |  ");
+					programsLine.AddLink("Visual Studio", FontStyle.Regular, () => { OpenSolution(); });
+					programsLine.AddText("  |  ");
 				}
 
-				List<ToolDefinition> Tools = Owner.ToolUpdateMonitor.Tools;
-				foreach(ToolDefinition Tool in Tools)
+				List<ToolDefinition> tools = _owner.ToolUpdateMonitor.Tools;
+				foreach(ToolDefinition tool in tools)
 				{
-					if (Tool.Enabled)
+					if (tool.Enabled)
 					{
-						foreach (ToolLink Link in Tool.StatusPanelLinks)
+						foreach (ToolLink link in tool.StatusPanelLinks)
 						{
-							ProgramsLine.AddLink(Link.Label, FontStyle.Regular, () => RunTool(Tool, Link));
-							ProgramsLine.AddText("  |  ");
+							programsLine.AddLink(link.Label, FontStyle.Regular, () => RunTool(tool, link));
+							programsLine.AddText("  |  ");
 						}
 					}
 				}
 
-				foreach ((string Name, Action<Point, Rectangle> Action) in CustomStatusPanelLinks)
+				foreach ((string name, Action<Point, Rectangle> action) in _customStatusPanelLinks)
 				{
-					ProgramsLine.AddLink(Name, FontStyle.Regular, Action);
-					ProgramsLine.AddText("  |  ");
+					programsLine.AddLink(name, FontStyle.Regular, action);
+					programsLine.AddText("  |  ");
 				}
 
-				ProgramsLine.AddLink("Windows Explorer", FontStyle.Regular, () => { SafeProcessStart("explorer.exe", String.Format("\"{0}\"", SelectedFileName.Directory.FullName)); });
+				programsLine.AddLink("Windows Explorer", FontStyle.Regular, () => { SafeProcessStart("explorer.exe", String.Format("\"{0}\"", SelectedFileName.Directory.FullName)); });
 
 				if (GetDefaultIssueFilter() != null)
 				{
-					ProgramsLine.AddText("  |  ");
-					if (bUserHasOpenIssues)
+					programsLine.AddText("  |  ");
+					if (_userHasOpenIssues)
 					{
-						ProgramsLine.AddBadge("Build Health", GetBuildBadgeColor(BadgeResult.Failure), (P, R) => ShowBuildHealthMenu(R));
+						programsLine.AddBadge("Build Health", GetBuildBadgeColor(BadgeResult.Failure), (p, r) => ShowBuildHealthMenu(r));
 					}
 					else
 					{
-						ProgramsLine.AddLink("Build Health", FontStyle.Regular, (P, R) => { ShowBuildHealthMenu(R); });
+						programsLine.AddLink("Build Health", FontStyle.Regular, (p, r) => { ShowBuildHealthMenu(r); });
 					}
 				}
 
-				foreach (ServiceBadgeInfo ServiceBadge in ServiceBadges)
+				foreach (ServiceBadgeInfo serviceBadge in _serviceBadges)
 				{
-					ProgramsLine.AddText("  |  ");
-					if (ServiceBadge.Url == null)
+					programsLine.AddText("  |  ");
+					if (serviceBadge.Url == null)
 					{
-						ProgramsLine.AddBadge(ServiceBadge.Name, GetBuildBadgeColor(ServiceBadge.Result), null);
+						programsLine.AddBadge(serviceBadge.Name, GetBuildBadgeColor(serviceBadge.Result), null);
 					}
 					else
 					{
-						ProgramsLine.AddBadge(ServiceBadge.Name, GetBuildBadgeColor(ServiceBadge.Result), (P, R) => { SafeProcessStart(ServiceBadge.Url); });
+						programsLine.AddBadge(serviceBadge.Name, GetBuildBadgeColor(serviceBadge.Result), (p, r) => { SafeProcessStart(serviceBadge.Url); });
 					}
 				}
-				ProgramsLine.AddText("  |  ");
-				ProgramsLine.AddLink("More... \u25BE", FontStyle.Regular, (P, R) => { ShowActionsMenu(R); });
-				Lines.Add(ProgramsLine);
+				programsLine.AddText("  |  ");
+				programsLine.AddLink("More... \u25BE", FontStyle.Regular, (p, r) => { ShowActionsMenu(r); });
+				lines.Add(programsLine);
 
 				// Get the summary of the last sync
-				if (WorkspaceState.LastSyncChangeNumber > 0)
+				if (_workspaceState.LastSyncChangeNumber > 0)
 				{
-					string? SummaryText;
-					if (WorkspaceState.LastSyncChangeNumber == Workspace.CurrentChangeNumber && WorkspaceState.LastSyncResult == WorkspaceUpdateResult.Success && WorkspaceState.LastSyncTime.HasValue)
+					string? summaryText;
+					if (_workspaceState.LastSyncChangeNumber == _workspace.CurrentChangeNumber && _workspaceState.LastSyncResult == WorkspaceUpdateResult.Success && _workspaceState.LastSyncTime.HasValue)
 					{
-						Lines.Add(new StatusLine() { LineHeight = 0.5f });
+						lines.Add(new StatusLine() { LineHeight = 0.5f });
 
-						StatusLine SuccessLine = new StatusLine();
-						SuccessLine.AddIcon(Properties.Resources.StatusIcons, new Size(16, 16), 0);
-						SuccessLine.AddText(String.Format("  Sync took {0}{1}s, completed at {2}.", (WorkspaceState.LastSyncDurationSeconds >= 60) ? String.Format("{0}m ", WorkspaceState.LastSyncDurationSeconds / 60) : "", WorkspaceState.LastSyncDurationSeconds % 60, WorkspaceState.LastSyncTime.Value.ToLocalTime().ToString("h\\:mmtt").ToLowerInvariant()));
-						Lines.Add(SuccessLine);
+						StatusLine successLine = new StatusLine();
+						successLine.AddIcon(Properties.Resources.StatusIcons, new Size(16, 16), 0);
+						successLine.AddText(String.Format("  Sync took {0}{1}s, completed at {2}.", (_workspaceState.LastSyncDurationSeconds >= 60) ? String.Format("{0}m ", _workspaceState.LastSyncDurationSeconds / 60) : "", _workspaceState.LastSyncDurationSeconds % 60, _workspaceState.LastSyncTime.Value.ToLocalTime().ToString("h\\:mmtt").ToLowerInvariant()));
+						lines.Add(successLine);
 					}
-					else if (GetLastUpdateMessage(WorkspaceState.LastSyncResult, WorkspaceState.LastSyncResultMessage, out SummaryText))
+					else if (GetLastUpdateMessage(_workspaceState.LastSyncResult, _workspaceState.LastSyncResultMessage, out summaryText))
 					{
-						Lines.Add(new StatusLine() { LineHeight = 0.5f });
+						lines.Add(new StatusLine() { LineHeight = 0.5f });
 
-						int SummaryTextLength = SummaryText.IndexOf('\n');
-						if (SummaryTextLength == -1)
+						int summaryTextLength = summaryText.IndexOf('\n');
+						if (summaryTextLength == -1)
 						{
-							SummaryTextLength = SummaryText.Length;
+							summaryTextLength = summaryText.Length;
 						}
-						SummaryTextLength = Math.Min(SummaryTextLength, 80);
+						summaryTextLength = Math.Min(summaryTextLength, 80);
 
-						StatusLine FailLine = new StatusLine();
-						FailLine.AddIcon(Properties.Resources.StatusIcons, new Size(16, 16), 1);
+						StatusLine failLine = new StatusLine();
+						failLine.AddIcon(Properties.Resources.StatusIcons, new Size(16, 16), 1);
 
-						if (SummaryTextLength == SummaryText.Length)
+						if (summaryTextLength == summaryText.Length)
 						{
-							FailLine.AddText(String.Format("  {0}  ", SummaryText));
+							failLine.AddText(String.Format("  {0}  ", summaryText));
 						}
 						else
 						{
-							FailLine.AddText(String.Format("  {0}...  ", SummaryText.Substring(0, SummaryTextLength).TrimEnd()));
-							FailLine.AddLink("More...", FontStyle.Bold | FontStyle.Underline, () => { ViewLastSyncStatus(); });
-							FailLine.AddText("  |  ");
+							failLine.AddText(String.Format("  {0}...  ", summaryText.Substring(0, summaryTextLength).TrimEnd()));
+							failLine.AddLink("More...", FontStyle.Bold | FontStyle.Underline, () => { ViewLastSyncStatus(); });
+							failLine.AddText("  |  ");
 						}
-						FailLine.AddLink("Show Log", FontStyle.Bold | FontStyle.Underline, () => { ShowErrorInLog(); });
-						Lines.Add(FailLine);
+						failLine.AddLink("Show Log", FontStyle.Bold | FontStyle.Underline, () => { ShowErrorInLog(); });
+						lines.Add(failLine);
 					}
 				}
 			}
 
-			StatusLine? Caption = null;
-			if (StreamName != null && !Workspace.IsBusy())
+			StatusLine? caption = null;
+			if (StreamName != null && !_workspace.IsBusy())
 			{
-				Caption = new StatusLine();
-				Caption.AddLink(StreamName + "\u25BE", FontStyle.Bold, (P, R) => { SelectOtherStream(R); });
+				caption = new StatusLine();
+				caption.AddLink(StreamName + "\u25BE", FontStyle.Bold, (p, r) => { SelectOtherStream(r); });
 			}
 
-			StatusLine? Alert = null;
-			Color? TintColor = null;
+			StatusLine? alert = null;
+			Color? tintColor = null;
 
-			ConfigFile ProjectConfigFile = PerforceMonitor.LatestProjectConfigFile;
-			if (ProjectConfigFile != null)
+			ConfigFile projectConfigFile = _perforceMonitor.LatestProjectConfigFile;
+			if (projectConfigFile != null)
 			{
-				string? Message;
-				if (TryGetProjectSetting(ProjectConfigFile, "Message", out Message))
+				string? message;
+				if (TryGetProjectSetting(projectConfigFile, "Message", out message))
 				{
-					Alert = CreateStatusLineFromMarkdown(Message);
+					alert = CreateStatusLineFromMarkdown(message);
 				}
 
-				string? StatusPanelColor;
-				if (TryGetProjectSetting(ProjectConfigFile, "StatusPanelColor", out StatusPanelColor))
+				string? statusPanelColor;
+				if (TryGetProjectSetting(projectConfigFile, "StatusPanelColor", out statusPanelColor))
 				{
-					TintColor = System.Drawing.ColorTranslator.FromHtml(StatusPanelColor);
+					tintColor = System.Drawing.ColorTranslator.FromHtml(statusPanelColor);
 				}
 			}
 
-			using (Graphics Graphics = CreateGraphics())
+			using (Graphics graphics = CreateGraphics())
 			{
-				int StatusPanelHeight = (int)(148.0f * Graphics.DpiY / 96.0f);
-				if (Alert != null)
+				int statusPanelHeight = (int)(148.0f * graphics.DpiY / 96.0f);
+				if (alert != null)
 				{
-					StatusPanelHeight += (int)(40.0f * Graphics.DpiY / 96.0f);
+					statusPanelHeight += (int)(40.0f * graphics.DpiY / 96.0f);
 				}
-				if (StatusLayoutPanel.RowStyles[0].Height != StatusPanelHeight)
+				if (StatusLayoutPanel.RowStyles[0].Height != statusPanelHeight)
 				{
-					StatusLayoutPanel.RowStyles[0].Height = StatusPanelHeight;
+					StatusLayoutPanel.RowStyles[0].Height = statusPanelHeight;
 				}
 			}
 
-			StatusPanel.Set(Lines, Caption, Alert, TintColor);
+			StatusPanel.Set(lines, caption, alert, tintColor);
 		}
 
-		private void ShowBuildHealthMenu(Rectangle Bounds)
+		private void ShowBuildHealthMenu(Rectangle bounds)
 		{
-			int MinSeparatorIdx = BuildHealthContextMenu.Items.IndexOf(BuildHealthContextMenu_MinSeparator);
+			int minSeparatorIdx = BuildHealthContextMenu.Items.IndexOf(BuildHealthContextMenu_MinSeparator);
 
-			while (BuildHealthContextMenu.Items[MinSeparatorIdx + 1] != BuildHealthContextMenu_MaxSeparator)
+			while (BuildHealthContextMenu.Items[minSeparatorIdx + 1] != BuildHealthContextMenu_MaxSeparator)
 			{
-				BuildHealthContextMenu.Items.RemoveAt(MinSeparatorIdx + 1);
+				BuildHealthContextMenu.Items.RemoveAt(minSeparatorIdx + 1);
 			}
 
-			Func<IssueData, bool>? Predicate = null;
-			string? DefaultIssueFilter = GetDefaultIssueFilter();
-			if (DefaultIssueFilter == null || !GetCustomIssueFilters().TryGetValue(DefaultIssueFilter, out Predicate))
+			Func<IssueData, bool>? predicate = null;
+			string? defaultIssueFilter = GetDefaultIssueFilter();
+			if (defaultIssueFilter == null || !GetCustomIssueFilters().TryGetValue(defaultIssueFilter, out predicate))
 			{
-				Predicate = x => true;
+				predicate = x => true;
 			}
 
-			List<IssueData> Issues = new List<IssueData>();
-			bool bHasOtherAssignedIssue = false;
-			foreach (IssueData Issue in IssueMonitor.GetIssues().OrderByDescending(x => x.Id))
+			List<IssueData> issues = new List<IssueData>();
+			bool hasOtherAssignedIssue = false;
+			foreach (IssueData issue in _issueMonitor.GetIssues().OrderByDescending(x => x.Id))
 			{
-				if (Predicate(Issue))
+				if (predicate(issue))
 				{
-					Issues.Add(Issue);
+					issues.Add(issue);
 				}
-				else if (Issue.FixChange == 0 && Issue.Owner != null && String.Compare(Issue.Owner, PerforceSettings.UserName, StringComparison.OrdinalIgnoreCase) == 0)
+				else if (issue.FixChange == 0 && issue.Owner != null && String.Compare(issue.Owner, _perforceSettings.UserName, StringComparison.OrdinalIgnoreCase) == 0)
 				{
-					bHasOtherAssignedIssue = true;
+					hasOtherAssignedIssue = true;
 				}
 			}
 
-			BuildHealthContextMenu_Browse.Font = new Font(BuildHealthContextMenu_Browse.Font, bHasOtherAssignedIssue ? FontStyle.Bold : FontStyle.Regular);
+			BuildHealthContextMenu_Browse.Font = new Font(BuildHealthContextMenu_Browse.Font, hasOtherAssignedIssue ? FontStyle.Bold : FontStyle.Regular);
 
-			if (Issues.Count == 0)
+			if (issues.Count == 0)
 			{
 				BuildHealthContextMenu_MaxSeparator.Visible = false;
 			}
 			else
 			{
 				BuildHealthContextMenu_MaxSeparator.Visible = true;
-				for (int Idx = 0; Idx < Issues.Count; Idx++)
+				for (int idx = 0; idx < issues.Count; idx++)
 				{
-					IssueData Issue = Issues[Idx];
+					IssueData issue = issues[idx];
 
-					string Summary = Issue.Summary;
-					if (Summary.Length > 100)
+					string summary = issue.Summary;
+					if (summary.Length > 100)
 					{
-						Summary = Summary.Substring(0, 100).TrimEnd() + "...";
+						summary = summary.Substring(0, 100).TrimEnd() + "...";
 					}
 
-					StringBuilder Description = new StringBuilder();
-					Description.AppendFormat("{0}: {1}", Issue.Id, Summary);
-					if (Issue.Owner == null)
+					StringBuilder description = new StringBuilder();
+					description.AppendFormat("{0}: {1}", issue.Id, summary);
+					if (issue.Owner == null)
 					{
-						Description.AppendFormat(" - Unassigned");
+						description.AppendFormat(" - Unassigned");
 					}
 					else
 					{
-						Description.AppendFormat(" - {0}", FormatUserName(Issue.Owner));
-						if (!Issue.AcknowledgedAt.HasValue)
+						description.AppendFormat(" - {0}", FormatUserName(issue.Owner));
+						if (!issue.AcknowledgedAt.HasValue)
 						{
-							Description.Append(" (?)");
+							description.Append(" (?)");
 						}
 					}
-					if (Issue.FixChange > 0)
+					if (issue.FixChange > 0)
 					{
-						Description.AppendFormat(" (Unverified fix in CL {0})", Issue.FixChange);
+						description.AppendFormat(" (Unverified fix in CL {0})", issue.FixChange);
 					}
 
-					Description.AppendFormat(" ({0})", Utility.FormatDurationMinutes((int)((Issue.RetrievedAt - Issue.CreatedAt).TotalMinutes + 1)));
+					description.AppendFormat(" ({0})", Utility.FormatDurationMinutes((int)((issue.RetrievedAt - issue.CreatedAt).TotalMinutes + 1)));
 
-					ToolStripMenuItem Item = new ToolStripMenuItem(Description.ToString());
-					if (Issue.FixChange > 0)
+					ToolStripMenuItem item = new ToolStripMenuItem(description.ToString());
+					if (issue.FixChange > 0)
 					{
-						Item.ForeColor = SystemColors.GrayText;
+						item.ForeColor = SystemColors.GrayText;
 					}
-					else if (Issue.Owner != null && String.Compare(Issue.Owner, PerforceSettings.UserName, StringComparison.OrdinalIgnoreCase) == 0)
+					else if (issue.Owner != null && String.Compare(issue.Owner, _perforceSettings.UserName, StringComparison.OrdinalIgnoreCase) == 0)
 					{
-						Item.Font = new Font(Item.Font, FontStyle.Bold);
+						item.Font = new Font(item.Font, FontStyle.Bold);
 					}
-					Item.Click += (s, e) => { BuildHealthContextMenu_Issue_Click(Issue); };
-					BuildHealthContextMenu.Items.Insert(MinSeparatorIdx + Idx + 1, Item);
+					item.Click += (s, e) => { BuildHealthContextMenu_Issue_Click(issue); };
+					BuildHealthContextMenu.Items.Insert(minSeparatorIdx + idx + 1, item);
 				}
 			}
 
-			BuildHealthContextMenu.Show(StatusPanel, new Point(Bounds.Right, Bounds.Bottom), ToolStripDropDownDirection.BelowLeft);
+			BuildHealthContextMenu.Show(StatusPanel, new Point(bounds.Right, bounds.Bottom), ToolStripDropDownDirection.BelowLeft);
 		}
 
-		private void BuildHealthContextMenu_Issue_Click(IssueData Issue)
+		private void BuildHealthContextMenu_Issue_Click(IssueData issue)
 		{
-			ShowIssueDetails(Issue);
+			ShowIssueDetails(issue);
 		}
 
 		private TimeSpan? GetServerTimeOffset()
 		{
-			TimeSpan? Offset = null;
-			if (Settings.bShowLocalTimes)
+			TimeSpan? offset = null;
+			if (_settings.ShowLocalTimes)
 			{
-				Offset = PerforceMonitor.ServerTimeZone;
+				offset = _perforceMonitor.ServerTimeZone;
 			}
-			return Offset;
+			return offset;
 		}
 
-		public void ShowIssueDetails(IssueData Issue)
+		public void ShowIssueDetails(IssueData issue)
 		{
-			IssueDetailsWindow.Show(ParentForm, IssueMonitor, PerforceSettings, GetServerTimeOffset(), Issue, ServiceProvider, StreamName);
+			IssueDetailsWindow.Show(ParentForm, _issueMonitor, _perforceSettings, GetServerTimeOffset(), issue, _serviceProvider, StreamName);
 		}
 
 		private void BuildHealthContextMenu_Browse_Click(object sender, EventArgs e)
 		{
-			string? DefaultFilter = GetDefaultIssueFilter();
-			Dictionary<string, Func<IssueData, bool>> CustomFilters = GetCustomIssueFilters();
-			IssueBrowserWindow.Show(ParentForm, IssueMonitor, PerforceSettings, GetServerTimeOffset(), ServiceProvider, StreamName, CustomFilters, DefaultFilter);
+			string? defaultFilter = GetDefaultIssueFilter();
+			Dictionary<string, Func<IssueData, bool>> customFilters = GetCustomIssueFilters();
+			IssueBrowserWindow.Show(ParentForm, _issueMonitor, _perforceSettings, GetServerTimeOffset(), _serviceProvider, StreamName, customFilters, defaultFilter);
 		}
 
 		private void BuildHealthContextMenu_Settings_Click(object sender, EventArgs e)
 		{
-			string? BuildHealthProject;
-			TryGetProjectSetting(PerforceMonitor.LatestProjectConfigFile, "BuildHealthProject", out BuildHealthProject);
+			string? buildHealthProject;
+			TryGetProjectSetting(_perforceMonitor.LatestProjectConfigFile, "BuildHealthProject", out buildHealthProject);
 
-			IssueSettingsWindow IssueSettings = new IssueSettingsWindow(Settings, BuildHealthProject ?? "", Logger);
-			if (IssueSettings.ShowDialog(this) == DialogResult.OK)
+			IssueSettingsWindow issueSettings = new IssueSettingsWindow(_settings, buildHealthProject ?? "", _logger);
+			if (issueSettings.ShowDialog(this) == DialogResult.OK)
 			{
-				Owner.UpdateAlertWindows();
+				_owner.UpdateAlertWindows();
 			}
 		}
 
 		private void ShowRequiredSdkInfo()
 		{
-			string[]? SdkInfoEntries;
-			if (TryGetProjectSetting(PerforceMonitor.LatestProjectConfigFile, "SdkInfo", out SdkInfoEntries))
+			string[]? sdkInfoEntries;
+			if (TryGetProjectSetting(_perforceMonitor.LatestProjectConfigFile, "SdkInfo", out sdkInfoEntries))
 			{
-				Dictionary<string, string> Variables = Workspace.GetVariables(GetEditorBuildConfig());
-				SdkInfoWindow Window = new SdkInfoWindow(SdkInfoEntries, Variables, BadgeFont!);
-				Window.ShowDialog();
+				Dictionary<string, string> variables = _workspace.GetVariables(GetEditorBuildConfig());
+				SdkInfoWindow window = new SdkInfoWindow(sdkInfoEntries, variables, _badgeFont!);
+				window.ShowDialog();
 			}
 		}
 
-		private StatusLine CreateStatusLineFromMarkdown(string Text)
+		private StatusLine CreateStatusLineFromMarkdown(string text)
 		{
-			StatusLine Line = new StatusLine();
+			StatusLine line = new StatusLine();
 
-			FontStyle Style = FontStyle.Regular;
+			FontStyle style = FontStyle.Regular;
 
-			StringBuilder ElementText = new StringBuilder();
-			for (int Idx = 0; Idx < Text.Length;)
+			StringBuilder elementText = new StringBuilder();
+			for (int idx = 0; idx < text.Length;)
 			{
 				// Bold and italic
-				if (Text[Idx] == '_' || Text[Idx] == '*')
+				if (text[idx] == '_' || text[idx] == '*')
 				{
-					if (Idx + 1 < Text.Length && Text[Idx + 1] == Text[Idx])
+					if (idx + 1 < text.Length && text[idx + 1] == text[idx])
 					{
-						FlushMarkdownText(Line, ElementText, Style);
-						Idx += 2;
-						Style ^= FontStyle.Bold;
+						FlushMarkdownText(line, elementText, style);
+						idx += 2;
+						style ^= FontStyle.Bold;
 						continue;
 					}
 					else
 					{
-						FlushMarkdownText(Line, ElementText, Style);
-						Idx++;
-						Style ^= FontStyle.Italic;
+						FlushMarkdownText(line, elementText, style);
+						idx++;
+						style ^= FontStyle.Italic;
 						continue;
 					}
 				}
 
 				// Strikethrough
-				if (Idx + 2 < Text.Length && Text[Idx] == '~' && Text[Idx + 1] == '~')
+				if (idx + 2 < text.Length && text[idx] == '~' && text[idx + 1] == '~')
 				{
-					FlushMarkdownText(Line, ElementText, Style);
-					Idx += 2;
-					Style ^= FontStyle.Strikeout;
+					FlushMarkdownText(line, elementText, style);
+					idx += 2;
+					style ^= FontStyle.Strikeout;
 					continue;
 				}
 
 				// Link
-				if (Text[Idx] == '[')
+				if (text[idx] == '[')
 				{
-					if (Idx + 1 < Text.Length && Text[Idx + 1] == '[')
+					if (idx + 1 < text.Length && text[idx + 1] == '[')
 					{
-						ElementText.Append(Text[Idx]);
-						Idx += 2;
+						elementText.Append(text[idx]);
+						idx += 2;
 						continue;
 					}
 
-					int EndIdx = Text.IndexOf("](", Idx);
-					if (EndIdx != -1)
+					int endIdx = text.IndexOf("](", idx);
+					if (endIdx != -1)
 					{
-						int UrlEndIdx = Text.IndexOf(')', EndIdx + 2);
-						if (UrlEndIdx != -1)
+						int urlEndIdx = text.IndexOf(')', endIdx + 2);
+						if (urlEndIdx != -1)
 						{
-							FlushMarkdownText(Line, ElementText, Style);
-							string LinkText = Text.Substring(Idx + 1, EndIdx - (Idx + 1));
-							string LinkUrl = Text.Substring(EndIdx + 2, UrlEndIdx - (EndIdx + 2));
-							Line.AddLink(LinkText, Style, () => SafeProcessStart(LinkUrl));
-							Idx = UrlEndIdx + 1;
+							FlushMarkdownText(line, elementText, style);
+							string linkText = text.Substring(idx + 1, endIdx - (idx + 1));
+							string linkUrl = text.Substring(endIdx + 2, urlEndIdx - (endIdx + 2));
+							line.AddLink(linkText, style, () => SafeProcessStart(linkUrl));
+							idx = urlEndIdx + 1;
 							continue;
 						}
 					}
 				}
 
 				// Icon
-				if (Text[Idx] == ':')
+				if (text[idx] == ':')
 				{
-					int EndIdx = Text.IndexOf(':', Idx + 1);
-					if (EndIdx != -1)
+					int endIdx = text.IndexOf(':', idx + 1);
+					if (endIdx != -1)
 					{
-						if (String.Compare(":alert:", 0, Text, Idx, EndIdx - Idx) == 0)
+						if (String.Compare(":alert:", 0, text, idx, endIdx - idx) == 0)
 						{
-							FlushMarkdownText(Line, ElementText, Style);
-							Line.AddIcon(Properties.Resources.StatusIcons, new Size(16, 16), 4);
-							Idx = EndIdx + 1;
+							FlushMarkdownText(line, elementText, style);
+							line.AddIcon(Properties.Resources.StatusIcons, new Size(16, 16), 4);
+							idx = endIdx + 1;
 							continue;
 						}
 					}
 				}
 
 				// Otherwise, just append the current character
-				ElementText.Append(Text[Idx++]);
+				elementText.Append(text[idx++]);
 			}
-			FlushMarkdownText(Line, ElementText, Style);
+			FlushMarkdownText(line, elementText, style);
 
-			return Line;
+			return line;
 		}
 
-		private void FlushMarkdownText(StatusLine Line, StringBuilder ElementText, FontStyle Style)
+		private void FlushMarkdownText(StatusLine line, StringBuilder elementText, FontStyle style)
 		{
-			if (ElementText.Length > 0)
+			if (elementText.Length > 0)
 			{
-				Line.AddText(ElementText.ToString(), Style);
-				ElementText.Clear();
+				line.AddText(elementText.ToString(), style);
+				elementText.Clear();
 			}
 		}
 
-		private void SelectOtherStream(Rectangle Bounds)
+		private void SelectOtherStream(Rectangle bounds)
 		{
-			bool bShownContextMenu = false;
+			bool shownContextMenu = false;
 			if (StreamName != null)
 			{
-				IReadOnlyList<string>? OtherStreamNames = Workspace.ProjectStreamFilter;
-				if (OtherStreamNames != null)
+				IReadOnlyList<string>? otherStreamNames = _workspace.ProjectStreamFilter;
+				if (otherStreamNames != null)
 				{
 					StreamContextMenu.Items.Clear();
 
-					ToolStripMenuItem CurrentStreamItem = new ToolStripMenuItem(StreamName, null, new EventHandler((S, E) => SelectStream(StreamName)));
-					CurrentStreamItem.Checked = true;
-					StreamContextMenu.Items.Add(CurrentStreamItem);
+					ToolStripMenuItem currentStreamItem = new ToolStripMenuItem(StreamName, null, new EventHandler((s, e) => SelectStream(StreamName)));
+					currentStreamItem.Checked = true;
+					StreamContextMenu.Items.Add(currentStreamItem);
 
 					StreamContextMenu.Items.Add(new ToolStripSeparator());
 
-					foreach (string OtherStreamName in OtherStreamNames.OrderBy(x => x).Where(x => !x.EndsWith("/Dev-Binaries")))
+					foreach (string otherStreamName in otherStreamNames.OrderBy(x => x).Where(x => !x.EndsWith("/Dev-Binaries")))
 					{
-						string ThisStreamName = OtherStreamName; // Local for lambda capture
-						if (String.Compare(StreamName, OtherStreamName, StringComparison.InvariantCultureIgnoreCase) != 0)
+						string thisStreamName = otherStreamName; // Local for lambda capture
+						if (String.Compare(StreamName, otherStreamName, StringComparison.InvariantCultureIgnoreCase) != 0)
 						{
-							ToolStripMenuItem Item = new ToolStripMenuItem(ThisStreamName, null, new EventHandler((S, E) => SelectStream(ThisStreamName)));
-							StreamContextMenu.Items.Add(Item);
+							ToolStripMenuItem item = new ToolStripMenuItem(thisStreamName, null, new EventHandler((s, e) => SelectStream(thisStreamName)));
+							StreamContextMenu.Items.Add(item);
 						}
 					}
 
@@ -3809,16 +3809,16 @@ namespace UnrealGameSync
 						StreamContextMenu.Items.Add(new ToolStripSeparator());
 					}
 
-					StreamContextMenu.Items.Add(new ToolStripMenuItem("Select Other...", null, new EventHandler((S, E) => SelectOtherStreamDialog())));
+					StreamContextMenu.Items.Add(new ToolStripMenuItem("Select Other...", null, new EventHandler((s, e) => SelectOtherStreamDialog())));
 
-					int X = (Bounds.Left + Bounds.Right) / 2 + StreamContextMenu.Bounds.Width / 2;
-					int Y = Bounds.Bottom + 2;
-					StreamContextMenu.Show(StatusPanel, new Point(X, Y), ToolStripDropDownDirection.Left);
+					int x = (bounds.Left + bounds.Right) / 2 + StreamContextMenu.Bounds.Width / 2;
+					int y = bounds.Bottom + 2;
+					StreamContextMenu.Show(StatusPanel, new Point(x, y), ToolStripDropDownDirection.Left);
 
-					bShownContextMenu = true;
+					shownContextMenu = true;
 				}
 			}
-			if (!bShownContextMenu)
+			if (!shownContextMenu)
 			{
 				SelectOtherStreamDialog();
 			}
@@ -3826,41 +3826,41 @@ namespace UnrealGameSync
 
 		private void SelectOtherStreamDialog()
 		{
-			string? NewStreamName;
-			if (SelectStreamWindow.ShowModal(this, PerforceSettings, StreamName, ServiceProvider, out NewStreamName))
+			string? newStreamName;
+			if (SelectStreamWindow.ShowModal(this, _perforceSettings, StreamName, _serviceProvider, out newStreamName))
 			{
-				SelectStream(NewStreamName);
+				SelectStream(newStreamName);
 			}
 		}
 
-		private void SelectStream(string NewStreamName)
+		private void SelectStream(string newStreamName)
 		{
-			if (StreamName != NewStreamName)
+			if (StreamName != newStreamName)
 			{
-				if (Workspace.IsBusy())
+				if (_workspace.IsBusy())
 				{
 					MessageBox.Show("Please retry after the current sync has finished.", "Sync in Progress");
 				}
 				else
 				{
-					for (int Idx = 0; Idx < 2; Idx++)
+					for (int idx = 0; idx < 2; idx++)
 					{
-						Func<IPerforceConnection, CancellationToken, Task<bool>> SwitchFunc = async (Perforce, CancellationToken) =>
+						Func<IPerforceConnection, CancellationToken, Task<bool>> switchFunc = async (perforce, cancellationToken) =>
 						{
-							if (Idx == 1 || !await Perforce.OpenedAsync(OpenedOptions.None, FileSpecList.Any, CancellationToken).AnyAsync())
+							if (idx == 1 || !await perforce.OpenedAsync(OpenedOptions.None, FileSpecList.Any, cancellationToken).AnyAsync())
 							{
-								await Perforce.SwitchClientToStreamAsync(NewStreamName, SwitchClientOptions.IgnoreOpenFiles, CancellationToken);
+								await perforce.SwitchClientToStreamAsync(newStreamName, SwitchClientOptions.IgnoreOpenFiles, cancellationToken);
 								return true;
 							}
 							return false;
 						};
 
-						ModalTask<bool>? SwitchTask = PerforceModalTask.Execute<bool>(this, "Switching streams", "Please wait...", PerforceSettings, SwitchFunc, Logger);
-						if (SwitchTask == null || !SwitchTask.Succeeded)
+						ModalTask<bool>? switchTask = PerforceModalTask.Execute<bool>(this, "Switching streams", "Please wait...", _perforceSettings, switchFunc, _logger);
+						if (switchTask == null || !switchTask.Succeeded)
 						{
 							break;
 						}
-						if (SwitchTask.Result)
+						if (switchTask.Result)
 						{
 							StatusPanel.SuspendLayout();
 							StreamChanged();
@@ -3874,15 +3874,15 @@ namespace UnrealGameSync
 
 		private void ViewLastSyncStatus()
 		{
-			string? SummaryText;
-			if (GetLastUpdateMessage(WorkspaceState.LastSyncResult, WorkspaceState.LastSyncResultMessage, out SummaryText))
+			string? summaryText;
+			if (GetLastUpdateMessage(_workspaceState.LastSyncResult, _workspaceState.LastSyncResultMessage, out summaryText))
 			{
-				string? CaptionText;
-				if (!GetGenericLastUpdateMessage(WorkspaceState.LastSyncResult, out CaptionText))
+				string? captionText;
+				if (!GetGenericLastUpdateMessage(_workspaceState.LastSyncResult, out captionText))
 				{
-					CaptionText = "Sync error";
+					captionText = "Sync error";
 				}
-				MessageBox.Show(SummaryText, CaptionText);
+				MessageBox.Show(summaryText, captionText);
 			}
 		}
 
@@ -3895,168 +3895,168 @@ namespace UnrealGameSync
 			SyncLog.ScrollToEnd();
 		}
 
-		private void ShowToolContextMenu(Rectangle Bounds, ContextMenuStrip MenuStrip)
+		private void ShowToolContextMenu(Rectangle bounds, ContextMenuStrip menuStrip)
 		{
-			MenuStrip.Show(StatusPanel, new Point(Bounds.Left, Bounds.Bottom), ToolStripDropDownDirection.BelowRight);
+			menuStrip.Show(StatusPanel, new Point(bounds.Left, bounds.Bottom), ToolStripDropDownDirection.BelowRight);
 		}
 
-		private void ShowActionsMenu(Rectangle Bounds)
+		private void ShowActionsMenu(Rectangle bounds)
 		{
-			MoreToolsContextMenu.Show(StatusPanel, new Point(Bounds.Left, Bounds.Bottom), ToolStripDropDownDirection.BelowRight);
+			MoreToolsContextMenu.Show(StatusPanel, new Point(bounds.Left, bounds.Bottom), ToolStripDropDownDirection.BelowRight);
 		}
 
-		private void SelectChange(int ChangeNumber)
+		private void SelectChange(int changeNumber)
 		{
 			BuildList.SelectedItems.Clear();
 
-			PendingSelectedChangeNumber = -1;
+			_pendingSelectedChangeNumber = -1;
 
-			foreach (ListViewItem? Item in BuildList.Items)
+			foreach (ListViewItem? item in BuildList.Items)
 			{
-				if (Item != null)
+				if (item != null)
 				{
-					ChangesRecord Summary = (ChangesRecord)Item.Tag;
-					if (Summary != null && Summary.Number <= ChangeNumber)
+					ChangesRecord summary = (ChangesRecord)item.Tag;
+					if (summary != null && summary.Number <= changeNumber)
 					{
-						Item.Selected = true;
-						Item.EnsureVisible();
+						item.Selected = true;
+						item.EnsureVisible();
 						return;
 					}
 				}
 			}
 
-			PendingSelectedChangeNumber = ChangeNumber;
+			_pendingSelectedChangeNumber = changeNumber;
 
-			int CurrentMaxChanges = PerforceMonitor.CurrentMaxChanges;
-			if (PerforceMonitor.PendingMaxChanges <= CurrentMaxChanges && BuildList.SelectedItems.Count == 0)
+			int currentMaxChanges = _perforceMonitor.CurrentMaxChanges;
+			if (_perforceMonitor.PendingMaxChanges <= currentMaxChanges && BuildList.SelectedItems.Count == 0)
 			{
-				PerforceMonitor.PendingMaxChanges = CurrentMaxChanges + BuildListExpandCount;
-				if (ExpandItem != null)
+				_perforceMonitor.PendingMaxChanges = currentMaxChanges + BuildListExpandCount;
+				if (_expandItem != null)
 				{
-					ExpandItem.EnsureVisible();
+					_expandItem.EnsureVisible();
 				}
 			}
 		}
 
-		private void BuildList_MouseClick(object Sender, MouseEventArgs Args)
+		private void BuildList_MouseClick(object sender, MouseEventArgs args)
 		{
-			if (Args.Button == MouseButtons.Left)
+			if (args.Button == MouseButtons.Left)
 			{
-				ListViewHitTestInfo HitTest = BuildList.HitTest(Args.Location);
-				if (HitTest.Item != null)
+				ListViewHitTestInfo hitTest = BuildList.HitTest(args.Location);
+				if (hitTest.Item != null)
 				{
-					if (HitTest.Item == ExpandItem)
+					if (hitTest.Item == _expandItem)
 					{
-						if (HitTestExpandLink(Args.Location))
+						if (HitTestExpandLink(args.Location))
 						{
-							int CurrentMaxChanges = PerforceMonitor.CurrentMaxChanges;
-							if (PerforceMonitor.PendingMaxChanges > CurrentMaxChanges)
+							int currentMaxChanges = _perforceMonitor.CurrentMaxChanges;
+							if (_perforceMonitor.PendingMaxChanges > currentMaxChanges)
 							{
-								PerforceMonitor.PendingMaxChanges = CurrentMaxChanges;
+								_perforceMonitor.PendingMaxChanges = currentMaxChanges;
 							}
 							else
 							{
-								PerforceMonitor.PendingMaxChanges = CurrentMaxChanges + BuildListExpandCount;
+								_perforceMonitor.PendingMaxChanges = currentMaxChanges + BuildListExpandCount;
 							}
 							BuildList.Invalidate();
 						}
 					}
 					else
 					{
-						ChangesRecord Change = (ChangesRecord)HitTest.Item.Tag;
-						if (Workspace.PendingChangeNumber == Change.Number)
+						ChangesRecord change = (ChangesRecord)hitTest.Item.Tag;
+						if (_workspace.PendingChangeNumber == change.Number)
 						{
-							Rectangle SubItemRect = HitTest.Item.SubItems[StatusColumn.Index].Bounds;
+							Rectangle subItemRect = hitTest.Item.SubItems[StatusColumn.Index].Bounds;
 
-							if (Workspace.IsBusy())
+							if (_workspace.IsBusy())
 							{
-								Rectangle CancelRect = new Rectangle(SubItemRect.Right - 16, SubItemRect.Top, 16, SubItemRect.Height);
-								Rectangle InfoRect = new Rectangle(SubItemRect.Right - 32, SubItemRect.Top, 16, SubItemRect.Height);
-								if (CancelRect.Contains(Args.Location))
+								Rectangle cancelRect = new Rectangle(subItemRect.Right - 16, subItemRect.Top, 16, subItemRect.Height);
+								Rectangle infoRect = new Rectangle(subItemRect.Right - 32, subItemRect.Top, 16, subItemRect.Height);
+								if (cancelRect.Contains(args.Location))
 								{
 									CancelWorkspaceUpdate();
 								}
-								else if (InfoRect.Contains(Args.Location) && !Splitter.IsLogVisible())
+								else if (infoRect.Contains(args.Location) && !Splitter.IsLogVisible())
 								{
 									ToggleLogVisibility();
 								}
 							}
 							else
 							{
-								Rectangle HappyRect = new Rectangle(SubItemRect.Right - 32, SubItemRect.Top, 16, SubItemRect.Height);
-								Rectangle FrownRect = new Rectangle(SubItemRect.Right - 16, SubItemRect.Top, 16, SubItemRect.Height);
-								if (HappyRect.Contains(Args.Location))
+								Rectangle happyRect = new Rectangle(subItemRect.Right - 32, subItemRect.Top, 16, subItemRect.Height);
+								Rectangle frownRect = new Rectangle(subItemRect.Right - 16, subItemRect.Top, 16, subItemRect.Height);
+								if (happyRect.Contains(args.Location))
 								{
-									EventMonitor.PostEvent(Change.Number, EventType.Good);
+									_eventMonitor.PostEvent(change.Number, EventType.Good);
 									BuildList.Invalidate();
 								}
-								else if (FrownRect.Contains(Args.Location))
+								else if (frownRect.Contains(args.Location))
 								{
-									EventMonitor.PostEvent(Change.Number, EventType.Bad);
+									_eventMonitor.PostEvent(change.Number, EventType.Bad);
 									BuildList.Invalidate();
 								}
 							}
 						}
 						else
 						{
-							Rectangle SyncBadgeRectangle = GetSyncBadgeRectangle(HitTest.Item.SubItems[StatusColumn.Index].Bounds);
-							if (SyncBadgeRectangle.Contains(Args.Location) && CanSyncChange(Change.Number))
+							Rectangle syncBadgeRectangle = GetSyncBadgeRectangle(hitTest.Item.SubItems[StatusColumn.Index].Bounds);
+							if (syncBadgeRectangle.Contains(args.Location) && CanSyncChange(change.Number))
 							{
-								StartSync(Change.Number, false, null);
+								StartSync(change.Number, false, null);
 							}
 						}
 
-						if (DescriptionColumn.Index < HitTest.Item.SubItems.Count && HitTest.Item.SubItems[DescriptionColumn.Index] == HitTest.SubItem)
+						if (DescriptionColumn.Index < hitTest.Item.SubItems.Count && hitTest.Item.SubItems[DescriptionColumn.Index] == hitTest.SubItem)
 						{
-							ChangeLayoutInfo LayoutInfo = GetChangeLayoutInfo(Change);
-							if (LayoutInfo.DescriptionBadges.Count > 0)
+							ChangeLayoutInfo layoutInfo = GetChangeLayoutInfo(change);
+							if (layoutInfo.DescriptionBadges.Count > 0)
 							{
-								Point BuildListLocation = GetBadgeListLocation(LayoutInfo.DescriptionBadges, HitTest.SubItem.Bounds, HorizontalAlign.Right, VerticalAlignment.Middle);
-								BuildListLocation.Offset(-2, 0);
+								Point buildListLocation = GetBadgeListLocation(layoutInfo.DescriptionBadges, hitTest.SubItem.Bounds, HorizontalAlign.Right, VerticalAlignment.Middle);
+								buildListLocation.Offset(-2, 0);
 
-								foreach (BadgeInfo Badge in LayoutInfo.DescriptionBadges)
+								foreach (BadgeInfo badge in layoutInfo.DescriptionBadges)
 								{
-									Rectangle BadgeBounds = Badge.GetBounds(BuildListLocation);
-									if (BadgeBounds.Contains(Args.Location) && Badge.ClickHandler != null)
+									Rectangle badgeBounds = badge.GetBounds(buildListLocation);
+									if (badgeBounds.Contains(args.Location) && badge.ClickHandler != null)
 									{
-										Badge.ClickHandler();
+										badge.ClickHandler();
 										break;
 									}
 								}
 							}
 						}
 
-						if (CISColumn.Index < HitTest.Item.SubItems.Count && HitTest.Item.SubItems[CISColumn.Index] == HitTest.SubItem)
+						if (CISColumn.Index < hitTest.Item.SubItems.Count && hitTest.Item.SubItems[CISColumn.Index] == hitTest.SubItem)
 						{
-							ChangeLayoutInfo LayoutInfo = GetChangeLayoutInfo(Change);
-							if (LayoutInfo.BuildBadges.Count > 0)
+							ChangeLayoutInfo layoutInfo = GetChangeLayoutInfo(change);
+							if (layoutInfo.BuildBadges.Count > 0)
 							{
-								Point BuildListLocation = GetBadgeListLocation(LayoutInfo.BuildBadges, HitTest.SubItem.Bounds, HorizontalAlign.Center, VerticalAlignment.Middle);
-								BuildListLocation.X = Math.Max(BuildListLocation.X, HitTest.SubItem.Bounds.Left);
+								Point buildListLocation = GetBadgeListLocation(layoutInfo.BuildBadges, hitTest.SubItem.Bounds, HorizontalAlign.Center, VerticalAlignment.Middle);
+								buildListLocation.X = Math.Max(buildListLocation.X, hitTest.SubItem.Bounds.Left);
 
-								BadgeInfo? BadgeInfo = HitTestBadge(Args.Location, LayoutInfo.BuildBadges, BuildListLocation);
-								if (BadgeInfo != null && BadgeInfo.ClickHandler != null)
+								BadgeInfo? badgeInfo = HitTestBadge(args.Location, layoutInfo.BuildBadges, buildListLocation);
+								if (badgeInfo != null && badgeInfo.ClickHandler != null)
 								{
-									BadgeInfo.ClickHandler();
+									badgeInfo.ClickHandler();
 								}
 							}
 						}
 
-						foreach (ColumnHeader CustomColumn in CustomColumns)
+						foreach (ColumnHeader customColumn in _customColumns)
 						{
-							if (CustomColumn.Index < HitTest.Item.SubItems.Count && HitTest.Item.SubItems[CustomColumn.Index] == HitTest.SubItem)
+							if (customColumn.Index < hitTest.Item.SubItems.Count && hitTest.Item.SubItems[customColumn.Index] == hitTest.SubItem)
 							{
-								ChangeLayoutInfo LayoutInfo = GetChangeLayoutInfo((ChangesRecord)HitTest.Item.Tag);
+								ChangeLayoutInfo layoutInfo = GetChangeLayoutInfo((ChangesRecord)hitTest.Item.Tag);
 
-								List<BadgeInfo>? Badges;
-								if (LayoutInfo.CustomBadges.TryGetValue(CustomColumn.Text, out Badges) && Badges.Count > 0)
+								List<BadgeInfo>? badges;
+								if (layoutInfo.CustomBadges.TryGetValue(customColumn.Text, out badges) && badges.Count > 0)
 								{
-									Point ListLocation = GetBadgeListLocation(Badges, HitTest.SubItem.Bounds, HorizontalAlign.Center, VerticalAlignment.Middle);
+									Point listLocation = GetBadgeListLocation(badges, hitTest.SubItem.Bounds, HorizontalAlign.Center, VerticalAlignment.Middle);
 
-									BadgeInfo? BadgeInfo = HitTestBadge(Args.Location, Badges, ListLocation);
-									if (BadgeInfo != null && BadgeInfo.ClickHandler != null)
+									BadgeInfo? badgeInfo = HitTestBadge(args.Location, badges, listLocation);
+									if (badgeInfo != null && badgeInfo.ClickHandler != null)
 									{
-										BadgeInfo.ClickHandler();
+										badgeInfo.ClickHandler();
 									}
 								}
 							}
@@ -4064,115 +4064,115 @@ namespace UnrealGameSync
 					}
 				}
 			}
-			else if (Args.Button == MouseButtons.Right)
+			else if (args.Button == MouseButtons.Right)
 			{
-				ListViewHitTestInfo HitTest = BuildList.HitTest(Args.Location);
-				if (HitTest.Item != null && HitTest.Item.Tag != null)
+				ListViewHitTestInfo hitTest = BuildList.HitTest(args.Location);
+				if (hitTest.Item != null && hitTest.Item.Tag != null)
 				{
-					if (BuildList.SelectedItems.Count > 1 && BuildList.SelectedItems.Contains(HitTest.Item))
+					if (BuildList.SelectedItems.Count > 1 && BuildList.SelectedItems.Contains(hitTest.Item))
 					{
-						bool bIsTimeColumn = (HitTest.Item.SubItems.IndexOf(HitTest.SubItem) == TimeColumn.Index);
-						BuildListMultiContextMenu_TimeZoneSeparator.Visible = bIsTimeColumn;
-						BuildListMultiContextMenu_ShowLocalTimes.Visible = bIsTimeColumn;
-						BuildListMultiContextMenu_ShowLocalTimes.Checked = Settings.bShowLocalTimes;
-						BuildListMultiContextMenu_ShowServerTimes.Visible = bIsTimeColumn;
-						BuildListMultiContextMenu_ShowServerTimes.Checked = !Settings.bShowLocalTimes;
-
-						BuildListMultiContextMenu.Show(BuildList, Args.Location);
+						bool isTimeColumn = (hitTest.Item.SubItems.IndexOf(hitTest.SubItem) == TimeColumn.Index);
+						BuildListMultiContextMenu_TimeZoneSeparator.Visible = isTimeColumn;
+						BuildListMultiContextMenu_ShowLocalTimes.Visible = isTimeColumn;
+						BuildListMultiContextMenu_ShowLocalTimes.Checked = _settings.ShowLocalTimes;
+						BuildListMultiContextMenu_ShowServerTimes.Visible = isTimeColumn;
+						BuildListMultiContextMenu_ShowServerTimes.Checked = !_settings.ShowLocalTimes;
+
+						BuildListMultiContextMenu.Show(BuildList, args.Location);
 					}
 					else
 					{
-						ContextMenuChange = (ChangesRecord)HitTest.Item.Tag;
-
-						BuildListContextMenu_WithdrawReview.Visible = (EventMonitor.GetReviewByCurrentUser(ContextMenuChange.Number) != null);
-						BuildListContextMenu_StartInvestigating.Visible = !EventMonitor.IsUnderInvestigationByCurrentUser(ContextMenuChange.Number);
-						BuildListContextMenu_FinishInvestigating.Visible = EventMonitor.IsUnderInvestigation(ContextMenuChange.Number);
-
-						string? CommentText;
-						bool bHasExistingComment = EventMonitor.GetCommentByCurrentUser(ContextMenuChange.Number, out CommentText);
-						BuildListContextMenu_LeaveComment.Visible = !bHasExistingComment;
-						BuildListContextMenu_EditComment.Visible = bHasExistingComment;
-
-						bool bIsBusy = Workspace.IsBusy();
-						bool bIsCurrentChange = (ContextMenuChange.Number == Workspace.CurrentChangeNumber);
-						BuildListContextMenu_Sync.Visible = !bIsBusy;
-						BuildListContextMenu_Sync.Font = new Font(SystemFonts.MenuFont, bIsCurrentChange ? FontStyle.Regular : FontStyle.Bold);
-						BuildListContextMenu_SyncContentOnly.Visible = !bIsBusy && ShouldSyncPrecompiledEditor;
-						BuildListContextMenu_SyncOnlyThisChange.Visible = !bIsBusy && !bIsCurrentChange && ContextMenuChange.Number > Workspace.CurrentChangeNumber && Workspace.CurrentChangeNumber != -1;
-						BuildListContextMenu_Build.Visible = !bIsBusy && bIsCurrentChange && !ShouldSyncPrecompiledEditor;
-						BuildListContextMenu_Rebuild.Visible = !bIsBusy && bIsCurrentChange && !ShouldSyncPrecompiledEditor;
-						BuildListContextMenu_GenerateProjectFiles.Visible = !bIsBusy && bIsCurrentChange;
-						BuildListContextMenu_LaunchEditor.Visible = !bIsBusy && ContextMenuChange.Number == Workspace.CurrentChangeNumber;
+						_contextMenuChange = (ChangesRecord)hitTest.Item.Tag;
+
+						BuildListContextMenu_WithdrawReview.Visible = (_eventMonitor.GetReviewByCurrentUser(_contextMenuChange.Number) != null);
+						BuildListContextMenu_StartInvestigating.Visible = !_eventMonitor.IsUnderInvestigationByCurrentUser(_contextMenuChange.Number);
+						BuildListContextMenu_FinishInvestigating.Visible = _eventMonitor.IsUnderInvestigation(_contextMenuChange.Number);
+
+						string? commentText;
+						bool hasExistingComment = _eventMonitor.GetCommentByCurrentUser(_contextMenuChange.Number, out commentText);
+						BuildListContextMenu_LeaveComment.Visible = !hasExistingComment;
+						BuildListContextMenu_EditComment.Visible = hasExistingComment;
+
+						bool isBusy = _workspace.IsBusy();
+						bool isCurrentChange = (_contextMenuChange.Number == _workspace.CurrentChangeNumber);
+						BuildListContextMenu_Sync.Visible = !isBusy;
+						BuildListContextMenu_Sync.Font = new Font(SystemFonts.MenuFont, isCurrentChange ? FontStyle.Regular : FontStyle.Bold);
+						BuildListContextMenu_SyncContentOnly.Visible = !isBusy && ShouldSyncPrecompiledEditor;
+						BuildListContextMenu_SyncOnlyThisChange.Visible = !isBusy && !isCurrentChange && _contextMenuChange.Number > _workspace.CurrentChangeNumber && _workspace.CurrentChangeNumber != -1;
+						BuildListContextMenu_Build.Visible = !isBusy && isCurrentChange && !ShouldSyncPrecompiledEditor;
+						BuildListContextMenu_Rebuild.Visible = !isBusy && isCurrentChange && !ShouldSyncPrecompiledEditor;
+						BuildListContextMenu_GenerateProjectFiles.Visible = !isBusy && isCurrentChange;
+						BuildListContextMenu_LaunchEditor.Visible = !isBusy && _contextMenuChange.Number == _workspace.CurrentChangeNumber;
 						BuildListContextMenu_LaunchEditor.Font = new Font(SystemFonts.MenuFont, FontStyle.Bold);
-						BuildListContextMenu_OpenVisualStudio.Visible = !bIsBusy && bIsCurrentChange;
-						BuildListContextMenu_Cancel.Visible = bIsBusy;
-
-						BisectState State = WorkspaceState.BisectChanges.FirstOrDefault(x => x.Change == ContextMenuChange.Number)?.State ?? BisectState.Include;
-						bool bIsBisectMode = IsBisectModeEnabled();
-						BuildListContextMenu_Bisect_Pass.Visible = bIsBisectMode && State != BisectState.Pass;
-						BuildListContextMenu_Bisect_Fail.Visible = bIsBisectMode && State != BisectState.Fail;
-						BuildListContextMenu_Bisect_Exclude.Visible = bIsBisectMode && State != BisectState.Exclude;
-						BuildListContextMenu_Bisect_Include.Visible = bIsBisectMode && State != BisectState.Include;
-						BuildListContextMenu_Bisect_Separator.Visible = bIsBisectMode;
-
-						BuildListContextMenu_MarkGood.Visible = !bIsBisectMode;
-						BuildListContextMenu_MarkBad.Visible = !bIsBisectMode;
-						BuildListContextMenu_WithdrawReview.Visible = !bIsBisectMode;
-
-						EventSummary? Summary = EventMonitor.GetSummaryForChange(ContextMenuChange.Number);
-						bool bStarred = (Summary != null && Summary.LastStarReview != null && Summary.LastStarReview.Type == EventType.Starred);
-						BuildListContextMenu_AddStar.Visible = !bStarred;
-						BuildListContextMenu_RemoveStar.Visible = bStarred;
-
-						bool bIsTimeColumn = (HitTest.Item.SubItems.IndexOf(HitTest.SubItem) == TimeColumn.Index);
-						BuildListContextMenu_TimeZoneSeparator.Visible = bIsTimeColumn;
-						BuildListContextMenu_ShowLocalTimes.Visible = bIsTimeColumn;
-						BuildListContextMenu_ShowLocalTimes.Checked = Settings.bShowLocalTimes;
-						BuildListContextMenu_ShowServerTimes.Visible = bIsTimeColumn;
-						BuildListContextMenu_ShowServerTimes.Checked = !Settings.bShowLocalTimes;
-
-						int CustomToolStart = BuildListContextMenu.Items.IndexOf(BuildListContextMenu_CustomTool_Start) + 1;
-						int CustomToolEnd = BuildListContextMenu.Items.IndexOf(BuildListContextMenu_CustomTool_End);
-						while (CustomToolEnd > CustomToolStart)
+						BuildListContextMenu_OpenVisualStudio.Visible = !isBusy && isCurrentChange;
+						BuildListContextMenu_Cancel.Visible = isBusy;
+
+						BisectState state = _workspaceState.BisectChanges.FirstOrDefault(x => x.Change == _contextMenuChange.Number)?.State ?? BisectState.Include;
+						bool isBisectMode = IsBisectModeEnabled();
+						BuildListContextMenu_Bisect_Pass.Visible = isBisectMode && state != BisectState.Pass;
+						BuildListContextMenu_Bisect_Fail.Visible = isBisectMode && state != BisectState.Fail;
+						BuildListContextMenu_Bisect_Exclude.Visible = isBisectMode && state != BisectState.Exclude;
+						BuildListContextMenu_Bisect_Include.Visible = isBisectMode && state != BisectState.Include;
+						BuildListContextMenu_Bisect_Separator.Visible = isBisectMode;
+
+						BuildListContextMenu_MarkGood.Visible = !isBisectMode;
+						BuildListContextMenu_MarkBad.Visible = !isBisectMode;
+						BuildListContextMenu_WithdrawReview.Visible = !isBisectMode;
+
+						EventSummary? summary = _eventMonitor.GetSummaryForChange(_contextMenuChange.Number);
+						bool starred = (summary != null && summary.LastStarReview != null && summary.LastStarReview.Type == EventType.Starred);
+						BuildListContextMenu_AddStar.Visible = !starred;
+						BuildListContextMenu_RemoveStar.Visible = starred;
+
+						bool isTimeColumn = (hitTest.Item.SubItems.IndexOf(hitTest.SubItem) == TimeColumn.Index);
+						BuildListContextMenu_TimeZoneSeparator.Visible = isTimeColumn;
+						BuildListContextMenu_ShowLocalTimes.Visible = isTimeColumn;
+						BuildListContextMenu_ShowLocalTimes.Checked = _settings.ShowLocalTimes;
+						BuildListContextMenu_ShowServerTimes.Visible = isTimeColumn;
+						BuildListContextMenu_ShowServerTimes.Checked = !_settings.ShowLocalTimes;
+
+						int customToolStart = BuildListContextMenu.Items.IndexOf(BuildListContextMenu_CustomTool_Start) + 1;
+						int customToolEnd = BuildListContextMenu.Items.IndexOf(BuildListContextMenu_CustomTool_End);
+						while (customToolEnd > customToolStart)
 						{
-							BuildListContextMenu.Items.RemoveAt(CustomToolEnd - 1);
-							CustomToolEnd--;
+							BuildListContextMenu.Items.RemoveAt(customToolEnd - 1);
+							customToolEnd--;
 						}
 
-						ConfigFile? ProjectConfigFile = PerforceMonitor.LatestProjectConfigFile;
-						if (ProjectConfigFile != null)
+						ConfigFile? projectConfigFile = _perforceMonitor.LatestProjectConfigFile;
+						if (projectConfigFile != null)
 						{
-							Dictionary<string, string> Variables = Workspace.GetVariables(GetEditorBuildConfig(), ContextMenuChange.Number, -1);
+							Dictionary<string, string> variables = _workspace.GetVariables(GetEditorBuildConfig(), _contextMenuChange.Number, -1);
 
-							string[] ChangeContextMenuEntries = ProjectConfigFile.GetValues("Options.ContextMenu", new string[0]);
-							foreach (string ChangeContextMenuEntry in ChangeContextMenuEntries)
+							string[] changeContextMenuEntries = projectConfigFile.GetValues("Options.ContextMenu", new string[0]);
+							foreach (string changeContextMenuEntry in changeContextMenuEntries)
 							{
-								ConfigObject Object = new ConfigObject(ChangeContextMenuEntry);
+								ConfigObject obj = new ConfigObject(changeContextMenuEntry);
 
-								string? Label = Object.GetValue("Label");
-								string? Execute = Object.GetValue("Execute");
-								string? Arguments = Object.GetValue("Arguments");
+								string? label = obj.GetValue("Label");
+								string? execute = obj.GetValue("Execute");
+								string? arguments = obj.GetValue("Arguments");
 
-								if (Label != null && Execute != null)
+								if (label != null && execute != null)
 								{
-									Label = Utility.ExpandVariables(Label, Variables);
-									Execute = Utility.ExpandVariables(Execute, Variables);
-									Arguments = Utility.ExpandVariables(Arguments ?? "", Variables);
+									label = Utility.ExpandVariables(label, variables);
+									execute = Utility.ExpandVariables(execute, variables);
+									arguments = Utility.ExpandVariables(arguments ?? "", variables);
 
-									ToolStripMenuItem Item = new ToolStripMenuItem(Label, null, new EventHandler((o, a) => SafeProcessStart(Execute, Arguments)));
+									ToolStripMenuItem item = new ToolStripMenuItem(label, null, new EventHandler((o, a) => SafeProcessStart(execute, arguments)));
 
-									BuildListContextMenu.Items.Insert(CustomToolEnd, Item);
-									CustomToolEnd++;
+									BuildListContextMenu.Items.Insert(customToolEnd, item);
+									customToolEnd++;
 								}
 							}
 						}
 
-						BuildListContextMenu_CustomTool_End.Visible = (CustomToolEnd > CustomToolStart);
+						BuildListContextMenu_CustomTool_End.Visible = (customToolEnd > customToolStart);
 
-						string? SwarmURL;
-						BuildListContextMenu_ViewInSwarm.Visible = ProjectConfigFile != null && TryGetProjectSetting(ProjectConfigFile, "SwarmURL", out SwarmURL);
+						string? swarmUrl;
+						BuildListContextMenu_ViewInSwarm.Visible = projectConfigFile != null && TryGetProjectSetting(projectConfigFile, "SwarmURL", out swarmUrl);
 							
-						BuildListContextMenu.Show(BuildList, Args.Location);
+						BuildListContextMenu.Show(BuildList, args.Location);
 					}
 				}
 			}
@@ -4180,45 +4180,45 @@ namespace UnrealGameSync
 
 		private void BuildList_MouseMove(object sender, MouseEventArgs e)
 		{
-			ListViewHitTestInfo HitTest = BuildList.HitTest(e.Location);
+			ListViewHitTestInfo hitTest = BuildList.HitTest(e.Location);
 
-			bool bNewMouseOverExpandLink = HitTest.Item != null && HitTestExpandLink(e.Location);
-			if (bMouseOverExpandLink != bNewMouseOverExpandLink)
+			bool newMouseOverExpandLink = hitTest.Item != null && HitTestExpandLink(e.Location);
+			if (_mouseOverExpandLink != newMouseOverExpandLink)
 			{
-				bMouseOverExpandLink = bNewMouseOverExpandLink;
-				Cursor = bMouseOverExpandLink ? NativeCursors.Hand : Cursors.Arrow;
+				_mouseOverExpandLink = newMouseOverExpandLink;
+				Cursor = _mouseOverExpandLink ? NativeCursors.Hand : Cursors.Arrow;
 				BuildList.Invalidate();
 			}
 
-			string? NewHoverBadgeUniqueId = null;
-			if (HitTest.Item != null && HitTest.Item.Tag is ChangesRecord)
+			string? newHoverBadgeUniqueId = null;
+			if (hitTest.Item != null && hitTest.Item.Tag is ChangesRecord)
 			{
-				int ColumnIndex = HitTest.Item.SubItems.IndexOf(HitTest.SubItem);
-				if (ColumnIndex == DescriptionColumn.Index)
+				int columnIndex = hitTest.Item.SubItems.IndexOf(hitTest.SubItem);
+				if (columnIndex == DescriptionColumn.Index)
 				{
-					ChangeLayoutInfo LayoutInfo = GetChangeLayoutInfo((ChangesRecord)HitTest.Item.Tag);
-					if (LayoutInfo.DescriptionBadges.Count > 0)
+					ChangeLayoutInfo layoutInfo = GetChangeLayoutInfo((ChangesRecord)hitTest.Item.Tag);
+					if (layoutInfo.DescriptionBadges.Count > 0)
 					{
-						Point ListLocation = GetBadgeListLocation(LayoutInfo.DescriptionBadges, HitTest.SubItem.Bounds, HorizontalAlign.Right, VerticalAlignment.Middle);
-						NewHoverBadgeUniqueId = HitTestBadge(e.Location, LayoutInfo.DescriptionBadges, ListLocation)?.UniqueId;
+						Point listLocation = GetBadgeListLocation(layoutInfo.DescriptionBadges, hitTest.SubItem.Bounds, HorizontalAlign.Right, VerticalAlignment.Middle);
+						newHoverBadgeUniqueId = HitTestBadge(e.Location, layoutInfo.DescriptionBadges, listLocation)?.UniqueId;
 					}
 				}
-				else if (ColumnIndex == CISColumn.Index)
+				else if (columnIndex == CISColumn.Index)
 				{
-					ChangeLayoutInfo LayoutInfo = GetChangeLayoutInfo((ChangesRecord)HitTest.Item.Tag);
-					if (LayoutInfo.BuildBadges.Count > 0)
+					ChangeLayoutInfo layoutInfo = GetChangeLayoutInfo((ChangesRecord)hitTest.Item.Tag);
+					if (layoutInfo.BuildBadges.Count > 0)
 					{
-						Point BuildListLocation = GetBadgeListLocation(LayoutInfo.BuildBadges, HitTest.SubItem.Bounds, HorizontalAlign.Center, VerticalAlignment.Middle);
-						BuildListLocation.X = Math.Max(BuildListLocation.X, HitTest.SubItem.Bounds.Left);
+						Point buildListLocation = GetBadgeListLocation(layoutInfo.BuildBadges, hitTest.SubItem.Bounds, HorizontalAlign.Center, VerticalAlignment.Middle);
+						buildListLocation.X = Math.Max(buildListLocation.X, hitTest.SubItem.Bounds.Left);
 
-						BadgeInfo? Badge = HitTestBadge(e.Location, LayoutInfo.BuildBadges, BuildListLocation);
-						NewHoverBadgeUniqueId = (Badge != null) ? Badge.UniqueId : null;
+						BadgeInfo? badge = HitTestBadge(e.Location, layoutInfo.BuildBadges, buildListLocation);
+						newHoverBadgeUniqueId = (badge != null) ? badge.UniqueId : null;
 
-						if (HoverBadgeUniqueId != NewHoverBadgeUniqueId)
+						if (_hoverBadgeUniqueId != newHoverBadgeUniqueId)
 						{
-							if (Badge != null && Badge.ToolTip != null && Badge.BackgroundColor.A != 0)
+							if (badge != null && badge.ToolTip != null && badge.BackgroundColor.A != 0)
 							{
-								BuildListToolTip.Show(Badge.ToolTip, BuildList, new Point(BuildListLocation.X + Badge.Offset, HitTest.Item.Bounds.Bottom + 2));
+								BuildListToolTip.Show(badge.ToolTip, BuildList, new Point(buildListLocation.X + badge.Offset, hitTest.Item.Bounds.Bottom + 2));
 							}
 							else
 							{
@@ -4227,105 +4227,105 @@ namespace UnrealGameSync
 						}
 					}
 				}
-				else if (CustomColumns.Contains(BuildList.Columns[ColumnIndex]))
+				else if (_customColumns.Contains(BuildList.Columns[columnIndex]))
 				{
-					ColumnHeader Column = BuildList.Columns[ColumnIndex];
+					ColumnHeader column = BuildList.Columns[columnIndex];
 
-					ChangeLayoutInfo LayoutInfo = GetChangeLayoutInfo((ChangesRecord)HitTest.Item.Tag);
+					ChangeLayoutInfo layoutInfo = GetChangeLayoutInfo((ChangesRecord)hitTest.Item.Tag);
 
-					List<BadgeInfo>? Badges;
-					if (LayoutInfo.CustomBadges.TryGetValue(Column.Text, out Badges) && Badges.Count > 0)
+					List<BadgeInfo>? badges;
+					if (layoutInfo.CustomBadges.TryGetValue(column.Text, out badges) && badges.Count > 0)
 					{
-						Point ListLocation = GetBadgeListLocation(Badges, HitTest.SubItem.Bounds, HorizontalAlign.Center, VerticalAlignment.Middle);
-						NewHoverBadgeUniqueId = HitTestBadge(e.Location, Badges, ListLocation)?.UniqueId;
+						Point listLocation = GetBadgeListLocation(badges, hitTest.SubItem.Bounds, HorizontalAlign.Center, VerticalAlignment.Middle);
+						newHoverBadgeUniqueId = HitTestBadge(e.Location, badges, listLocation)?.UniqueId;
 					}
 				}
 			}
-			if (HoverBadgeUniqueId != NewHoverBadgeUniqueId)
+			if (_hoverBadgeUniqueId != newHoverBadgeUniqueId)
 			{
-				HoverBadgeUniqueId = NewHoverBadgeUniqueId;
+				_hoverBadgeUniqueId = newHoverBadgeUniqueId;
 				BuildList.Invalidate();
 			}
 
-			bool bNewHoverSync = false;
-			if (HitTest.Item != null && StatusColumn.Index < HitTest.Item.SubItems.Count)
+			bool newHoverSync = false;
+			if (hitTest.Item != null && StatusColumn.Index < hitTest.Item.SubItems.Count)
 			{
-				bNewHoverSync = GetSyncBadgeRectangle(HitTest.Item.SubItems[StatusColumn.Index].Bounds).Contains(e.Location);
+				newHoverSync = GetSyncBadgeRectangle(hitTest.Item.SubItems[StatusColumn.Index].Bounds).Contains(e.Location);
 			}
-			if (bNewHoverSync != bHoverSync)
+			if (newHoverSync != _hoverSync)
 			{
-				bHoverSync = bNewHoverSync;
+				_hoverSync = newHoverSync;
 				BuildList.Invalidate();
 			}
 		}
 
-		private BadgeInfo? HitTestBadge(Point Location, List<BadgeInfo> BadgeList, Point ListLocation)
+		private BadgeInfo? HitTestBadge(Point location, List<BadgeInfo> badgeList, Point listLocation)
 		{
-			foreach (BadgeInfo Badge in BadgeList)
+			foreach (BadgeInfo badge in badgeList)
 			{
-				Rectangle BadgeBounds = Badge.GetBounds(ListLocation);
-				if (BadgeBounds.Contains(Location))
+				Rectangle badgeBounds = badge.GetBounds(listLocation);
+				if (badgeBounds.Contains(location))
 				{
-					return Badge;
+					return badge;
 				}
 			}
 			return null;
 		}
 
-		private List<IArchiveInfo> GetSelectedArchives(IReadOnlyList<IArchiveInfo> Archives)
+		private List<IArchiveInfo> GetSelectedArchives(IReadOnlyList<IArchiveInfo> archives)
 		{
-			Dictionary<string, KeyValuePair<IArchiveInfo, int>> ArchiveTypeToSelection = new Dictionary<string, KeyValuePair<IArchiveInfo, int>>();
-			foreach (IArchiveInfo Archive in Archives)
+			Dictionary<string, KeyValuePair<IArchiveInfo, int>> archiveTypeToSelection = new Dictionary<string, KeyValuePair<IArchiveInfo, int>>();
+			foreach (IArchiveInfo archive in archives)
 			{
-				ArchiveSettings ArchiveSettings = Settings.Archives.FirstOrDefault(x => x.Type == Archive.Type);
-				if (ArchiveSettings != null && ArchiveSettings.bEnabled)
+				ArchiveSettings archiveSettings = _settings.Archives.FirstOrDefault(x => x.Type == archive.Type);
+				if (archiveSettings != null && archiveSettings.Enabled)
 				{
-					int Preference = ArchiveSettings.Order.IndexOf(Archive.Name);
-					if (Preference == -1)
+					int preference = archiveSettings.Order.IndexOf(archive.Name);
+					if (preference == -1)
 					{
-						Preference = ArchiveSettings.Order.Count;
+						preference = archiveSettings.Order.Count;
 					}
 
-					KeyValuePair<IArchiveInfo, int> ExistingItem;
-					if (!ArchiveTypeToSelection.TryGetValue(Archive.Type, out ExistingItem) || ExistingItem.Value > Preference)
+					KeyValuePair<IArchiveInfo, int> existingItem;
+					if (!archiveTypeToSelection.TryGetValue(archive.Type, out existingItem) || existingItem.Value > preference)
 					{
-						ArchiveTypeToSelection[Archive.Type] = new KeyValuePair<IArchiveInfo, int>(Archive, Preference);
+						archiveTypeToSelection[archive.Type] = new KeyValuePair<IArchiveInfo, int>(archive, preference);
 					}
 				}
 			}
-			return ArchiveTypeToSelection.Select(x => x.Value.Key).ToList();
+			return archiveTypeToSelection.Select(x => x.Value.Key).ToList();
 		}
 
 		private void ClearSelectedArchives()
 		{
-			foreach (ArchiveSettings ArchiveSettings in Settings.Archives)
+			foreach (ArchiveSettings archiveSettings in _settings.Archives)
 			{
-				ArchiveSettings.bEnabled = false;
+				archiveSettings.Enabled = false;
 			}
-			Settings.Save(Logger);
+			_settings.Save(_logger);
 
 			UpdateSelectedArchives();
 		}
 
-		private void SetSelectedArchive(IArchiveInfo Archive, bool bSelected)
+		private void SetSelectedArchive(IArchiveInfo archive, bool selected)
 		{
-			ArchiveSettings ArchiveSettings = Settings.Archives.FirstOrDefault(x => x.Type == Archive.Type);
-			if (ArchiveSettings == null)
+			ArchiveSettings archiveSettings = _settings.Archives.FirstOrDefault(x => x.Type == archive.Type);
+			if (archiveSettings == null)
 			{
-				ArchiveSettings = new ArchiveSettings(bSelected, Archive.Type, new string[] { Archive.Name });
-				Settings.Archives.Add(ArchiveSettings);
+				archiveSettings = new ArchiveSettings(selected, archive.Type, new string[] { archive.Name });
+				_settings.Archives.Add(archiveSettings);
 			}
-			else if (bSelected)
+			else if (selected)
 			{
-				ArchiveSettings.bEnabled = true;
-				ArchiveSettings.Order.Remove(Archive.Name);
-				ArchiveSettings.Order.Insert(0, Archive.Name);
+				archiveSettings.Enabled = true;
+				archiveSettings.Order.Remove(archive.Name);
+				archiveSettings.Order.Insert(0, archive.Name);
 			}
 			else
 			{
-				ArchiveSettings.bEnabled = false;
+				archiveSettings.Enabled = false;
 			}
-			Settings.Save(Logger);
+			_settings.Save(_logger);
 
 			UpdateSelectedArchives();
 		}
@@ -4340,12 +4340,12 @@ namespace UnrealGameSync
 
 		private void OptionsButton_Click(object sender, EventArgs e)
 		{
-			OptionsContextMenu_AutoResolveConflicts.Checked = Settings.bAutoResolveConflicts;
+			OptionsContextMenu_AutoResolveConflicts.Checked = _settings.AutoResolveConflicts;
 
 			OptionsContextMenu_SyncPrecompiledBinaries.DropDownItems.Clear();
 
-			IReadOnlyList<IArchiveInfo> Archives = GetArchives();
-			if (Archives == null || Archives.Count == 0)
+			IReadOnlyList<IArchiveInfo> archives = GetArchives();
+			if (archives == null || archives.Count == 0)
 			{
 				OptionsContextMenu_SyncPrecompiledBinaries.Enabled = false;
 				OptionsContextMenu_SyncPrecompiledBinaries.ToolTipText = String.Format("Precompiled binaries are not available for {0}", SelectedProjectIdentifier);
@@ -4353,63 +4353,63 @@ namespace UnrealGameSync
 			}
 			else
 			{
-				List<IArchiveInfo> SelectedArchives = GetSelectedArchives(Archives);
+				List<IArchiveInfo> selectedArchives = GetSelectedArchives(archives);
 
 				OptionsContextMenu_SyncPrecompiledBinaries.Enabled = true;
 				OptionsContextMenu_SyncPrecompiledBinaries.ToolTipText = null;
-				OptionsContextMenu_SyncPrecompiledBinaries.Checked = SelectedArchives.Count > 0;
+				OptionsContextMenu_SyncPrecompiledBinaries.Checked = selectedArchives.Count > 0;
 
-				if (Archives.Count > 1 || Archives[0].Type != IArchiveInfo.EditorArchiveType)
+				if (archives.Count > 1 || archives[0].Type != IArchiveInfo.EditorArchiveType)
 				{
-					ToolStripMenuItem DisableItem = new ToolStripMenuItem("Disable (compile locally)");
-					DisableItem.Checked = (SelectedArchives.Count == 0);
-					DisableItem.Click += (Sender, Args) => ClearSelectedArchives();
-					OptionsContextMenu_SyncPrecompiledBinaries.DropDownItems.Add(DisableItem);
+					ToolStripMenuItem disableItem = new ToolStripMenuItem("Disable (compile locally)");
+					disableItem.Checked = (selectedArchives.Count == 0);
+					disableItem.Click += (sender, args) => ClearSelectedArchives();
+					OptionsContextMenu_SyncPrecompiledBinaries.DropDownItems.Add(disableItem);
 
-					ToolStripSeparator Separator = new ToolStripSeparator();
-					OptionsContextMenu_SyncPrecompiledBinaries.DropDownItems.Add(Separator);
+					ToolStripSeparator separator = new ToolStripSeparator();
+					OptionsContextMenu_SyncPrecompiledBinaries.DropDownItems.Add(separator);
 
-					foreach (IArchiveInfo Archive in Archives)
+					foreach (IArchiveInfo archive in archives)
 					{
-						ToolStripMenuItem Item = new ToolStripMenuItem(Archive.Name);
-						Item.Enabled = Archive.Exists();
-						if (!Item.Enabled)
+						ToolStripMenuItem item = new ToolStripMenuItem(archive.Name);
+						item.Enabled = archive.Exists();
+						if (!item.Enabled)
 						{
-							Item.ToolTipText = String.Format("No valid archives found at {0}", Archive.BasePath);
+							item.ToolTipText = String.Format("No valid archives found at {0}", archive.BasePath);
 						}
-						Item.Checked = SelectedArchives.Contains(Archive);
-						Item.Click += (Sender, Args) => SetSelectedArchive(Archive, !Item.Checked);
-						OptionsContextMenu_SyncPrecompiledBinaries.DropDownItems.Add(Item);
+						item.Checked = selectedArchives.Contains(archive);
+						item.Click += (sender, args) => SetSelectedArchive(archive, !item.Checked);
+						OptionsContextMenu_SyncPrecompiledBinaries.DropDownItems.Add(item);
 					}
 				}
 			}
 
 			OptionsContextMenu_EditorBuildConfiguration.Enabled = !ShouldSyncPrecompiledEditor;
 			UpdateCheckedBuildConfig();
-			OptionsContextMenu_CustomizeBuildSteps.Enabled = (Workspace != null);
-			OptionsContextMenu_EditorArguments.Checked = Settings.EditorArguments.Any(x => x.Item2);
-			OptionsContextMenu_ScheduledSync.Checked = Settings.bScheduleEnabled;
-			OptionsContextMenu_TimeZone_Local.Checked = Settings.bShowLocalTimes;
-			OptionsContextMenu_TimeZone_PerforceServer.Checked = !Settings.bShowLocalTimes;
-			OptionsContextMenu_ShowChanges_ShowUnreviewed.Checked = Settings.bShowUnreviewedChanges;
-			OptionsContextMenu_ShowChanges_ShowAutomated.Checked = Settings.bShowAutomatedChanges;
-			OptionsContextMenu_TabNames_Stream.Checked = Settings.TabLabels == TabLabels.Stream;
-			OptionsContextMenu_TabNames_WorkspaceName.Checked = Settings.TabLabels == TabLabels.WorkspaceName;
-			OptionsContextMenu_TabNames_WorkspaceRoot.Checked = Settings.TabLabels == TabLabels.WorkspaceRoot;
-			OptionsContextMenu_TabNames_ProjectFile.Checked = Settings.TabLabels == TabLabels.ProjectFile;
+			OptionsContextMenu_CustomizeBuildSteps.Enabled = (_workspace != null);
+			OptionsContextMenu_EditorArguments.Checked = _settings.EditorArguments.Any(x => x.Item2);
+			OptionsContextMenu_ScheduledSync.Checked = _settings.ScheduleEnabled;
+			OptionsContextMenu_TimeZone_Local.Checked = _settings.ShowLocalTimes;
+			OptionsContextMenu_TimeZone_PerforceServer.Checked = !_settings.ShowLocalTimes;
+			OptionsContextMenu_ShowChanges_ShowUnreviewed.Checked = _settings.ShowUnreviewedChanges;
+			OptionsContextMenu_ShowChanges_ShowAutomated.Checked = _settings.ShowAutomatedChanges;
+			OptionsContextMenu_TabNames_Stream.Checked = _settings.TabLabels == TabLabels.Stream;
+			OptionsContextMenu_TabNames_WorkspaceName.Checked = _settings.TabLabels == TabLabels.WorkspaceName;
+			OptionsContextMenu_TabNames_WorkspaceRoot.Checked = _settings.TabLabels == TabLabels.WorkspaceRoot;
+			OptionsContextMenu_TabNames_ProjectFile.Checked = _settings.TabLabels == TabLabels.ProjectFile;
 			OptionsContextMenu.Show(OptionsButton, new Point(OptionsButton.Width - OptionsContextMenu.Size.Width, OptionsButton.Height));
 		}
 
 		private IReadOnlyList<IArchiveInfo> GetArchives()
 		{
-			IReadOnlyList<IArchiveInfo>? AvailableArchives = JupiterMonitor?.AvailableArchives;
-			if (AvailableArchives != null && AvailableArchives.Count != 0)
-				return AvailableArchives;
+			IReadOnlyList<IArchiveInfo>? availableArchives = _jupiterMonitor?.AvailableArchives;
+			if (availableArchives != null && availableArchives.Count != 0)
+				return availableArchives;
 
 			// if jupiter had no archives we fallback to the perforce monitor
-			if (PerforceMonitor != null)
+			if (_perforceMonitor != null)
 			{
-				return PerforceMonitor.AvailableArchives;
+				return _perforceMonitor.AvailableArchives;
 			}
 
 			return new List<IArchiveInfo>();
@@ -4417,35 +4417,35 @@ namespace UnrealGameSync
 
 		private void BuildAfterSyncCheckBox_CheckedChanged(object sender, EventArgs e)
 		{
-			Settings.bBuildAfterSync = BuildAfterSyncCheckBox.Checked;
-			Settings.Save(Logger);
+			_settings.BuildAfterSync = BuildAfterSyncCheckBox.Checked;
+			_settings.Save(_logger);
 
 			UpdateSyncActionCheckboxes();
 		}
 
 		private void RunAfterSyncCheckBox_CheckedChanged(object sender, EventArgs e)
 		{
-			Settings.bRunAfterSync = RunAfterSyncCheckBox.Checked;
-			Settings.Save(Logger);
+			_settings.RunAfterSync = RunAfterSyncCheckBox.Checked;
+			_settings.Save(_logger);
 
 			UpdateSyncActionCheckboxes();
 		}
 
 		private void OpenSolutionAfterSyncCheckBox_CheckedChanged(object sender, EventArgs e)
 		{
-			Settings.bOpenSolutionAfterSync = OpenSolutionAfterSyncCheckBox.Checked;
-			Settings.Save(Logger);
+			_settings.OpenSolutionAfterSync = OpenSolutionAfterSyncCheckBox.Checked;
+			_settings.Save(_logger);
 
 			UpdateSyncActionCheckboxes();
 		}
 
 		private void UpdateSyncActionCheckboxes()
 		{
-			BuildAfterSyncCheckBox.Checked = Settings.bBuildAfterSync;
-			RunAfterSyncCheckBox.Checked = Settings.bRunAfterSync;
-			OpenSolutionAfterSyncCheckBox.Checked = Settings.bOpenSolutionAfterSync;
+			BuildAfterSyncCheckBox.Checked = _settings.BuildAfterSync;
+			RunAfterSyncCheckBox.Checked = _settings.RunAfterSync;
+			OpenSolutionAfterSyncCheckBox.Checked = _settings.OpenSolutionAfterSync;
 
-			if (Workspace == null || Workspace.IsBusy())
+			if (_workspace == null || _workspace.IsBusy())
 			{
 				AfterSyncingLabel.Enabled = false;
 				BuildAfterSyncCheckBox.Enabled = false;
@@ -4455,7 +4455,7 @@ namespace UnrealGameSync
 			else
 			{
 				AfterSyncingLabel.Enabled = true;
-				BuildAfterSyncCheckBox.Enabled = bHasBuildSteps;
+				BuildAfterSyncCheckBox.Enabled = _hasBuildSteps;
 				RunAfterSyncCheckBox.Enabled = BuildAfterSyncCheckBox.Checked || ShouldSyncPrecompiledEditor;
 				OpenSolutionAfterSyncCheckBox.Enabled = !ShouldSyncPrecompiledEditor;
 			}
@@ -4463,48 +4463,48 @@ namespace UnrealGameSync
 
 		private void BuildList_FontChanged(object? sender, EventArgs e)
 		{
-			if (BuildFont != null)
+			if (_buildFont != null)
 			{
-				BuildFont.Dispose();
+				_buildFont.Dispose();
 			}
-			BuildFont = BuildList.Font;
+			_buildFont = BuildList.Font;
 
-			if (SelectedBuildFont != null)
+			if (_selectedBuildFont != null)
 			{
-				SelectedBuildFont.Dispose();
+				_selectedBuildFont.Dispose();
 			}
-			SelectedBuildFont = new Font(BuildFont, FontStyle.Bold);
+			_selectedBuildFont = new Font(_buildFont, FontStyle.Bold);
 
-			if (BadgeFont != null)
+			if (_badgeFont != null)
 			{
-				BadgeFont.Dispose();
+				_badgeFont.Dispose();
 			}
-			BadgeFont = new Font(BuildFont.FontFamily, BuildFont.SizeInPoints - 2, FontStyle.Bold);
+			_badgeFont = new Font(_buildFont.FontFamily, _buildFont.SizeInPoints - 2, FontStyle.Bold);
 		}
 
-		public void ExecCommand(string Description, string StatusText, string FileName, string Arguments, string WorkingDir, bool bUseLogWindow)
+		public void ExecCommand(string description, string statusText, string fileName, string arguments, string workingDir, bool useLogWindow)
 		{
-			if (Workspace != null && !Workspace.IsBusy())
+			if (_workspace != null && !_workspace.IsBusy())
 			{
-				Guid Id = Guid.NewGuid();
+				Guid id = Guid.NewGuid();
 
-				BuildStep CustomBuildStep = new BuildStep(Guid.NewGuid(), 0, Description, StatusText, 1, FileName, Arguments, WorkingDir, bUseLogWindow);
+				BuildStep customBuildStep = new BuildStep(Guid.NewGuid(), 0, description, statusText, 1, fileName, arguments, workingDir, useLogWindow);
 
-				List<ConfigObject> UserBuildSteps = new List<ConfigObject>();
-				UserBuildSteps.Add(CustomBuildStep.ToConfigObject());
+				List<ConfigObject> userBuildSteps = new List<ConfigObject>();
+				userBuildSteps.Add(customBuildStep.ToConfigObject());
 
-				WorkspaceUpdateContext Context = new WorkspaceUpdateContext(Workspace.CurrentChangeNumber, WorkspaceUpdateOptions.Build, GetEditorBuildConfig(), null, UserBuildSteps, new HashSet<Guid> { CustomBuildStep.UniqueId });
-				StartWorkspaceUpdate(Context, null);
+				WorkspaceUpdateContext context = new WorkspaceUpdateContext(_workspace.CurrentChangeNumber, WorkspaceUpdateOptions.Build, GetEditorBuildConfig(), null, userBuildSteps, new HashSet<Guid> { customBuildStep.UniqueId });
+				StartWorkspaceUpdate(context, null);
 			}
 		}
 
-		public void SyncChange(int ChangeNumber, bool bSyncOnly, WorkspaceUpdateCallback? Callback)
+		public void SyncChange(int changeNumber, bool syncOnly, WorkspaceUpdateCallback? callback)
 		{
-			if(Workspace != null)
+			if(_workspace != null)
 			{
-				Owner.ShowAndActivate();
-				SelectChange(ChangeNumber);
-				StartSync(ChangeNumber, bSyncOnly, Callback);
+				_owner.ShowAndActivate();
+				SelectChange(changeNumber);
+				StartSync(changeNumber, syncOnly, callback);
 			}
 		}
 
@@ -4513,96 +4513,96 @@ namespace UnrealGameSync
 			SyncLatestChange(null);
 		}
 
-		public void SyncLatestChange(WorkspaceUpdateCallback? Callback)
+		public void SyncLatestChange(WorkspaceUpdateCallback? callback)
 		{
-			if (Workspace != null)
+			if (_workspace != null)
 			{
-				int ChangeNumber;
-				if (!FindChangeToSync(Settings.SyncTypeID, out ChangeNumber))
+				int changeNumber;
+				if (!FindChangeToSync(_settings.SyncTypeId, out changeNumber))
 				{
-					string SyncTypeName = "";
-					foreach (LatestChangeType ChangeType in GetCustomLatestChangeTypes())
+					string syncTypeName = "";
+					foreach (LatestChangeType changeType in GetCustomLatestChangeTypes())
 					{
-						if (ChangeType.Name == Settings.SyncTypeID)
+						if (changeType.Name == _settings.SyncTypeId)
 						{
-							SyncTypeName = ChangeType.Name;
+							syncTypeName = changeType.Name;
 							break;
 						}
 					}
-					ShowErrorDialog(String.Format("Couldn't find any {0}changelist. Double-click on the change you want to sync manually.", SyncTypeName));
+					ShowErrorDialog(String.Format("Couldn't find any {0}changelist. Double-click on the change you want to sync manually.", syncTypeName));
 				}
-				else if (ChangeNumber < Workspace.CurrentChangeNumber)
+				else if (changeNumber < _workspace.CurrentChangeNumber)
 				{
-					ShowErrorDialog("Workspace is already synced to a later change ({0} vs {1}).", Workspace.CurrentChangeNumber, ChangeNumber);
+					ShowErrorDialog("Workspace is already synced to a later change ({0} vs {1}).", _workspace.CurrentChangeNumber, changeNumber);
 				}
-				else if (ChangeNumber >= PerforceMonitor.LastChangeByCurrentUser || MessageBox.Show(String.Format("The changelist that would be synced is before the last change you submitted.\n\nIf you continue, your changes submitted after CL {0} will be locally removed from your workspace until you can sync past them again.\n\nAre you sure you want to continue?", ChangeNumber), "Local changes will be removed", MessageBoxButtons.YesNo) == DialogResult.Yes)
+				else if (changeNumber >= _perforceMonitor.LastChangeByCurrentUser || MessageBox.Show(String.Format("The changelist that would be synced is before the last change you submitted.\n\nIf you continue, your changes submitted after CL {0} will be locally removed from your workspace until you can sync past them again.\n\nAre you sure you want to continue?", changeNumber), "Local changes will be removed", MessageBoxButtons.YesNo) == DialogResult.Yes)
 				{
-					SyncChange(ChangeNumber, false, Callback);
+					SyncChange(changeNumber, false, callback);
 				}
 			}
 		}
 
 		private void BuildListContextMenu_MarkGood_Click(object sender, EventArgs e)
 		{
-			if (ContextMenuChange != null)
+			if (_contextMenuChange != null)
 			{
-				EventMonitor.PostEvent(ContextMenuChange.Number, EventType.Good);
+				_eventMonitor.PostEvent(_contextMenuChange.Number, EventType.Good);
 			}
 		}
 
 		private void BuildListContextMenu_MarkBad_Click(object sender, EventArgs e)
 		{
-			if (ContextMenuChange != null)
+			if (_contextMenuChange != null)
 			{
-				EventMonitor.PostEvent(ContextMenuChange.Number, EventType.Bad);
+				_eventMonitor.PostEvent(_contextMenuChange.Number, EventType.Bad);
 			}
 		}
 
 		private void BuildListContextMenu_WithdrawReview_Click(object sender, EventArgs e)
 		{
-			if (ContextMenuChange != null)
+			if (_contextMenuChange != null)
 			{
-				EventMonitor.PostEvent(ContextMenuChange.Number, EventType.Unknown);
+				_eventMonitor.PostEvent(_contextMenuChange.Number, EventType.Unknown);
 			}
 		}
 
 		private void BuildListContextMenu_Sync_Click(object sender, EventArgs e)
 		{
-			if (ContextMenuChange != null)
+			if (_contextMenuChange != null)
 			{
-				StartSync(ContextMenuChange.Number);
+				StartSync(_contextMenuChange.Number);
 			}
 		}
 
 		private void BuildListContextMenu_SyncContentOnly_Click(object sender, EventArgs e)
 		{
-			if (ContextMenuChange != null)
+			if (_contextMenuChange != null)
 			{
-				StartWorkspaceUpdate(ContextMenuChange.Number, WorkspaceUpdateOptions.Sync | WorkspaceUpdateOptions.ContentOnly);
+				StartWorkspaceUpdate(_contextMenuChange.Number, WorkspaceUpdateOptions.Sync | WorkspaceUpdateOptions.ContentOnly);
 			}
 		}
 
 		private void BuildListContextMenu_SyncOnlyThisChange_Click(object sender, EventArgs e)
 		{
-			if (ContextMenuChange != null)
+			if (_contextMenuChange != null)
 			{
-				StartWorkspaceUpdate(ContextMenuChange.Number, WorkspaceUpdateOptions.SyncSingleChange);
+				StartWorkspaceUpdate(_contextMenuChange.Number, WorkspaceUpdateOptions.SyncSingleChange);
 			}
 		}
 
 		private void BuildListContextMenu_Build_Click(object sender, EventArgs e)
 		{
-			StartWorkspaceUpdate(Workspace.CurrentChangeNumber, WorkspaceUpdateOptions.SyncArchives | WorkspaceUpdateOptions.Build);
+			StartWorkspaceUpdate(_workspace.CurrentChangeNumber, WorkspaceUpdateOptions.SyncArchives | WorkspaceUpdateOptions.Build);
 		}
 
 		private void BuildListContextMenu_Rebuild_Click(object sender, EventArgs e)
 		{
-			StartWorkspaceUpdate(Workspace.CurrentChangeNumber, WorkspaceUpdateOptions.SyncArchives | WorkspaceUpdateOptions.Build | WorkspaceUpdateOptions.Clean);
+			StartWorkspaceUpdate(_workspace.CurrentChangeNumber, WorkspaceUpdateOptions.SyncArchives | WorkspaceUpdateOptions.Build | WorkspaceUpdateOptions.Clean);
 		}
 
 		private void BuildListContextMenu_GenerateProjectFiles_Click(object sender, EventArgs e)
 		{
-			StartWorkspaceUpdate(Workspace.CurrentChangeNumber, WorkspaceUpdateOptions.GenerateProjectFiles);
+			StartWorkspaceUpdate(_workspace.CurrentChangeNumber, WorkspaceUpdateOptions.GenerateProjectFiles);
 		}
 
 		private void BuildListContextMenu_CancelSync_Click(object sender, EventArgs e)
@@ -4612,20 +4612,20 @@ namespace UnrealGameSync
 
 		private void BuildListContextMenu_MoreInfo_Click(object sender, EventArgs e)
 		{
-			if (ContextMenuChange != null)
+			if (_contextMenuChange != null)
 			{
-				Program.SpawnP4VC(String.Format("{0} change {1}", PerforceSettings.GetArgumentsForExternalProgram(true), ContextMenuChange.Number));
+				Program.SpawnP4Vc(String.Format("{0} change {1}", _perforceSettings.GetArgumentsForExternalProgram(true), _contextMenuChange.Number));
 			}
 		}
 
 		private void BuildListContextMenu_ViewInSwarm_Click(object sender, EventArgs e)
 		{
-			if (ContextMenuChange != null)
+			if (_contextMenuChange != null)
 			{
-				string? SwarmURL;
-				if (TryGetProjectSetting(PerforceMonitor.LatestProjectConfigFile, "SwarmURL", out SwarmURL))
+				string? swarmUrl;
+				if (TryGetProjectSetting(_perforceMonitor.LatestProjectConfigFile, "SwarmURL", out swarmUrl))
 				{
-					Utility.OpenUrl(String.Format("{0}/changes/{1}", SwarmURL, ContextMenuChange.Number));
+					Utility.OpenUrl(String.Format("{0}/changes/{1}", swarmUrl, _contextMenuChange.Number));
 				}
 				else
 				{
@@ -4636,34 +4636,34 @@ namespace UnrealGameSync
 
 		private void BuildListContextMenu_CopyChangelistNumber_Click(object sender, EventArgs e)
 		{
-			if (ContextMenuChange != null)
+			if (_contextMenuChange != null)
 			{
-				System.Windows.Forms.Clipboard.SetText(ContextMenuChange.Number.ToString());
+				System.Windows.Forms.Clipboard.SetText(_contextMenuChange.Number.ToString());
 			}
 		}
 
 		private void BuildListContextMenu_AddStar_Click(object sender, EventArgs e)
 		{
-			if (ContextMenuChange != null)
+			if (_contextMenuChange != null)
 			{
 				if (MessageBox.Show("Starred builds are meant to convey a stable, verified build to the rest of the team. Do not star a build unless it has been fully tested.\n\nAre you sure you want to star this build?", "Confirm star", MessageBoxButtons.YesNo) == DialogResult.Yes)
 				{
-					EventMonitor.PostEvent(ContextMenuChange.Number, EventType.Starred);
+					_eventMonitor.PostEvent(_contextMenuChange.Number, EventType.Starred);
 				}
 			}
 		}
 
 		private void BuildListContextMenu_RemoveStar_Click(object sender, EventArgs e)
 		{
-			if (ContextMenuChange != null)
+			if (_contextMenuChange != null)
 			{
-				EventSummary? Summary = EventMonitor.GetSummaryForChange(ContextMenuChange.Number);
-				if (Summary != null && Summary.LastStarReview != null && Summary.LastStarReview.Type == EventType.Starred)
+				EventSummary? summary = _eventMonitor.GetSummaryForChange(_contextMenuChange.Number);
+				if (summary != null && summary.LastStarReview != null && summary.LastStarReview.Type == EventType.Starred)
 				{
-					string Message = String.Format("This change was starred by {0}. Are you sure you want to remove it?", FormatUserName(Summary.LastStarReview.UserName));
-					if (MessageBox.Show(Message, "Confirm removing star", MessageBoxButtons.YesNo) == DialogResult.Yes)
+					string message = String.Format("This change was starred by {0}. Are you sure you want to remove it?", FormatUserName(summary.LastStarReview.UserName));
+					if (MessageBox.Show(message, "Confirm removing star", MessageBoxButtons.YesNo) == DialogResult.Yes)
 					{
-						EventMonitor.PostEvent(ContextMenuChange.Number, EventType.Unstarred);
+						_eventMonitor.PostEvent(_contextMenuChange.Number, EventType.Unstarred);
 					}
 				}
 			}
@@ -4676,38 +4676,38 @@ namespace UnrealGameSync
 
 		private void BuildListContextMenu_StartInvestigating_Click(object sender, EventArgs e)
 		{
-			if (ContextMenuChange != null)
+			if (_contextMenuChange != null)
 			{
-				string Message = String.Format("All changes from {0} onwards will be marked as bad while you are investigating an issue.\n\nAre you sure you want to continue?", ContextMenuChange.Number);
-				if (MessageBox.Show(Message, "Confirm investigating", MessageBoxButtons.YesNo) == DialogResult.Yes)
+				string message = String.Format("All changes from {0} onwards will be marked as bad while you are investigating an issue.\n\nAre you sure you want to continue?", _contextMenuChange.Number);
+				if (MessageBox.Show(message, "Confirm investigating", MessageBoxButtons.YesNo) == DialogResult.Yes)
 				{
-					EventMonitor.StartInvestigating(ContextMenuChange.Number);
+					_eventMonitor.StartInvestigating(_contextMenuChange.Number);
 				}
 			}
 		}
 
 		private void BuildListContextMenu_FinishInvestigating_Click(object sender, EventArgs e)
 		{
-			if (ContextMenuChange != null)
+			if (_contextMenuChange != null)
 			{
-				int StartChangeNumber = EventMonitor.GetInvestigationStartChangeNumber(ContextMenuChange.Number);
-				if (StartChangeNumber != -1)
+				int startChangeNumber = _eventMonitor.GetInvestigationStartChangeNumber(_contextMenuChange.Number);
+				if (startChangeNumber != -1)
 				{
-					if (ContextMenuChange.Number > StartChangeNumber)
+					if (_contextMenuChange.Number > startChangeNumber)
 					{
-						string Message = String.Format("Mark all changes between {0} and {1} as bad?", StartChangeNumber, ContextMenuChange.Number);
-						if (MessageBox.Show(Message, "Finish investigating", MessageBoxButtons.YesNo) == DialogResult.Yes)
+						string message = String.Format("Mark all changes between {0} and {1} as bad?", startChangeNumber, _contextMenuChange.Number);
+						if (MessageBox.Show(message, "Finish investigating", MessageBoxButtons.YesNo) == DialogResult.Yes)
 						{
-							foreach (ChangesRecord Change in PerforceMonitor.GetChanges())
+							foreach (ChangesRecord change in _perforceMonitor.GetChanges())
 							{
-								if (Change.Number >= StartChangeNumber && Change.Number < ContextMenuChange.Number)
+								if (change.Number >= startChangeNumber && change.Number < _contextMenuChange.Number)
 								{
-									EventMonitor.PostEvent(Change.Number, EventType.Bad);
+									_eventMonitor.PostEvent(change.Number, EventType.Bad);
 								}
 							}
 						}
 					}
-					EventMonitor.FinishInvestigating(ContextMenuChange.Number);
+					_eventMonitor.FinishInvestigating(_contextMenuChange.Number);
 				}
 			}
 		}
@@ -4722,15 +4722,15 @@ namespace UnrealGameSync
 		{
 			UpdateSyncActionCheckboxes();
 
-			if (PerforceMonitor != null)
+			if (_perforceMonitor != null)
 			{
-				PerforceMonitor.Refresh();
+				_perforceMonitor.Refresh();
 			}
 
 			if (DesiredTaskbarState.Item1 == TaskbarState.Error)
 			{
 				DesiredTaskbarState = Tuple.Create(TaskbarState.NoProgress, 0.0f);
-				Owner.UpdateProgress();
+				_owner.UpdateProgress();
 			}
 		}
 
@@ -4739,66 +4739,66 @@ namespace UnrealGameSync
 			ShrinkNumRequestedBuilds();
 		}
 
-		private void BuildList_ItemMouseHover(object sender, ListViewItemMouseHoverEventArgs Args)
+		private void BuildList_ItemMouseHover(object sender, ListViewItemMouseHoverEventArgs args)
 		{
-			Point ClientPoint = BuildList.PointToClient(Cursor.Position);
-			if (Args.Item.SubItems.Count >= 6)
+			Point clientPoint = BuildList.PointToClient(Cursor.Position);
+			if (args.Item.SubItems.Count >= 6)
 			{
-				if (Args.Item.Bounds.Contains(ClientPoint))
+				if (args.Item.Bounds.Contains(clientPoint))
 				{
-					ChangesRecord Change = (ChangesRecord)Args.Item.Tag;
-					if (Change == null)
+					ChangesRecord change = (ChangesRecord)args.Item.Tag;
+					if (change == null)
 					{
 						return;
 					}
 
-					Rectangle CISBounds = Args.Item.SubItems[CISColumn.Index].Bounds;
-					if (CISBounds.Contains(ClientPoint))
+					Rectangle cisBounds = args.Item.SubItems[CISColumn.Index].Bounds;
+					if (cisBounds.Contains(clientPoint))
 					{
 						return;
 					}
 
-					Rectangle DescriptionBounds = Args.Item.SubItems[DescriptionColumn.Index].Bounds;
-					if (DescriptionBounds.Contains(ClientPoint))
+					Rectangle descriptionBounds = args.Item.SubItems[DescriptionColumn.Index].Bounds;
+					if (descriptionBounds.Contains(clientPoint))
 					{
-						ChangeLayoutInfo LayoutInfo = GetChangeLayoutInfo(Change);
-						if (LayoutInfo.DescriptionBadges.Count == 0 || ClientPoint.X < GetBadgeListLocation(LayoutInfo.DescriptionBadges, DescriptionBounds, HorizontalAlign.Right, VerticalAlignment.Middle).X - 2)
+						ChangeLayoutInfo layoutInfo = GetChangeLayoutInfo(change);
+						if (layoutInfo.DescriptionBadges.Count == 0 || clientPoint.X < GetBadgeListLocation(layoutInfo.DescriptionBadges, descriptionBounds, HorizontalAlign.Right, VerticalAlignment.Middle).X - 2)
 						{
-							BuildListToolTip.Show(Change.Description, BuildList, new Point(DescriptionBounds.Left, DescriptionBounds.Bottom + 2));
+							BuildListToolTip.Show(change.Description, BuildList, new Point(descriptionBounds.Left, descriptionBounds.Bottom + 2));
 							return;
 						}
 					}
 
-					EventSummary? Summary = EventMonitor.GetSummaryForChange(Change.Number);
-					if (Summary != null)
+					EventSummary? summary = _eventMonitor.GetSummaryForChange(change.Number);
+					if (summary != null)
 					{
-						StringBuilder SummaryText = new StringBuilder();
-						if (Summary.Comments.Count > 0)
+						StringBuilder summaryText = new StringBuilder();
+						if (summary.Comments.Count > 0)
 						{
-							foreach (CommentData Comment in Summary.Comments)
+							foreach (CommentData comment in summary.Comments)
 							{
-								if (!String.IsNullOrWhiteSpace(Comment.Text))
+								if (!String.IsNullOrWhiteSpace(comment.Text))
 								{
-									SummaryText.AppendFormat("{0}: \"{1}\"\n", FormatUserName(Comment.UserName), Comment.Text);
+									summaryText.AppendFormat("{0}: \"{1}\"\n", FormatUserName(comment.UserName), comment.Text);
 								}
 							}
-							if (SummaryText.Length > 0)
+							if (summaryText.Length > 0)
 							{
-								SummaryText.Append("\n");
+								summaryText.Append("\n");
 							}
 						}
-						AppendUserList(SummaryText, "\n", "Compiled by {0}.", Summary.Reviews.Where(x => x.Type == EventType.Compiles));
-						AppendUserList(SummaryText, "\n", "Failed to compile for {0}.", Summary.Reviews.Where(x => x.Type == EventType.DoesNotCompile));
-						AppendUserList(SummaryText, "\n", "Marked good by {0}.", Summary.Reviews.Where(x => x.Type == EventType.Good));
-						AppendUserList(SummaryText, "\n", "Marked bad by {0}.", Summary.Reviews.Where(x => x.Type == EventType.Bad));
-						if (Summary.LastStarReview != null)
+						AppendUserList(summaryText, "\n", "Compiled by {0}.", summary.Reviews.Where(x => x.Type == EventType.Compiles));
+						AppendUserList(summaryText, "\n", "Failed to compile for {0}.", summary.Reviews.Where(x => x.Type == EventType.DoesNotCompile));
+						AppendUserList(summaryText, "\n", "Marked good by {0}.", summary.Reviews.Where(x => x.Type == EventType.Good));
+						AppendUserList(summaryText, "\n", "Marked bad by {0}.", summary.Reviews.Where(x => x.Type == EventType.Bad));
+						if (summary.LastStarReview != null)
 						{
-							AppendUserList(SummaryText, "\n", "Starred by {0}.", new EventData[] { Summary.LastStarReview });
+							AppendUserList(summaryText, "\n", "Starred by {0}.", new EventData[] { summary.LastStarReview });
 						}
-						if (SummaryText.Length > 0)
+						if (summaryText.Length > 0)
 						{
-							Rectangle SummaryBounds = Args.Item.SubItems[StatusColumn.Index].Bounds;
-							BuildListToolTip.Show(SummaryText.ToString(), BuildList, new Point(SummaryBounds.Left, SummaryBounds.Bottom));
+							Rectangle summaryBounds = args.Item.SubItems[StatusColumn.Index].Bounds;
+							BuildListToolTip.Show(summaryText.ToString(), BuildList, new Point(summaryBounds.Left, summaryBounds.Bottom));
 							return;
 						}
 					}
@@ -4812,16 +4812,16 @@ namespace UnrealGameSync
 		{
 			BuildListToolTip.Hide(BuildList);
 
-			if (HoverBadgeUniqueId != null)
+			if (_hoverBadgeUniqueId != null)
 			{
-				HoverBadgeUniqueId = null;
+				_hoverBadgeUniqueId = null;
 				BuildList.Invalidate();
 			}
 
-			if (bMouseOverExpandLink)
+			if (_mouseOverExpandLink)
 			{
 				Cursor = Cursors.Arrow;
-				bMouseOverExpandLink = false;
+				_mouseOverExpandLink = false;
 				BuildList.Invalidate();
 			}
 		}
@@ -4836,10 +4836,10 @@ namespace UnrealGameSync
 			Splitter.SetLogVisibility(!Splitter.IsLogVisible());
 		}
 
-		private void Splitter_OnVisibilityChanged(bool bVisible)
+		private void Splitter_OnVisibilityChanged(bool visible)
 		{
-			Settings.bShowLogWindow = bVisible;
-			Settings.Save(Logger);
+			_settings.ShowLogWindow = visible;
+			_settings.Save(_logger);
 
 			UpdateStatusPanel();
 		}
@@ -4847,8 +4847,8 @@ namespace UnrealGameSync
 		private void OptionsContextMenu_AutoResolveConflicts_Click(object sender, EventArgs e)
 		{
 			OptionsContextMenu_AutoResolveConflicts.Checked ^= true;
-			Settings.bAutoResolveConflicts = OptionsContextMenu_AutoResolveConflicts.Checked;
-			Settings.Save(Logger);
+			_settings.AutoResolveConflicts = OptionsContextMenu_AutoResolveConflicts.Checked;
+			_settings.Save(_logger);
 		}
 
 		private void OptionsContextMenu_EditorArguments_Click(object sender, EventArgs e)
@@ -4858,12 +4858,12 @@ namespace UnrealGameSync
 
 		private bool ModifyEditorArguments()
 		{
-			ArgumentsWindow Arguments = new ArgumentsWindow(Settings.EditorArguments, Settings.bEditorArgumentsPrompt);
-			if (Arguments.ShowDialog(this) == DialogResult.OK)
+			ArgumentsWindow arguments = new ArgumentsWindow(_settings.EditorArguments, _settings.EditorArgumentsPrompt);
+			if (arguments.ShowDialog(this) == DialogResult.OK)
 			{
-				Settings.EditorArguments = Arguments.GetItems();
-				Settings.bEditorArgumentsPrompt = Arguments.PromptBeforeLaunch;
-				Settings.Save(Logger);
+				_settings.EditorArguments = arguments.GetItems();
+				_settings.EditorArgumentsPrompt = arguments.PromptBeforeLaunch;
+				_settings.Save(_logger);
 				return true;
 			}
 			return false;
@@ -4876,36 +4876,36 @@ namespace UnrealGameSync
 
 		private void OpenSolution()
 		{
-			string MasterProjectName = "UE4";
+			string masterProjectName = "UE4";
 			if (FileReference.Exists(FileReference.Combine(BranchDirectoryName, "UE5.sln")))
 			{
-				MasterProjectName = "UE5";
+				masterProjectName = "UE5";
 			}
 
-			FileReference MasterProjectNameFileName = FileReference.Combine(BranchDirectoryName, "Engine", "Intermediate", "ProjectFiles", "MasterProjectName.txt");
-			if (FileReference.Exists(MasterProjectNameFileName))
+			FileReference masterProjectNameFileName = FileReference.Combine(BranchDirectoryName, "Engine", "Intermediate", "ProjectFiles", "MasterProjectName.txt");
+			if (FileReference.Exists(masterProjectNameFileName))
 			{
 				try
 				{
-					MasterProjectName = FileReference.ReadAllText(MasterProjectNameFileName).Trim();
+					masterProjectName = FileReference.ReadAllText(masterProjectNameFileName).Trim();
 				}
-				catch (Exception Ex)
+				catch (Exception ex)
 				{
-					Logger.LogError(Ex, "Unable to read '{File}'", MasterProjectNameFileName);
+					_logger.LogError(ex, "Unable to read '{File}'", masterProjectNameFileName);
 				}
 			}
 
-			FileReference SolutionFileName = FileReference.Combine(BranchDirectoryName, MasterProjectName + ".sln");
-			if (!FileReference.Exists(SolutionFileName))
+			FileReference solutionFileName = FileReference.Combine(BranchDirectoryName, masterProjectName + ".sln");
+			if (!FileReference.Exists(solutionFileName))
 			{
-				MessageBox.Show(String.Format("Couldn't find solution at {0}", SolutionFileName));
+				MessageBox.Show(String.Format("Couldn't find solution at {0}", solutionFileName));
 			}
 			else
 			{
-				ProcessStartInfo StartInfo = new ProcessStartInfo(SolutionFileName.FullName);
-				StartInfo.UseShellExecute = true;
-				StartInfo.WorkingDirectory = BranchDirectoryName.FullName;
-				SafeProcessStart(StartInfo);
+				ProcessStartInfo startInfo = new ProcessStartInfo(solutionFileName.FullName);
+				startInfo.UseShellExecute = true;
+				startInfo.WorkingDirectory = BranchDirectoryName.FullName;
+				SafeProcessStart(startInfo);
 			}
 		}
 
@@ -4924,112 +4924,112 @@ namespace UnrealGameSync
 			UpdateBuildConfig(BuildConfig.Development);
 		}
 
-		void UpdateBuildConfig(BuildConfig NewConfig)
+		void UpdateBuildConfig(BuildConfig newConfig)
 		{
-			Settings.CompiledEditorBuildConfig = NewConfig;
-			Settings.Save(Logger);
+			_settings.CompiledEditorBuildConfig = newConfig;
+			_settings.Save(_logger);
 			UpdateCheckedBuildConfig();
 		}
 
 		void UpdateCheckedBuildConfig()
 		{
-			BuildConfig EditorBuildConfig = GetEditorBuildConfig();
+			BuildConfig editorBuildConfig = GetEditorBuildConfig();
 
-			OptionsContextMenu_BuildConfig_Debug.Checked = (EditorBuildConfig == BuildConfig.Debug);
-			OptionsContextMenu_BuildConfig_Debug.Enabled = (!ShouldSyncPrecompiledEditor || EditorBuildConfig == BuildConfig.Debug);
+			OptionsContextMenu_BuildConfig_Debug.Checked = (editorBuildConfig == BuildConfig.Debug);
+			OptionsContextMenu_BuildConfig_Debug.Enabled = (!ShouldSyncPrecompiledEditor || editorBuildConfig == BuildConfig.Debug);
 
-			OptionsContextMenu_BuildConfig_DebugGame.Checked = (EditorBuildConfig == BuildConfig.DebugGame);
-			OptionsContextMenu_BuildConfig_DebugGame.Enabled = (!ShouldSyncPrecompiledEditor || EditorBuildConfig == BuildConfig.DebugGame);
+			OptionsContextMenu_BuildConfig_DebugGame.Checked = (editorBuildConfig == BuildConfig.DebugGame);
+			OptionsContextMenu_BuildConfig_DebugGame.Enabled = (!ShouldSyncPrecompiledEditor || editorBuildConfig == BuildConfig.DebugGame);
 
-			OptionsContextMenu_BuildConfig_Development.Checked = (EditorBuildConfig == BuildConfig.Development);
-			OptionsContextMenu_BuildConfig_Development.Enabled = (!ShouldSyncPrecompiledEditor || EditorBuildConfig == BuildConfig.Development);
+			OptionsContextMenu_BuildConfig_Development.Checked = (editorBuildConfig == BuildConfig.Development);
+			OptionsContextMenu_BuildConfig_Development.Enabled = (!ShouldSyncPrecompiledEditor || editorBuildConfig == BuildConfig.Development);
 		}
 
 		private void OptionsContextMenu_ScheduleSync_Click(object sender, EventArgs e)
 		{
-			Owner.SetupScheduledSync();
+			_owner.SetupScheduledSync();
 		}
 
 		public void ScheduleTimerElapsed()
 		{
-			if (Settings.bScheduleEnabled)
+			if (_settings.ScheduleEnabled)
 			{
 				// Check the project config if they have a scheduled setting selected.
-				string? ScheduledSyncTypeID = null;
-				UserSelectedProjectSettings ProjectSetting = Settings.ScheduleProjects.FirstOrDefault(x => x.LocalPath != null && x.LocalPath.Equals(SelectedProject.LocalPath, StringComparison.OrdinalIgnoreCase));
-				if (ProjectSetting != null)
+				string? scheduledSyncTypeId = null;
+				UserSelectedProjectSettings projectSetting = _settings.ScheduleProjects.FirstOrDefault(x => x.LocalPath != null && x.LocalPath.Equals(SelectedProject.LocalPath, StringComparison.OrdinalIgnoreCase));
+				if (projectSetting != null)
 				{
-					ScheduledSyncTypeID = ProjectSetting.ScheduledSyncTypeID;
+					scheduledSyncTypeId = projectSetting.ScheduledSyncTypeId;
 				}
 
-				Logger.LogInformation("Scheduled sync at {0} for {1}.",
+				_logger.LogInformation("Scheduled sync at {0} for {1}.",
 					DateTime.Now,
-					ScheduledSyncTypeID != null ? ScheduledSyncTypeID : "Default");
+					scheduledSyncTypeId != null ? scheduledSyncTypeId : "Default");
 
-				int ChangeNumber;
-				if (!FindChangeToSync(ScheduledSyncTypeID, out ChangeNumber))
+				int changeNumber;
+				if (!FindChangeToSync(scheduledSyncTypeId, out changeNumber))
 				{
 					SyncLog.AppendLine("Couldn't find any matching change");
 				}
-				else if (Workspace.CurrentChangeNumber >= ChangeNumber)
+				else if (_workspace.CurrentChangeNumber >= changeNumber)
 				{
-					Logger.LogInformation("Sync ignored; already at or ahead of CL ({CurrentChange} >= {Change})", Workspace.CurrentChangeNumber, ChangeNumber);
+					_logger.LogInformation("Sync ignored; already at or ahead of CL ({CurrentChange} >= {Change})", _workspace.CurrentChangeNumber, changeNumber);
 				}
 				else
 				{
-					WorkspaceUpdateOptions Options = WorkspaceUpdateOptions.Sync | WorkspaceUpdateOptions.SyncArchives | WorkspaceUpdateOptions.GenerateProjectFiles | WorkspaceUpdateOptions.Build | WorkspaceUpdateOptions.ScheduledBuild;
-					if (Settings.bAutoResolveConflicts)
+					WorkspaceUpdateOptions options = WorkspaceUpdateOptions.Sync | WorkspaceUpdateOptions.SyncArchives | WorkspaceUpdateOptions.GenerateProjectFiles | WorkspaceUpdateOptions.Build | WorkspaceUpdateOptions.ScheduledBuild;
+					if (_settings.AutoResolveConflicts)
 					{
-						Options |= WorkspaceUpdateOptions.AutoResolveChanges;
+						options |= WorkspaceUpdateOptions.AutoResolveChanges;
 					}
-					StartWorkspaceUpdate(ChangeNumber, Options);
+					StartWorkspaceUpdate(changeNumber, options);
 				}
 			}
 		}
 
-		bool FindChangeToSync(string? ChangeTypeID, out int ChangeNumber)
+		bool FindChangeToSync(string? changeTypeId, out int changeNumber)
 		{
-			if (ChangeTypeID != null)
+			if (changeTypeId != null)
 			{
-				LatestChangeType? ChangeTypeToSyncTo = null;
-				foreach (LatestChangeType ChangeType in GetCustomLatestChangeTypes())
+				LatestChangeType? changeTypeToSyncTo = null;
+				foreach (LatestChangeType changeType in GetCustomLatestChangeTypes())
 				{
-					if (ChangeType.Name == ChangeTypeID)
+					if (changeType.Name == changeTypeId)
 					{
-						ChangeTypeToSyncTo = ChangeType;
+						changeTypeToSyncTo = changeType;
 						break;
 					}
 				}
 
-				if (ChangeTypeToSyncTo != null)
+				if (changeTypeToSyncTo != null)
 				{
-					return FindChangeToSync(ChangeTypeToSyncTo, out ChangeNumber);
+					return FindChangeToSync(changeTypeToSyncTo, out changeNumber);
 				}
 			}
 
 			SyncLog.AppendLine("Invalid or not set Scheduled Sync type using (Sync to Latest Good)");
 
-			return FindChangeToSync(LatestChangeType.LatestGoodChange(), out ChangeNumber);
+			return FindChangeToSync(LatestChangeType.LatestGoodChange(), out changeNumber);
 		}
 
-		bool FindChangeToSync(LatestChangeType ChangeType, out int ChangeNumber)
+		bool FindChangeToSync(LatestChangeType changeType, out int changeNumber)
 		{
-			for (int Idx = 0; Idx < BuildList.Items.Count; Idx++)
+			for (int idx = 0; idx < BuildList.Items.Count; idx++)
 			{
-				ChangesRecord Change = (ChangesRecord)BuildList.Items[Idx].Tag;
-				if (Change != null)
+				ChangesRecord change = (ChangesRecord)BuildList.Items[idx].Tag;
+				if (change != null)
 				{
-					EventSummary? Summary = EventMonitor.GetSummaryForChange(Change.Number);
-					if (CanSyncChange(Change.Number)
-						&& CanSyncChangeType(ChangeType, Change, Summary))
+					EventSummary? summary = _eventMonitor.GetSummaryForChange(change.Number);
+					if (CanSyncChange(change.Number)
+						&& CanSyncChangeType(changeType, change, summary))
 					{
-						if (ChangeType.bFindNewestGoodContent)
+						if (changeType.FindNewestGoodContent)
 						{
-							ChangeNumber = FindNewestGoodContentChange(Change.Number);
+							changeNumber = FindNewestGoodContentChange(change.Number);
 						}
 						else
 						{
-							ChangeNumber = Change.Number;
+							changeNumber = change.Number;
 						}
 						
 						return true;
@@ -5037,71 +5037,71 @@ namespace UnrealGameSync
 				}
 			}
 
-			ChangeNumber = -1;
+			changeNumber = -1;
 			return false;
 		}
 
-		int FindNewestGoodContentChange(int ChangeNumber)
+		int FindNewestGoodContentChange(int changeNumber)
 		{
-			int Index = SortedChangeNumbers.BinarySearch(ChangeNumber);
-			if (Index <= 0)
+			int index = _sortedChangeNumbers.BinarySearch(changeNumber);
+			if (index <= 0)
 			{
-				return ChangeNumber;
+				return changeNumber;
 			}
 
-			for (int NextIndex = Index + 1; NextIndex < SortedChangeNumbers.Count; NextIndex++)
+			for (int nextIndex = index + 1; nextIndex < _sortedChangeNumbers.Count; nextIndex++)
 			{
-				int NextChangeNumber = SortedChangeNumbers[NextIndex];
+				int nextChangeNumber = _sortedChangeNumbers[nextIndex];
 
-				PerforceChangeDetails? Details;
-				if (!PerforceMonitor.TryGetChangeDetails(NextChangeNumber, out Details) || Details.bContainsCode)
+				PerforceChangeDetails? details;
+				if (!_perforceMonitor.TryGetChangeDetails(nextChangeNumber, out details) || details.ContainsCode)
 				{
 					break;
 				}
 
-				EventSummary? Summary = EventMonitor.GetSummaryForChange(NextChangeNumber);
-				if (Summary != null && Summary.Verdict == ReviewVerdict.Bad)
+				EventSummary? summary = _eventMonitor.GetSummaryForChange(nextChangeNumber);
+				if (summary != null && summary.Verdict == ReviewVerdict.Bad)
 				{
 					break;
 				}
 
-				Index = NextIndex;
+				index = nextIndex;
 			}
 
-			return SortedChangeNumbers[Index];
+			return _sortedChangeNumbers[index];
 		}
 
 		private void BuildListContextMenu_LeaveOrEditComment_Click(object sender, EventArgs e)
 		{
-			if (ContextMenuChange != null)
+			if (_contextMenuChange != null)
 			{
-				LeaveCommentWindow LeaveComment = new LeaveCommentWindow();
+				LeaveCommentWindow leaveComment = new LeaveCommentWindow();
 
-				string? CommentText;
-				if (EventMonitor.GetCommentByCurrentUser(ContextMenuChange.Number, out CommentText))
+				string? commentText;
+				if (_eventMonitor.GetCommentByCurrentUser(_contextMenuChange.Number, out commentText))
 				{
-					LeaveComment.CommentTextBox.Text = CommentText;
+					leaveComment.CommentTextBox.Text = commentText;
 				}
 
-				if (LeaveComment.ShowDialog() == System.Windows.Forms.DialogResult.OK)
+				if (leaveComment.ShowDialog() == System.Windows.Forms.DialogResult.OK)
 				{
-					EventMonitor.PostComment(ContextMenuChange.Number, LeaveComment.CommentTextBox.Text);
+					_eventMonitor.PostComment(_contextMenuChange.Number, leaveComment.CommentTextBox.Text);
 				}
 			}
 		}
 
 		private void BuildListContextMenu_ShowServerTimes_Click(object sender, EventArgs e)
 		{
-			Settings.bShowLocalTimes = false;
-			Settings.Save(Logger);
+			_settings.ShowLocalTimes = false;
+			_settings.Save(_logger);
 			BuildList.Items.Clear(); // Need to clear list to rebuild groups, populate time column again
 			UpdateBuildList();
 		}
 
 		private void BuildListContextMenu_ShowLocalTimes_Click(object sender, EventArgs e)
 		{
-			Settings.bShowLocalTimes = true;
-			Settings.Save(Logger);
+			_settings.ShowLocalTimes = true;
+			_settings.Save(_logger);
 			BuildList.Items.Clear(); // Need to clear list to rebuild groups, populate time column again
 			UpdateBuildList();
 		}
@@ -5113,107 +5113,107 @@ namespace UnrealGameSync
 				return;
 			}
 
-			string? ExtraSafeToDeleteFolders;
-			if (!TryGetProjectSetting(Workspace.ProjectConfigFile, "SafeToDeleteFolders", out ExtraSafeToDeleteFolders))
+			string? extraSafeToDeleteFolders;
+			if (!TryGetProjectSetting(_workspace.ProjectConfigFile, "SafeToDeleteFolders", out extraSafeToDeleteFolders))
 			{
-				ExtraSafeToDeleteFolders = "";
+				extraSafeToDeleteFolders = "";
 			}
 
-			string? ExtraSafeToDeleteExtensions;
-			if (!TryGetProjectSetting(Workspace.ProjectConfigFile, "SafeToDeleteExtensions", out ExtraSafeToDeleteExtensions))
+			string? extraSafeToDeleteExtensions;
+			if (!TryGetProjectSetting(_workspace.ProjectConfigFile, "SafeToDeleteExtensions", out extraSafeToDeleteExtensions))
 			{
-				ExtraSafeToDeleteExtensions = "";
+				extraSafeToDeleteExtensions = "";
 			}
 
-			string[] CombinedSyncFilter = UserSettings.GetCombinedSyncFilter(GetSyncCategories(), Settings.Global.Filter, WorkspaceSettings.Filter);
-			List<string> SyncPaths = WorkspaceUpdate.GetSyncPaths(Workspace.Project, WorkspaceSettings.Filter.AllProjects ?? Settings.Global.Filter.AllProjects ?? false, CombinedSyncFilter);
+			string[] combinedSyncFilter = UserSettings.GetCombinedSyncFilter(GetSyncCategories(), _settings.Global.Filter, _workspaceSettings.Filter);
+			List<string> syncPaths = WorkspaceUpdate.GetSyncPaths(_workspace.Project, _workspaceSettings.Filter.AllProjects ?? _settings.Global.Filter.AllProjects ?? false, combinedSyncFilter);
 
-			CleanWorkspaceWindow.DoClean(ParentForm, PerforceSettings, BranchDirectoryName, Workspace.Project.ClientRootPath, SyncPaths, ExtraSafeToDeleteFolders.Split('\n'), ExtraSafeToDeleteExtensions.Split('\n'), ServiceProvider.GetRequiredService<ILogger<CleanWorkspaceWindow>>());
+			CleanWorkspaceWindow.DoClean(ParentForm, _perforceSettings, BranchDirectoryName, _workspace.Project.ClientRootPath, syncPaths, extraSafeToDeleteFolders.Split('\n'), extraSafeToDeleteExtensions.Split('\n'), _serviceProvider.GetRequiredService<ILogger<CleanWorkspaceWindow>>());
 		}
 
 		private void UpdateBuildSteps()
 		{
-			bHasBuildSteps = false;
+			_hasBuildSteps = false;
 
-			int MoreToolsItemCount = MoreToolsContextMenu.Items.IndexOf(MoreActionsContextMenu_CustomToolSeparator);
-			while (MoreToolsItemCount > 0)
+			int moreToolsItemCount = MoreToolsContextMenu.Items.IndexOf(MoreActionsContextMenu_CustomToolSeparator);
+			while (moreToolsItemCount > 0)
 			{
-				MoreToolsContextMenu.Items.RemoveAt(--MoreToolsItemCount);
+				MoreToolsContextMenu.Items.RemoveAt(--moreToolsItemCount);
 			}
 
-			for (int Idx = 0; Idx < CustomStatusPanelMenus.Count; Idx++)
+			for (int idx = 0; idx < _customStatusPanelMenus.Count; idx++)
 			{
-				components.Remove(CustomStatusPanelMenus[Idx]);
-				CustomStatusPanelMenus[Idx].Dispose();
+				components.Remove(_customStatusPanelMenus[idx]);
+				_customStatusPanelMenus[idx].Dispose();
 			}
 
-			CustomStatusPanelLinks.Clear();
-			CustomStatusPanelMenus.Clear();
+			_customStatusPanelLinks.Clear();
+			_customStatusPanelMenus.Clear();
 
-			if (Workspace != null)
+			if (_workspace != null)
 			{
-				ConfigFile ProjectConfigFile = Workspace.ProjectConfigFile;
-				if (ProjectConfigFile != null)
+				ConfigFile projectConfigFile = _workspace.ProjectConfigFile;
+				if (projectConfigFile != null)
 				{
-					Dictionary<Guid, ConfigObject> ProjectBuildStepObjects = GetProjectBuildStepObjects(ProjectConfigFile);
-					List<BuildStep> UserSteps = GetUserBuildSteps(ProjectBuildStepObjects);
+					Dictionary<Guid, ConfigObject> projectBuildStepObjects = GetProjectBuildStepObjects(projectConfigFile);
+					List<BuildStep> userSteps = GetUserBuildSteps(projectBuildStepObjects);
 
-					Dictionary<string, ContextMenuStrip> NameToMenu = new Dictionary<string, ContextMenuStrip>();
-					foreach (BuildStep Step in UserSteps)
+					Dictionary<string, ContextMenuStrip> nameToMenu = new Dictionary<string, ContextMenuStrip>();
+					foreach (BuildStep step in userSteps)
 					{
-						if (!String.IsNullOrEmpty(Step.StatusPanelLink) && CanRunStep(Step))
+						if (!String.IsNullOrEmpty(step.StatusPanelLink) && CanRunStep(step))
 						{
-							int BaseMenuIdx = Step.StatusPanelLink.IndexOf('|');
-							if (BaseMenuIdx == -1)
+							int baseMenuIdx = step.StatusPanelLink.IndexOf('|');
+							if (baseMenuIdx == -1)
 							{
-								CustomStatusPanelLinks.Add((Step.StatusPanelLink, (P, R) => RunCustomTool(Step, UserSteps)));
+								_customStatusPanelLinks.Add((step.StatusPanelLink, (p, r) => RunCustomTool(step, userSteps)));
 							}
 							else
 							{
-								string MenuName = Step.StatusPanelLink.Substring(0, BaseMenuIdx);
-								string ItemName = Step.StatusPanelLink.Substring(BaseMenuIdx + 1).Replace("&", "&&");
+								string menuName = step.StatusPanelLink.Substring(0, baseMenuIdx);
+								string itemName = step.StatusPanelLink.Substring(baseMenuIdx + 1).Replace("&", "&&");
 
-								ToolStripMenuItem NewMenuItem = new ToolStripMenuItem(ItemName);
-								NewMenuItem.Click += new EventHandler((sender, e) => { RunCustomTool(Step, UserSteps); });
+								ToolStripMenuItem newMenuItem = new ToolStripMenuItem(itemName);
+								newMenuItem.Click += new EventHandler((sender, e) => { RunCustomTool(step, userSteps); });
 
-								if (MenuName == "More...")
+								if (menuName == "More...")
 								{
-									MoreToolsContextMenu.Items.Insert(MoreToolsItemCount++, NewMenuItem);
+									MoreToolsContextMenu.Items.Insert(moreToolsItemCount++, newMenuItem);
 								}
 								else
 								{
-									ContextMenuStrip? Menu;
-									if (!NameToMenu.TryGetValue(MenuName, out Menu))
+									ContextMenuStrip? menu;
+									if (!nameToMenu.TryGetValue(menuName, out menu))
 									{
-										Menu = new ContextMenuStrip();
-										NameToMenu.Add(MenuName, Menu);
-										CustomStatusPanelLinks.Add(($"{MenuName} \u25BE", (P, R) => ShowToolContextMenu(R, Menu)));
-										CustomStatusPanelMenus.Add(Menu);
-										components.Add(Menu);
+										menu = new ContextMenuStrip();
+										nameToMenu.Add(menuName, menu);
+										_customStatusPanelLinks.Add(($"{menuName} \u25BE", (p, r) => ShowToolContextMenu(r, menu)));
+										_customStatusPanelMenus.Add(menu);
+										components.Add(menu);
 									}
-									Menu.Items.Add(NewMenuItem);
+									menu.Items.Add(newMenuItem);
 								}
 							}
 						}
 
-						bHasBuildSteps |= Step.bNormalSync;
+						_hasBuildSteps |= step.NormalSync;
 					}
 				}
 			}
 
-			MoreActionsContextMenu_CustomToolSeparator.Visible = (MoreToolsItemCount > 0);
+			MoreActionsContextMenu_CustomToolSeparator.Visible = (moreToolsItemCount > 0);
 		}
 
-		private bool CanRunStep(BuildStep Step)
+		private bool CanRunStep(BuildStep step)
 		{
-			if(Step.ToolId != Guid.Empty)
+			if(step.ToolId != Guid.Empty)
 			{
-				string? ToolName = Owner.ToolUpdateMonitor.GetToolName(Step.ToolId);
-				if (ToolName == null)
+				string? toolName = _owner.ToolUpdateMonitor.GetToolName(step.ToolId);
+				if (toolName == null)
 				{
 					return false;
 				}
-				if (Owner.ToolUpdateMonitor.GetToolPath(ToolName) == null)
+				if (_owner.ToolUpdateMonitor.GetToolPath(toolName) == null)
 				{
 					return false;
 				}
@@ -5221,107 +5221,107 @@ namespace UnrealGameSync
 			return true;
 		}
 
-		private void RunCustomTool(BuildStep Step, List<BuildStep> AllSteps)
+		private void RunCustomTool(BuildStep step, List<BuildStep> allSteps)
 		{
-			if (Workspace != null)
+			if (_workspace != null)
 			{
-				if (Workspace.IsBusy())
+				if (_workspace.IsBusy())
 				{
 					MessageBox.Show("Please retry after the current sync has finished.", "Sync in Progress");
 				}
 				else
 				{
-					HashSet<Guid> StepSet = new HashSet<Guid> { Step.UniqueId };
+					HashSet<Guid> stepSet = new HashSet<Guid> { step.UniqueId };
 
-					ConfigFile ProjectConfigFile = Workspace.ProjectConfigFile;
-					if (ProjectConfigFile != null && Step.Requires.Length > 0)
+					ConfigFile projectConfigFile = _workspace.ProjectConfigFile;
+					if (projectConfigFile != null && step.Requires.Length > 0)
 					{
-						Stack<Guid> Stack = new Stack<Guid>(StepSet);
-						while (Stack.Count > 0)
+						Stack<Guid> stack = new Stack<Guid>(stepSet);
+						while (stack.Count > 0)
 						{
-							Guid Id = Stack.Pop();
-							BuildStep NextStep = AllSteps.FirstOrDefault(x => x.UniqueId == Id);
-							if (NextStep != null)
+							Guid id = stack.Pop();
+							BuildStep nextStep = allSteps.FirstOrDefault(x => x.UniqueId == id);
+							if (nextStep != null)
 							{
-								foreach (Guid RequiresId in NextStep.Requires)
+								foreach (Guid requiresId in nextStep.Requires)
 								{
-									if (StepSet.Add(RequiresId))
+									if (stepSet.Add(requiresId))
 									{
-										Stack.Push(RequiresId);
+										stack.Push(requiresId);
 									}
 								}
 							}
 						}
 					}
 
-					WorkspaceUpdateContext Context = new WorkspaceUpdateContext(Workspace.CurrentChangeNumber, WorkspaceUpdateOptions.Build, GetEditorBuildConfig(), null, ProjectSettings.BuildSteps, StepSet);
+					WorkspaceUpdateContext context = new WorkspaceUpdateContext(_workspace.CurrentChangeNumber, WorkspaceUpdateOptions.Build, GetEditorBuildConfig(), null, _projectSettings.BuildSteps, stepSet);
 
-					if (Step.ToolId != Guid.Empty)
+					if (step.ToolId != Guid.Empty)
 					{
-						string? ToolName = Owner.ToolUpdateMonitor.GetToolName(Step.ToolId);
-						if (ToolName != null)
+						string? toolName = _owner.ToolUpdateMonitor.GetToolName(step.ToolId);
+						if (toolName != null)
 						{
-							DirectoryReference? ToolPath = Owner.ToolUpdateMonitor.GetToolPath(ToolName);
-							if (ToolPath != null)
+							DirectoryReference? toolPath = _owner.ToolUpdateMonitor.GetToolPath(toolName);
+							if (toolPath != null)
 							{
-								Context.AdditionalVariables["ToolDir"] = ToolPath.FullName;
+								context.AdditionalVariables["ToolDir"] = toolPath.FullName;
 							}
 						}
 					}
 
-					StartWorkspaceUpdate(Context, null);
+					StartWorkspaceUpdate(context, null);
 				}
 			}
 		}
 
 		private void OptionsContextMenu_EditBuildSteps_Click(object sender, EventArgs e)
 		{
-			ConfigFile ProjectConfigFile = Workspace.ProjectConfigFile;
-			if (Workspace != null && ProjectConfigFile != null)
+			ConfigFile projectConfigFile = _workspace.ProjectConfigFile;
+			if (_workspace != null && projectConfigFile != null)
 			{
 				// Find all the target names for this project
-				List<string> TargetNames = new List<string>();
+				List<string> targetNames = new List<string>();
 				if (SelectedFileName != null && SelectedFileName.HasExtension(".uproject"))
 				{
-					DirectoryInfo SourceDirectory = DirectoryReference.Combine(SelectedFileName.Directory, "Source").ToDirectoryInfo();
-					if (SourceDirectory.Exists)
+					DirectoryInfo sourceDirectory = DirectoryReference.Combine(SelectedFileName.Directory, "Source").ToDirectoryInfo();
+					if (sourceDirectory.Exists)
 					{
-						foreach (FileInfo TargetFile in SourceDirectory.EnumerateFiles("*.target.cs", SearchOption.TopDirectoryOnly))
+						foreach (FileInfo targetFile in sourceDirectory.EnumerateFiles("*.target.cs", SearchOption.TopDirectoryOnly))
 						{
-							TargetNames.Add(TargetFile.Name.Substring(0, TargetFile.Name.IndexOf('.')));
+							targetNames.Add(targetFile.Name.Substring(0, targetFile.Name.IndexOf('.')));
 						}
 					}
 				}
 
 				// Get all the task objects
-				Dictionary<Guid, ConfigObject> ProjectBuildStepObjects = GetProjectBuildStepObjects(ProjectConfigFile);
-				List<BuildStep> UserSteps = GetUserBuildSteps(ProjectBuildStepObjects);
+				Dictionary<Guid, ConfigObject> projectBuildStepObjects = GetProjectBuildStepObjects(projectConfigFile);
+				List<BuildStep> userSteps = GetUserBuildSteps(projectBuildStepObjects);
 
 				// Show the dialog
-				Dictionary<string, string> Variables = Workspace.GetVariables(GetEditorBuildConfig());
-				ModifyBuildStepsWindow EditStepsWindow = new ModifyBuildStepsWindow(TargetNames, UserSteps, new HashSet<Guid>(ProjectBuildStepObjects.Keys), BranchDirectoryName, Variables);
-				EditStepsWindow.ShowDialog();
+				Dictionary<string, string> variables = _workspace.GetVariables(GetEditorBuildConfig());
+				ModifyBuildStepsWindow editStepsWindow = new ModifyBuildStepsWindow(targetNames, userSteps, new HashSet<Guid>(projectBuildStepObjects.Keys), BranchDirectoryName, variables);
+				editStepsWindow.ShowDialog();
 
 				// Update the user settings
-				List<ConfigObject> ModifiedBuildSteps = new List<ConfigObject>();
-				foreach (BuildStep Step in UserSteps)
+				List<ConfigObject> modifiedBuildSteps = new List<ConfigObject>();
+				foreach (BuildStep step in userSteps)
 				{
-					if (Step.IsValid())
+					if (step.IsValid())
 					{
-						ConfigObject? DefaultObject;
-						ProjectBuildStepObjects.TryGetValue(Step.UniqueId, out DefaultObject);
+						ConfigObject? defaultObject;
+						projectBuildStepObjects.TryGetValue(step.UniqueId, out defaultObject);
 
-						ConfigObject? UserConfigObject = Step.ToConfigObject(DefaultObject);
-						if (UserConfigObject != null && UserConfigObject.Pairs.Any(x => x.Key != "UniqueId"))
+						ConfigObject? userConfigObject = step.ToConfigObject(defaultObject);
+						if (userConfigObject != null && userConfigObject.Pairs.Any(x => x.Key != "UniqueId"))
 						{
-							ModifiedBuildSteps.Add(UserConfigObject);
+							modifiedBuildSteps.Add(userConfigObject);
 						}
 					}
 				}
 
 				// Save the settings
-				ProjectSettings.BuildSteps = ModifiedBuildSteps;
-				ProjectSettings.Save(Logger);
+				_projectSettings.BuildSteps = modifiedBuildSteps;
+				_projectSettings.Save(_logger);
 
 				// Update the custom tools menu, because we might have changed it
 				UpdateBuildSteps();
@@ -5330,107 +5330,107 @@ namespace UnrealGameSync
 			}
 		}
 
-		private void AddOrUpdateBuildStep(Dictionary<Guid, ConfigObject> Steps, ConfigObject Object)
+		private void AddOrUpdateBuildStep(Dictionary<Guid, ConfigObject> steps, ConfigObject obj)
 		{
-			Guid UniqueId;
-			if (Guid.TryParse(Object.GetValue(BuildStep.UniqueIdKey, ""), out UniqueId))
+			Guid uniqueId;
+			if (Guid.TryParse(obj.GetValue(BuildStep.UniqueIdKey, ""), out uniqueId))
 			{
 				// Add or apply Object to the list of steps in Steps. Do not modify Object; make a copy first.
-				ConfigObject NewObject = new ConfigObject(Object);
+				ConfigObject newObject = new ConfigObject(obj);
 
-				ConfigObject? DefaultObject;
-				if (Steps.TryGetValue(UniqueId, out DefaultObject))
+				ConfigObject? defaultObject;
+				if (steps.TryGetValue(uniqueId, out defaultObject))
 				{
-					NewObject.SetDefaults(DefaultObject);
+					newObject.SetDefaults(defaultObject);
 				}
 
-				Steps[UniqueId] = NewObject;
+				steps[uniqueId] = newObject;
 			}
 		}
 
 		private bool ShouldSyncPrecompiledEditor
 		{
-			get { return Settings.Archives.Any(x => x.bEnabled && x.Type == IArchiveInfo.EditorArchiveType) && GetArchives().Any(x => x.Type == "Editor"); }
+			get { return _settings.Archives.Any(x => x.Enabled && x.Type == IArchiveInfo.EditorArchiveType) && GetArchives().Any(x => x.Type == "Editor"); }
 		}
 
 		public BuildConfig GetEditorBuildConfig()
 		{
-			return ShouldSyncPrecompiledEditor ? BuildConfig.Development : Settings.CompiledEditorBuildConfig;
+			return ShouldSyncPrecompiledEditor ? BuildConfig.Development : _settings.CompiledEditorBuildConfig;
 		}
 
-		private Dictionary<Guid, ConfigObject> GetProjectBuildStepObjects(ConfigFile ProjectConfigFile)
+		private Dictionary<Guid, ConfigObject> GetProjectBuildStepObjects(ConfigFile projectConfigFile)
 		{
-			FileReference EditorTargetFile = ConfigUtils.GetEditorTargetFile(ProjectInfo, ProjectConfigFile);
-			string EditorTargetName = EditorTargetFile.GetFileNameWithoutAnyExtensions();
+			FileReference editorTargetFile = ConfigUtils.GetEditorTargetFile(ProjectInfo, projectConfigFile);
+			string editorTargetName = editorTargetFile.GetFileNameWithoutAnyExtensions();
 
-			Dictionary<Guid, ConfigObject> ProjectBuildSteps = ConfigUtils.GetDefaultBuildStepObjects(ProjectInfo, EditorTargetName, Settings.CompiledEditorBuildConfig, ProjectConfigFile, ShouldSyncPrecompiledEditor);
-			foreach (string Line in ProjectConfigFile.GetValues("Build.Step", new string[0]))
+			Dictionary<Guid, ConfigObject> projectBuildSteps = ConfigUtils.GetDefaultBuildStepObjects(ProjectInfo, editorTargetName, _settings.CompiledEditorBuildConfig, projectConfigFile, ShouldSyncPrecompiledEditor);
+			foreach (string line in projectConfigFile.GetValues("Build.Step", new string[0]))
 			{
-				AddOrUpdateBuildStep(ProjectBuildSteps, new ConfigObject(Line));
+				AddOrUpdateBuildStep(projectBuildSteps, new ConfigObject(line));
 			}
-			return ProjectBuildSteps;
+			return projectBuildSteps;
 		}
 
-		private List<BuildStep> GetUserBuildSteps(Dictionary<Guid, ConfigObject> ProjectBuildStepObjects)
+		private List<BuildStep> GetUserBuildSteps(Dictionary<Guid, ConfigObject> projectBuildStepObjects)
 		{
 			// Read all the user-defined build tasks and modifications to the default list
-			Dictionary<Guid, ConfigObject> UserBuildStepObjects = ProjectBuildStepObjects.ToDictionary(x => x.Key, y => new ConfigObject(y.Value));
-			foreach (ConfigObject UserBuildStep in ProjectSettings.BuildSteps)
+			Dictionary<Guid, ConfigObject> userBuildStepObjects = projectBuildStepObjects.ToDictionary(x => x.Key, y => new ConfigObject(y.Value));
+			foreach (ConfigObject userBuildStep in _projectSettings.BuildSteps)
 			{
-				AddOrUpdateBuildStep(UserBuildStepObjects, UserBuildStep);
+				AddOrUpdateBuildStep(userBuildStepObjects, userBuildStep);
 			}
 
 			// Create the expanded task objects
-			return UserBuildStepObjects.Values.Select(x => new BuildStep(x)).OrderBy(x => (x.OrderIndex == -1)? 10000 : x.OrderIndex).ToList();
+			return userBuildStepObjects.Values.Select(x => new BuildStep(x)).OrderBy(x => (x.OrderIndex == -1)? 10000 : x.OrderIndex).ToList();
 		}
 
 		private void OptionsContextMenu_SyncPrecompiledBinaries_Click(object sender, EventArgs e)
 		{
 			if (OptionsContextMenu_SyncPrecompiledBinaries.DropDownItems.Count == 0)
 			{
-				IArchiveInfo EditorArchive = GetArchives().FirstOrDefault(x => x.Type == IArchiveInfo.EditorArchiveType);
-				if (EditorArchive != null)
+				IArchiveInfo editorArchive = GetArchives().FirstOrDefault(x => x.Type == IArchiveInfo.EditorArchiveType);
+				if (editorArchive != null)
 				{
-					SetSelectedArchive(EditorArchive, !OptionsContextMenu_SyncPrecompiledBinaries.Checked);
+					SetSelectedArchive(editorArchive, !OptionsContextMenu_SyncPrecompiledBinaries.Checked);
 				}
 			}
 		}
 
 		private void BuildList_SelectedIndexChanged(object sender, EventArgs e)
 		{
-			PendingSelectedChangeNumber = -1;
+			_pendingSelectedChangeNumber = -1;
 		}
 
 		private void OptionsContextMenu_Diagnostics_Click(object sender, EventArgs e)
 		{
-			StringBuilder DiagnosticsText = new StringBuilder();
-			DiagnosticsText.AppendFormat("Application version: {0}\n", Program.GetVersionString());
-			DiagnosticsText.AppendFormat("Synced from: {0}\n", Program.SyncVersion ?? "(unknown)");
-			DiagnosticsText.AppendFormat("Selected file: {0}\n", (SelectedFileName == null) ? "(none)" : SelectedFileName.FullName);
-			if (Workspace != null)
+			StringBuilder diagnosticsText = new StringBuilder();
+			diagnosticsText.AppendFormat("Application version: {0}\n", Program.GetVersionString());
+			diagnosticsText.AppendFormat("Synced from: {0}\n", Program.SyncVersion ?? "(unknown)");
+			diagnosticsText.AppendFormat("Selected file: {0}\n", (SelectedFileName == null) ? "(none)" : SelectedFileName.FullName);
+			if (_workspace != null)
 			{
-				DiagnosticsText.AppendFormat("P4 server: {0}\n", PerforceSettings.ServerAndPort ?? "(default)");
-				DiagnosticsText.AppendFormat("P4 user: {0}\n", PerforceSettings.UserName);
-				DiagnosticsText.AppendFormat("P4 workspace: {0}\n", PerforceSettings.ClientName);
+				diagnosticsText.AppendFormat("P4 server: {0}\n", _perforceSettings.ServerAndPort ?? "(default)");
+				diagnosticsText.AppendFormat("P4 user: {0}\n", _perforceSettings.UserName);
+				diagnosticsText.AppendFormat("P4 workspace: {0}\n", _perforceSettings.ClientName);
 			}
-			DiagnosticsText.AppendFormat("Perforce monitor: {0}\n", (PerforceMonitor == null) ? "(inactive)" : PerforceMonitor.LastStatusMessage);
-			DiagnosticsText.AppendFormat("Event monitor: {0}\n", (EventMonitor == null) ? "(inactive)" : EventMonitor.LastStatusMessage);
-			DiagnosticsText.AppendFormat("Issue monitor: {0}\n", (IssueMonitor == null) ? "(inactive)" : IssueMonitor.LastStatusMessage);
+			diagnosticsText.AppendFormat("Perforce monitor: {0}\n", (_perforceMonitor == null) ? "(inactive)" : _perforceMonitor.LastStatusMessage);
+			diagnosticsText.AppendFormat("Event monitor: {0}\n", (_eventMonitor == null) ? "(inactive)" : _eventMonitor.LastStatusMessage);
+			diagnosticsText.AppendFormat("Issue monitor: {0}\n", (_issueMonitor == null) ? "(inactive)" : _issueMonitor.LastStatusMessage);
 
-			DiagnosticsWindow Diagnostics = new DiagnosticsWindow(AppDataFolder, WorkspaceDataFolder, DiagnosticsText.ToString(), Settings.GetCachedFilePaths());
-			Diagnostics.ShowDialog(this);
+			DiagnosticsWindow diagnostics = new DiagnosticsWindow(_appDataFolder, _workspaceDataFolder, diagnosticsText.ToString(), _settings.GetCachedFilePaths());
+			diagnostics.ShowDialog(this);
 		}
 
 		private void OptionsContextMenu_SyncFilter_Click(object sender, EventArgs e)
 		{
-			SyncFilter Filter = new SyncFilter(GetSyncCategories(), Settings.Global.Filter, WorkspaceSettings.Filter);
-			if (Filter.ShowDialog() == DialogResult.OK)
+			SyncFilter filter = new SyncFilter(GetSyncCategories(), _settings.Global.Filter, _workspaceSettings.Filter);
+			if (filter.ShowDialog() == DialogResult.OK)
 			{
-				Settings.Global.Filter = Filter.GlobalFilter;
-				Settings.Save(Logger);
+				_settings.Global.Filter = filter.GlobalFilter;
+				_settings.Save(_logger);
 
-				WorkspaceSettings.Filter = Filter.WorkspaceFilter;
-				WorkspaceSettings.Save(Logger);
+				_workspaceSettings.Filter = filter.WorkspaceFilter;
+				_workspaceSettings.Save(_logger);
 			}
 		}
 
@@ -5444,58 +5444,58 @@ namespace UnrealGameSync
 			this.SyncContexMenu_EnterChangelist,
 			this.toolStripSeparator8});
 
-			foreach (LatestChangeType ChangeType in GetCustomLatestChangeTypes())
+			foreach (LatestChangeType changeType in GetCustomLatestChangeTypes())
 			{
-				System.Windows.Forms.ToolStripMenuItem MenuItem = new System.Windows.Forms.ToolStripMenuItem();
-				MenuItem.Name = ChangeType.Name;
-				MenuItem.Text = ChangeType.Description;
-				MenuItem.Size = new System.Drawing.Size(189, 22);
-				MenuItem.Click += (sender, e) => this.SyncContextMenu_LatestChangeType_Click(sender, e, ChangeType.Name);
+				System.Windows.Forms.ToolStripMenuItem menuItem = new System.Windows.Forms.ToolStripMenuItem();
+				menuItem.Name = changeType.Name;
+				menuItem.Text = changeType.Description;
+				menuItem.Size = new System.Drawing.Size(189, 22);
+				menuItem.Click += (sender, e) => this.SyncContextMenu_LatestChangeType_Click(sender, e, changeType.Name);
 
-				SyncContextMenu.Items.Add(MenuItem);
+				SyncContextMenu.Items.Add(menuItem);
 			}
 		}
 
-		private void ShowSyncMenu(Rectangle Bounds)
+		private void ShowSyncMenu(Rectangle bounds)
 		{
 			RefreshSyncMenuItems();
 
-			string? NameOfLatestChangeType = null;
-			foreach (LatestChangeType ChangeType in GetCustomLatestChangeTypes())
+			string? nameOfLatestChangeType = null;
+			foreach (LatestChangeType changeType in GetCustomLatestChangeTypes())
 			{
-				if (ChangeType.Name == Settings.SyncTypeID)
+				if (changeType.Name == _settings.SyncTypeId)
 				{
-					NameOfLatestChangeType = ChangeType.Name;
+					nameOfLatestChangeType = changeType.Name;
 					break;
 				}
 			}
 
-			if (NameOfLatestChangeType != null)
+			if (nameOfLatestChangeType != null)
 			{
 				foreach (int value in Enumerable.Range(0, SyncContextMenu.Items.Count))
 				{
-					System.Windows.Forms.ToolStripMenuItem? Item = SyncContextMenu.Items[value] as System.Windows.Forms.ToolStripMenuItem;
-					if (Item != null)
+					System.Windows.Forms.ToolStripMenuItem? item = SyncContextMenu.Items[value] as System.Windows.Forms.ToolStripMenuItem;
+					if (item != null)
 					{
-						if (Item.Name == NameOfLatestChangeType)
+						if (item.Name == nameOfLatestChangeType)
 						{
-							Item.Checked = true;
+							item.Checked = true;
 						}
 						else
 						{
-							Item.Checked = false;
+							item.Checked = false;
 						}
 					}
 				}
 			}
 
-			SyncContextMenu.Show(StatusPanel, new Point(Bounds.Left, Bounds.Bottom), ToolStripDropDownDirection.BelowRight);
+			SyncContextMenu.Show(StatusPanel, new Point(bounds.Left, bounds.Bottom), ToolStripDropDownDirection.BelowRight);
 		}
 
-		private void SyncContextMenu_LatestChangeType_Click(object? sender, EventArgs e, string SyncTypeID)
+		private void SyncContextMenu_LatestChangeType_Click(object? sender, EventArgs e, string syncTypeId)
 		{
-			Settings.SyncTypeID = SyncTypeID;
-			Settings.Save(Logger);
+			_settings.SyncTypeId = syncTypeId;
+			_settings.Save(_logger);
 		}
 
 		private void SyncContextMenu_EnterChangelist_Click(object sender, EventArgs e)
@@ -5505,79 +5505,79 @@ namespace UnrealGameSync
 				return;
 			}
 
-			ChangelistWindow ChildWindow = new ChangelistWindow((Workspace == null) ? -1 : Workspace.CurrentChangeNumber);
-			if (ChildWindow.ShowDialog() == DialogResult.OK)
+			ChangelistWindow childWindow = new ChangelistWindow((_workspace == null) ? -1 : _workspace.CurrentChangeNumber);
+			if (childWindow.ShowDialog() == DialogResult.OK)
 			{
-				StartSync(ChildWindow.ChangeNumber);
+				StartSync(childWindow.ChangeNumber);
 			}
 		}
 
-		private void BuildList_KeyDown(object Sender, KeyEventArgs Args)
+		private void BuildList_KeyDown(object sender, KeyEventArgs args)
 		{
-			if (Args.Control && Args.KeyCode == Keys.C && BuildList.SelectedItems.Count > 0)
+			if (args.Control && args.KeyCode == Keys.C && BuildList.SelectedItems.Count > 0)
 			{
-				int SelectedChange = ((ChangesRecord)BuildList.SelectedItems[0].Tag).Number;
-				Clipboard.SetText(String.Format("{0}", SelectedChange));
+				int selectedChange = ((ChangesRecord)BuildList.SelectedItems[0].Tag).Number;
+				Clipboard.SetText(String.Format("{0}", selectedChange));
 			}
 		}
 
 		private void OptionsContextMenu_ApplicationSettings_Click(object sender, EventArgs e)
 		{
-			Owner.ModifyApplicationSettings();
+			_owner.ModifyApplicationSettings();
 		}
 
 		private void OptionsContextMenu_TabNames_Stream_Click(object sender, EventArgs e)
 		{
-			Owner.SetTabNames(TabLabels.Stream);
+			_owner.SetTabNames(TabLabels.Stream);
 		}
 
 		private void OptionsContextMenu_TabNames_WorkspaceName_Click(object sender, EventArgs e)
 		{
-			Owner.SetTabNames(TabLabels.WorkspaceName);
+			_owner.SetTabNames(TabLabels.WorkspaceName);
 		}
 
 		private void OptionsContextMenu_TabNames_WorkspaceRoot_Click(object sender, EventArgs e)
 		{
-			Owner.SetTabNames(TabLabels.WorkspaceRoot);
+			_owner.SetTabNames(TabLabels.WorkspaceRoot);
 		}
 
 		private void OptionsContextMenu_TabNames_ProjectFile_Click(object sender, EventArgs e)
 		{
-			Owner.SetTabNames(TabLabels.ProjectFile);
+			_owner.SetTabNames(TabLabels.ProjectFile);
 		}
 
-		private void SelectRecentProject(Rectangle Bounds)
+		private void SelectRecentProject(Rectangle bounds)
 		{
 			while (RecentMenu.Items[2] != RecentMenu_Separator)
 			{
 				RecentMenu.Items.RemoveAt(2);
 			}
 
-			foreach (UserSelectedProjectSettings RecentProject in Settings.RecentProjects)
+			foreach (UserSelectedProjectSettings recentProject in _settings.RecentProjects)
 			{
-				ToolStripMenuItem Item = new ToolStripMenuItem(RecentProject.ToString(), null, new EventHandler((o, e) => Owner.RequestProjectChange(this, RecentProject, true)));
-				RecentMenu.Items.Insert(RecentMenu.Items.Count - 2, Item);
+				ToolStripMenuItem item = new ToolStripMenuItem(recentProject.ToString(), null, new EventHandler((o, e) => _owner.RequestProjectChange(this, recentProject, true)));
+				RecentMenu.Items.Insert(RecentMenu.Items.Count - 2, item);
 			}
 
-			RecentMenu_Separator.Visible = (Settings.RecentProjects.Count > 0);
-			RecentMenu.Show(StatusPanel, new Point(Bounds.Left, Bounds.Bottom), ToolStripDropDownDirection.BelowRight);
+			RecentMenu_Separator.Visible = (_settings.RecentProjects.Count > 0);
+			RecentMenu.Show(StatusPanel, new Point(bounds.Left, bounds.Bottom), ToolStripDropDownDirection.BelowRight);
 		}
 
 		private void RecentMenu_Browse_Click(object sender, EventArgs e)
 		{
-			Owner.EditSelectedProject(this);
+			_owner.EditSelectedProject(this);
 		}
 
 		private void RecentMenu_ClearList_Click(object sender, EventArgs e)
 		{
-			Settings.RecentProjects.Clear();
-			Settings.Save(Logger);
+			_settings.RecentProjects.Clear();
+			_settings.Save(_logger);
 		}
 
 		private void OptionsContextMenu_ShowChanges_ShowUnreviewed_Click(object sender, EventArgs e)
 		{
-			Settings.bShowUnreviewedChanges ^= true;
-			Settings.Save(Logger);
+			_settings.ShowUnreviewedChanges ^= true;
+			_settings.Save(_logger);
 
 			UpdateBuildList();
 			ShrinkNumRequestedBuilds();
@@ -5585,8 +5585,8 @@ namespace UnrealGameSync
 
 		private void OptionsContextMenu_ShowChanges_ShowAutomated_Click(object sender, EventArgs e)
 		{
-			Settings.bShowAutomatedChanges ^= true;
-			Settings.Save(Logger);
+			_settings.ShowAutomatedChanges ^= true;
+			_settings.Save(_logger);
 
 			UpdateBuildListFilter();
 		}
@@ -5598,113 +5598,113 @@ namespace UnrealGameSync
 
 		private void BuildList_ColumnWidthChanged(object sender, ColumnWidthChangedEventArgs e)
 		{
-			if (ColumnWidths != null && MinColumnWidths != null)
+			if (_columnWidths != null && _minColumnWidths != null)
 			{
-				int NewWidth = BuildList.Columns[e.ColumnIndex].Width;
-				if (NewWidth < MinColumnWidths[e.ColumnIndex])
+				int newWidth = BuildList.Columns[e.ColumnIndex].Width;
+				if (newWidth < _minColumnWidths[e.ColumnIndex])
 				{
-					NewWidth = MinColumnWidths[e.ColumnIndex];
-					BuildList.Columns[e.ColumnIndex].Width = NewWidth;
+					newWidth = _minColumnWidths[e.ColumnIndex];
+					BuildList.Columns[e.ColumnIndex].Width = newWidth;
 				}
-				ColumnWidths[e.ColumnIndex] = NewWidth;
+				_columnWidths[e.ColumnIndex] = newWidth;
 			}
 			UpdateMaxBuildBadgeChars();
 		}
 
 		private void BuildList_Resize(object sender, EventArgs e)
 		{
-			int PrevBuildListWidth = BuildListWidth;
+			int prevBuildListWidth = _buildListWidth;
 
-			BuildListWidth = BuildList.Width;
+			_buildListWidth = BuildList.Width;
 
-			if (PrevBuildListWidth != 0 && BuildListWidth != PrevBuildListWidth && ColumnWidths != null)
+			if (prevBuildListWidth != 0 && _buildListWidth != prevBuildListWidth && _columnWidths != null)
 			{
-				float SafeWidth = ColumnWidths.Sum() + 50.0f;
-				if (BuildListWidth < PrevBuildListWidth)
+				float safeWidth = _columnWidths.Sum() + 50.0f;
+				if (_buildListWidth < prevBuildListWidth)
 				{
-					if (BuildListWidth <= SafeWidth)
+					if (_buildListWidth <= safeWidth)
 					{
-						ResizeColumns(ColumnWidths.Sum() + (BuildListWidth - Math.Max(PrevBuildListWidth, SafeWidth)));
+						ResizeColumns(_columnWidths.Sum() + (_buildListWidth - Math.Max(prevBuildListWidth, safeWidth)));
 					}
 				}
 				else
 				{
-					if (BuildListWidth >= SafeWidth)
+					if (_buildListWidth >= safeWidth)
 					{
-						ResizeColumns(ColumnWidths.Sum() + (BuildListWidth - Math.Max(PrevBuildListWidth, SafeWidth)));
+						ResizeColumns(_columnWidths.Sum() + (_buildListWidth - Math.Max(prevBuildListWidth, safeWidth)));
 					}
 				}
 			}
 		}
 
-		void ResizeColumns(float NextTotalWidth)
+		void ResizeColumns(float nextTotalWidth)
 		{
-			float[] TargetColumnWidths = GetTargetColumnWidths(NextTotalWidth);
+			float[] targetColumnWidths = GetTargetColumnWidths(nextTotalWidth);
 
 			// Get the current total width of the columns, and the new space that we'll aim to fill
-			float PrevTotalWidth = ColumnWidths.Sum();
-			float TotalDelta = Math.Abs(NextTotalWidth - PrevTotalWidth);
+			float prevTotalWidth = _columnWidths.Sum();
+			float totalDelta = Math.Abs(nextTotalWidth - prevTotalWidth);
 
-			float TotalColumnDelta = 0.0f;
-			for (int Idx = 0; Idx < BuildList.Columns.Count; Idx++)
+			float totalColumnDelta = 0.0f;
+			for (int idx = 0; idx < BuildList.Columns.Count; idx++)
 			{
-				TotalColumnDelta += Math.Abs(TargetColumnWidths[Idx] - ColumnWidths[Idx]);
+				totalColumnDelta += Math.Abs(targetColumnWidths[idx] - _columnWidths[idx]);
 			}
 
-			if (TotalColumnDelta > 0.5f)
+			if (totalColumnDelta > 0.5f)
 			{
-				float[] NextColumnWidths = new float[BuildList.Columns.Count];
-				for (int Idx = 0; Idx < BuildList.Columns.Count; Idx++)
+				float[] nextColumnWidths = new float[BuildList.Columns.Count];
+				for (int idx = 0; idx < BuildList.Columns.Count; idx++)
 				{
-					float MaxColumnDelta = TotalDelta * Math.Abs(TargetColumnWidths[Idx] - ColumnWidths[Idx]) / TotalColumnDelta;
-					NextColumnWidths[Idx] = Math.Max(Math.Min(TargetColumnWidths[Idx], ColumnWidths[Idx] + MaxColumnDelta), ColumnWidths[Idx] - MaxColumnDelta);
+					float maxColumnDelta = totalDelta * Math.Abs(targetColumnWidths[idx] - _columnWidths[idx]) / totalColumnDelta;
+					nextColumnWidths[idx] = Math.Max(Math.Min(targetColumnWidths[idx], _columnWidths[idx] + maxColumnDelta), _columnWidths[idx] - maxColumnDelta);
 				}
 
 				// Update the control
 				BuildList.BeginUpdate();
-				for (int Idx = 0; Idx < BuildList.Columns.Count; Idx++)
+				for (int idx = 0; idx < BuildList.Columns.Count; idx++)
 				{
-					BuildList.Columns[Idx].Width = (int)NextColumnWidths[Idx];
+					BuildList.Columns[idx].Width = (int)nextColumnWidths[idx];
 				}
-				ColumnWidths = NextColumnWidths;
+				_columnWidths = nextColumnWidths;
 				BuildList.EndUpdate();
 			}
 		}
 
-		float[] GetTargetColumnWidths(float NextTotalWidth)
+		float[] GetTargetColumnWidths(float nextTotalWidth)
 		{
 			// Array to store the output list
-			float[] TargetColumnWidths = new float[BuildList.Columns.Count];
+			float[] targetColumnWidths = new float[BuildList.Columns.Count];
 
 			// Array of flags to store columns which are clamped into position. We try to respect proportional resizing as well as clamping to min/max sizes,
 			// and remaining space can be distributed via non-clamped columns via another iteration.
-			bool[] ConstrainedColumns = new bool[BuildList.Columns.Count];
+			bool[] constrainedColumns = new bool[BuildList.Columns.Count];
 
 			// Keep track of the remaining width that we have to distribute between columns. Does not include the required minimum size of each column.
-			float RemainingWidth = Math.Max(NextTotalWidth - MinColumnWidths.Sum(), 0.0f);
+			float remainingWidth = Math.Max(nextTotalWidth - _minColumnWidths.Sum(), 0.0f);
 
 			// Keep track of the sum of the remaining column weights. Used to proportionally allocate remaining space.
-			float RemainingTotalWeight = ColumnWeights.Sum();
+			float remainingTotalWeight = _columnWeights.Sum();
 
 			// Handle special cases for shrinking/expanding
-			float PrevTotalWidth = ColumnWidths.Sum();
-			if (NextTotalWidth < PrevTotalWidth)
+			float prevTotalWidth = _columnWidths.Sum();
+			if (nextTotalWidth < prevTotalWidth)
 			{
 				// If target size is less than current size, keep it at the current size
-				for (int Idx = 0; Idx < BuildList.Columns.Count; Idx++)
+				for (int idx = 0; idx < BuildList.Columns.Count; idx++)
 				{
-					if (!ConstrainedColumns[Idx])
+					if (!constrainedColumns[idx])
 					{
-						float TargetColumnWidth = MinColumnWidths[Idx] + (RemainingWidth * ColumnWeights[Idx] / RemainingTotalWeight);
-						if (TargetColumnWidth > ColumnWidths[Idx])
+						float targetColumnWidth = _minColumnWidths[idx] + (remainingWidth * _columnWeights[idx] / remainingTotalWeight);
+						if (targetColumnWidth > _columnWidths[idx])
 						{
-							TargetColumnWidths[Idx] = ColumnWidths[Idx];
-							ConstrainedColumns[Idx] = true;
+							targetColumnWidths[idx] = _columnWidths[idx];
+							constrainedColumns[idx] = true;
 
-							RemainingWidth -= Math.Max(TargetColumnWidths[Idx] - MinColumnWidths[Idx], 0.0f);
-							RemainingTotalWeight -= ColumnWeights[Idx];
+							remainingWidth -= Math.Max(targetColumnWidths[idx] - _minColumnWidths[idx], 0.0f);
+							remainingTotalWeight -= _columnWeights[idx];
 
-							Idx = -1;
+							idx = -1;
 							continue;
 						}
 					}
@@ -5713,41 +5713,41 @@ namespace UnrealGameSync
 			else
 			{
 				// If target size is greater than desired size, clamp it to that
-				for (int Idx = 0; Idx < BuildList.Columns.Count; Idx++)
+				for (int idx = 0; idx < BuildList.Columns.Count; idx++)
 				{
-					if (!ConstrainedColumns[Idx])
+					if (!constrainedColumns[idx])
 					{
-						float TargetColumnWidth = MinColumnWidths[Idx] + (RemainingWidth * ColumnWeights[Idx] / RemainingTotalWeight);
-						if (TargetColumnWidth > DesiredColumnWidths[Idx])
+						float targetColumnWidth = _minColumnWidths[idx] + (remainingWidth * _columnWeights[idx] / remainingTotalWeight);
+						if (targetColumnWidth > _desiredColumnWidths[idx])
 						{
 							// Don't allow this column to expand above the maximum desired size
-							TargetColumnWidths[Idx] = DesiredColumnWidths[Idx];
-							ConstrainedColumns[Idx] = true;
+							targetColumnWidths[idx] = _desiredColumnWidths[idx];
+							constrainedColumns[idx] = true;
 
-							RemainingWidth -= Math.Max(TargetColumnWidths[Idx] - MinColumnWidths[Idx], 0.0f);
-							RemainingTotalWeight -= ColumnWeights[Idx];
+							remainingWidth -= Math.Max(targetColumnWidths[idx] - _minColumnWidths[idx], 0.0f);
+							remainingTotalWeight -= _columnWeights[idx];
 
-							Idx = -1;
+							idx = -1;
 							continue;
 						}
 					}
 				}
 
 				// If current size is greater than target size, keep it that way
-				for (int Idx = 0; Idx < BuildList.Columns.Count; Idx++)
+				for (int idx = 0; idx < BuildList.Columns.Count; idx++)
 				{
-					if (!ConstrainedColumns[Idx])
+					if (!constrainedColumns[idx])
 					{
-						float TargetColumnWidth = MinColumnWidths[Idx] + (RemainingWidth * ColumnWeights[Idx] / RemainingTotalWeight);
-						if (TargetColumnWidth < ColumnWidths[Idx])
+						float targetColumnWidth = _minColumnWidths[idx] + (remainingWidth * _columnWeights[idx] / remainingTotalWeight);
+						if (targetColumnWidth < _columnWidths[idx])
 						{
-							TargetColumnWidths[Idx] = ColumnWidths[Idx];
-							ConstrainedColumns[Idx] = true;
+							targetColumnWidths[idx] = _columnWidths[idx];
+							constrainedColumns[idx] = true;
 
-							RemainingWidth -= Math.Max(TargetColumnWidths[Idx] - MinColumnWidths[Idx], 0.0f);
-							RemainingTotalWeight -= ColumnWeights[Idx];
+							remainingWidth -= Math.Max(targetColumnWidths[idx] - _minColumnWidths[idx], 0.0f);
+							remainingTotalWeight -= _columnWeights[idx];
 
-							Idx = -1;
+							idx = -1;
 							continue;
 						}
 					}
@@ -5755,15 +5755,15 @@ namespace UnrealGameSync
 			}
 
 			// Allocate the remaining space equally
-			for (int Idx = 0; Idx < BuildList.Columns.Count; Idx++)
+			for (int idx = 0; idx < BuildList.Columns.Count; idx++)
 			{
-				if (!ConstrainedColumns[Idx])
+				if (!constrainedColumns[idx])
 				{
-					TargetColumnWidths[Idx] = MinColumnWidths[Idx] + (RemainingWidth * ColumnWeights[Idx] / RemainingTotalWeight);
+					targetColumnWidths[idx] = _minColumnWidths[idx] + (remainingWidth * _columnWeights[idx] / remainingTotalWeight);
 				}
 			}
 
-			return TargetColumnWidths;
+			return targetColumnWidths;
 		}
 
 		private void BuidlListMultiContextMenu_Bisect_Click(object sender, EventArgs e)
@@ -5773,28 +5773,28 @@ namespace UnrealGameSync
 
 		private bool IsBisectModeEnabled()
 		{
-			return WorkspaceState.BisectChanges.Count >= 2;
+			return _workspaceState.BisectChanges.Count >= 2;
 		}
 
 		private void EnableBisectMode()
 		{
 			if (BuildList.SelectedItems.Count >= 2)
 			{
-				Dictionary<int, BisectState> ChangeNumberToBisectState = new Dictionary<int, BisectState>();
-				foreach (ListViewItem? SelectedItem in BuildList.SelectedItems)
+				Dictionary<int, BisectState> changeNumberToBisectState = new Dictionary<int, BisectState>();
+				foreach (ListViewItem? selectedItem in BuildList.SelectedItems)
 				{
-					if (SelectedItem != null)
+					if (selectedItem != null)
 					{
-						ChangesRecord Change = (ChangesRecord)SelectedItem.Tag;
-						ChangeNumberToBisectState[Change.Number] = BisectState.Include;
+						ChangesRecord change = (ChangesRecord)selectedItem.Tag;
+						changeNumberToBisectState[change.Number] = BisectState.Include;
 					}
 				}
 
-				ChangeNumberToBisectState[ChangeNumberToBisectState.Keys.Min()] = BisectState.Pass;
-				ChangeNumberToBisectState[ChangeNumberToBisectState.Keys.Max()] = BisectState.Fail;
+				changeNumberToBisectState[changeNumberToBisectState.Keys.Min()] = BisectState.Pass;
+				changeNumberToBisectState[changeNumberToBisectState.Keys.Max()] = BisectState.Fail;
 
-				WorkspaceState.BisectChanges = ChangeNumberToBisectState.Select(x => new BisectEntry {  Change = x.Key, State = x.Value }).ToList();
-				WorkspaceState.Save(Logger);
+				_workspaceState.BisectChanges = changeNumberToBisectState.Select(x => new BisectEntry {  Change = x.Key, State = x.Value }).ToList();
+				_workspaceState.Save(_logger);
 
 				UpdateBuildList();
 				UpdateStatusPanel();
@@ -5803,27 +5803,27 @@ namespace UnrealGameSync
 
 		private void CancelBisectMode()
 		{
-			WorkspaceState.BisectChanges.Clear();
-			WorkspaceState.Save(Logger);
+			_workspaceState.BisectChanges.Clear();
+			_workspaceState.Save(_logger);
 
 			UpdateBuildList();
 			UpdateStatusPanel();
 		}
 
-		private void SetBisectStateForSelection(BisectState State)
+		private void SetBisectStateForSelection(BisectState state)
 		{
-			foreach (ListViewItem? SelectedItem in BuildList.SelectedItems)
+			foreach (ListViewItem? selectedItem in BuildList.SelectedItems)
 			{
-				ChangesRecord? Change = SelectedItem?.Tag as ChangesRecord;
-				if (Change != null)
+				ChangesRecord? change = selectedItem?.Tag as ChangesRecord;
+				if (change != null)
 				{
-					WorkspaceState.SetBisectState(Change.Number, State);
+					_workspaceState.SetBisectState(change.Number, state);
 				}
 			}
 
-			WorkspaceState.Save(Logger);
+			_workspaceState.Save(_logger);
 
-			ChangeNumberToLayoutInfo.Clear();
+			_changeNumberToLayoutInfo.Clear();
 			BuildList.Invalidate();
 
 			UpdateStatusPanel();
@@ -5831,32 +5831,32 @@ namespace UnrealGameSync
 
 		private int GetBisectChangeNumber()
 		{
-			int PassChangeNumber;
-			int FailChangeNumber;
-			GetRemainingBisectRange(out PassChangeNumber, out FailChangeNumber);
+			int passChangeNumber;
+			int failChangeNumber;
+			GetRemainingBisectRange(out passChangeNumber, out failChangeNumber);
 
-			List<int> ChangeNumbers = new List<int>();
-			foreach (BisectEntry Entry in WorkspaceState.BisectChanges)
+			List<int> changeNumbers = new List<int>();
+			foreach (BisectEntry entry in _workspaceState.BisectChanges)
 			{
-				if (Entry.State == BisectState.Include && Entry.Change > PassChangeNumber && Entry.Change < FailChangeNumber)
+				if (entry.State == BisectState.Include && entry.Change > passChangeNumber && entry.Change < failChangeNumber)
 				{
-					ChangeNumbers.Add(Entry.Change);
+					changeNumbers.Add(entry.Change);
 				}
 			}
 
-			ChangeNumbers.Sort();
+			changeNumbers.Sort();
 
-			return (ChangeNumbers.Count > 0) ? ChangeNumbers[ChangeNumbers.Count / 2] : -1;
+			return (changeNumbers.Count > 0) ? changeNumbers[changeNumbers.Count / 2] : -1;
 		}
 
 		private void SyncBisectChange()
 		{
-			int BisectChange = GetBisectChangeNumber();
-			if (BisectChange != -1)
+			int bisectChange = GetBisectChangeNumber();
+			if (bisectChange != -1)
 			{
-				Owner.ShowAndActivate();
-				SelectChange(BisectChange);
-				StartSync(BisectChange);
+				_owner.ShowAndActivate();
+				SelectChange(bisectChange);
+				StartSync(bisectChange);
 			}
 		}
 
@@ -5882,124 +5882,124 @@ namespace UnrealGameSync
 
 		private void WorkspaceControl_VisibleChanged(object sender, EventArgs e)
 		{
-			if (PerforceMonitor != null && PerforceMonitor.IsActive != Visible)
+			if (_perforceMonitor != null && _perforceMonitor.IsActive != Visible)
 			{
-				PerforceMonitor.IsActive = Visible;
+				_perforceMonitor.IsActive = Visible;
 			}
 		}
 
 		private void FilterButton_Click(object sender, EventArgs e)
 		{
-			FilterContextMenu_Default.Checked = !Settings.bShowAutomatedChanges && ProjectSettings.FilterType == FilterType.None && ProjectSettings.FilterBadges.Count == 0;
+			FilterContextMenu_Default.Checked = !_settings.ShowAutomatedChanges && _projectSettings.FilterType == FilterType.None && _projectSettings.FilterBadges.Count == 0;
 
-			FilterContextMenu_Type.Checked = ProjectSettings.FilterType != FilterType.None;
-			FilterContextMenu_Type_ShowAll.Checked = ProjectSettings.FilterType == FilterType.None;
-			FilterContextMenu_Type_Code.Checked = ProjectSettings.FilterType == FilterType.Code;
-			FilterContextMenu_Type_Content.Checked = ProjectSettings.FilterType == FilterType.Content;
+			FilterContextMenu_Type.Checked = _projectSettings.FilterType != FilterType.None;
+			FilterContextMenu_Type_ShowAll.Checked = _projectSettings.FilterType == FilterType.None;
+			FilterContextMenu_Type_Code.Checked = _projectSettings.FilterType == FilterType.Code;
+			FilterContextMenu_Type_Content.Checked = _projectSettings.FilterType == FilterType.Content;
 
 			FilterContextMenu_Badges.DropDownItems.Clear();
-			FilterContextMenu_Badges.Checked = ProjectSettings.FilterBadges.Count > 0;
+			FilterContextMenu_Badges.Checked = _projectSettings.FilterBadges.Count > 0;
 
-			HashSet<string> BadgeNames = new HashSet<string>(ProjectSettings.FilterBadges, StringComparer.OrdinalIgnoreCase);
-			BadgeNames.ExceptWith(BadgeNameAndGroupPairs.Select(x => x.Key));
+			HashSet<string> badgeNames = new HashSet<string>(_projectSettings.FilterBadges, StringComparer.OrdinalIgnoreCase);
+			badgeNames.ExceptWith(_badgeNameAndGroupPairs.Select(x => x.Key));
 
-			List<KeyValuePair<string, string>> DisplayBadgeNameAndGroupPairs = new List<KeyValuePair<string, string>>(BadgeNameAndGroupPairs);
-			DisplayBadgeNameAndGroupPairs.AddRange(BadgeNames.Select(x => new KeyValuePair<string, string>(x, "User")));
+			List<KeyValuePair<string, string>> displayBadgeNameAndGroupPairs = new List<KeyValuePair<string, string>>(_badgeNameAndGroupPairs);
+			displayBadgeNameAndGroupPairs.AddRange(badgeNames.Select(x => new KeyValuePair<string, string>(x, "User")));
 
-			string? LastGroup = null;
-			foreach (KeyValuePair<string, string> BadgeNameAndGroupPair in DisplayBadgeNameAndGroupPairs)
+			string? lastGroup = null;
+			foreach (KeyValuePair<string, string> badgeNameAndGroupPair in displayBadgeNameAndGroupPairs)
 			{
-				if (LastGroup != BadgeNameAndGroupPair.Value)
+				if (lastGroup != badgeNameAndGroupPair.Value)
 				{
-					if (LastGroup != null)
+					if (lastGroup != null)
 					{
 						FilterContextMenu_Badges.DropDownItems.Add(new ToolStripSeparator());
 					}
-					LastGroup = BadgeNameAndGroupPair.Value;
+					lastGroup = badgeNameAndGroupPair.Value;
 				}
 
-				ToolStripMenuItem Item = new ToolStripMenuItem(BadgeNameAndGroupPair.Key);
-				Item.Checked = ProjectSettings.FilterBadges.Contains(BadgeNameAndGroupPair.Key, StringComparer.OrdinalIgnoreCase);
-				Item.Click += (Sender, Args) => FilterContextMenu_Badge_Click(BadgeNameAndGroupPair.Key);
-				FilterContextMenu_Badges.DropDownItems.Add(Item);
+				ToolStripMenuItem item = new ToolStripMenuItem(badgeNameAndGroupPair.Key);
+				item.Checked = _projectSettings.FilterBadges.Contains(badgeNameAndGroupPair.Key, StringComparer.OrdinalIgnoreCase);
+				item.Click += (sender, args) => FilterContextMenu_Badge_Click(badgeNameAndGroupPair.Key);
+				FilterContextMenu_Badges.DropDownItems.Add(item);
 			}
 
 			FilterContextMenu_Badges.Enabled = FilterContextMenu_Badges.DropDownItems.Count > 0;
 
-			FilterContextMenu_ShowBuildMachineChanges.Checked = Settings.bShowAutomatedChanges;
+			FilterContextMenu_ShowBuildMachineChanges.Checked = _settings.ShowAutomatedChanges;
 
 			// Set checks if any robomerge filters are applied
-			FilterContextMenu_Robomerge.Checked = Settings.ShowRobomerge != UserSettings.RobomergeShowChangesOption.All || Settings.bAnnotateRobmergeChanges;
-			FilterContextMenu_Robomerge_ShowAll.Checked = Settings.ShowRobomerge == UserSettings.RobomergeShowChangesOption.All;
-			FilterContextMenu_Robomerge_ShowBadged.Checked = Settings.ShowRobomerge == UserSettings.RobomergeShowChangesOption.Badged;
-			FilterContextMenu_Robomerge_Annotate.Checked = Settings.bAnnotateRobmergeChanges;
+			FilterContextMenu_Robomerge.Checked = _settings.ShowRobomerge != UserSettings.RobomergeShowChangesOption.All || _settings.AnnotateRobmergeChanges;
+			FilterContextMenu_Robomerge_ShowAll.Checked = _settings.ShowRobomerge == UserSettings.RobomergeShowChangesOption.All;
+			FilterContextMenu_Robomerge_ShowBadged.Checked = _settings.ShowRobomerge == UserSettings.RobomergeShowChangesOption.Badged;
+			FilterContextMenu_Robomerge_Annotate.Checked = _settings.AnnotateRobmergeChanges;
 
 			// Set checks if an author filter string is set
-			FilterContextMenu_Author.Checked = !string.IsNullOrEmpty(this.AuthorFilterText) && this.AuthorFilterText != AuthorFilterPlaceholderText;
+			FilterContextMenu_Author.Checked = !string.IsNullOrEmpty(this._authorFilterText) && this._authorFilterText != _authorFilterPlaceholderText;
 
 			FilterContextMenu.Show(FilterButton, new Point(0, FilterButton.Height));
 		}
 
-		private void FilterContextMenu_Badge_Click(string BadgeName)
+		private void FilterContextMenu_Badge_Click(string badgeName)
 		{
-			if (ProjectSettings.FilterBadges.Contains(BadgeName))
+			if (_projectSettings.FilterBadges.Contains(badgeName))
 			{
-				ProjectSettings.FilterBadges.Remove(BadgeName);
+				_projectSettings.FilterBadges.Remove(badgeName);
 			}
 			else
 			{
-				ProjectSettings.FilterBadges.Add(BadgeName);
+				_projectSettings.FilterBadges.Add(badgeName);
 			}
-			Settings.Save(Logger);
+			_settings.Save(_logger);
 
 			UpdateBuildListFilter();
 		}
 
 		private void FilterContextMenu_Default_Click(object sender, EventArgs e)
 		{
-			ProjectSettings.FilterBadges.Clear();
-			ProjectSettings.FilterType = FilterType.None;
-			ProjectSettings.Save(Logger);
+			_projectSettings.FilterBadges.Clear();
+			_projectSettings.FilterType = FilterType.None;
+			_projectSettings.Save(_logger);
 
-			Settings.bShowAutomatedChanges = false;
-			Settings.Save(Logger);
+			_settings.ShowAutomatedChanges = false;
+			_settings.Save(_logger);
 
 			UpdateBuildListFilter();
 		}
 
 		private void FilterContextMenu_ShowBuildMachineChanges_Click(object sender, EventArgs e)
 		{
-			Settings.bShowAutomatedChanges ^= true;
-			Settings.Save(Logger);
+			_settings.ShowAutomatedChanges ^= true;
+			_settings.Save(_logger);
 
 			UpdateBuildListFilter();
 		}
 
 		private void FilterContextMenu_Robomerge_ShowAll_Click(object sender, EventArgs e)
 		{
-			Settings.ShowRobomerge = UserSettings.RobomergeShowChangesOption.All;	
-			Settings.Save(Logger);
+			_settings.ShowRobomerge = UserSettings.RobomergeShowChangesOption.All;	
+			_settings.Save(_logger);
 			UpdateBuildListFilter();
 		}
 
 		private void FilterContextMenu_Robomerge_ShowBadged_Click(object sender, EventArgs e)
 		{
-			Settings.ShowRobomerge = UserSettings.RobomergeShowChangesOption.Badged;
-			Settings.Save(Logger);
+			_settings.ShowRobomerge = UserSettings.RobomergeShowChangesOption.Badged;
+			_settings.Save(_logger);
 			UpdateBuildListFilter();
 		}
 
 		private void FilterContextMenu_Robomerge_ShowNone_Click(object sender, EventArgs e)
 		{
-			Settings.ShowRobomerge = UserSettings.RobomergeShowChangesOption.None;
-			Settings.Save(Logger);
+			_settings.ShowRobomerge = UserSettings.RobomergeShowChangesOption.None;
+			_settings.Save(_logger);
 			UpdateBuildListFilter();
 		}
 
 		private void FilterContextMenu_Robomerge_Annotate_Click(object sender, EventArgs e)
 		{
-			Settings.bAnnotateRobmergeChanges ^= true;
-			Settings.Save(Logger);
+			_settings.AnnotateRobmergeChanges ^= true;
+			_settings.Save(_logger);
 			BuildList.Items.Clear();	// need to reload to update user names
 
 			UpdateBuildListFilter();
@@ -6013,7 +6013,7 @@ namespace UnrealGameSync
 		/// <param name="e"></param>
 		private void FilterContextMenu_Author_Name_GotFocus(object sender, EventArgs e)
 		{
-			if (this.FilterContextMenu_Author_Name.Text == AuthorFilterPlaceholderText)
+			if (this.FilterContextMenu_Author_Name.Text == _authorFilterPlaceholderText)
 			{
 				this.FilterContextMenu_Author_Name.Text = "";
 				this.FilterContextMenu_Author_Name.TextBox.ForeColor = Color.Black;
@@ -6030,7 +6030,7 @@ namespace UnrealGameSync
 		{
 			if (string.IsNullOrEmpty(this.FilterContextMenu_Author_Name.Text))
 			{
-				this.FilterContextMenu_Author_Name.Text = AuthorFilterPlaceholderText;
+				this.FilterContextMenu_Author_Name.Text = _authorFilterPlaceholderText;
 				this.FilterContextMenu_Author_Name.TextBox.ForeColor = Color.DarkGray;
 			}
 		}
@@ -6043,31 +6043,31 @@ namespace UnrealGameSync
 		/// <param name="e"></param>
 		private void FilterContextMenu_Author_Name_Changed(object sender, EventArgs e)
 		{
-			AuthorFilterText = this.FilterContextMenu_Author_Name.Text;
+			_authorFilterText = this.FilterContextMenu_Author_Name.Text;
 			UpdateBuildListFilter();
 		}
 
 
 		private void FilterContextMenu_Type_ShowAll_Click(object sender, EventArgs e)
 		{
-			ProjectSettings.FilterType = FilterType.None;
-			ProjectSettings.Save(Logger);
+			_projectSettings.FilterType = FilterType.None;
+			_projectSettings.Save(_logger);
 
 			UpdateBuildListFilter();
 		}
 
 		private void FilterContextMenu_Type_Code_Click(object sender, EventArgs e)
 		{
-			ProjectSettings.FilterType = FilterType.Code;
-			ProjectSettings.Save(Logger);
+			_projectSettings.FilterType = FilterType.Code;
+			_projectSettings.Save(_logger);
 
 			UpdateBuildListFilter();
 		}
 
 		private void FilterContextMenu_Type_Content_Click(object sender, EventArgs e)
 		{
-			ProjectSettings.FilterType = FilterType.Content;
-			ProjectSettings.Save(Logger);
+			_projectSettings.FilterType = FilterType.Content;
+			_projectSettings.Save(_logger);
 
 			UpdateBuildListFilter();
 		}
@@ -6108,105 +6108,105 @@ namespace UnrealGameSync
 
 		private Color? GetNewTintColor()
 		{
-			FileReference NewConfigFile = FileReference.Combine(SelectedFileName.Directory, "Saved", "Config", "WindowsEditor", "EditorPerProjectUserSettings.ini");
-			if (!FileReference.Exists(NewConfigFile))
+			FileReference newConfigFile = FileReference.Combine(SelectedFileName.Directory, "Saved", "Config", "WindowsEditor", "EditorPerProjectUserSettings.ini");
+			if (!FileReference.Exists(newConfigFile))
 			{
 				return null;
 			}
 
-			ConfigFile Config = new ConfigFile();
-			Config.Load(NewConfigFile);
+			ConfigFile config = new ConfigFile();
+			config.Load(newConfigFile);
 
-			ConfigSection Section = Config.FindSection("/Script/EditorStyle.EditorStyleSettings");
-			if (Section == null)
+			ConfigSection section = config.FindSection("/Script/EditorStyle.EditorStyleSettings");
+			if (section == null)
 			{
 				return null;
 			}
 
-			string? EnableColor = Section.GetValue("bEnableEditorWindowBackgroundColor", null);
-			if (EnableColor == null || !String.Equals(EnableColor, "True", StringComparison.OrdinalIgnoreCase))
+			string? enableColor = section.GetValue("bEnableEditorWindowBackgroundColor", null);
+			if (enableColor == null || !String.Equals(enableColor, "True", StringComparison.OrdinalIgnoreCase))
 			{
 				return null;
 			}
 
-			string? BackgroundColor = Section.GetValue("EditorWindowBackgroundColor", null);
-			if (BackgroundColor == null)
+			string? backgroundColor = section.GetValue("EditorWindowBackgroundColor", null);
+			if (backgroundColor == null)
 			{
 				return null;
 			}
 
-			ConfigObject OverrideObject = new ConfigObject(BackgroundColor);
+			ConfigObject overrideObject = new ConfigObject(backgroundColor);
 
-			float R, G, B;
-			if (!float.TryParse(OverrideObject.GetValue("R", ""), out R) || !float.TryParse(OverrideObject.GetValue("G", ""), out G) || !float.TryParse(OverrideObject.GetValue("B", ""), out B))
+			float r, g, b;
+			if (!float.TryParse(overrideObject.GetValue("R", ""), out r) || !float.TryParse(overrideObject.GetValue("G", ""), out g) || !float.TryParse(overrideObject.GetValue("B", ""), out b))
 			{
 				return null;
 			}
 
-			return Color.FromArgb((int)(255.0f * R), (int)(255.0f * G), (int)(255.0f * B));
+			return Color.FromArgb((int)(255.0f * r), (int)(255.0f * g), (int)(255.0f * b));
 		}
 
 		private Color? GetLegacyTintColor()
 		{
-			FileReference FileName = FileReference.Combine(SelectedFileName.Directory, "Saved", "Config", "Windows", "EditorPerProjectUserSettings.ini");
-			if (!FileReference.Exists(FileName))
+			FileReference fileName = FileReference.Combine(SelectedFileName.Directory, "Saved", "Config", "Windows", "EditorPerProjectUserSettings.ini");
+			if (!FileReference.Exists(fileName))
 			{
 				return null;
 			}
 
-			ConfigFile Config = new ConfigFile();
-			Config.Load(FileName);
+			ConfigFile config = new ConfigFile();
+			config.Load(fileName);
 
-			ConfigSection Section = Config.FindSection("/Script/EditorStyle.EditorStyleSettings");
-			if (Section == null)
+			ConfigSection section = config.FindSection("/Script/EditorStyle.EditorStyleSettings");
+			if (section == null)
 			{
 				return null;
 			}
 
-			string? OverrideValue = Section.GetValue("EditorMainWindowBackgroundOverride", null);
-			if (OverrideValue == null)
+			string? overrideValue = section.GetValue("EditorMainWindowBackgroundOverride", null);
+			if (overrideValue == null)
 			{
 				return null;
 			}
 
-			ConfigObject OverrideObject = new ConfigObject(OverrideValue);
+			ConfigObject overrideObject = new ConfigObject(overrideValue);
 
-			string? TintColorValue = OverrideObject.GetValue("TintColor");
-			if (TintColorValue == null)
+			string? tintColorValue = overrideObject.GetValue("TintColor");
+			if (tintColorValue == null)
 			{
 				return null;
 			}
 
-			ConfigObject TintColorObject = new ConfigObject(TintColorValue);
-			if (TintColorObject.GetValue("ColorUseRule", "") != "UseColor_Specified")
+			ConfigObject tintColorObject = new ConfigObject(tintColorValue);
+			if (tintColorObject.GetValue("ColorUseRule", "") != "UseColor_Specified")
 			{
 				return null;
 			}
 
-			string? SpecifiedColorValue = TintColorObject.GetValue("SpecifiedColor");
-			if (SpecifiedColorValue == null)
+			string? specifiedColorValue = tintColorObject.GetValue("SpecifiedColor");
+			if (specifiedColorValue == null)
 			{
 				return null;
 			}
 
-			ConfigObject SpecifiedColorObject = new ConfigObject(SpecifiedColorValue);
+			ConfigObject specifiedColorObject = new ConfigObject(specifiedColorValue);
 
-			float R, G, B;
-			if (!float.TryParse(SpecifiedColorObject.GetValue("R", ""), out R) || !float.TryParse(SpecifiedColorObject.GetValue("G", ""), out G) || !float.TryParse(SpecifiedColorObject.GetValue("B", ""), out B))
+			float r, g, b;
+			if (!float.TryParse(specifiedColorObject.GetValue("R", ""), out r) || !float.TryParse(specifiedColorObject.GetValue("G", ""), out g) || !float.TryParse(specifiedColorObject.GetValue("B", ""), out b))
 			{
 				return null;
 			}
 
-			return Color.FromArgb((int)(255.0f * R), (int)(255.0f * G), (int)(255.0f * B));
+			return Color.FromArgb((int)(255.0f * r), (int)(255.0f * g), (int)(255.0f * b));
 		}
 
 		private void UpdateTintColor()
 		{
-			Color? NewTintColor = GetTintColor();
-			if (TintColor != NewTintColor)
+			Color? newTintColor = GetTintColor();
+			if (TintColor != newTintColor)
 			{
-				TintColor = NewTintColor;
-				Owner.UpdateTintColors();
+				TintColor = newTintColor;
+				_owner.UpdateTintColors();
 			}
 		}
 	}
