UGS: Reformat to standard C# style conventions.

diff --git a/Engine/Source/Programs/UnrealGameSync/UnrealGameSync/DeleteFilesTask.cs b/Engine/Source/Programs/UnrealGameSync/UnrealGameSync/DeleteFilesTask.cs
--- a/Engine/Source/Programs/UnrealGameSync/UnrealGameSync/DeleteFilesTask.cs
+++ b/Engine/Source/Programs/UnrealGameSync/UnrealGameSync/DeleteFilesTask.cs
@@ -14,56 +14,56 @@ namespace UnrealGameSync
 {
 	static class DeleteFilesTask
 	{
-		public static async Task RunAsync(IPerforceSettings PerforceSettings, List<FileInfo> FilesToSync, List<FileInfo> FilesToDelete, List<DirectoryInfo> DirectoriesToDelete, ILogger Logger, CancellationToken CancellationToken)
+		public static async Task RunAsync(IPerforceSettings perforceSettings, List<FileInfo> filesToSync, List<FileInfo> filesToDelete, List<DirectoryInfo> directoriesToDelete, ILogger logger, CancellationToken cancellationToken)
 		{
-			StringBuilder FailMessage = new StringBuilder();
+			StringBuilder failMessage = new StringBuilder();
 
-			if(FilesToSync.Count > 0)
+			if(filesToSync.Count > 0)
 			{
-				using IPerforceConnection Perforce = await PerforceConnection.CreateAsync(PerforceSettings, Logger);
+				using IPerforceConnection perforce = await PerforceConnection.CreateAsync(perforceSettings, logger);
 
-				List<string> RevisionsToSync = new List<string>();
-				foreach(FileInfo FileToSync in FilesToSync)
+				List<string> revisionsToSync = new List<string>();
+				foreach(FileInfo fileToSync in filesToSync)
 				{
-					RevisionsToSync.Add(String.Format("{0}#have", PerforceUtils.EscapePath(FileToSync.FullName)));
+					revisionsToSync.Add(String.Format("{0}#have", PerforceUtils.EscapePath(fileToSync.FullName)));
 				}
 
-				List<PerforceResponse<SyncRecord>> FailedRecords = await Perforce.TrySyncAsync(SyncOptions.Force, RevisionsToSync, CancellationToken).Where(x => x.Failed).ToListAsync(CancellationToken);
-				foreach (PerforceResponse<SyncRecord> FailedRecord in FailedRecords)
+				List<PerforceResponse<SyncRecord>> failedRecords = await perforce.TrySyncAsync(SyncOptions.Force, revisionsToSync, cancellationToken).Where(x => x.Failed).ToListAsync(cancellationToken);
+				foreach (PerforceResponse<SyncRecord> failedRecord in failedRecords)
 				{
-					FailMessage.Append(FailedRecord.ToString());
+					failMessage.Append(failedRecord.ToString());
 				}
 			}
 
-			foreach(FileInfo FileToDelete in FilesToDelete)
+			foreach(FileInfo fileToDelete in filesToDelete)
 			{
 				try
 				{
-					FileToDelete.IsReadOnly = false;
-					FileToDelete.Delete();
+					fileToDelete.IsReadOnly = false;
+					fileToDelete.Delete();
 				}
-				catch(Exception Ex)
+				catch(Exception ex)
 				{
-					Logger.LogWarning(Ex, "Unable to delete {File}", FileToDelete.FullName);
-					FailMessage.AppendFormat("{0} ({1})\r\n", FileToDelete.FullName, Ex.Message.Trim());
+					logger.LogWarning(ex, "Unable to delete {File}", fileToDelete.FullName);
+					failMessage.AppendFormat("{0} ({1})\r\n", fileToDelete.FullName, ex.Message.Trim());
 				}
 			}
-			foreach(DirectoryInfo DirectoryToDelete in DirectoriesToDelete)
+			foreach(DirectoryInfo directoryToDelete in directoriesToDelete)
 			{
 				try
 				{
-					DirectoryToDelete.Delete(true);
+					directoryToDelete.Delete(true);
 				}
-				catch(Exception Ex)
+				catch(Exception ex)
 				{
-					Logger.LogWarning(Ex, "Unable to delete {Directory}", DirectoryToDelete.FullName);
-					FailMessage.AppendFormat("{0} ({1})\r\n", DirectoryToDelete.FullName, Ex.Message.Trim());
+					logger.LogWarning(ex, "Unable to delete {Directory}", directoryToDelete.FullName);
+					failMessage.AppendFormat("{0} ({1})\r\n", directoryToDelete.FullName, ex.Message.Trim());
 				}
 			}
 
-			if(FailMessage.Length > 0)
+			if(failMessage.Length > 0)
 			{
-				throw new UserErrorException(FailMessage.ToString());
+				throw new UserErrorException(failMessage.ToString());
 			}
 		}
 	}
diff --git a/Engine/Source/Programs/UnrealGameSync/UnrealGameSync/DeploymentSettings.cs b/Engine/Source/Programs/UnrealGameSync/UnrealGameSync/DeploymentSettings.cs
--- a/Engine/Source/Programs/UnrealGameSync/UnrealGameSync/DeploymentSettings.cs
+++ b/Engine/Source/Programs/UnrealGameSync/UnrealGameSync/DeploymentSettings.cs
@@ -20,11 +20,11 @@ namespace UnrealGameSync
 		/// <summary>
 		/// Delegate used to create a telemetry sink
 		/// </summary>
-		/// <param name="UserName">The default Perforce user name</param>
-		/// <param name="SessionId">Unique identifier for this session</param>
-		/// <param name="Logger">Log writer</param>
+		/// <param name="userName">The default Perforce user name</param>
+		/// <param name="sessionId">Unique identifier for this session</param>
+		/// <param name="logger">Log writer</param>
 		/// <returns>New telemetry sink instance</returns>
-		public delegate ITelemetrySink CreateTelemetrySinkDelegate(string UserName, string SessionId, ILogger Logger);
+		public delegate ITelemetrySink CreateTelemetrySinkDelegate(string userName, string sessionId, ILogger logger);
 #endif
 
 		/// <summary>
@@ -68,7 +68,7 @@ namespace UnrealGameSync
 		/// <summary>
 		/// Delegate used to create a new telemetry sink
 		/// </summary>
-		public static readonly CreateTelemetrySinkDelegate CreateTelemetrySink = (UserName, SessionId, Log) => new NullTelemetrySink();
+		public static readonly CreateTelemetrySinkDelegate CreateTelemetrySink = (userName, sessionId, log) => new NullTelemetrySink();
 #endif
 
 #if !UGS_LAUNCHER
@@ -76,10 +76,10 @@ namespace UnrealGameSync
 		/// Delegate to allow validating a project being opened
 		/// </summary>
 		/// <param name="Workspace">The detected settings for the project</param>
-		/// <param name="Logger">The logger</param>
-		/// <param name="Error">Receives an error on failure</param>
+		/// <param name="logger">The logger</param>
+		/// <param name="error">Receives an error on failure</param>
 		/// <returns></returns>
-		public delegate bool DetectProjectSettingsEvent(OpenProjectInfo OpenProjectInfo, ILogger Logger, [NotNullWhen(false)] out string? Error);
+		public delegate bool DetectProjectSettingsEvent(OpenProjectInfo openProjectInfo, ILogger logger, [NotNullWhen(false)] out string? error);
 
 		/// <summary>
 		/// Called to validate the project settings
diff --git a/Engine/Source/Programs/UnrealGameSync/UnrealGameSync/EventMonitor.cs b/Engine/Source/Programs/UnrealGameSync/UnrealGameSync/EventMonitor.cs
--- a/Engine/Source/Programs/UnrealGameSync/UnrealGameSync/EventMonitor.cs
+++ b/Engine/Source/Programs/UnrealGameSync/UnrealGameSync/EventMonitor.cs
@@ -100,14 +100,14 @@ namespace UnrealGameSync
 					return "Unknown";
 				}
 
-				int Idx = BuildType.IndexOf(':');
-				if(Idx == -1)
+				int idx = BuildType.IndexOf(':');
+				if(idx == -1)
 				{
 					return BuildType;
 				}
 				else
 				{
-					return BuildType.Substring(0, Idx);
+					return BuildType.Substring(0, idx);
 				}
 			}
 		}
@@ -121,14 +121,14 @@ namespace UnrealGameSync
 					return "Unknown";
 				}
 
-				int Idx = BuildType.IndexOf(':');
-				if(Idx == -1)
+				int idx = BuildType.IndexOf(':');
+				if(idx == -1)
 				{
 					return BuildType;
 				}
 				else
 				{
-					return BuildType.Substring(Idx + 1);
+					return BuildType.Substring(idx + 1);
 				}
 			}
 		}
@@ -212,82 +212,82 @@ namespace UnrealGameSync
 
 	class EventMonitor : IDisposable
 	{
-		string? ApiUrl;
-		int ApiVersion;
-		string Project;
-		string CurrentUserName;
-		SynchronizationContext SynchronizationContext;
-		CancellationTokenSource CancellationSource;
-		Task? WorkerTask;
-		AsyncEvent RefreshEvent = new AsyncEvent();
-		ConcurrentQueue<EventData> OutgoingEvents = new ConcurrentQueue<EventData>();
-		ConcurrentQueue<EventData> IncomingEvents = new ConcurrentQueue<EventData>();
-		ConcurrentQueue<CommentData> OutgoingComments = new ConcurrentQueue<CommentData>();
-		ConcurrentQueue<CommentData> IncomingComments = new ConcurrentQueue<CommentData>();
-		ConcurrentQueue<BadgeData> IncomingBadges = new ConcurrentQueue<BadgeData>();
-		SortedDictionary<int, EventSummary> ChangeNumberToSummary = new SortedDictionary<int, EventSummary>();
-		Dictionary<string, EventData> UserNameToLastSyncEvent = new Dictionary<string, EventData>(StringComparer.InvariantCultureIgnoreCase);
-		Dictionary<string, BadgeData> BadgeNameToLatestData = new Dictionary<string, BadgeData>();
-		ILogger Logger;
-		IAsyncDisposer AsyncDisposer;
-		LatestData LatestIds;
-		HashSet<int> FilterChangeNumbers = new HashSet<int>();
-		List<EventData> InvestigationEvents = new List<EventData>();
-		List<EventData>? ActiveInvestigations;
+		string? _apiUrl;
+		int _apiVersion;
+		string _project;
+		string _currentUserName;
+		SynchronizationContext _synchronizationContext;
+		CancellationTokenSource _cancellationSource;
+		Task? _workerTask;
+		AsyncEvent _refreshEvent = new AsyncEvent();
+		ConcurrentQueue<EventData> _outgoingEvents = new ConcurrentQueue<EventData>();
+		ConcurrentQueue<EventData> _incomingEvents = new ConcurrentQueue<EventData>();
+		ConcurrentQueue<CommentData> _outgoingComments = new ConcurrentQueue<CommentData>();
+		ConcurrentQueue<CommentData> _incomingComments = new ConcurrentQueue<CommentData>();
+		ConcurrentQueue<BadgeData> _incomingBadges = new ConcurrentQueue<BadgeData>();
+		SortedDictionary<int, EventSummary> _changeNumberToSummary = new SortedDictionary<int, EventSummary>();
+		Dictionary<string, EventData> _userNameToLastSyncEvent = new Dictionary<string, EventData>(StringComparer.InvariantCultureIgnoreCase);
+		Dictionary<string, BadgeData> _badgeNameToLatestData = new Dictionary<string, BadgeData>();
+		ILogger _logger;
+		IAsyncDisposer _asyncDisposer;
+		LatestData _latestIds;
+		HashSet<int> _filterChangeNumbers = new HashSet<int>();
+		List<EventData> _investigationEvents = new List<EventData>();
+		List<EventData>? _activeInvestigations;
 
 		// MetadataV2
-		string MetadataStream;
-		string MetadataProject;
-		ConcurrentQueue<GetMetadataResponseV2> IncomingMetadata = new ConcurrentQueue<GetMetadataResponseV2>();
-		int MinChange;
-		int NewMinChange;
-		long MetadataSequenceNumber;
+		string _metadataStream;
+		string _metadataProject;
+		ConcurrentQueue<GetMetadataResponseV2> _incomingMetadata = new ConcurrentQueue<GetMetadataResponseV2>();
+		int _minChange;
+		int _newMinChange;
+		long _metadataSequenceNumber;
 
 		public Action? OnUpdatesReady;
 
-		public EventMonitor(string? InApiUrl, string InProject, string InCurrentUserName, IServiceProvider ServiceProvider)
+		public EventMonitor(string? inApiUrl, string inProject, string inCurrentUserName, IServiceProvider serviceProvider)
 		{
-			ApiUrl = InApiUrl;
-			Project = InProject;
-			CurrentUserName = InCurrentUserName;
-			Logger = ServiceProvider.GetRequiredService<ILogger<EventMonitor>>();
-			AsyncDisposer = ServiceProvider.GetRequiredService<IAsyncDisposer>();
-			SynchronizationContext = SynchronizationContext.Current!;
-			CancellationSource = new CancellationTokenSource();
-
-			LatestIds = new LatestData { LastBuildId = 0, LastCommentId = 0, LastEventId = 0 };
-
-			MetadataProject = String.Empty;
-			MetadataStream = Project.ToLowerInvariant().TrimEnd('/');
-			if (MetadataStream.StartsWith("//", StringComparison.Ordinal))
+			_apiUrl = inApiUrl;
+			_project = inProject;
+			_currentUserName = inCurrentUserName;
+			_logger = serviceProvider.GetRequiredService<ILogger<EventMonitor>>();
+			_asyncDisposer = serviceProvider.GetRequiredService<IAsyncDisposer>();
+			_synchronizationContext = SynchronizationContext.Current!;
+			_cancellationSource = new CancellationTokenSource();
+
+			_latestIds = new LatestData { LastBuildId = 0, LastCommentId = 0, LastEventId = 0 };
+
+			_metadataProject = String.Empty;
+			_metadataStream = _project.ToLowerInvariant().TrimEnd('/');
+			if (_metadataStream.StartsWith("//", StringComparison.Ordinal))
 			{
-				int NextIdx = MetadataStream.IndexOf('/', 2);
-				if (NextIdx != -1)
+				int nextIdx = _metadataStream.IndexOf('/', 2);
+				if (nextIdx != -1)
 				{
-					NextIdx = MetadataStream.IndexOf('/', NextIdx + 1);
-					if (NextIdx != -1)
+					nextIdx = _metadataStream.IndexOf('/', nextIdx + 1);
+					if (nextIdx != -1)
 					{
-						MetadataProject = MetadataStream.Substring(NextIdx + 1);
-						MetadataStream = MetadataStream.Substring(0, NextIdx);
+						_metadataProject = _metadataStream.Substring(nextIdx + 1);
+						_metadataStream = _metadataStream.Substring(0, nextIdx);
 					}
 				}
 			}
 
-			if(ApiUrl == null)
+			if(_apiUrl == null)
 			{
 				LastStatusMessage = "Database functionality disabled due to empty ApiUrl.";
 			}
 			else
 			{
-				Logger.LogInformation("Using connection string: {ApiUrl}", ApiUrl);
+				_logger.LogInformation("Using connection string: {ApiUrl}", _apiUrl);
 			}
 		}
 
 		public void Start()
 		{
-			if (WorkerTask == null)
+			if (_workerTask == null)
 			{
-				WorkerTask = Task.Run(() => PollForUpdatesAsync(CancellationSource.Token));
+				_workerTask = Task.Run(() => PollForUpdatesAsync(_cancellationSource.Token));
 			}
 		}
 
@@ -295,80 +295,80 @@ namespace UnrealGameSync
 		{
 			OnUpdatesReady = null;
 
-			if(WorkerTask != null)
+			if(_workerTask != null)
 			{
-				CancellationSource.Cancel();
-				AsyncDisposer.Add(WorkerTask.ContinueWith(_ => CancellationSource.Dispose()));
-				WorkerTask = null;
+				_cancellationSource.Cancel();
+				_asyncDisposer.Add(_workerTask.ContinueWith(_ => _cancellationSource.Dispose()));
+				_workerTask = null;
 			}
 		}
 
-		public void FilterChanges(IEnumerable<int> ChangeNumbers)
+		public void FilterChanges(IEnumerable<int> changeNumbers)
 		{
 			// Build a lookup for all the change numbers
-			FilterChangeNumbers = new HashSet<int>(ChangeNumbers);
+			_filterChangeNumbers = new HashSet<int>(changeNumbers);
 
 			// Figure out the minimum changelist number to fetch
-			int PrevNewMinChange = NewMinChange;
-			if (ChangeNumbers.Any())
+			int prevNewMinChange = _newMinChange;
+			if (changeNumbers.Any())
 			{
-				NewMinChange = ChangeNumbers.Min(x => x);
+				_newMinChange = changeNumbers.Min(x => x);
 			}
 			else
 			{
-				NewMinChange = 0;
+				_newMinChange = 0;
 			}
 
 			// Remove any changes which are no longer relevant
-			if (ApiVersion == 2)
+			if (_apiVersion == 2)
 			{
-				while (ChangeNumberToSummary.Count > 0)
+				while (_changeNumberToSummary.Count > 0)
 				{
-					int FirstChange = ChangeNumberToSummary.Keys.First();
-					if (FirstChange >= NewMinChange)
+					int firstChange = _changeNumberToSummary.Keys.First();
+					if (firstChange >= _newMinChange)
 					{
 						break;
 					}
-					ChangeNumberToSummary.Remove(FirstChange);
+					_changeNumberToSummary.Remove(firstChange);
 				}
 			}
 
 			// Clear out the list of active users for each review we have
-			UserNameToLastSyncEvent.Clear();
-			foreach(EventSummary Summary in ChangeNumberToSummary.Values)
+			_userNameToLastSyncEvent.Clear();
+			foreach(EventSummary summary in _changeNumberToSummary.Values)
 			{
-				Summary.CurrentUsers.Clear();
+				summary.CurrentUsers.Clear();
 			}
 
 			// Add all the user reviews back in again
-			foreach(EventSummary Summary in ChangeNumberToSummary.Values)
+			foreach(EventSummary summary in _changeNumberToSummary.Values)
 			{
-				foreach(EventData SyncEvent in Summary.SyncEvents)
+				foreach(EventData syncEvent in summary.SyncEvents)
 				{
-					ApplyFilteredUpdate(SyncEvent);
+					ApplyFilteredUpdate(syncEvent);
 				}
 			}
 
 			// Clear the list of active investigations, since this depends on the changes we're showing
-			ActiveInvestigations = null;
+			_activeInvestigations = null;
 
 			// Trigger an update if there's something to do
-			if (NewMinChange < PrevNewMinChange || (NewMinChange != 0 && PrevNewMinChange == 0))
+			if (_newMinChange < prevNewMinChange || (_newMinChange != 0 && prevNewMinChange == 0))
 			{
-				RefreshEvent.Set();
+				_refreshEvent.Set();
 			}
 		}
 
-		protected EventSummary FindOrAddSummary(int ChangeNumber)
+		protected EventSummary FindOrAddSummary(int changeNumber)
 		{
-			EventSummary? Summary;
-			if(!ChangeNumberToSummary.TryGetValue(ChangeNumber, out Summary))
+			EventSummary? summary;
+			if(!_changeNumberToSummary.TryGetValue(changeNumber, out summary))
 			{
-				Summary = new EventSummary();
-				Summary.ChangeNumber = ChangeNumber;
-				ChangeNumberToSummary.Add(ChangeNumber, Summary);
+				summary = new EventSummary();
+				summary.ChangeNumber = changeNumber;
+				_changeNumberToSummary.Add(changeNumber, summary);
 			}
-			return Summary;
+			return summary;
 		}
 
 		public string LastStatusMessage
@@ -379,71 +379,71 @@ namespace UnrealGameSync
 
 		public void ApplyUpdates()
 		{
-			GetMetadataResponseV2? Metadata;
-			while (IncomingMetadata.TryDequeue(out Metadata))
+			GetMetadataResponseV2? metadata;
+			while (_incomingMetadata.TryDequeue(out metadata))
 			{
-				ConvertMetadataToEvents(Metadata);
+				ConvertMetadataToEvents(metadata);
 			}
 
-			EventData? Event;
-			while(IncomingEvents.TryDequeue(out Event))
+			EventData? evt;
+			while(_incomingEvents.TryDequeue(out evt))
 			{
-				ApplyEventUpdate(Event);
+				ApplyEventUpdate(evt);
 			}
 
-			BadgeData? Badge;
-			while(IncomingBadges.TryDequeue(out Badge))
+			BadgeData? badge;
+			while(_incomingBadges.TryDequeue(out badge))
 			{
-				ApplyBadgeUpdate(Badge);
+				ApplyBadgeUpdate(badge);
 			}
 
-			CommentData? Comment;
-			while(IncomingComments.TryDequeue(out Comment))
+			CommentData? comment;
+			while(_incomingComments.TryDequeue(out comment))
 			{
-				ApplyCommentUpdate(Comment);
+				ApplyCommentUpdate(comment);
 			}
 		}
 
-		void ApplyEventUpdate(EventData Event)
+		void ApplyEventUpdate(EventData evt)
 		{
-			EventSummary Summary = FindOrAddSummary(Event.Change);
-			if(Event.Type == EventType.Starred || Event.Type == EventType.Unstarred)
+			EventSummary summary = FindOrAddSummary(evt.Change);
+			if(evt.Type == EventType.Starred || evt.Type == EventType.Unstarred)
 			{
 				// If it's a star or un-star review, process that separately
-				if(Summary.LastStarReview == null || Event.Id > Summary.LastStarReview.Id)
+				if(summary.LastStarReview == null || evt.Id > summary.LastStarReview.Id)
 				{
-					Summary.LastStarReview = Event;
+					summary.LastStarReview = evt;
 				}
 			}
-			else if(Event.Type == EventType.Investigating || Event.Type == EventType.Resolved)
+			else if(evt.Type == EventType.Investigating || evt.Type == EventType.Resolved)
 			{
 				// Insert it sorted in the investigation list
-				int InsertIdx = 0;
-				while(InsertIdx < InvestigationEvents.Count && InvestigationEvents[InsertIdx].Id < Event.Id)
+				int insertIdx = 0;
+				while(insertIdx < _investigationEvents.Count && _investigationEvents[insertIdx].Id < evt.Id)
 				{
-					InsertIdx++;
+					insertIdx++;
 				}
-				if(InsertIdx == InvestigationEvents.Count || InvestigationEvents[InsertIdx].Id != Event.Id)
+				if(insertIdx == _investigationEvents.Count || _investigationEvents[insertIdx].Id != evt.Id)
 				{
-					InvestigationEvents.Insert(InsertIdx, Event);
+					_investigationEvents.Insert(insertIdx, evt);
 				}
-				ActiveInvestigations = null;
+				_activeInvestigations = null;
 			}
-			else if(Event.Type == EventType.Syncing)
+			else if(evt.Type == EventType.Syncing)
 			{
-				Summary.SyncEvents.RemoveAll(x => String.Compare(x.UserName, Event.UserName, true) == 0);
-				Summary.SyncEvents.Add(Event);
-				ApplyFilteredUpdate(Event);
+				summary.SyncEvents.RemoveAll(x => String.Compare(x.UserName, evt.UserName, true) == 0);
+				summary.SyncEvents.Add(evt);
+				ApplyFilteredUpdate(evt);
 			}
-			else if(IsReview(Event.Type))
+			else if(IsReview(evt.Type))
 			{
 				// Try to find an existing review by this user. If we already have a newer review, ignore this one. Otherwise remove it.
-				EventData? ExistingReview = Summary.Reviews.Find(x => String.Compare(x.UserName, Event.UserName, true) == 0);
-				if(ExistingReview != null)
+				EventData? existingReview = summary.Reviews.Find(x => String.Compare(x.UserName, evt.UserName, true) == 0);
+				if(existingReview != null)
 				{
-					if(ExistingReview.Id <= Event.Id)
+					if(existingReview.Id <= evt.Id)
 					{
-						Summary.Reviews.Remove(ExistingReview);
+						summary.Reviews.Remove(existingReview);
 					}
 					else
 					{
@@ -452,8 +452,8 @@ namespace UnrealGameSync
 				}
 
 				// Add the new review, and find the new verdict for this change
-				Summary.Reviews.Add(Event);
-				Summary.Verdict = GetVerdict(Summary.Reviews, Summary.Badges);
+				summary.Reviews.Add(evt);
+				summary.Verdict = GetVerdict(summary.Reviews, summary.Badges);
 			}
 			else
 			{
@@ -461,16 +461,16 @@ namespace UnrealGameSync
 			}
 		}
 
-		void ApplyBadgeUpdate(BadgeData Badge)
+		void ApplyBadgeUpdate(BadgeData badge)
 		{
-			EventSummary Summary = FindOrAddSummary(Badge.ChangeNumber);
+			EventSummary summary = FindOrAddSummary(badge.ChangeNumber);
 
-			BadgeData? ExistingBadge = Summary.Badges.Find(x => x.ChangeNumber == Badge.ChangeNumber && x.BuildType == Badge.BuildType);
-			if(ExistingBadge != null)
+			BadgeData? existingBadge = summary.Badges.Find(x => x.ChangeNumber == badge.ChangeNumber && x.BuildType == badge.BuildType);
+			if(existingBadge != null)
 			{
-				if(ExistingBadge.Id <= Badge.Id)
+				if(existingBadge.Id <= badge.Id)
 				{
-					Summary.Badges.Remove(ExistingBadge);
+					summary.Badges.Remove(existingBadge);
 				}
 				else
 				{
@@ -478,57 +478,57 @@ namespace UnrealGameSync
 				}
 			}
 
-			Summary.Badges.Add(Badge);
-			Summary.Verdict = GetVerdict(Summary.Reviews, Summary.Badges);
+			summary.Badges.Add(badge);
+			summary.Verdict = GetVerdict(summary.Reviews, summary.Badges);
 
-			BadgeData? LatestBadge;
-			if(!BadgeNameToLatestData.TryGetValue(Badge.BadgeName, out LatestBadge) || Badge.ChangeNumber > LatestBadge.ChangeNumber || (Badge.ChangeNumber == LatestBadge.ChangeNumber && Badge.Id > LatestBadge.Id))
+			BadgeData? latestBadge;
+			if(!_badgeNameToLatestData.TryGetValue(badge.BadgeName, out latestBadge) || badge.ChangeNumber > latestBadge.ChangeNumber || (badge.ChangeNumber == latestBadge.ChangeNumber && badge.Id > latestBadge.Id))
 			{
-				BadgeNameToLatestData[Badge.BadgeName] = Badge;
+				_badgeNameToLatestData[badge.BadgeName] = badge;
 			}
 		}
 
-		void ApplyCommentUpdate(CommentData Comment)
+		void ApplyCommentUpdate(CommentData comment)
 		{
-			EventSummary Summary = FindOrAddSummary(Comment.ChangeNumber);
-			if(String.Compare(Comment.UserName, CurrentUserName, true) == 0 && Summary.Comments.Count > 0 && Summary.Comments.Last().Id == long.MaxValue)
+			EventSummary summary = FindOrAddSummary(comment.ChangeNumber);
+			if(String.Compare(comment.UserName, _currentUserName, true) == 0 && summary.Comments.Count > 0 && summary.Comments.Last().Id == long.MaxValue)
 			{
 				// This comment was added by PostComment(), to mask the latency of a round trip to the server. Remove it now we have the sorted comment.
-				Summary.Comments.RemoveAt(Summary.Comments.Count - 1);
+				summary.Comments.RemoveAt(summary.Comments.Count - 1);
 			}
-			AddPerUserItem(Summary.Comments, Comment, x => x.Id, x => x.UserName);
+			AddPerUserItem(summary.Comments, comment, x => x.Id, x => x.UserName);
 		}
 
-		static bool AddPerUserItem<T>(List<T> Items, T NewItem, Func<T, long> IdSelector, Func<T, string> UserSelector)
+		static bool AddPerUserItem<T>(List<T> items, T newItem, Func<T, long> idSelector, Func<T, string> userSelector)
 		{
-			int InsertIdx = Items.Count;
+			int insertIdx = items.Count;
 
-			for(; InsertIdx > 0 && IdSelector(Items[InsertIdx - 1]) >= IdSelector(NewItem); InsertIdx--)
+			for(; insertIdx > 0 && idSelector(items[insertIdx - 1]) >= idSelector(newItem); insertIdx--)
 			{
-				if(String.Compare(UserSelector(Items[InsertIdx - 1]), UserSelector(NewItem), true) == 0)
+				if(String.Compare(userSelector(items[insertIdx - 1]), userSelector(newItem), true) == 0)
 				{
 					return false;
 				}
 			}
 
-			Items.Insert(InsertIdx, NewItem);
+			items.Insert(insertIdx, newItem);
 
-			for(; InsertIdx > 0; InsertIdx--)
+			for(; insertIdx > 0; insertIdx--)
 			{
-				if(String.Compare(UserSelector(Items[InsertIdx - 1]), UserSelector(NewItem), true) == 0)
+				if(String.Compare(userSelector(items[insertIdx - 1]), userSelector(newItem), true) == 0)
 				{
-					Items.RemoveAt(InsertIdx - 1);
+					items.RemoveAt(insertIdx - 1);
 				}
 			}
 
 			return true;
 		}
 
-		static EventType? GetEventTypeFromVote(UgsUserVote? State)
+		static EventType? GetEventTypeFromVote(UgsUserVote? state)
 		{
-			if (State != null)
+			if (state != null)
 			{
-				switch (State.Value)
+				switch (state.Value)
 				{
 					case UgsUserVote.CompileSuccess:
 						return EventType.Compiles;
@@ -543,135 +543,135 @@ namespace UnrealGameSync
 			return null;
 		}
 
-		void ConvertMetadataToEvents(GetMetadataResponseV2 Metadata)
+		void ConvertMetadataToEvents(GetMetadataResponseV2 metadata)
 		{
-			EventSummary NewSummary = new EventSummary();
-			NewSummary.ChangeNumber = Metadata.Change;
+			EventSummary newSummary = new EventSummary();
+			newSummary.ChangeNumber = metadata.Change;
 
-			EventSummary? Summary;
-			if (ChangeNumberToSummary.TryGetValue(Metadata.Change, out Summary))
+			EventSummary? summary;
+			if (_changeNumberToSummary.TryGetValue(metadata.Change, out summary))
 			{
-				foreach (string CurrentUser in Summary.CurrentUsers)
+				foreach (string currentUser in summary.CurrentUsers)
 				{
-					UserNameToLastSyncEvent.Remove(CurrentUser);
+					_userNameToLastSyncEvent.Remove(currentUser);
 				}
 
-				NewSummary.SharedMetadata = Summary.SharedMetadata;
-				NewSummary.ProjectMetadata = Summary.ProjectMetadata;
+				newSummary.SharedMetadata = summary.SharedMetadata;
+				newSummary.ProjectMetadata = summary.ProjectMetadata;
 			}
 
-			if (string.IsNullOrEmpty(Metadata.Project))
+			if (string.IsNullOrEmpty(metadata.Project))
 			{
-				NewSummary.SharedMetadata = Metadata;
+				newSummary.SharedMetadata = metadata;
 			}
 			else
 			{
-				NewSummary.ProjectMetadata = Metadata;
+				newSummary.ProjectMetadata = metadata;
 			}
 
-			ChangeNumberToSummary[NewSummary.ChangeNumber] = NewSummary;
+			_changeNumberToSummary[newSummary.ChangeNumber] = newSummary;
 
-			if (NewSummary.SharedMetadata != null)
+			if (newSummary.SharedMetadata != null)
 			{
-				PostEvents(NewSummary.SharedMetadata);
+				PostEvents(newSummary.SharedMetadata);
 			}
 
-			if (NewSummary.ProjectMetadata != null)
+			if (newSummary.ProjectMetadata != null)
 			{
-				PostEvents(NewSummary.ProjectMetadata);
+				PostEvents(newSummary.ProjectMetadata);
 			}
 		}
 
-		void PostEvents(GetMetadataResponseV2 Metadata)
+		void PostEvents(GetMetadataResponseV2 metadata)
 		{
-			if (Metadata.Badges != null)
+			if (metadata.Badges != null)
 			{
-				foreach (GetBadgeDataResponseV2 BadgeData in Metadata.Badges)
+				foreach (GetBadgeDataResponseV2 badgeData in metadata.Badges)
 				{
-					BadgeData Badge = new BadgeData();
-					Badge.Id = ++LatestIds.LastBuildId;
-					Badge.ChangeNumber = Metadata.Change;
-					Badge.BuildType = BadgeData.Name;
-					Badge.Result = BadgeData.State;
-					Badge.Url = BadgeData.Url;
-					Badge.Project = Metadata.Project;
-					IncomingBadges.Enqueue(Badge);
+					BadgeData badge = new BadgeData();
+					badge.Id = ++_latestIds.LastBuildId;
+					badge.ChangeNumber = metadata.Change;
+					badge.BuildType = badgeData.Name;
+					badge.Result = badgeData.State;
+					badge.Url = badgeData.Url;
+					badge.Project = metadata.Project;
+					_incomingBadges.Enqueue(badge);
 				}
 			}
 
-			if (Metadata.Users != null)
+			if (metadata.Users != null)
 			{
-				foreach (GetUserDataResponseV2 UserData in Metadata.Users)
+				foreach (GetUserDataResponseV2 userData in metadata.Users)
 				{
-					if (UserData.SyncTime != null)
+					if (userData.SyncTime != null)
 					{
-						EventData Event = new EventData { Id = UserData.SyncTime.Value, Change = Metadata.Change, Project = Metadata.Project, UserName = UserData.User, Type = EventType.Syncing };
-						IncomingEvents.Enqueue(Event);
+						EventData evt = new EventData { Id = userData.SyncTime.Value, Change = metadata.Change, Project = metadata.Project, UserName = userData.User, Type = EventType.Syncing };
+						_incomingEvents.Enqueue(evt);
 					}
 
-					EventType? Type = GetEventTypeFromVote(UserData.Vote);
-					if (Type != null)
+					EventType? type = GetEventTypeFromVote(userData.Vote);
+					if (type != null)
 					{
-						EventData Event = new EventData { Id = ++LatestIds.LastEventId, Change = Metadata.Change, Project = Metadata.Project, UserName = UserData.User, Type = Type.Value };
-						IncomingEvents.Enqueue(Event);
+						EventData evt = new EventData { Id = ++_latestIds.LastEventId, Change = metadata.Change, Project = metadata.Project, UserName = userData.User, Type = type.Value };
+						_incomingEvents.Enqueue(evt);
 					}
 
-					if (UserData.Investigating != null)
+					if (userData.Investigating != null)
 					{
-						EventType InvestigationEventType = (UserData.Investigating.Value ? EventType.Investigating : EventType.Resolved);
-						EventData Event = new EventData { Id = ++LatestIds.LastEventId, Change = Metadata.Change, Project = Metadata.Project, UserName = UserData.User, Type = InvestigationEventType };
-						IncomingEvents.Enqueue(Event);
+						EventType investigationEventType = (userData.Investigating.Value ? EventType.Investigating : EventType.Resolved);
+						EventData evt = new EventData { Id = ++_latestIds.LastEventId, Change = metadata.Change, Project = metadata.Project, UserName = userData.User, Type = investigationEventType };
+						_incomingEvents.Enqueue(evt);
 					}
 
-					if (UserData.Starred != null)
+					if (userData.Starred != null)
 					{
-						EventType StarEventType = (UserData.Starred.Value ? EventType.Starred : EventType.Unstarred);
-						EventData Event = new EventData { Id = ++LatestIds.LastEventId, Change = Metadata.Change, Project = Metadata.Project, UserName = UserData.User, Type = StarEventType };
-						IncomingEvents.Enqueue(Event);
+						EventType starEventType = (userData.Starred.Value ? EventType.Starred : EventType.Unstarred);
+						EventData evt = new EventData { Id = ++_latestIds.LastEventId, Change = metadata.Change, Project = metadata.Project, UserName = userData.User, Type = starEventType };
+						_incomingEvents.Enqueue(evt);
 					}
 
-					if (UserData.Comment != null)
+					if (userData.Comment != null)
 					{
-						CommentData Comment = new CommentData { Id = ++LatestIds.LastCommentId, ChangeNumber = Metadata.Change, Project = Metadata.Project, UserName = UserData.User, Text = UserData.Comment };
-						IncomingComments.Enqueue(Comment);
+						CommentData comment = new CommentData { Id = ++_latestIds.LastCommentId, ChangeNumber = metadata.Change, Project = metadata.Project, UserName = userData.User, Text = userData.Comment };
+						_incomingComments.Enqueue(comment);
 					}
 				}
 			}
 		}
 
-		static ReviewVerdict GetVerdict(IEnumerable<EventData> Events, IEnumerable<BadgeData> Badges)
+		static ReviewVerdict GetVerdict(IEnumerable<EventData> events, IEnumerable<BadgeData> badges)
 		{
-			int NumPositiveReviews = Events.Count(x => x.Type == EventType.Good);
-			int NumNegativeReviews = Events.Count(x => x.Type == EventType.Bad);
-			if(NumPositiveReviews > 0 || NumNegativeReviews > 0)
+			int numPositiveReviews = events.Count(x => x.Type == EventType.Good);
+			int numNegativeReviews = events.Count(x => x.Type == EventType.Bad);
+			if(numPositiveReviews > 0 || numNegativeReviews > 0)
 			{
-				return GetVerdict(NumPositiveReviews, NumNegativeReviews);
+				return GetVerdict(numPositiveReviews, numNegativeReviews);
 			}
 
-			int NumCompiles = Events.Count(x => x.Type == EventType.Compiles);
-			int NumFailedCompiles = Events.Count(x => x.Type == EventType.DoesNotCompile);
-			if(NumCompiles > 0 || NumFailedCompiles > 0)
+			int numCompiles = events.Count(x => x.Type == EventType.Compiles);
+			int numFailedCompiles = events.Count(x => x.Type == EventType.DoesNotCompile);
+			if(numCompiles > 0 || numFailedCompiles > 0)
 			{
-				return GetVerdict(NumCompiles, NumFailedCompiles);
+				return GetVerdict(numCompiles, numFailedCompiles);
 			}
 
-			int NumBadges = Badges.Count(x => x.BuildType == "Editor" && x.IsSuccess);
-			int NumFailedBadges = Badges.Count(x => x.BuildType == "Editor" && x.IsFailure);
-			if(NumBadges > 0 || NumFailedBadges > 0)
+			int numBadges = badges.Count(x => x.BuildType == "Editor" && x.IsSuccess);
+			int numFailedBadges = badges.Count(x => x.BuildType == "Editor" && x.IsFailure);
+			if(numBadges > 0 || numFailedBadges > 0)
 			{
-				return GetVerdict(NumBadges, NumFailedBadges);
+				return GetVerdict(numBadges, numFailedBadges);
 			}
 
 			return ReviewVerdict.Unknown;
 		}
 
-		static ReviewVerdict GetVerdict(int NumPositive, int NumNegative)
+		static ReviewVerdict GetVerdict(int numPositive, int numNegative)
 		{
-			if(NumPositive > (int)(NumNegative * 1.5))
+			if(numPositive > (int)(numNegative * 1.5))
 			{
 				return ReviewVerdict.Good;
 			}
-			else if(NumPositive >= NumNegative)
+			else if(numPositive >= numNegative)
 			{
 				return ReviewVerdict.Mixed;
 			}
@@ -681,90 +681,90 @@ namespace UnrealGameSync
 			}
 		}
 
-		void ApplyFilteredUpdate(EventData Event)
+		void ApplyFilteredUpdate(EventData evt)
 		{
-			if(Event.Type == EventType.Syncing && FilterChangeNumbers.Contains(Event.Change) && !String.IsNullOrEmpty(Event.UserName))
+			if(evt.Type == EventType.Syncing && _filterChangeNumbers.Contains(evt.Change) && !String.IsNullOrEmpty(evt.UserName))
 			{
 				// Update the active users list for this change
-				EventData? LastSync;
-				if(UserNameToLastSyncEvent.TryGetValue(Event.UserName, out LastSync))
+				EventData? lastSync;
+				if(_userNameToLastSyncEvent.TryGetValue(evt.UserName, out lastSync))
 				{
-					if(Event.Id > LastSync.Id)
+					if(evt.Id > lastSync.Id)
 					{
-						ChangeNumberToSummary[LastSync.Change].CurrentUsers.RemoveAll(x => String.Compare(x, Event.UserName, true) == 0);
-						FindOrAddSummary(Event.Change).CurrentUsers.Add(Event.UserName);
-						UserNameToLastSyncEvent[Event.UserName] = Event;
+						_changeNumberToSummary[lastSync.Change].CurrentUsers.RemoveAll(x => String.Compare(x, evt.UserName, true) == 0);
+						FindOrAddSummary(evt.Change).CurrentUsers.Add(evt.UserName);
+						_userNameToLastSyncEvent[evt.UserName] = evt;
 					}
 				}
 				else
 				{
-					FindOrAddSummary(Event.Change).CurrentUsers.Add(Event.UserName);
-					UserNameToLastSyncEvent[Event.UserName] = Event;
+					FindOrAddSummary(evt.Change).CurrentUsers.Add(evt.UserName);
+					_userNameToLastSyncEvent[evt.UserName] = evt;
 				}
 			}
 		}
 
-		async Task PollForUpdatesAsync(CancellationToken CancellationToken)
+		async Task PollForUpdatesAsync(CancellationToken cancellationToken)
 		{
-			EventData? Event = null;
-			CommentData? Comment = null;
-			bool bUpdateThrottledRequests = true;
-			double RequestThrottle = 90; // seconds to wait for throttled request;
-			Stopwatch Timer = Stopwatch.StartNew();
-			while (!CancellationToken.IsCancellationRequested)
+			EventData? evt = null;
+			CommentData? comment = null;
+			bool updateThrottledRequests = true;
+			double requestThrottle = 90; // seconds to wait for throttled request;
+			Stopwatch timer = Stopwatch.StartNew();
+			while (!cancellationToken.IsCancellationRequested)
 			{
-				Task RefreshTask = RefreshEvent.Task;
+				Task refreshTask = _refreshEvent.Task;
 
 				// If there's no connection string, just empty out the queue
-				if (ApiUrl != null)
+				if (_apiUrl != null)
 				{
 					// Post all the reviews to the database. We don't send them out of order, so keep the review outside the queue until the next update if it fails
-					while(Event != null || OutgoingEvents.TryDequeue(out Event))
+					while(evt != null || _outgoingEvents.TryDequeue(out evt))
 					{
-						await SendEventToBackendAsync(Event, CancellationToken);
-						Event = null;
+						await SendEventToBackendAsync(evt, cancellationToken);
+						evt = null;
 					}
 
 					// Post all the comments to the database.
-					while(Comment != null || OutgoingComments.TryDequeue(out Comment))
+					while(comment != null || _outgoingComments.TryDequeue(out comment))
 					{
-						await SendCommentToBackendAsync(Comment, CancellationToken);
-						Comment = null;
+						await SendCommentToBackendAsync(comment, cancellationToken);
+						comment = null;
 					}
 
-					if(Timer.Elapsed > TimeSpan.FromSeconds(RequestThrottle))
+					if(timer.Elapsed > TimeSpan.FromSeconds(requestThrottle))
 					{
-						bUpdateThrottledRequests = true;
-						Timer.Restart();
+						updateThrottledRequests = true;
+						timer.Restart();
 					}
 
 					// Read all the new reviews, pass whether or not to fire the throttled requests
-					await ReadEventsFromBackendAsync(bUpdateThrottledRequests, CancellationToken);
+					await ReadEventsFromBackendAsync(updateThrottledRequests, cancellationToken);
 
 					// Send a notification that we're ready to update
-					if(IncomingMetadata.Count > 0 || IncomingEvents.Count > 0 || IncomingBadges.Count > 0 || IncomingComments.Count > 0)
+					if(_incomingMetadata.Count > 0 || _incomingEvents.Count > 0 || _incomingBadges.Count > 0 || _incomingComments.Count > 0)
 					{
-						SynchronizationContext.Post(_ => OnUpdatesReady?.Invoke(), null);
+						_synchronizationContext.Post(_ => OnUpdatesReady?.Invoke(), null);
 					}
 				}
 
 				// Wait for something else to do
-				Task DelayTask = Task.Delay(TimeSpan.FromSeconds(30.0), CancellationToken);
-				if (await Task.WhenAny(DelayTask, RefreshTask) == DelayTask)
+				Task delayTask = Task.Delay(TimeSpan.FromSeconds(30.0), cancellationToken);
+				if (await Task.WhenAny(delayTask, refreshTask) == delayTask)
 				{
-					bUpdateThrottledRequests = true;
+					updateThrottledRequests = true;
 				}
 			}
 		}
 
 		public static bool HasNetCore3()
 		{
-			DirectoryInfo BaseDirInfo = new DirectoryInfo(Path.Combine(Environment.GetFolderPath(Environment.SpecialFolder.ProgramFiles), "dotnet", "shared", "Microsoft.NETCore.App"));
-			if (BaseDirInfo.Exists)
+			DirectoryInfo baseDirInfo = new DirectoryInfo(Path.Combine(Environment.GetFolderPath(Environment.SpecialFolder.ProgramFiles), "dotnet", "shared", "Microsoft.NETCore.App"));
+			if (baseDirInfo.Exists)
 			{
-				foreach (DirectoryInfo SubDir in BaseDirInfo.EnumerateDirectories())
+				foreach (DirectoryInfo subDir in baseDirInfo.EnumerateDirectories())
 				{
-					if (SubDir.Name.StartsWith("3.", StringComparison.Ordinal))
+					if (subDir.Name.StartsWith("3.", StringComparison.Ordinal))
 					{
 						return true;
 					}
@@ -773,156 +773,156 @@ namespace UnrealGameSync
 			return false;
 		}
 
-		async Task<bool> SendEventToBackendAsync(EventData Event, CancellationToken CancellationToken)
+		async Task<bool> SendEventToBackendAsync(EventData evt, CancellationToken cancellationToken)
 		{
 			try
 			{
-				Stopwatch Timer = Stopwatch.StartNew();
-				Logger.LogInformation("Posting event... ({Change}, {UserName}, {Type})", Event.Change, Event.UserName, Event.Type);
-				if (ApiVersion == 2)
+				Stopwatch timer = Stopwatch.StartNew();
+				_logger.LogInformation("Posting event... ({Change}, {UserName}, {Type})", evt.Change, evt.UserName, evt.Type);
+				if (_apiVersion == 2)
 				{
-					await SendMetadataUpdateUpdateV2Async(Event.Change, Event.Project, Event.UserName, Event.Type, null, CancellationToken);
+					await SendMetadataUpdateUpdateV2Async(evt.Change, evt.Project, evt.UserName, evt.Type, null, cancellationToken);
 				}
 				else
 				{
-					await RESTApi.PostAsync($"{ApiUrl}/api/event", JsonSerializer.Serialize(Event), CancellationToken);
+					await RestApi.PostAsync($"{_apiUrl}/api/event", JsonSerializer.Serialize(evt), cancellationToken);
 				}
 				return true;
 			}
-			catch(Exception Ex)
+			catch(Exception ex)
 			{
-				Logger.LogError(Ex, "Failed with exception.");
+				_logger.LogError(ex, "Failed with exception.");
 				return false;
 			}
 		}
 
-		async Task<bool> SendCommentToBackendAsync(CommentData Comment, CancellationToken CancellationToken)
+		async Task<bool> SendCommentToBackendAsync(CommentData comment, CancellationToken cancellationToken)
 		{
 			try
 			{
-				Stopwatch Timer = Stopwatch.StartNew();
-				Logger.LogInformation("Posting comment... ({Change}, {User}, {Text}, {Project})", Comment.ChangeNumber, Comment.UserName, Comment.Text, Comment.Project);
-				if (ApiVersion == 2)
+				Stopwatch timer = Stopwatch.StartNew();
+				_logger.LogInformation("Posting comment... ({Change}, {User}, {Text}, {Project})", comment.ChangeNumber, comment.UserName, comment.Text, comment.Project);
+				if (_apiVersion == 2)
 				{
-					await SendMetadataUpdateUpdateV2Async(Comment.ChangeNumber, Comment.Project, Comment.UserName, null, Comment.Text, CancellationToken);
+					await SendMetadataUpdateUpdateV2Async(comment.ChangeNumber, comment.Project, comment.UserName, null, comment.Text, cancellationToken);
 				}
 				else
 				{
-					await RESTApi.PostAsync($"{ApiUrl}/api/comment", JsonSerializer.Serialize(Comment), CancellationToken);
+					await RestApi.PostAsync($"{_apiUrl}/api/comment", JsonSerializer.Serialize(comment), cancellationToken);
 				}
-				Logger.LogInformation("Done in {Time}ms.", Timer.ElapsedMilliseconds);
+				_logger.LogInformation("Done in {Time}ms.", timer.ElapsedMilliseconds);
 				return true;
 			}
-			catch(Exception Ex)
+			catch(Exception ex)
 			{
-				Logger.LogError(Ex, "Failed with exception.");
+				_logger.LogError(ex, "Failed with exception.");
 				return false;
 			}
 		}
 
-		async Task SendMetadataUpdateUpdateV2Async(int Change, string Project, string UserName, EventType? Event, string? Comment, CancellationToken CancellationToken)
+		async Task SendMetadataUpdateUpdateV2Async(int change, string project, string userName, EventType? evt, string? comment, CancellationToken cancellationToken)
 		{
-			UpdateMetadataRequestV2 Update = new UpdateMetadataRequestV2();
-			Update.Stream = MetadataStream;
-			Update.Project = MetadataProject;
-			Update.Change = Change;
-			Update.UserName = UserName;
+			UpdateMetadataRequestV2 update = new UpdateMetadataRequestV2();
+			update.Stream = _metadataStream;
+			update.Project = _metadataProject;
+			update.Change = change;
+			update.UserName = userName;
 
-			if (Event != null)
+			if (evt != null)
 			{
-				switch (Event)
+				switch (evt)
 				{
 					case EventType.Syncing:
-						Update.Synced = true;
+						update.Synced = true;
 						break;
 					case EventType.Compiles:
-						Update.Vote = nameof(UgsUserVote.CompileSuccess);
+						update.Vote = nameof(UgsUserVote.CompileSuccess);
 						break;
 					case EventType.DoesNotCompile:
-						Update.Vote = nameof(UgsUserVote.CompileFailure);
+						update.Vote = nameof(UgsUserVote.CompileFailure);
 						break;
 					case EventType.Good:
-						Update.Vote = nameof(UgsUserVote.Good);
+						update.Vote = nameof(UgsUserVote.Good);
 						break;
 					case EventType.Bad:
-						Update.Vote = nameof(UgsUserVote.Bad);
+						update.Vote = nameof(UgsUserVote.Bad);
 						break;
 					case EventType.Unknown:
-						Update.Vote = nameof(UgsUserVote.None);
+						update.Vote = nameof(UgsUserVote.None);
 						break;
 					case EventType.Starred:
-						Update.Starred = true;
+						update.Starred = true;
 						break;
 					case EventType.Unstarred:
-						Update.Starred = false;
+						update.Starred = false;
 						break;
 					case EventType.Investigating:
-						Update.Investigating = true;
+						update.Investigating = true;
 						break;
 					case EventType.Resolved:
-						Update.Investigating = false;
+						update.Investigating = false;
 						break;
 				}
 			}
-			Update.Comment = Comment;
+			update.Comment = comment;
 
-			await RESTApi.PostAsync($"{ApiUrl}/api/metadata", JsonSerializer.Serialize(Update), CancellationToken);
+			await RestApi.PostAsync($"{_apiUrl}/api/metadata", JsonSerializer.Serialize(update), cancellationToken);
 		}
 
-		async Task<bool> ReadEventsFromBackendAsync(bool bFireThrottledRequests, CancellationToken CancellationToken)
+		async Task<bool> ReadEventsFromBackendAsync(bool fireThrottledRequests, CancellationToken cancellationToken)
 		{
 			try
 			{
-				Stopwatch Timer = Stopwatch.StartNew();
-				Logger.LogInformation("Polling for events...");
+				Stopwatch timer = Stopwatch.StartNew();
+				_logger.LogInformation("Polling for events...");
 
 				//////////////
 				/// Initial Ids 
 				//////////////
-				if (ApiVersion == 0)
+				if (_apiVersion == 0)
 				{
-					LatestData InitialIds = await RESTApi.GetAsync<LatestData>($"{ApiUrl}/api/latest?project={Project}", CancellationToken);
-					ApiVersion = (InitialIds.Version == 0)? 1 : InitialIds.Version;
-					LatestIds.LastBuildId = InitialIds.LastBuildId;
-					LatestIds.LastCommentId = InitialIds.LastCommentId;
-					LatestIds.LastEventId = InitialIds.LastEventId;
+					LatestData initialIds = await RestApi.GetAsync<LatestData>($"{_apiUrl}/api/latest?project={_project}", cancellationToken);
+					_apiVersion = (initialIds.Version == 0)? 1 : initialIds.Version;
+					_latestIds.LastBuildId = initialIds.LastBuildId;
+					_latestIds.LastCommentId = initialIds.LastCommentId;
+					_latestIds.LastEventId = initialIds.LastEventId;
 				}
 
-				if (ApiVersion == 2)
+				if (_apiVersion == 2)
 				{
-					int NewMinChangeCopy = NewMinChange;
-					if (NewMinChangeCopy != 0)
+					int newMinChangeCopy = _newMinChange;
+					if (newMinChangeCopy != 0)
 					{
 						// If the range of changes has decreased, update the MinChange value before we fetch anything
-						MinChange = Math.Max(MinChange, NewMinChangeCopy);
+						_minChange = Math.Max(_minChange, newMinChangeCopy);
 
 						// Get the first part of the query
-						string CommonArgs = String.Format("stream={0}", MetadataStream);
-						if (MetadataProject != null)
+						string commonArgs = String.Format("stream={0}", _metadataStream);
+						if (_metadataProject != null)
 						{
-							CommonArgs += String.Format("&project={0}", MetadataProject);
+							commonArgs += String.Format("&project={0}", _metadataProject);
 						}
 
 						// Fetch any updates in the current range of changes
-						if (MinChange != 0)
+						if (_minChange != 0)
 						{
-							GetMetadataListResponseV2 NewEventList = await RESTApi.GetAsync<GetMetadataListResponseV2>($"{ApiUrl}/api/metadata?{CommonArgs}&minchange={MinChange}&sequence={MetadataSequenceNumber}", CancellationToken);
-							foreach (GetMetadataResponseV2 NewEvent in NewEventList.Items)
+							GetMetadataListResponseV2 newEventList = await RestApi.GetAsync<GetMetadataListResponseV2>($"{_apiUrl}/api/metadata?{commonArgs}&minchange={_minChange}&sequence={_metadataSequenceNumber}", cancellationToken);
+							foreach (GetMetadataResponseV2 newEvent in newEventList.Items)
 							{
-								IncomingMetadata.Enqueue(NewEvent);
+								_incomingMetadata.Enqueue(newEvent);
 							}
-							MetadataSequenceNumber = Math.Max(NewEventList.SequenceNumber, MetadataSequenceNumber);
+							_metadataSequenceNumber = Math.Max(newEventList.SequenceNumber, _metadataSequenceNumber);
 						}
 
 						// Fetch any new changes
-						if (NewMinChangeCopy < MinChange)
+						if (newMinChangeCopy < _minChange)
 						{
-							GetMetadataListResponseV2 NewEvents = await RESTApi.GetAsync<GetMetadataListResponseV2>($"{ApiUrl}/api/metadata?{CommonArgs}&minchange={NewMinChangeCopy}&maxchange={MinChange}", CancellationToken);
-							foreach (GetMetadataResponseV2 NewEvent in NewEvents.Items)
+							GetMetadataListResponseV2 newEvents = await RestApi.GetAsync<GetMetadataListResponseV2>($"{_apiUrl}/api/metadata?{commonArgs}&minchange={newMinChangeCopy}&maxchange={_minChange}", cancellationToken);
+							foreach (GetMetadataResponseV2 newEvent in newEvents.Items)
 							{
-								IncomingMetadata.Enqueue(NewEvent);
+								_incomingMetadata.Enqueue(newEvent);
 							}
-							MinChange = NewMinChangeCopy;
+							_minChange = newMinChangeCopy;
 						}
 					}
 				}
@@ -931,207 +931,207 @@ namespace UnrealGameSync
 					//////////////
 					/// Builds
 					//////////////
-					List<BadgeData> Builds = await RESTApi.GetAsync<List<BadgeData>>($"{ApiUrl}/api/build?project={Project}&lastbuildid={LatestIds.LastBuildId}", CancellationToken);
-					foreach (BadgeData Build in Builds)
+					List<BadgeData> builds = await RestApi.GetAsync<List<BadgeData>>($"{_apiUrl}/api/build?project={_project}&lastbuildid={_latestIds.LastBuildId}", cancellationToken);
+					foreach (BadgeData build in builds)
 					{
-						IncomingBadges.Enqueue(Build);
-						LatestIds.LastBuildId = Math.Max(LatestIds.LastBuildId, Build.Id);
+						_incomingBadges.Enqueue(build);
+						_latestIds.LastBuildId = Math.Max(_latestIds.LastBuildId, build.Id);
 					}
 
 					//////////////////////////
 					/// Throttled Requests
 					//////////////////////////
-					if (bFireThrottledRequests)
+					if (fireThrottledRequests)
 					{
 						//////////////
 						/// Reviews 
 						//////////////
-						List<EventData> Events = await RESTApi.GetAsync<List<EventData>>($"{ApiUrl}/api/event?project={Project}&lasteventid={LatestIds.LastEventId}", CancellationToken);
-						foreach (EventData Review in Events)
+						List<EventData> events = await RestApi.GetAsync<List<EventData>>($"{_apiUrl}/api/event?project={_project}&lasteventid={_latestIds.LastEventId}", cancellationToken);
+						foreach (EventData review in events)
 						{
-							IncomingEvents.Enqueue(Review);
-							LatestIds.LastEventId = Math.Max(LatestIds.LastEventId, Review.Id);
+							_incomingEvents.Enqueue(review);
+							_latestIds.LastEventId = Math.Max(_latestIds.LastEventId, review.Id);
 						}
 
 						//////////////
 						/// Comments 
 						//////////////
-						List<CommentData> Comments = await RESTApi.GetAsync<List<CommentData>>($"{ApiUrl}/api/comment?project={Project}&lastcommentid={LatestIds.LastCommentId}", CancellationToken);
-						foreach (CommentData Comment in Comments)
+						List<CommentData> comments = await RestApi.GetAsync<List<CommentData>>($"{_apiUrl}/api/comment?project={_project}&lastcommentid={_latestIds.LastCommentId}", cancellationToken);
+						foreach (CommentData comment in comments)
 						{
-							IncomingComments.Enqueue(Comment);
-							LatestIds.LastCommentId = Math.Max(LatestIds.LastCommentId, Comment.Id);
+							_incomingComments.Enqueue(comment);
+							_latestIds.LastCommentId = Math.Max(_latestIds.LastCommentId, comment.Id);
 						}
 					}
 				}
 
-				LastStatusMessage = String.Format("Last update took {0}ms", Timer.ElapsedMilliseconds);
-				Logger.LogInformation("Done in {Time}ms.", Timer.ElapsedMilliseconds);
+				LastStatusMessage = String.Format("Last update took {0}ms", timer.ElapsedMilliseconds);
+				_logger.LogInformation("Done in {Time}ms.", timer.ElapsedMilliseconds);
 				return true;
 			}
-			catch(Exception Ex)
+			catch(Exception ex)
 			{
-				Logger.LogError(Ex, "Failed with exception.");
-				LastStatusMessage = String.Format("Last update failed: ({0})", Ex.ToString());
+				_logger.LogError(ex, "Failed with exception.");
+				LastStatusMessage = String.Format("Last update failed: ({0})", ex.ToString());
 				return false;
 			}
 		}
 
-		static bool MatchesWildcard(string Wildcard, string Project)
+		static bool MatchesWildcard(string wildcard, string project)
 		{
-			return Wildcard.EndsWith("...") && Project.StartsWith(Wildcard.Substring(0, Wildcard.Length - 3), StringComparison.InvariantCultureIgnoreCase);
+			return wildcard.EndsWith("...") && project.StartsWith(wildcard.Substring(0, wildcard.Length - 3), StringComparison.InvariantCultureIgnoreCase);
 		}
 
-		public void PostEvent(int ChangeNumber, EventType Type)
+		public void PostEvent(int changeNumber, EventType type)
 		{
-			if(ApiUrl != null)
+			if(_apiUrl != null)
 			{
-				EventData Event = new EventData();
-				Event.Id = ++LatestIds.LastEventId;
-				Event.Change = ChangeNumber;
-				Event.UserName = CurrentUserName;
-				Event.Type = Type;
-				Event.Project = Project;
-				OutgoingEvents.Enqueue(Event);
+				EventData evt = new EventData();
+				evt.Id = ++_latestIds.LastEventId;
+				evt.Change = changeNumber;
+				evt.UserName = _currentUserName;
+				evt.Type = type;
+				evt.Project = _project;
+				_outgoingEvents.Enqueue(evt);
 
-				ApplyEventUpdate(Event);
+				ApplyEventUpdate(evt);
 
-				RefreshEvent.Set();
+				_refreshEvent.Set();
 			}
 		}
 
-		public void PostComment(int ChangeNumber, string Text)
+		public void PostComment(int changeNumber, string text)
 		{
-			if(ApiUrl != null)
+			if(_apiUrl != null)
 			{
-				CommentData Comment = new CommentData();
-				Comment.Id = long.MaxValue;
-				Comment.ChangeNumber = ChangeNumber;
-				Comment.UserName = CurrentUserName;
-				Comment.Text = Text;
-				Comment.Project = Project;
-				OutgoingComments.Enqueue(Comment);
+				CommentData comment = new CommentData();
+				comment.Id = long.MaxValue;
+				comment.ChangeNumber = changeNumber;
+				comment.UserName = _currentUserName;
+				comment.Text = text;
+				comment.Project = _project;
+				_outgoingComments.Enqueue(comment);
 
-				ApplyCommentUpdate(Comment);
+				ApplyCommentUpdate(comment);
 
-				RefreshEvent.Set();
+				_refreshEvent.Set();
 			}
 		}
 
-		public bool GetCommentByCurrentUser(int ChangeNumber, [NotNullWhen(true)] out string? CommentText)
+		public bool GetCommentByCurrentUser(int changeNumber, [NotNullWhen(true)] out string? commentText)
 		{
-			EventSummary? Summary = GetSummaryForChange(ChangeNumber);
-			if(Summary == null)
+			EventSummary? summary = GetSummaryForChange(changeNumber);
+			if(summary == null)
 			{
-				CommentText = null;
+				commentText = null;
 				return false;
 			}
 
-			CommentData? Comment = Summary.Comments.Find(x => String.Compare(x.UserName, CurrentUserName, true) == 0);
-			if(Comment == null || String.IsNullOrWhiteSpace(Comment.Text))
+			CommentData? comment = summary.Comments.Find(x => String.Compare(x.UserName, _currentUserName, true) == 0);
+			if(comment == null || String.IsNullOrWhiteSpace(comment.Text))
 			{
-				CommentText = null;
+				commentText = null;
 				return false;
 			}
 
-			CommentText = Comment.Text;
+			commentText = comment.Text;
 			return true;
 		}
 
-		public EventData? GetReviewByCurrentUser(int ChangeNumber)
+		public EventData? GetReviewByCurrentUser(int changeNumber)
 		{
-			EventSummary? Summary = GetSummaryForChange(ChangeNumber);
-			if(Summary == null)
+			EventSummary? summary = GetSummaryForChange(changeNumber);
+			if(summary == null)
 			{
 				return null;
 			}
 
-			EventData? Event = Summary.Reviews.FirstOrDefault(x => String.Compare(x.UserName, CurrentUserName, true) == 0);
-			if(Event == null || Event.Type == EventType.Unknown)
+			EventData? evt = summary.Reviews.FirstOrDefault(x => String.Compare(x.UserName, _currentUserName, true) == 0);
+			if(evt == null || evt.Type == EventType.Unknown)
 			{
 				return null;
 			}
 
-			return Event;
+			return evt;
 		}
 
-		public EventSummary? GetSummaryForChange(int ChangeNumber)
+		public EventSummary? GetSummaryForChange(int changeNumber)
 		{
-			EventSummary? Summary;
-			ChangeNumberToSummary.TryGetValue(ChangeNumber, out Summary);
-			return Summary;
+			EventSummary? summary;
+			_changeNumberToSummary.TryGetValue(changeNumber, out summary);
+			return summary;
 		}
 
-		public bool TryGetLatestBadge(string BuildType, [NotNullWhen(true)] out BadgeData? BadgeData)
+		public bool TryGetLatestBadge(string buildType, [NotNullWhen(true)] out BadgeData? badgeData)
 		{
-			return BadgeNameToLatestData.TryGetValue(BuildType, out BadgeData);
+			return _badgeNameToLatestData.TryGetValue(buildType, out badgeData);
 		}
 
-		public static bool IsReview(EventType Type)
+		public static bool IsReview(EventType type)
 		{
-			return IsPositiveReview(Type) || IsNegativeReview(Type) || Type == EventType.Unknown;
+			return IsPositiveReview(type) || IsNegativeReview(type) || type == EventType.Unknown;
 		}
 
-		public static bool IsPositiveReview(EventType Type)
+		public static bool IsPositiveReview(EventType type)
 		{
-			return Type == EventType.Good || Type == EventType.Compiles;
+			return type == EventType.Good || type == EventType.Compiles;
 		}
 
-		public static bool IsNegativeReview(EventType Type)
+		public static bool IsNegativeReview(EventType type)
 		{
-			return Type == EventType.DoesNotCompile || Type == EventType.Bad;
+			return type == EventType.DoesNotCompile || type == EventType.Bad;
 		}
 
-		public bool WasSyncedByCurrentUser(int ChangeNumber)
+		public bool WasSyncedByCurrentUser(int changeNumber)
 		{
-			EventSummary? Summary = GetSummaryForChange(ChangeNumber);
-			return (Summary != null && Summary.SyncEvents.Any(x => x.Type == EventType.Syncing && String.Compare(x.UserName, CurrentUserName, true) == 0));
+			EventSummary? summary = GetSummaryForChange(changeNumber);
+			return (summary != null && summary.SyncEvents.Any(x => x.Type == EventType.Syncing && String.Compare(x.UserName, _currentUserName, true) == 0));
 		}
 
-		public void StartInvestigating(int ChangeNumber)
+		public void StartInvestigating(int changeNumber)
 		{
-			PostEvent(ChangeNumber, EventType.Investigating);
+			PostEvent(changeNumber, EventType.Investigating);
 		}
 
-		public void FinishInvestigating(int ChangeNumber)
+		public void FinishInvestigating(int changeNumber)
 		{
-			PostEvent(ChangeNumber, EventType.Resolved);
+			PostEvent(changeNumber, EventType.Resolved);
 		}
 
 		protected void UpdateActiveInvestigations()
 		{
-			if(ActiveInvestigations == null)
+			if(_activeInvestigations == null)
 			{
 				// Insert investigation events into the active list, sorted by change number.
-				ActiveInvestigations = new List<EventData>();
-				foreach(EventData InvestigationEvent in InvestigationEvents)
+				_activeInvestigations = new List<EventData>();
+				foreach(EventData investigationEvent in _investigationEvents)
 				{
-					if(FilterChangeNumbers.Contains(InvestigationEvent.Change))
+					if(_filterChangeNumbers.Contains(investigationEvent.Change))
 					{
-						if(InvestigationEvent.Type == EventType.Investigating)
+						if(investigationEvent.Type == EventType.Investigating)
 						{
-							int InsertIdx = 0;
-							while(InsertIdx < ActiveInvestigations.Count && ActiveInvestigations[InsertIdx].Change > InvestigationEvent.Change)
+							int insertIdx = 0;
+							while(insertIdx < _activeInvestigations.Count && _activeInvestigations[insertIdx].Change > investigationEvent.Change)
 							{
-								InsertIdx++;
+								insertIdx++;
 							}
-							ActiveInvestigations.Insert(InsertIdx, InvestigationEvent);
+							_activeInvestigations.Insert(insertIdx, investigationEvent);
 						}
 						else
 						{
-							ActiveInvestigations.RemoveAll(x => String.Compare(x.UserName, InvestigationEvent.UserName, true) == 0 && x.Change <= InvestigationEvent.Change);
+							_activeInvestigations.RemoveAll(x => String.Compare(x.UserName, investigationEvent.UserName, true) == 0 && x.Change <= investigationEvent.Change);
 						}
 					}
 				}
 
 				// Remove any duplicate users
-				for(int Idx = 0; Idx < ActiveInvestigations.Count; Idx++)
+				for(int idx = 0; idx < _activeInvestigations.Count; idx++)
 				{
-					for(int OtherIdx = 0; OtherIdx < Idx; OtherIdx++)
+					for(int otherIdx = 0; otherIdx < idx; otherIdx++)
 					{
-						if(String.Compare(ActiveInvestigations[Idx].UserName, ActiveInvestigations[OtherIdx].UserName, true) == 0)
+						if(String.Compare(_activeInvestigations[idx].UserName, _activeInvestigations[otherIdx].UserName, true) == 0)
 						{
-							ActiveInvestigations.RemoveAt(Idx--);
+							_activeInvestigations.RemoveAt(idx--);
 							break;
 						}
 					}
@@ -1139,43 +1139,43 @@ namespace UnrealGameSync
 			}
 		}
 
-		public bool IsUnderInvestigation(int ChangeNumber)
+		public bool IsUnderInvestigation(int changeNumber)
 		{
 			UpdateActiveInvestigations();
-			return ActiveInvestigations.Any(x => x.Change <= ChangeNumber);
+			return _activeInvestigations.Any(x => x.Change <= changeNumber);
 		}
 
-		public bool IsUnderInvestigationByCurrentUser(int ChangeNumber)
+		public bool IsUnderInvestigationByCurrentUser(int changeNumber)
 		{
 			UpdateActiveInvestigations();
-			return ActiveInvestigations.Any(x => x.Change <= ChangeNumber && String.Compare(x.UserName, CurrentUserName, true) == 0);
+			return _activeInvestigations.Any(x => x.Change <= changeNumber && String.Compare(x.UserName, _currentUserName, true) == 0);
 		}
 
-		public IEnumerable<string> GetInvestigatingUsers(int ChangeNumber)
+		public IEnumerable<string> GetInvestigatingUsers(int changeNumber)
 		{
 			UpdateActiveInvestigations();
-			return ActiveInvestigations.Where(x => ChangeNumber >= x.Change).Select(x => x.UserName);
+			return _activeInvestigations.Where(x => changeNumber >= x.Change).Select(x => x.UserName);
 		}
 
-		public int GetInvestigationStartChangeNumber(int LastChangeNumber)
+		public int GetInvestigationStartChangeNumber(int lastChangeNumber)
 		{
 			UpdateActiveInvestigations();
 
-			int StartChangeNumber = -1;
-			if (ActiveInvestigations != null)
+			int startChangeNumber = -1;
+			if (_activeInvestigations != null)
 			{
-				foreach (EventData ActiveInvestigation in ActiveInvestigations)
+				foreach (EventData activeInvestigation in _activeInvestigations)
 				{
-					if (String.Compare(ActiveInvestigation.UserName, CurrentUserName, true) == 0)
+					if (String.Compare(activeInvestigation.UserName, _currentUserName, true) == 0)
 					{
-						if (ActiveInvestigation.Change <= LastChangeNumber && (StartChangeNumber == -1 || ActiveInvestigation.Change < StartChangeNumber))
+						if (activeInvestigation.Change <= lastChangeNumber && (startChangeNumber == -1 || activeInvestigation.Change < startChangeNumber))
 						{
-							StartChangeNumber = ActiveInvestigation.Change;
+							startChangeNumber = activeInvestigation.Change;
 						}
 					}
 				}
 			}
-			return StartChangeNumber;
+			return startChangeNumber;
 		}
 	}
 }
diff --git a/Engine/Source/Programs/UnrealGameSync/UnrealGameSync/FindFoldersToCleanTask.cs b/Engine/Source/Programs/UnrealGameSync/UnrealGameSync/FindFoldersToCleanTask.cs
--- a/Engine/Source/Programs/UnrealGameSync/UnrealGameSync/FindFoldersToCleanTask.cs
+++ b/Engine/Source/Programs/UnrealGameSync/UnrealGameSync/FindFoldersToCleanTask.cs
@@ -20,12 +20,12 @@ namespace UnrealGameSync
 		public Dictionary<string, FileInfo> NameToFile = new Dictionary<string, FileInfo>(StringComparer.InvariantCultureIgnoreCase);
 		public List<FileInfo> FilesToDelete = new List<FileInfo>();
 		public List<FileInfo> FilesToSync = new List<FileInfo>();
-		public bool bEmptyLeaf = false;
-		public bool bEmptyAfterClean = true;
+		public bool EmptyLeaf = false;
+		public bool EmptyAfterClean = true;
 
-		public FolderToClean(DirectoryInfo InDirectory)
+		public FolderToClean(DirectoryInfo inDirectory)
 		{
-			Directory = InDirectory;
+			Directory = inDirectory;
 		}
 
 		public string Name
@@ -47,222 +47,222 @@ namespace UnrealGameSync
 			public Dictionary<string, FStatRecord> NameToFile = new Dictionary<string, FStatRecord>(StringComparer.OrdinalIgnoreCase);
 		}
 
-		IPerforceSettings PerforceSettings;
-		string ClientRootPath;
-		IReadOnlyList<string> SyncPaths;
-		ILogger Logger;
-		FolderToClean RootFolderToClean;
+		IPerforceSettings _perforceSettings;
+		string _clientRootPath;
+		IReadOnlyList<string> _syncPaths;
+		ILogger _logger;
+		FolderToClean _rootFolderToClean;
 
-		int RemainingFoldersToScan;
-		ManualResetEvent FinishedScan = new ManualResetEvent(false);
-		bool bAbortScan;
-		string? ScanError;
+		int _remainingFoldersToScan;
+		ManualResetEvent _finishedScan = new ManualResetEvent(false);
+		bool _abortScan;
+		string? _scanError;
 
 		public List<string> FileNames = new List<string>();
 
-		public FindFoldersToCleanTask(IPerforceSettings InPerforceSettings, FolderToClean InRootFolderToClean, string InClientRootPath, IReadOnlyList<string> InSyncPaths, ILogger Logger)
+		public FindFoldersToCleanTask(IPerforceSettings inPerforceSettings, FolderToClean inRootFolderToClean, string inClientRootPath, IReadOnlyList<string> inSyncPaths, ILogger logger)
 		{
-			PerforceSettings = InPerforceSettings;
-			ClientRootPath = InClientRootPath.TrimEnd('/') + "/";
-			SyncPaths = new List<string>(InSyncPaths);
-			this.Logger = Logger;
-			RootFolderToClean = InRootFolderToClean;
-			FinishedScan = new ManualResetEvent(true);
+			_perforceSettings = inPerforceSettings;
+			_clientRootPath = inClientRootPath.TrimEnd('/') + "/";
+			_syncPaths = new List<string>(inSyncPaths);
+			this._logger = logger;
+			_rootFolderToClean = inRootFolderToClean;
+			_finishedScan = new ManualResetEvent(true);
 		}
 
-		void QueueFolderToPopulate(FolderToClean Folder)
+		void QueueFolderToPopulate(FolderToClean folder)
 		{
-			if(Interlocked.Increment(ref RemainingFoldersToScan) == 1)
+			if(Interlocked.Increment(ref _remainingFoldersToScan) == 1)
 			{
-				FinishedScan.Reset();
+				_finishedScan.Reset();
 			}
-			ThreadPool.QueueUserWorkItem(x => PopulateFolder(Folder));
+			ThreadPool.QueueUserWorkItem(x => PopulateFolder(folder));
 		}
 
-		void PopulateFolder(FolderToClean Folder)
+		void PopulateFolder(FolderToClean folder)
 		{
-			if(!bAbortScan)
+			if(!_abortScan)
 			{
 				try
 				{
-					if ((Folder.Directory.Attributes & FileAttributes.ReparsePoint) == 0)
+					if ((folder.Directory.Attributes & FileAttributes.ReparsePoint) == 0)
 					{
-						foreach (DirectoryInfo SubDirectory in Folder.Directory.EnumerateDirectories())
+						foreach (DirectoryInfo subDirectory in folder.Directory.EnumerateDirectories())
 						{
-							FolderToClean SubFolder = new FolderToClean(SubDirectory);
-							Folder.NameToSubFolder[SubFolder.Name] = SubFolder;
-							QueueFolderToPopulate(SubFolder);
+							FolderToClean subFolder = new FolderToClean(subDirectory);
+							folder.NameToSubFolder[subFolder.Name] = subFolder;
+							QueueFolderToPopulate(subFolder);
 						}
-						foreach (FileInfo File in Folder.Directory.EnumerateFiles())
+						foreach (FileInfo file in folder.Directory.EnumerateFiles())
 						{
-							FileAttributes Attributes = File.Attributes; // Force the value to be cached.
-							Folder.NameToFile[File.Name] = File;
+							FileAttributes attributes = file.Attributes; // Force the value to be cached.
+							folder.NameToFile[file.Name] = file;
 						}
 					}
 				}
-				catch (Exception Ex)
+				catch (Exception ex)
 				{
-					string NewError = String.Format("Unable to enumerate contents of {0} due to an error:\n\n{1}", Folder.Directory.FullName, Ex);
-					Interlocked.CompareExchange(ref ScanError, NewError, null);
-					bAbortScan = true;
+					string newError = String.Format("Unable to enumerate contents of {0} due to an error:\n\n{1}", folder.Directory.FullName, ex);
+					Interlocked.CompareExchange(ref _scanError, newError, null);
+					_abortScan = true;
 				}
 			}
 
-			if(Interlocked.Decrement(ref RemainingFoldersToScan) == 0)
+			if(Interlocked.Decrement(ref _remainingFoldersToScan) == 0)
 			{
-				FinishedScan.Set();
+				_finishedScan.Set();
 			}
 		}
 
-		void MergeTrees(FolderToClean LocalFolder, PerforceHaveFolder? PerforceFolder, HashSet<string> OpenClientPaths, string? PerforceConfigFile)
+		void MergeTrees(FolderToClean localFolder, PerforceHaveFolder? perforceFolder, HashSet<string> openClientPaths, string? perforceConfigFile)
 		{
-			if(PerforceFolder == null)
+			if(perforceFolder == null)
 			{
 				// Loop through all the local sub-folders
-				foreach(FolderToClean LocalSubFolder in LocalFolder.NameToSubFolder.Values)
+				foreach(FolderToClean localSubFolder in localFolder.NameToSubFolder.Values)
 				{
-					MergeTrees(LocalSubFolder, null, OpenClientPaths, PerforceConfigFile);
+					MergeTrees(localSubFolder, null, openClientPaths, perforceConfigFile);
 				}
 
 				// Delete everything
-				LocalFolder.FilesToDelete.AddRange(LocalFolder.NameToFile.Values);
+				localFolder.FilesToDelete.AddRange(localFolder.NameToFile.Values);
 			}
 			else
 			{
 				// Loop through all the local sub-folders
-				foreach(FolderToClean LocalSubFolder in LocalFolder.NameToSubFolder.Values)
+				foreach(FolderToClean localSubFolder in localFolder.NameToSubFolder.Values)
 				{
-					PerforceHaveFolder? PerforceSubFolder;
-					PerforceFolder.NameToSubFolder.TryGetValue(LocalSubFolder.Name, out PerforceSubFolder);
-					MergeTrees(LocalSubFolder, PerforceSubFolder, OpenClientPaths, PerforceConfigFile);
+					PerforceHaveFolder? perforceSubFolder;
+					perforceFolder.NameToSubFolder.TryGetValue(localSubFolder.Name, out perforceSubFolder);
+					MergeTrees(localSubFolder, perforceSubFolder, openClientPaths, perforceConfigFile);
 				}
 
 				// Also merge all the Perforce folders that no longer exist
-				foreach(KeyValuePair<string, PerforceHaveFolder> PerforceSubFolderPair in PerforceFolder.NameToSubFolder)
+				foreach(KeyValuePair<string, PerforceHaveFolder> perforceSubFolderPair in perforceFolder.NameToSubFolder)
 				{
-					FolderToClean? LocalSubFolder;
-					if(!LocalFolder.NameToSubFolder.TryGetValue(PerforceSubFolderPair.Key, out LocalSubFolder))
+					FolderToClean? localSubFolder;
+					if(!localFolder.NameToSubFolder.TryGetValue(perforceSubFolderPair.Key, out localSubFolder))
 					{
-						LocalSubFolder = new FolderToClean(new DirectoryInfo(Path.Combine(LocalFolder.Directory.FullName, PerforceSubFolderPair.Key)));
-						MergeTrees(LocalSubFolder, PerforceSubFolderPair.Value, OpenClientPaths, PerforceConfigFile);
-						LocalFolder.NameToSubFolder.Add(LocalSubFolder.Name, LocalSubFolder);
+						localSubFolder = new FolderToClean(new DirectoryInfo(Path.Combine(localFolder.Directory.FullName, perforceSubFolderPair.Key)));
+						MergeTrees(localSubFolder, perforceSubFolderPair.Value, openClientPaths, perforceConfigFile);
+						localFolder.NameToSubFolder.Add(localSubFolder.Name, localSubFolder);
 					}
 				}
 
 				// Find all the files that need to be re-synced
-				foreach(KeyValuePair<string, FStatRecord> FilePair in PerforceFolder.NameToFile)
+				foreach(KeyValuePair<string, FStatRecord> filePair in perforceFolder.NameToFile)
 				{
-					FileInfo? LocalFile;
-					if(!LocalFolder.NameToFile.TryGetValue(FilePair.Key, out LocalFile))
+					FileInfo? localFile;
+					if(!localFolder.NameToFile.TryGetValue(filePair.Key, out localFile))
 					{
-						LocalFolder.FilesToSync.Add(new FileInfo(Path.Combine(LocalFolder.Directory.FullName, FilePair.Key)));
+						localFolder.FilesToSync.Add(new FileInfo(Path.Combine(localFolder.Directory.FullName, filePair.Key)));
 					}
-					else if(IsFileTypeWritable(FilePair.Value?.HeadType ?? FilePair.Value?.Type) && (LocalFile.Attributes & FileAttributes.ReadOnly) == 0 && !OpenClientPaths.Contains(FilePair.Value?.ClientFile ?? ""))
+					else if(IsFileTypeWritable(filePair.Value?.HeadType ?? filePair.Value?.Type) && (localFile.Attributes & FileAttributes.ReadOnly) == 0 && !openClientPaths.Contains(filePair.Value?.ClientFile ?? ""))
 					{
-						LocalFolder.FilesToSync.Add(LocalFile);
+						localFolder.FilesToSync.Add(localFile);
 					}
 				}
 
 				// Find all the files that should be deleted
-				foreach(FileInfo LocalFileInfo in LocalFolder.NameToFile.Values)
+				foreach(FileInfo localFileInfo in localFolder.NameToFile.Values)
 				{
-					if(!PerforceFolder.NameToFile.ContainsKey(LocalFileInfo.Name) && !OpenClientPaths.Contains(LocalFileInfo.FullName))
+					if(!perforceFolder.NameToFile.ContainsKey(localFileInfo.Name) && !openClientPaths.Contains(localFileInfo.FullName))
 					{
-						LocalFolder.FilesToDelete.Add(LocalFileInfo);
+						localFolder.FilesToDelete.Add(localFileInfo);
 					}
 				}
 			}
 
 			// Remove any config files
-			if(PerforceConfigFile != null)
+			if(perforceConfigFile != null)
 			{
-				LocalFolder.FilesToDelete.RemoveAll(x => String.Compare(x.Name, PerforceConfigFile, StringComparison.OrdinalIgnoreCase) == 0);
+				localFolder.FilesToDelete.RemoveAll(x => String.Compare(x.Name, perforceConfigFile, StringComparison.OrdinalIgnoreCase) == 0);
 			}
 
 			// Figure out if this folder is just an empty directory that needs to be removed
-			LocalFolder.bEmptyLeaf = LocalFolder.NameToFile.Count == 0 && LocalFolder.NameToSubFolder.Count == 0 && LocalFolder.FilesToSync.Count == 0;
+			localFolder.EmptyLeaf = localFolder.NameToFile.Count == 0 && localFolder.NameToSubFolder.Count == 0 && localFolder.FilesToSync.Count == 0;
 
 			// Figure out if it the folder will be empty after the clean operation
-			LocalFolder.bEmptyAfterClean = LocalFolder.NameToSubFolder.Values.All(x => x.bEmptyAfterClean) && LocalFolder.FilesToDelete.Count == LocalFolder.NameToFile.Count && LocalFolder.FilesToSync.Count == 0;
+			localFolder.EmptyAfterClean = localFolder.NameToSubFolder.Values.All(x => x.EmptyAfterClean) && localFolder.FilesToDelete.Count == localFolder.NameToFile.Count && localFolder.FilesToSync.Count == 0;
 		}
 
-		static bool IsFileTypeWritable(string? Type)
+		static bool IsFileTypeWritable(string? type)
 		{
-			if (Type != null)
+			if (type != null)
 			{
-				int Idx = Type.IndexOf('+');
-				if (Idx != -1)
+				int idx = type.IndexOf('+');
+				if (idx != -1)
 				{
-					return Type.IndexOf('w', Idx) != -1;
+					return type.IndexOf('w', idx) != -1;
 				}
 			}
 			return true;
 		}
 
-		void RemoveEmptyFolders(FolderToClean Folder)
+		void RemoveEmptyFolders(FolderToClean folder)
 		{
-			foreach(FolderToClean SubFolder in Folder.NameToSubFolder.Values)
+			foreach(FolderToClean subFolder in folder.NameToSubFolder.Values)
 			{
-				RemoveEmptyFolders(SubFolder);
+				RemoveEmptyFolders(subFolder);
 			}
 
-			Folder.NameToSubFolder = Folder.NameToSubFolder.Values.Where(x => x.NameToSubFolder.Count > 0 || x.FilesToSync.Count > 0 || x.FilesToDelete.Count > 0 || x.bEmptyLeaf).ToDictionary(x => x.Name, x => x, StringComparer.InvariantCultureIgnoreCase);
+			folder.NameToSubFolder = folder.NameToSubFolder.Values.Where(x => x.NameToSubFolder.Count > 0 || x.FilesToSync.Count > 0 || x.FilesToDelete.Count > 0 || x.EmptyLeaf).ToDictionary(x => x.Name, x => x, StringComparer.InvariantCultureIgnoreCase);
 		}
 
 		public void Dispose()
 		{
-			bAbortScan = true;
+			_abortScan = true;
 
-			if(FinishedScan != null)
+			if(_finishedScan != null)
 			{
-				FinishedScan.WaitOne();
-				FinishedScan.Dispose();
-				FinishedScan = null!;
+				_finishedScan.WaitOne();
+				_finishedScan.Dispose();
+				_finishedScan = null!;
 			}
 		}
 
-		public async Task RunAsync(CancellationToken CancellationToken)
+		public async Task RunAsync(CancellationToken cancellationToken)
 		{
-			using IPerforceConnection PerforceClient = await PerforceConnection.CreateAsync(PerforceSettings, Logger);
+			using IPerforceConnection perforceClient = await PerforceConnection.CreateAsync(_perforceSettings, _logger);
 
-			Logger.LogInformation("Finding files in workspace...");
-			Logger.LogInformation("");
+			_logger.LogInformation("Finding files in workspace...");
+			_logger.LogInformation("");
 
 			// Clear the current error
-			ScanError = null;
+			_scanError = null;
 
 			// Start enumerating all the files that exist locally
-			foreach(string SyncPath in SyncPaths)
+			foreach(string syncPath in _syncPaths)
 			{
-				Debug.Assert(SyncPath.StartsWith(ClientRootPath));
-				if(SyncPath.StartsWith(ClientRootPath, StringComparison.InvariantCultureIgnoreCase))
+				Debug.Assert(syncPath.StartsWith(_clientRootPath));
+				if(syncPath.StartsWith(_clientRootPath, StringComparison.InvariantCultureIgnoreCase))
 				{
-					string[] Fragments = SyncPath.Substring(ClientRootPath.Length).Split('/');
+					string[] fragments = syncPath.Substring(_clientRootPath.Length).Split('/');
 
-					FolderToClean SyncFolder = RootFolderToClean;
-					for(int Idx = 0; Idx < Fragments.Length - 1; Idx++)
+					FolderToClean syncFolder = _rootFolderToClean;
+					for(int idx = 0; idx < fragments.Length - 1; idx++)
 					{
-						FolderToClean? NextSyncFolder;
-						if(!SyncFolder.NameToSubFolder.TryGetValue(Fragments[Idx], out NextSyncFolder))
+						FolderToClean? nextSyncFolder;
+						if(!syncFolder.NameToSubFolder.TryGetValue(fragments[idx], out nextSyncFolder))
 						{
-							NextSyncFolder = new FolderToClean(new DirectoryInfo(Path.Combine(SyncFolder.Directory.FullName, Fragments[Idx])));
-							SyncFolder.NameToSubFolder[NextSyncFolder.Name] = NextSyncFolder;
+							nextSyncFolder = new FolderToClean(new DirectoryInfo(Path.Combine(syncFolder.Directory.FullName, fragments[idx])));
+							syncFolder.NameToSubFolder[nextSyncFolder.Name] = nextSyncFolder;
 						}
-						SyncFolder = NextSyncFolder;
+						syncFolder = nextSyncFolder;
 					}
 
-					string Wildcard = Fragments[Fragments.Length - 1];
-					if(Wildcard == "...")
+					string wildcard = fragments[fragments.Length - 1];
+					if(wildcard == "...")
 					{
-						QueueFolderToPopulate(SyncFolder);
+						QueueFolderToPopulate(syncFolder);
 					}
 					else
 					{
-						if(SyncFolder.Directory.Exists)
+						if(syncFolder.Directory.Exists)
 						{
-							foreach(FileInfo File in SyncFolder.Directory.EnumerateFiles(Wildcard))
+							foreach(FileInfo file in syncFolder.Directory.EnumerateFiles(wildcard))
 							{
-								SyncFolder.NameToFile[File.Name] = File;
+								syncFolder.NameToFile[file.Name] = file;
 							}
 						}
 					}
@@ -270,68 +270,68 @@ namespace UnrealGameSync
 			}
 
 			// Get the prefix for any local file
-			string LocalRootPrefix = RootFolderToClean.Directory.FullName.Replace(Path.AltDirectorySeparatorChar, Path.DirectorySeparatorChar) + Path.DirectorySeparatorChar;
+			string localRootPrefix = _rootFolderToClean.Directory.FullName.Replace(Path.AltDirectorySeparatorChar, Path.DirectorySeparatorChar) + Path.DirectorySeparatorChar;
 
 			// Query the have table and build a separate tree from it
-			PerforceHaveFolder RootHaveFolder = new PerforceHaveFolder();
-			foreach(string SyncPath in SyncPaths)
+			PerforceHaveFolder rootHaveFolder = new PerforceHaveFolder();
+			foreach(string syncPath in _syncPaths)
 			{
-				List<FStatRecord> FileRecords = await PerforceClient.FStatAsync($"{SyncPath}#have", CancellationToken).ToListAsync(CancellationToken);
-				foreach(FStatRecord FileRecord in FileRecords)
+				List<FStatRecord> fileRecords = await perforceClient.FStatAsync($"{syncPath}#have", cancellationToken).ToListAsync(cancellationToken);
+				foreach(FStatRecord fileRecord in fileRecords)
 				{
-					if(FileRecord.ClientFile == null || !FileRecord.ClientFile.StartsWith(LocalRootPrefix, StringComparison.InvariantCultureIgnoreCase))
+					if(fileRecord.ClientFile == null || !fileRecord.ClientFile.StartsWith(localRootPrefix, StringComparison.InvariantCultureIgnoreCase))
 					{
-						throw new UserErrorException($"Failed to get have table; file '{FileRecord.ClientFile}' doesn't start with root path ('{RootFolderToClean.Directory.FullName}')");
+						throw new UserErrorException($"Failed to get have table; file '{fileRecord.ClientFile}' doesn't start with root path ('{_rootFolderToClean.Directory.FullName}')");
 					}
 
-					string[] Tokens = FileRecord.ClientFile.Substring(LocalRootPrefix.Length).Split('/', '\\');
+					string[] tokens = fileRecord.ClientFile.Substring(localRootPrefix.Length).Split('/', '\\');
 
-					PerforceHaveFolder FileFolder = RootHaveFolder;
-					for(int Idx = 0; Idx < Tokens.Length - 1; Idx++)
+					PerforceHaveFolder fileFolder = rootHaveFolder;
+					for(int idx = 0; idx < tokens.Length - 1; idx++)
 					{
-						PerforceHaveFolder? NextFileFolder;
-						if(!FileFolder.NameToSubFolder.TryGetValue(Tokens[Idx], out NextFileFolder))
+						PerforceHaveFolder? nextFileFolder;
+						if(!fileFolder.NameToSubFolder.TryGetValue(tokens[idx], out nextFileFolder))
 						{
-							NextFileFolder = new PerforceHaveFolder();
-							FileFolder.NameToSubFolder.Add(Tokens[Idx], NextFileFolder);
+							nextFileFolder = new PerforceHaveFolder();
+							fileFolder.NameToSubFolder.Add(tokens[idx], nextFileFolder);
 						}
-						FileFolder = NextFileFolder;
+						fileFolder = nextFileFolder;
 					}
-					FileFolder.NameToFile[Tokens[Tokens.Length - 1]] = FileRecord;
+					fileFolder.NameToFile[tokens[tokens.Length - 1]] = fileRecord;
 				}
 			}
 
 			// Find all the files which are currently open for edit. We don't want to force sync these.
-			List<OpenedRecord> OpenFileRecords = await PerforceClient.OpenedAsync(OpenedOptions.None, -1, null, null, -1, "//...", CancellationToken).ToListAsync(CancellationToken);
+			List<OpenedRecord> openFileRecords = await perforceClient.OpenedAsync(OpenedOptions.None, -1, null, null, -1, "//...", cancellationToken).ToListAsync(cancellationToken);
 
 			// Build a set of all the open local files
-			HashSet<string> OpenLocalFiles = new HashSet<string>(StringComparer.InvariantCultureIgnoreCase);
-			foreach (OpenedRecord OpenFileRecord in OpenFileRecords)
+			HashSet<string> openLocalFiles = new HashSet<string>(StringComparer.InvariantCultureIgnoreCase);
+			foreach (OpenedRecord openFileRecord in openFileRecords)
 			{
-				if(!OpenFileRecord.ClientFile.StartsWith(ClientRootPath, StringComparison.InvariantCultureIgnoreCase))
+				if(!openFileRecord.ClientFile.StartsWith(_clientRootPath, StringComparison.InvariantCultureIgnoreCase))
 				{
-					throw new UserErrorException($"Failed to get open file list; file '{OpenFileRecord.ClientFile}' doesn't start with client root path ('{ClientRootPath}')");
+					throw new UserErrorException($"Failed to get open file list; file '{openFileRecord.ClientFile}' doesn't start with client root path ('{_clientRootPath}')");
 				}
-				OpenLocalFiles.Add(LocalRootPrefix + PerforceUtils.UnescapePath(OpenFileRecord.ClientFile).Substring(ClientRootPath.Length).Replace('/', Path.DirectorySeparatorChar));
+				openLocalFiles.Add(localRootPrefix + PerforceUtils.UnescapePath(openFileRecord.ClientFile).Substring(_clientRootPath.Length).Replace('/', Path.DirectorySeparatorChar));
 			}
 
 			// Wait to finish scanning the directory
-			FinishedScan.WaitOne();
+			_finishedScan.WaitOne();
 
 			// Check if there was an error
-			if (ScanError != null)
+			if (_scanError != null)
 			{
-				throw new UserErrorException(ScanError);
+				throw new UserErrorException(_scanError);
 			}
 
 			// Find the value of the P4CONFIG variable
-			string? PerforceConfigFile = PerforceEnvironment.Default.GetValue("P4CONFIG");
+			string? perforceConfigFile = PerforceEnvironment.Default.GetValue("P4CONFIG");
 
 			// Merge the trees
-			MergeTrees(RootFolderToClean, RootHaveFolder, OpenLocalFiles, PerforceConfigFile);
+			MergeTrees(_rootFolderToClean, rootHaveFolder, openLocalFiles, perforceConfigFile);
 
 			// Remove all the empty folders
-			RemoveEmptyFolders(RootFolderToClean);
+			RemoveEmptyFolders(_rootFolderToClean);
 		}
 	}
 }
diff --git a/Engine/Source/Programs/UnrealGameSync/UnrealGameSync/Forms/ApplicationSettingsWindow.cs b/Engine/Source/Programs/UnrealGameSync/UnrealGameSync/Forms/ApplicationSettingsWindow.cs
--- a/Engine/Source/Programs/UnrealGameSync/UnrealGameSync/Forms/ApplicationSettingsWindow.cs
+++ b/Engine/Source/Programs/UnrealGameSync/UnrealGameSync/Forms/ApplicationSettingsWindow.cs
@@ -23,41 +23,41 @@ namespace UnrealGameSync
 	{
 		static class PerforceTestConnectionTask
 		{
-			public static async Task RunAsync(IPerforceConnection Perforce, string DepotPath, CancellationToken CancellationToken)
+			public static async Task RunAsync(IPerforceConnection perforce, string depotPath, CancellationToken cancellationToken)
 			{
-				string CheckFilePath = String.Format("{0}/Release/UnrealGameSync.exe", DepotPath);
+				string checkFilePath = String.Format("{0}/Release/UnrealGameSync.exe", depotPath);
 
-				List<FStatRecord> FileRecords = await Perforce.FStatAsync(CheckFilePath, CancellationToken).ToListAsync(CancellationToken);
-				if(FileRecords.Count == 0)
+				List<FStatRecord> fileRecords = await perforce.FStatAsync(checkFilePath, cancellationToken).ToListAsync(cancellationToken);
+				if(fileRecords.Count == 0)
 				{
-					throw new UserErrorException($"Unable to find {CheckFilePath}");
+					throw new UserErrorException($"Unable to find {checkFilePath}");
 				}
 			}
 		}
 
-		string OriginalExecutableFileName;
-		IPerforceSettings DefaultPerforceSettings;
-		UserSettings Settings;
-		ILogger Logger;
+		string _originalExecutableFileName;
+		IPerforceSettings _defaultPerforceSettings;
+		UserSettings _settings;
+		ILogger _logger;
 
-		string? InitialServerAndPort;
-		string? InitialUserName;
-		string? InitialDepotPath;
-		bool bInitialPreview;
-		int InitialAutomationPortNumber;
-		ProtocolHandlerState InitialProtocolHandlerState;
+		string? _initialServerAndPort;
+		string? _initialUserName;
+		string? _initialDepotPath;
+		bool _initialPreview;
+		int _initialAutomationPortNumber;
+		ProtocolHandlerState _initialProtocolHandlerState;
 
-		bool? bRestartPreview;
+		bool? _restartPreview;
 
-		ToolUpdateMonitor ToolUpdateMonitor;
+		ToolUpdateMonitor _toolUpdateMonitor;
 
 		class ToolItem
 		{
 			public ToolDefinition Definition { get; }
 
-			public ToolItem(ToolDefinition Definition)
+			public ToolItem(ToolDefinition definition)
 			{
-				this.Definition = Definition;
+				this.Definition = definition;
 			}
 
 			public override string ToString()
@@ -66,46 +66,46 @@ namespace UnrealGameSync
 			}
 		}
 
-		private ApplicationSettingsWindow(IPerforceSettings DefaultPerforceSettings, bool bPreview, string OriginalExecutableFileName, UserSettings Settings, ToolUpdateMonitor ToolUpdateMonitor, ILogger<ApplicationSettingsWindow> Logger)
+		private ApplicationSettingsWindow(IPerforceSettings defaultPerforceSettings, bool preview, string originalExecutableFileName, UserSettings settings, ToolUpdateMonitor toolUpdateMonitor, ILogger<ApplicationSettingsWindow> logger)
 		{
 			InitializeComponent();
 
-			this.OriginalExecutableFileName = OriginalExecutableFileName;
-			this.DefaultPerforceSettings = DefaultPerforceSettings;
-			this.Settings = Settings;
-			this.ToolUpdateMonitor = ToolUpdateMonitor;
-			this.Logger = Logger;
+			this._originalExecutableFileName = originalExecutableFileName;
+			this._defaultPerforceSettings = defaultPerforceSettings;
+			this._settings = settings;
+			this._toolUpdateMonitor = toolUpdateMonitor;
+			this._logger = logger;
 
-			GlobalPerforceSettings.ReadGlobalPerforceSettings(ref InitialServerAndPort, ref InitialUserName, ref InitialDepotPath, ref bPreview);
-			bInitialPreview = bPreview;
+			GlobalPerforceSettings.ReadGlobalPerforceSettings(ref _initialServerAndPort, ref _initialUserName, ref _initialDepotPath, ref preview);
+			_initialPreview = preview;
 
-			InitialAutomationPortNumber = AutomationServer.GetPortNumber();
-			InitialProtocolHandlerState = ProtocolHandlerUtils.GetState();
+			_initialAutomationPortNumber = AutomationServer.GetPortNumber();
+			_initialProtocolHandlerState = ProtocolHandlerUtils.GetState();
 
 			this.AutomaticallyRunAtStartupCheckBox.Checked = IsAutomaticallyRunAtStartup();
-			this.KeepInTrayCheckBox.Checked = Settings.bKeepInTray;
+			this.KeepInTrayCheckBox.Checked = settings.KeepInTray;
 						
-			this.ServerTextBox.Text = InitialServerAndPort;
+			this.ServerTextBox.Text = _initialServerAndPort;
 			this.ServerTextBox.Select(ServerTextBox.TextLength, 0);
-			this.ServerTextBox.CueBanner = $"Default ({DefaultPerforceSettings.ServerAndPort})";
+			this.ServerTextBox.CueBanner = $"Default ({defaultPerforceSettings.ServerAndPort})";
 
-			this.UserNameTextBox.Text = InitialUserName;
+			this.UserNameTextBox.Text = _initialUserName;
 			this.UserNameTextBox.Select(UserNameTextBox.TextLength, 0);
-			this.UserNameTextBox.CueBanner = $"Default ({DefaultPerforceSettings.UserName})";
+			this.UserNameTextBox.CueBanner = $"Default ({defaultPerforceSettings.UserName})";
 
-			this.ParallelSyncThreadsSpinner.Value = Math.Max(Math.Min(Settings.SyncOptions.NumThreads, ParallelSyncThreadsSpinner.Maximum), ParallelSyncThreadsSpinner.Minimum);
+			this.ParallelSyncThreadsSpinner.Value = Math.Max(Math.Min(settings.SyncOptions.NumThreads, ParallelSyncThreadsSpinner.Maximum), ParallelSyncThreadsSpinner.Minimum);
 
-			this.DepotPathTextBox.Text = InitialDepotPath;
+			this.DepotPathTextBox.Text = _initialDepotPath;
 			this.DepotPathTextBox.Select(DepotPathTextBox.TextLength, 0);
 			this.DepotPathTextBox.CueBanner = DeploymentSettings.DefaultDepotPath ?? String.Empty;
 
-			this.UsePreviewBuildCheckBox.Checked = bPreview;
+			this.UsePreviewBuildCheckBox.Checked = preview;
 
-			if(InitialAutomationPortNumber > 0)
+			if(_initialAutomationPortNumber > 0)
 			{
 				this.EnableAutomationCheckBox.Checked = true;
 				this.AutomationPortTextBox.Enabled = true;
-				this.AutomationPortTextBox.Text = InitialAutomationPortNumber.ToString();
+				this.AutomationPortTextBox.Text = _initialAutomationPortNumber.ToString();
 			}
 			else
 			{
@@ -114,11 +114,11 @@ namespace UnrealGameSync
 				this.AutomationPortTextBox.Text = AutomationServer.DefaultPortNumber.ToString();
 			}
 
-			if(InitialProtocolHandlerState == ProtocolHandlerState.Installed)
+			if(_initialProtocolHandlerState == ProtocolHandlerState.Installed)
 			{
 				this.EnableProtocolHandlerCheckBox.CheckState = CheckState.Checked;
 			}
-			else if (InitialProtocolHandlerState == ProtocolHandlerState.NotInstalled)
+			else if (_initialProtocolHandlerState == ProtocolHandlerState.NotInstalled)
 			{
 				this.EnableProtocolHandlerCheckBox.CheckState = CheckState.Unchecked;
 			}
@@ -127,19 +127,19 @@ namespace UnrealGameSync
 				this.EnableProtocolHandlerCheckBox.CheckState = CheckState.Indeterminate;
 			}
 
-			List<ToolDefinition> Tools = ToolUpdateMonitor.Tools;
-			foreach (ToolDefinition Tool in Tools)
+			List<ToolDefinition> tools = toolUpdateMonitor.Tools;
+			foreach (ToolDefinition tool in tools)
 			{
-				this.CustomToolsListBox.Items.Add(new ToolItem(Tool), Settings.EnabledTools.Contains(Tool.Id));
+				this.CustomToolsListBox.Items.Add(new ToolItem(tool), settings.EnabledTools.Contains(tool.Id));
 			}
 		}
 
-		public static bool? ShowModal(IWin32Window Owner, IPerforceSettings DefaultPerforceSettings, bool bPreview, string OriginalExecutableFileName, UserSettings Settings, ToolUpdateMonitor ToolUpdateMonitor, ILogger<ApplicationSettingsWindow> Logger)
+		public static bool? ShowModal(IWin32Window owner, IPerforceSettings defaultPerforceSettings, bool preview, string originalExecutableFileName, UserSettings settings, ToolUpdateMonitor toolUpdateMonitor, ILogger<ApplicationSettingsWindow> logger)
 		{
-			ApplicationSettingsWindow ApplicationSettings = new ApplicationSettingsWindow(DefaultPerforceSettings, bPreview, OriginalExecutableFileName, Settings, ToolUpdateMonitor, Logger);
-			if(ApplicationSettings.ShowDialog() == DialogResult.OK)
+			ApplicationSettingsWindow applicationSettings = new ApplicationSettingsWindow(defaultPerforceSettings, preview, originalExecutableFileName, settings, toolUpdateMonitor, logger);
+			if(applicationSettings.ShowDialog() == DialogResult.OK)
 			{
-				return ApplicationSettings.bRestartPreview;
+				return applicationSettings._restartPreview;
 			}
 			else
 			{
@@ -149,52 +149,52 @@ namespace UnrealGameSync
 
 		private bool IsAutomaticallyRunAtStartup()
 		{
-			RegistryKey? Key = Registry.CurrentUser.OpenSubKey("SOFTWARE\\Microsoft\\Windows\\CurrentVersion\\Run");
-			return (Key?.GetValue("UnrealGameSync") != null);
+			RegistryKey? key = Registry.CurrentUser.OpenSubKey("SOFTWARE\\Microsoft\\Windows\\CurrentVersion\\Run");
+			return (key?.GetValue("UnrealGameSync") != null);
 		}
 
 		private void OkBtn_Click(object sender, EventArgs e)
 		{
 			// Update the settings
-			string? ServerAndPort = ServerTextBox.Text.Trim();
-			if(ServerAndPort.Length == 0)
+			string? serverAndPort = ServerTextBox.Text.Trim();
+			if(serverAndPort.Length == 0)
 			{
-				ServerAndPort = null;
+				serverAndPort = null;
 			}
 
-			string? UserName = UserNameTextBox.Text.Trim();
-			if(UserName.Length == 0)
+			string? userName = UserNameTextBox.Text.Trim();
+			if(userName.Length == 0)
 			{
-				UserName = null;
+				userName = null;
 			}
 
-			string? DepotPath = DepotPathTextBox.Text.Trim();
-			if(DepotPath.Length == 0 || DepotPath == DeploymentSettings.DefaultDepotPath)
+			string? depotPath = DepotPathTextBox.Text.Trim();
+			if(depotPath.Length == 0 || depotPath == DeploymentSettings.DefaultDepotPath)
 			{
-				DepotPath = null;
+				depotPath = null;
 			}
 
-			bool bPreview = UsePreviewBuildCheckBox.Checked;
+			bool preview = UsePreviewBuildCheckBox.Checked;
 
 
-			int AutomationPortNumber;
-			if(!EnableAutomationCheckBox.Checked || !int.TryParse(AutomationPortTextBox.Text, out AutomationPortNumber))
+			int automationPortNumber;
+			if(!EnableAutomationCheckBox.Checked || !int.TryParse(AutomationPortTextBox.Text, out automationPortNumber))
 			{
-				AutomationPortNumber = -1;
+				automationPortNumber = -1;
 			}
 			
-			if(ServerAndPort != InitialServerAndPort || UserName != InitialUserName || DepotPath != InitialDepotPath || bPreview != bInitialPreview || AutomationPortNumber != InitialAutomationPortNumber)
+			if(serverAndPort != _initialServerAndPort || userName != _initialUserName || depotPath != _initialDepotPath || preview != _initialPreview || automationPortNumber != _initialAutomationPortNumber)
 			{
 				// Try to log in to the new server, and check the application is there
-				if(ServerAndPort != InitialServerAndPort || UserName != InitialUserName || DepotPath != InitialDepotPath)
+				if(serverAndPort != _initialServerAndPort || userName != _initialUserName || depotPath != _initialDepotPath)
 				{
-					PerforceSettings Settings = Utility.OverridePerforceSettings(DefaultPerforceSettings, ServerAndPort, UserName);
+					PerforceSettings settings = Utility.OverridePerforceSettings(_defaultPerforceSettings, serverAndPort, userName);
 
-					string? TestDepotPath = DepotPath ?? DeploymentSettings.DefaultDepotPath;
-					if (TestDepotPath != null)
+					string? testDepotPath = depotPath ?? DeploymentSettings.DefaultDepotPath;
+					if (testDepotPath != null)
 					{
-						ModalTask? Task = PerforceModalTask.Execute(this, "Checking connection", "Checking connection, please wait...", Settings, (p, c) => PerforceTestConnectionTask.RunAsync(p, TestDepotPath, c), Logger);
-						if (Task == null || !Task.Succeeded)
+						ModalTask? task = PerforceModalTask.Execute(this, "Checking connection", "Checking connection, please wait...", settings, (p, c) => PerforceTestConnectionTask.RunAsync(p, testDepotPath, c), _logger);
+						if (task == null || !task.Succeeded)
 						{
 							return;
 						}
@@ -206,53 +206,53 @@ namespace UnrealGameSync
 					return;
 				}
 
-				bRestartPreview = UsePreviewBuildCheckBox.Checked;
-				GlobalPerforceSettings.SaveGlobalPerforceSettings(ServerAndPort, UserName, DepotPath, bPreview);
-				AutomationServer.SetPortNumber(AutomationPortNumber);
+				_restartPreview = UsePreviewBuildCheckBox.Checked;
+				GlobalPerforceSettings.SaveGlobalPerforceSettings(serverAndPort, userName, depotPath, preview);
+				AutomationServer.SetPortNumber(automationPortNumber);
 			}
 
-			RegistryKey Key = Registry.CurrentUser.CreateSubKey("SOFTWARE\\Microsoft\\Windows\\CurrentVersion\\Run");
+			RegistryKey key = Registry.CurrentUser.CreateSubKey("SOFTWARE\\Microsoft\\Windows\\CurrentVersion\\Run");
 			if (AutomaticallyRunAtStartupCheckBox.Checked)
 			{
-				Key.SetValue("UnrealGameSync", String.Format("\"{0}\" -RestoreState", OriginalExecutableFileName));
+				key.SetValue("UnrealGameSync", String.Format("\"{0}\" -RestoreState", _originalExecutableFileName));
 			}
 			else
 			{
-				Key.DeleteValue("UnrealGameSync", false);
+				key.DeleteValue("UnrealGameSync", false);
 			}
 
-			if (Settings.bKeepInTray != KeepInTrayCheckBox.Checked || Settings.SyncOptions.NumThreads != ParallelSyncThreadsSpinner.Value)
+			if (_settings.KeepInTray != KeepInTrayCheckBox.Checked || _settings.SyncOptions.NumThreads != ParallelSyncThreadsSpinner.Value)
 			{
-				Settings.SyncOptions.NumThreads = (int)ParallelSyncThreadsSpinner.Value;
-				Settings.bKeepInTray = KeepInTrayCheckBox.Checked;
-				Settings.Save(Logger);
+				_settings.SyncOptions.NumThreads = (int)ParallelSyncThreadsSpinner.Value;
+				_settings.KeepInTray = KeepInTrayCheckBox.Checked;
+				_settings.Save(_logger);
 			}
 
-			List<Guid> NewEnabledTools = new List<Guid>();
-			foreach (ToolItem? Item in CustomToolsListBox.CheckedItems)
+			List<Guid> newEnabledTools = new List<Guid>();
+			foreach (ToolItem? item in CustomToolsListBox.CheckedItems)
 			{
-				if (Item != null)
+				if (item != null)
 				{
-					NewEnabledTools.Add(Item.Definition.Id);
+					newEnabledTools.Add(item.Definition.Id);
 				}
 			}
-			if (!NewEnabledTools.SequenceEqual(Settings.EnabledTools))
+			if (!newEnabledTools.SequenceEqual(_settings.EnabledTools))
 			{
-				Settings.EnabledTools = NewEnabledTools.ToArray();
-				Settings.Save(Logger);
-				ToolUpdateMonitor.UpdateNow();
+				_settings.EnabledTools = newEnabledTools.ToArray();
+				_settings.Save(_logger);
+				_toolUpdateMonitor.UpdateNow();
 			}
 
 			if (EnableProtocolHandlerCheckBox.CheckState == CheckState.Checked)
 			{
-				if (InitialProtocolHandlerState != ProtocolHandlerState.Installed)
+				if (_initialProtocolHandlerState != ProtocolHandlerState.Installed)
 				{
 					ProtocolHandlerUtils.Install();
 				}
 			}
 			else if (EnableProtocolHandlerCheckBox.CheckState == CheckState.Unchecked)
 			{
-				if (InitialProtocolHandlerState != ProtocolHandlerState.NotInstalled)
+				if (_initialProtocolHandlerState != ProtocolHandlerState.NotInstalled)
 				{
 					ProtocolHandlerUtils.Uninstall();
 				}
@@ -275,8 +275,8 @@ namespace UnrealGameSync
 
 		private void AdvancedBtn_Click(object sender, EventArgs e)
 		{
-			PerforceSyncSettingsWindow Window = new PerforceSyncSettingsWindow(Settings, Logger);
-			Window.ShowDialog();
+			PerforceSyncSettingsWindow window = new PerforceSyncSettingsWindow(_settings, _logger);
+			window.ShowDialog();
 		}
 	}
 }
