UGS: Reformat to standard C# style conventions.

diff --git a/Engine/Source/Programs/UnrealGameSync/UnrealGameSync/Forms/ArgumentsWindow.cs b/Engine/Source/Programs/UnrealGameSync/UnrealGameSync/Forms/ArgumentsWindow.cs
--- a/Engine/Source/Programs/UnrealGameSync/UnrealGameSync/Forms/ArgumentsWindow.cs
+++ b/Engine/Source/Programs/UnrealGameSync/UnrealGameSync/Forms/ArgumentsWindow.cs
@@ -1,4 +1,4 @@
-﻿// Copyright Epic Games, Inc. All Rights Reserved.
+// Copyright Epic Games, Inc. All Rights Reserved.
 
 using System;
 using System.Collections.Generic;
@@ -16,18 +16,18 @@ namespace UnrealGameSync
 {
 	public partial class ArgumentsWindow : Form
 	{
-		const int LVM_FIRST = 0x1000;
-		const int LVM_EDITLABELW = (LVM_FIRST + 118);
+		const int LvmFirst = 0x1000;
+		const int LvmEditlabelw = (LvmFirst + 118);
 
-		const int EM_SETSEL = 0xb1;
+		const int EmSetsel = 0xb1;
 
 		[DllImport("user32.dll", CharSet = CharSet.Auto)]
-		static extern IntPtr SendMessage(IntPtr hWnd, UInt32 Msg, IntPtr wParam, IntPtr lParam);
+		static extern IntPtr SendMessage(IntPtr hWnd, UInt32 msg, IntPtr wParam, IntPtr lParam);
 		
 		[DllImport("user32.dll")]
 		static extern bool SetWindowText(IntPtr hWnd, string text);
 
-		public ArgumentsWindow(List<Tuple<string, bool>> Arguments, bool bPromptBeforeLaunch)
+		public ArgumentsWindow(List<Tuple<string, bool>> arguments, bool promptBeforeLaunch)
 		{
 			InitializeComponent();
 			ActiveControl = ArgumentsList;
@@ -36,17 +36,17 @@ namespace UnrealGameSync
 
 			ArgumentsList.Items.Clear();
 
-			foreach(Tuple<string, bool> Argument in Arguments)
+			foreach(Tuple<string, bool> argument in arguments)
 			{
-				ListViewItem Item = new ListViewItem(Argument.Item1);
-				Item.Checked = Argument.Item2;
-				ArgumentsList.Items.Add(Item);
+				ListViewItem item = new ListViewItem(argument.Item1);
+				item.Checked = argument.Item2;
+				ArgumentsList.Items.Add(item);
 			}
 
-			ListViewItem AddAnotherItem = new ListViewItem("Click to add an item...", 0);
-			ArgumentsList.Items.Add(AddAnotherItem);
+			ListViewItem addAnotherItem = new ListViewItem("Click to add an item...", 0);
+			ArgumentsList.Items.Add(addAnotherItem);
 
-			PromptBeforeLaunchCheckBox.Checked = bPromptBeforeLaunch;
+			PromptBeforeLaunchCheckBox.Checked = promptBeforeLaunch;
 		}
 
 		public bool PromptBeforeLaunch
@@ -56,13 +56,13 @@ namespace UnrealGameSync
 
 		public List<Tuple<string, bool>> GetItems()
 		{
-			List<Tuple<string, bool>> Items = new List<Tuple<string,bool>>();
-			for(int Idx = 0; Idx < ArgumentsList.Items.Count - 1; Idx++)
+			List<Tuple<string, bool>> items = new List<Tuple<string,bool>>();
+			for(int idx = 0; idx < ArgumentsList.Items.Count - 1; idx++)
 			{
-				ListViewItem Item = ArgumentsList.Items[Idx];
-				Items.Add(new Tuple<string,bool>(Item.Text, Item.Checked));
+				ListViewItem item = ArgumentsList.Items[idx];
+				items.Add(new Tuple<string,bool>(item.Text, item.Checked));
 			}
-			return Items;
+			return items;
 		}
 
 		private void ClearButton_Click(object sender, EventArgs e)
@@ -91,22 +91,22 @@ namespace UnrealGameSync
 
 		private void ArgumentsList_MouseClick(object sender, MouseEventArgs e)
 		{
-			ListViewHitTestInfo Info = ArgumentsList.HitTest(e.Location);
-			if(Info.Item.Index == ArgumentsList.Items.Count - 1)
+			ListViewHitTestInfo info = ArgumentsList.HitTest(e.Location);
+			if(info.Item.Index == ArgumentsList.Items.Count - 1)
 			{
-				ListViewItem NewItem = new ListViewItem();
-				NewItem.Checked = true;
-				NewItem = ArgumentsList.Items.Insert(ArgumentsList.Items.Count - 1, NewItem);
-				NewItem.BeginEdit();
+				ListViewItem newItem = new ListViewItem();
+				newItem.Checked = true;
+				newItem = ArgumentsList.Items.Insert(ArgumentsList.Items.Count - 1, newItem);
+				newItem.BeginEdit();
 			}
 			else
 			{
-				using(Graphics Graphics = ArgumentsList.CreateGraphics())
+				using(Graphics graphics = ArgumentsList.CreateGraphics())
 				{
-					int LabelOffset = e.X - CheckBoxPadding - CheckBoxRenderer.GetGlyphSize(Graphics, CheckBoxState.CheckedNormal).Width - CheckBoxPadding;
-					if(LabelOffset >= 0 && LabelOffset < TextRenderer.MeasureText(Info.Item.Text, ArgumentsList.Font).Width)
+					int labelOffset = e.X - CheckBoxPadding - CheckBoxRenderer.GetGlyphSize(graphics, CheckBoxState.CheckedNormal).Width - CheckBoxPadding;
+					if(labelOffset >= 0 && labelOffset < TextRenderer.MeasureText(info.Item.Text, ArgumentsList.Font).Width)
 					{
-						Info.Item.BeginEdit();
+						info.Item.BeginEdit();
 					}
 				}
 			}
@@ -119,9 +119,9 @@ namespace UnrealGameSync
 			e.DrawBackground();
 			if(e.ItemIndex < ArgumentsList.Items.Count - 1)
 			{
-				CheckBoxState State = e.Item.Checked? CheckBoxState.CheckedNormal : CheckBoxState.UncheckedNormal;
-				Size CheckSize = CheckBoxRenderer.GetGlyphSize(e.Graphics, State);
-				CheckBoxRenderer.DrawCheckBox(e.Graphics, new Point(e.Bounds.Left + 4, e.Bounds.Top + (e.Bounds.Height - CheckSize.Height) / 2), State);
+				CheckBoxState state = e.Item.Checked? CheckBoxState.CheckedNormal : CheckBoxState.UncheckedNormal;
+				Size checkSize = CheckBoxRenderer.GetGlyphSize(e.Graphics, state);
+				CheckBoxRenderer.DrawCheckBox(e.Graphics, new Point(e.Bounds.Left + 4, e.Bounds.Top + (e.Bounds.Height - checkSize.Height) / 2), state);
 				DrawItemLabel(e.Graphics, SystemColors.WindowText, e.Item);
 			}
 			else
@@ -130,39 +130,39 @@ namespace UnrealGameSync
 			}
 		}
 
-		private void DrawItemLabel(Graphics Graphics, Color NormalColor, ListViewItem Item)
+		private void DrawItemLabel(Graphics graphics, Color normalColor, ListViewItem item)
 		{
-			Rectangle LabelRect = GetLabelRectangle(Graphics, Item);
-			if(Item.Selected)
+			Rectangle labelRect = GetLabelRectangle(graphics, item);
+			if(item.Selected)
 			{
-				Graphics.FillRectangle(SystemBrushes.Highlight, LabelRect);
-				TextRenderer.DrawText(Graphics, Item.Text, ArgumentsList.Font, LabelRect, SystemColors.HighlightText, SystemColors.Highlight, TextFormatFlags.Left | TextFormatFlags.VerticalCenter | TextFormatFlags.SingleLine);
+				graphics.FillRectangle(SystemBrushes.Highlight, labelRect);
+				TextRenderer.DrawText(graphics, item.Text, ArgumentsList.Font, labelRect, SystemColors.HighlightText, SystemColors.Highlight, TextFormatFlags.Left | TextFormatFlags.VerticalCenter | TextFormatFlags.SingleLine);
 			}
 			else
 			{
-				TextRenderer.DrawText(Graphics, Item.Text, ArgumentsList.Font, LabelRect, NormalColor, TextFormatFlags.Left | TextFormatFlags.VerticalCenter | TextFormatFlags.SingleLine);
+				TextRenderer.DrawText(graphics, item.Text, ArgumentsList.Font, labelRect, normalColor, TextFormatFlags.Left | TextFormatFlags.VerticalCenter | TextFormatFlags.SingleLine);
 			}
 		}
 
-		private Rectangle GetLabelRectangle(Graphics Graphics, ListViewItem Item)
+		private Rectangle GetLabelRectangle(Graphics graphics, ListViewItem item)
 		{
-			CheckBoxState State = Item.Checked? CheckBoxState.CheckedNormal : CheckBoxState.UncheckedNormal;
-			Size CheckSize = CheckBoxRenderer.GetGlyphSize(Graphics, State);
-			Size LabelSize = TextRenderer.MeasureText(Item.Text, ArgumentsList.Font);
+			CheckBoxState state = item.Checked? CheckBoxState.CheckedNormal : CheckBoxState.UncheckedNormal;
+			Size checkSize = CheckBoxRenderer.GetGlyphSize(graphics, state);
+			Size labelSize = TextRenderer.MeasureText(item.Text, ArgumentsList.Font);
 
-			int LabelIndent = CheckBoxPadding + CheckSize.Width + CheckBoxPadding;
-			return new Rectangle(Item.Bounds.Left + LabelIndent, Item.Bounds.Top, LabelSize.Width, Item.Bounds.Height);
+			int labelIndent = CheckBoxPadding + checkSize.Width + CheckBoxPadding;
+			return new Rectangle(item.Bounds.Left + labelIndent, item.Bounds.Top, labelSize.Width, item.Bounds.Height);
 		}
 
 		private void ArgumentsList_KeyUp(object sender, KeyEventArgs e)
 		{
 			if(e.KeyCode == Keys.Delete && ArgumentsList.SelectedIndices.Count == 1)
 			{
-				int Index = ArgumentsList.SelectedIndices[0];
-				ArgumentsList.Items.RemoveAt(Index);
-				if(Index < ArgumentsList.Items.Count - 1)
+				int index = ArgumentsList.SelectedIndices[0];
+				ArgumentsList.Items.RemoveAt(index);
+				if(index < ArgumentsList.Items.Count - 1)
 				{
-					ArgumentsList.Items[Index].Selected = true;
+					ArgumentsList.Items[index].Selected = true;
 				}
 			}
 		}
@@ -171,17 +171,17 @@ namespace UnrealGameSync
 		{
 			if(ArgumentsList.SelectedItems.Count == 1 && !Char.IsControl(e.KeyChar))
 			{
-				ListViewItem Item = ArgumentsList.SelectedItems[0];
-				if(Item.Index == ArgumentsList.Items.Count - 1)
+				ListViewItem item = ArgumentsList.SelectedItems[0];
+				if(item.Index == ArgumentsList.Items.Count - 1)
 				{
-					Item = new ListViewItem();
-					Item.Checked = true;
-					Item = ArgumentsList.Items.Insert(ArgumentsList.Items.Count - 1, Item);
+					item = new ListViewItem();
+					item.Checked = true;
+					item = ArgumentsList.Items.Insert(ArgumentsList.Items.Count - 1, item);
 				}
 
-				IntPtr NewHandle = SendMessage(ArgumentsList.Handle, LVM_EDITLABELW, new IntPtr(Item.Index), IntPtr.Zero);
-				SetWindowText(NewHandle, e.KeyChar.ToString());
-				SendMessage(NewHandle, EM_SETSEL, new IntPtr(1), new IntPtr(1));
+				IntPtr newHandle = SendMessage(ArgumentsList.Handle, LvmEditlabelw, new IntPtr(item.Index), IntPtr.Zero);
+				SetWindowText(newHandle, e.KeyChar.ToString());
+				SendMessage(newHandle, EmSetsel, new IntPtr(1), new IntPtr(1));
 
 				e.Handled = true;
 			}
@@ -191,12 +191,12 @@ namespace UnrealGameSync
 		{
 			if(ArgumentsList.SelectedIndices.Count == 1)
 			{
-				int Index = ArgumentsList.SelectedIndices[0];
-				if(Index > 0)
+				int index = ArgumentsList.SelectedIndices[0];
+				if(index > 0)
 				{
-					ListViewItem Item = ArgumentsList.Items[Index];
-					ArgumentsList.Items.RemoveAt(Index);
-					ArgumentsList.Items.Insert(Index - 1, Item);
+					ListViewItem item = ArgumentsList.Items[index];
+					ArgumentsList.Items.RemoveAt(index);
+					ArgumentsList.Items.Insert(index - 1, item);
 				}
 			}
 		}
@@ -205,12 +205,12 @@ namespace UnrealGameSync
 		{
 			if(ArgumentsList.SelectedIndices.Count == 1)
 			{
-				int Index = ArgumentsList.SelectedIndices[0];
-				if(Index < ArgumentsList.Items.Count - 2)
+				int index = ArgumentsList.SelectedIndices[0];
+				if(index < ArgumentsList.Items.Count - 2)
 				{
-					ListViewItem Item = ArgumentsList.Items[Index];
-					ArgumentsList.Items.RemoveAt(Index);
-					ArgumentsList.Items.Insert(Index + 1, Item);
+					ListViewItem item = ArgumentsList.Items[index];
+					ArgumentsList.Items.RemoveAt(index);
+					ArgumentsList.Items.Insert(index + 1, item);
 				}
 			}
 		}
diff --git a/Engine/Source/Programs/UnrealGameSync/UnrealGameSync/Forms/AutomatedBuildWindow.cs b/Engine/Source/Programs/UnrealGameSync/UnrealGameSync/Forms/AutomatedBuildWindow.cs
--- a/Engine/Source/Programs/UnrealGameSync/UnrealGameSync/Forms/AutomatedBuildWindow.cs
+++ b/Engine/Source/Programs/UnrealGameSync/UnrealGameSync/Forms/AutomatedBuildWindow.cs
@@ -23,32 +23,32 @@ namespace UnrealGameSync.Forms
 		{
 			public AutomatedSyncWindow.WorkspaceInfo SelectedWorkspaceInfo;
 			public string ProjectPath;
-			public bool bSync;
+			public bool Sync;
 			public string ExecCommand;
 
-			public BuildInfo(AutomatedSyncWindow.WorkspaceInfo SelectedWorkspaceInfo, string ProjectPath, bool bSync, string ExecCommand)
+			public BuildInfo(AutomatedSyncWindow.WorkspaceInfo selectedWorkspaceInfo, string projectPath, bool sync, string execCommand)
 			{
-				this.SelectedWorkspaceInfo = SelectedWorkspaceInfo;
-				this.ProjectPath = ProjectPath;
-				this.bSync = bSync;
-				this.ExecCommand = ExecCommand;
+				this.SelectedWorkspaceInfo = selectedWorkspaceInfo;
+				this.ProjectPath = projectPath;
+				this.Sync = sync;
+				this.ExecCommand = execCommand;
 			}
 		}
 
-		string StreamName;
-		IServiceProvider ServiceProvider;
+		string _streamName;
+		IServiceProvider _serviceProvider;
 
-		string? ServerAndPortOverride;
-		string? UserNameOverride;
-		IPerforceSettings DefaultPerforceSettings;
+		string? _serverAndPortOverride;
+		string? _userNameOverride;
+		IPerforceSettings _defaultPerforceSettings;
 
-		BuildInfo? Result;
+		BuildInfo? _result;
 
-		private AutomatedBuildWindow(string StreamName, int Changelist, string Command, IPerforceSettings DefaultPerforceSettings, string? DefaultWorkspaceName, string? DefaultProjectPath, IServiceProvider ServiceProvider)
+		private AutomatedBuildWindow(string streamName, int changelist, string command, IPerforceSettings defaultPerforceSettings, string? defaultWorkspaceName, string? defaultProjectPath, IServiceProvider serviceProvider)
 		{
-			this.StreamName = StreamName;
-			this.DefaultPerforceSettings = DefaultPerforceSettings;
-			this.ServiceProvider = ServiceProvider;
+			this._streamName = streamName;
+			this._defaultPerforceSettings = defaultPerforceSettings;
+			this._serviceProvider = serviceProvider;
 
 			InitializeComponent();
 
@@ -57,18 +57,18 @@ namespace UnrealGameSync.Forms
 			MinimumSize = Size;
 			MaximumSize = new Size(32768, Size.Height);
 
-			SyncToChangeCheckBox.Text = String.Format("Sync to changelist {0}", Changelist);
-			ExecCommandTextBox.Text = Command;
+			SyncToChangeCheckBox.Text = String.Format("Sync to changelist {0}", changelist);
+			ExecCommandTextBox.Text = command;
 
-			if (DefaultWorkspaceName != null)
+			if (defaultWorkspaceName != null)
 			{
-				WorkspaceNameTextBox.Text = DefaultWorkspaceName;
+				WorkspaceNameTextBox.Text = defaultWorkspaceName;
 				WorkspaceNameTextBox.Select(WorkspaceNameTextBox.Text.Length, 0);
 			}
 
-			if (DefaultProjectPath != null)
+			if (defaultProjectPath != null)
 			{
-				WorkspacePathTextBox.Text = DefaultProjectPath;
+				WorkspacePathTextBox.Text = defaultProjectPath;
 				WorkspacePathTextBox.Select(WorkspacePathTextBox.Text.Length, 0);
 			}
 
@@ -79,47 +79,47 @@ namespace UnrealGameSync.Forms
 
 		private IPerforceSettings Perforce
 		{
-			get => Utility.OverridePerforceSettings(DefaultPerforceSettings, ServerAndPortOverride, UserNameOverride);
+			get => Utility.OverridePerforceSettings(_defaultPerforceSettings, _serverAndPortOverride, _userNameOverride);
 		}
 
-		public static bool ShowModal(IWin32Window Owner, IPerforceSettings DefaultPerforceSettings, string StreamName, string ProjectPath, int Changelist, string Command, UserSettings Settings, IServiceProvider LoggerFactory, [NotNullWhen(true)] out BuildInfo? BuildInfo)
+		public static bool ShowModal(IWin32Window owner, IPerforceSettings defaultPerforceSettings, string streamName, string projectPath, int changelist, string command, UserSettings settings, IServiceProvider loggerFactory, [NotNullWhen(true)] out BuildInfo? buildInfo)
 		{
-			string? DefaultWorkspaceName = AutomatedSyncWindow.FindDefaultWorkspace(Owner, DefaultPerforceSettings, StreamName, LoggerFactory);
+			string? defaultWorkspaceName = AutomatedSyncWindow.FindDefaultWorkspace(owner, defaultPerforceSettings, streamName, loggerFactory);
 
-			string? DefaultProjectPath = null;
-			if(!String.IsNullOrEmpty(ProjectPath))
+			string? defaultProjectPath = null;
+			if(!String.IsNullOrEmpty(projectPath))
 			{
-				DefaultProjectPath = ProjectPath;
+				defaultProjectPath = projectPath;
 			}
-			else if(DefaultWorkspaceName != null)
+			else if(defaultWorkspaceName != null)
 			{
-				string ClientPrefix = String.Format("//{0}/", DefaultWorkspaceName);
-				foreach (UserSelectedProjectSettings ProjectSettings in Settings.RecentProjects)
+				string clientPrefix = String.Format("//{0}/", defaultWorkspaceName);
+				foreach (UserSelectedProjectSettings projectSettings in settings.RecentProjects)
 				{
-					if (ProjectSettings.ClientPath != null && ProjectSettings.ClientPath.StartsWith(ClientPrefix, StringComparison.OrdinalIgnoreCase))
+					if (projectSettings.ClientPath != null && projectSettings.ClientPath.StartsWith(clientPrefix, StringComparison.OrdinalIgnoreCase))
 					{
-						DefaultProjectPath = ProjectSettings.ClientPath.Substring(ClientPrefix.Length - 1);
+						defaultProjectPath = projectSettings.ClientPath.Substring(clientPrefix.Length - 1);
 						break;
 					}
 				}
 			}
 
-			AutomatedBuildWindow Window = new AutomatedBuildWindow(StreamName, Changelist, Command, DefaultPerforceSettings, DefaultWorkspaceName, DefaultProjectPath, LoggerFactory);
-			if (Window.ShowDialog() == DialogResult.OK)
+			AutomatedBuildWindow window = new AutomatedBuildWindow(streamName, changelist, command, defaultPerforceSettings, defaultWorkspaceName, defaultProjectPath, loggerFactory);
+			if (window.ShowDialog() == DialogResult.OK)
 			{
-				BuildInfo = Window.Result!;
+				buildInfo = window._result!;
 				return true;
 			}
 			else
 			{
-				BuildInfo = null;
+				buildInfo = null;
 				return false;
 			}
 		}
 
 		private void ChangeLink_LinkClicked(object sender, LinkLabelLinkClickedEventArgs e)
 		{
-			if (ConnectWindow.ShowModal(this, DefaultPerforceSettings, ref ServerAndPortOverride, ref UserNameOverride, ServiceProvider))
+			if (ConnectWindow.ShowModal(this, _defaultPerforceSettings, ref _serverAndPortOverride, ref _userNameOverride, _serviceProvider))
 			{
 				UpdateServerLabel();
 			}
@@ -127,25 +127,25 @@ namespace UnrealGameSync.Forms
 
 		private void UpdateServerLabel()
 		{
-			ServerLabel.Text = OpenProjectWindow.GetServerLabelText(DefaultPerforceSettings, ServerAndPortOverride, UserNameOverride);
+			ServerLabel.Text = OpenProjectWindow.GetServerLabelText(_defaultPerforceSettings, _serverAndPortOverride, _userNameOverride);
 		}
 
 		private void WorkspaceNameNewBtn_Click(object sender, EventArgs e)
 		{
-			string? WorkspaceName;
-			if (NewWorkspaceWindow.ShowModal(this, Perforce, StreamName, WorkspaceNameTextBox.Text, ServiceProvider, out WorkspaceName))
+			string? workspaceName;
+			if (NewWorkspaceWindow.ShowModal(this, Perforce, _streamName, WorkspaceNameTextBox.Text, _serviceProvider, out workspaceName))
 			{
-				WorkspaceNameTextBox.Text = WorkspaceName;
+				WorkspaceNameTextBox.Text = workspaceName;
 				UpdateOkButton();
 			}
 		}
 
 		private void WorkspaceNameBrowseBtn_Click(object sender, EventArgs e)
 		{
-			string? WorkspaceName = WorkspaceNameTextBox.Text;
-			if (SelectWorkspaceWindow.ShowModal(this, Perforce, WorkspaceName, ServiceProvider, out WorkspaceName))
+			string? workspaceName = WorkspaceNameTextBox.Text;
+			if (SelectWorkspaceWindow.ShowModal(this, Perforce, workspaceName, _serviceProvider, out workspaceName))
 			{
-				WorkspaceNameTextBox.Text = WorkspaceName;
+				WorkspaceNameTextBox.Text = workspaceName;
 				UpdateOkButton();
 			}
 		}
@@ -157,10 +157,10 @@ namespace UnrealGameSync.Forms
 
 		private void OkBtn_Click(object sender, EventArgs e)
 		{
-			AutomatedSyncWindow.WorkspaceInfo? SelectedWorkspaceInfo;
-			if (AutomatedSyncWindow.ValidateWorkspace(this, Perforce, WorkspaceNameTextBox.Text, StreamName, ServiceProvider, out SelectedWorkspaceInfo))
+			AutomatedSyncWindow.WorkspaceInfo? selectedWorkspaceInfo;
+			if (AutomatedSyncWindow.ValidateWorkspace(this, Perforce, WorkspaceNameTextBox.Text, _streamName, _serviceProvider, out selectedWorkspaceInfo))
 			{
-				Result = new BuildInfo(SelectedWorkspaceInfo, WorkspacePathTextBox.Text, SyncToChangeCheckBox.Checked, ExecCommandTextBox.Text);
+				_result = new BuildInfo(selectedWorkspaceInfo, WorkspacePathTextBox.Text, SyncToChangeCheckBox.Checked, ExecCommandTextBox.Text);
 				DialogResult = DialogResult.OK;
 				Close();
 			}
@@ -173,34 +173,34 @@ namespace UnrealGameSync.Forms
 
 		private void UpdateWorkspacePathBrowseButton()
 		{
-			string? WorkspaceName;
-			WorkspacePathBrowseBtn.Enabled = TryGetWorkspaceName(out WorkspaceName);
+			string? workspaceName;
+			WorkspacePathBrowseBtn.Enabled = TryGetWorkspaceName(out workspaceName);
 		}
 
 		private void WorkspacePathBrowseBtn_Click(object sender, EventArgs e)
 		{
-			string? WorkspaceName;
-			if (TryGetWorkspaceName(out WorkspaceName))
+			string? workspaceName;
+			if (TryGetWorkspaceName(out workspaceName))
 			{
-				string? WorkspacePath = WorkspacePathTextBox.Text.Trim();
-				if (SelectProjectFromWorkspaceWindow.ShowModal(this, Perforce, WorkspaceName, WorkspacePath, ServiceProvider, out WorkspacePath))
+				string? workspacePath = WorkspacePathTextBox.Text.Trim();
+				if (SelectProjectFromWorkspaceWindow.ShowModal(this, Perforce, workspaceName, workspacePath, _serviceProvider, out workspacePath))
 				{
-					WorkspacePathTextBox.Text = WorkspacePath;
+					WorkspacePathTextBox.Text = workspacePath;
 					UpdateOkButton();
 				}
 			}
 		}
 
-		private bool TryGetWorkspaceName([NotNullWhen(true)] out string? WorkspaceName)
+		private bool TryGetWorkspaceName([NotNullWhen(true)] out string? workspaceName)
 		{
-			string Text = WorkspaceNameTextBox.Text.Trim();
-			if (Text.Length == 0)
+			string text = WorkspaceNameTextBox.Text.Trim();
+			if (text.Length == 0)
 			{
-				WorkspaceName = null;
+				workspaceName = null;
 				return false;
 			}
 
-			WorkspaceName = Text;
+			workspaceName = text;
 			return true;
 		}
 
diff --git a/Engine/Source/Programs/UnrealGameSync/UnrealGameSync/Forms/AutomatedSyncWindow.cs b/Engine/Source/Programs/UnrealGameSync/UnrealGameSync/Forms/AutomatedSyncWindow.cs
--- a/Engine/Source/Programs/UnrealGameSync/UnrealGameSync/Forms/AutomatedSyncWindow.cs
+++ b/Engine/Source/Programs/UnrealGameSync/UnrealGameSync/Forms/AutomatedSyncWindow.cs
@@ -24,30 +24,30 @@ namespace UnrealGameSync
 	{
 		static class FindDefaultWorkspaceTask
 		{
-			public static async Task<string?> RunAsync(IPerforceConnection Perforce, string StreamName, CancellationToken CancellationToken)
+			public static async Task<string?> RunAsync(IPerforceConnection perforce, string streamName, CancellationToken cancellationToken)
 			{
-				InfoRecord Info = await Perforce.GetInfoAsync(InfoOptions.ShortOutput, CancellationToken);
+				InfoRecord info = await perforce.GetInfoAsync(InfoOptions.ShortOutput, cancellationToken);
 
-				List<ClientsRecord> Clients = await Perforce.GetClientsAsync(ClientsOptions.None, Perforce.Settings.UserName, CancellationToken);
+				List<ClientsRecord> clients = await perforce.GetClientsAsync(ClientsOptions.None, perforce.Settings.UserName, cancellationToken);
 
-				List<ClientsRecord> CandidateClients = new List<ClientsRecord>();
-				foreach(ClientsRecord Client in Clients)
+				List<ClientsRecord> candidateClients = new List<ClientsRecord>();
+				foreach(ClientsRecord client in clients)
 				{
-					if(Client.Host == null || Client.Host.Equals(Info.ClientHost, StringComparison.OrdinalIgnoreCase))
+					if(client.Host == null || client.Host.Equals(info.ClientHost, StringComparison.OrdinalIgnoreCase))
 					{
-						if(Client.Stream != null && Client.Stream.Equals(StreamName, StringComparison.OrdinalIgnoreCase))
+						if(client.Stream != null && client.Stream.Equals(streamName, StringComparison.OrdinalIgnoreCase))
 						{
-							CandidateClients.Add(Client);
+							candidateClients.Add(client);
 						}
 					}
 				}
 
-				string? WorkspaceName = null;
-				if(CandidateClients.Count >= 1)
+				string? workspaceName = null;
+				if(candidateClients.Count >= 1)
 				{
-					WorkspaceName = CandidateClients.OrderByDescending(x => x.Access).First().Name;
+					workspaceName = candidateClients.OrderByDescending(x => x.Access).First().Name;
 				}
-				return WorkspaceName;
+				return workspaceName;
 			}
 		}
 
@@ -55,26 +55,26 @@ namespace UnrealGameSync
 		{
 			public string WorkspaceName;
 			public string StreamName;
-			public bool bRequiresStreamSwitch;
-			public bool bHasOpenFiles;
+			public bool RequiresStreamSwitch;
+			public bool HasOpenFiles;
 
-			public ValidateWorkspaceTask(string WorkspaceName, string StreamName)
+			public ValidateWorkspaceTask(string workspaceName, string streamName)
 			{
-				this.WorkspaceName = WorkspaceName;
-				this.StreamName = StreamName;
+				this.WorkspaceName = workspaceName;
+				this.StreamName = streamName;
 			}
 
-			public async Task RunAsync(IPerforceConnection Perforce, CancellationToken CancellationToken)
+			public async Task RunAsync(IPerforceConnection perforce, CancellationToken cancellationToken)
 			{
-				ClientRecord Spec = await Perforce.GetClientAsync(WorkspaceName, CancellationToken);
+				ClientRecord spec = await perforce.GetClientAsync(WorkspaceName, cancellationToken);
 
-				string? CurrentStreamName = Spec.Stream;
-				if(CurrentStreamName == null || CurrentStreamName != StreamName)
+				string? currentStreamName = spec.Stream;
+				if(currentStreamName == null || currentStreamName != StreamName)
 				{
-					bRequiresStreamSwitch = true;
+					RequiresStreamSwitch = true;
 
-					List<PerforceResponse<OpenedRecord>> Records = await Perforce.TryOpenedAsync(OpenedOptions.None, FileSpecList.Any, CancellationToken).ToListAsync(CancellationToken);
-					bHasOpenFiles = Records.Succeeded() && Records.Count > 0;
+					List<PerforceResponse<OpenedRecord>> records = await perforce.TryOpenedAsync(OpenedOptions.None, FileSpecList.Any, cancellationToken).ToListAsync(cancellationToken);
+					HasOpenFiles = records.Succeeded() && records.Count > 0;
 				}
 			}
 		}
@@ -84,30 +84,30 @@ namespace UnrealGameSync
 			public string ServerAndPort { get; }
 			public string UserName { get; }
 			public string WorkspaceName { get; }
-			public bool bRequiresStreamSwitch { get; }
+			public bool RequiresStreamSwitch { get; }
 
-			public WorkspaceInfo(string ServerAndPort, string UserName, string WorkspaceName, bool bRequiresStreamSwitch)
+			public WorkspaceInfo(string serverAndPort, string userName, string workspaceName, bool requiresStreamSwitch)
 			{
-				this.ServerAndPort = ServerAndPort;
-				this.UserName = UserName;
-				this.WorkspaceName = WorkspaceName;
-				this.bRequiresStreamSwitch = bRequiresStreamSwitch;
+				this.ServerAndPort = serverAndPort;
+				this.UserName = userName;
+				this.WorkspaceName = workspaceName;
+				this.RequiresStreamSwitch = requiresStreamSwitch;
 			}
 		}
 
-		string StreamName;
-		IServiceProvider ServiceProvider;
+		string _streamName;
+		IServiceProvider _serviceProvider;
 
-		string? ServerAndPortOverride;
-		string? UserNameOverride;
-		IPerforceSettings DefaultPerforceSettings;
-		WorkspaceInfo? SelectedWorkspaceInfo;
+		string? _serverAndPortOverride;
+		string? _userNameOverride;
+		IPerforceSettings _defaultPerforceSettings;
+		WorkspaceInfo? _selectedWorkspaceInfo;
 
-		private AutomatedSyncWindow(string StreamName, string ProjectPath, string? WorkspaceName, IPerforceSettings DefaultPerforceSettings, IServiceProvider ServiceProvider)
+		private AutomatedSyncWindow(string streamName, string projectPath, string? workspaceName, IPerforceSettings defaultPerforceSettings, IServiceProvider serviceProvider)
 		{
-			this.StreamName = StreamName;
-			this.DefaultPerforceSettings = DefaultPerforceSettings;
-			this.ServiceProvider = ServiceProvider;
+			this._streamName = streamName;
+			this._defaultPerforceSettings = defaultPerforceSettings;
+			this._serviceProvider = serviceProvider;
 
 			InitializeComponent();
 
@@ -116,11 +116,11 @@ namespace UnrealGameSync
 			MinimumSize = Size;
 			MaximumSize = new Size(32768, Size.Height);
 
-			ProjectTextBox.Text = StreamName + ProjectPath;
+			ProjectTextBox.Text = streamName + projectPath;
 
-			if(WorkspaceName != null)
+			if(workspaceName != null)
 			{
-				WorkspaceNameTextBox.Text = WorkspaceName;
+				WorkspaceNameTextBox.Text = workspaceName;
 				WorkspaceNameTextBox.Select(WorkspaceNameTextBox.Text.Length, 0);
 			}
 
@@ -128,35 +128,35 @@ namespace UnrealGameSync
 			UpdateOkButton();
 		}
 
-		private IPerforceSettings Perforce => Utility.OverridePerforceSettings(DefaultPerforceSettings, ServerAndPortOverride, UserNameOverride);
+		private IPerforceSettings Perforce => Utility.OverridePerforceSettings(_defaultPerforceSettings, _serverAndPortOverride, _userNameOverride);
 
-		public static string? FindDefaultWorkspace(IWin32Window Owner, IPerforceSettings DefaultPerforceSettings, string StreamName, IServiceProvider ServiceProvider)
+		public static string? FindDefaultWorkspace(IWin32Window owner, IPerforceSettings defaultPerforceSettings, string streamName, IServiceProvider serviceProvider)
 		{
-			ILogger Logger = ServiceProvider.GetRequiredService<ILogger<AutomatedSyncWindow>>();
-			ModalTask<string?>? WorkspaceTask = PerforceModalTask.Execute(Owner, "Finding workspace", "Finding default workspace, please wait...", DefaultPerforceSettings, (p, c) => FindDefaultWorkspaceTask.RunAsync(p, StreamName, c), Logger);
-			return (WorkspaceTask != null && WorkspaceTask.Succeeded) ? WorkspaceTask.Result : null;
+			ILogger logger = serviceProvider.GetRequiredService<ILogger<AutomatedSyncWindow>>();
+			ModalTask<string?>? workspaceTask = PerforceModalTask.Execute(owner, "Finding workspace", "Finding default workspace, please wait...", defaultPerforceSettings, (p, c) => FindDefaultWorkspaceTask.RunAsync(p, streamName, c), logger);
+			return (workspaceTask != null && workspaceTask.Succeeded) ? workspaceTask.Result : null;
 		}
 
-		public static bool ShowModal(IWin32Window Owner, IPerforceSettings DefaultPerforceSettings, string StreamName, string ProjectPath, [NotNullWhen(true)] out WorkspaceInfo? WorkspaceInfo, IServiceProvider ServiceProvider)
+		public static bool ShowModal(IWin32Window owner, IPerforceSettings defaultPerforceSettings, string streamName, string projectPath, [NotNullWhen(true)] out WorkspaceInfo? workspaceInfo, IServiceProvider serviceProvider)
 		{
-			string? WorkspaceName = FindDefaultWorkspace(Owner, DefaultPerforceSettings, StreamName, ServiceProvider);
+			string? workspaceName = FindDefaultWorkspace(owner, defaultPerforceSettings, streamName, serviceProvider);
 
-			AutomatedSyncWindow Window = new AutomatedSyncWindow(StreamName, ProjectPath, WorkspaceName, DefaultPerforceSettings, ServiceProvider);
-			if(Window.ShowDialog() == DialogResult.OK)
+			AutomatedSyncWindow window = new AutomatedSyncWindow(streamName, projectPath, workspaceName, defaultPerforceSettings, serviceProvider);
+			if(window.ShowDialog() == DialogResult.OK)
 			{
-				WorkspaceInfo = Window.SelectedWorkspaceInfo!;
+				workspaceInfo = window._selectedWorkspaceInfo!;
 				return true;
 			}
 			else
 			{
-				WorkspaceInfo = null;
+				workspaceInfo = null;
 				return false;
 			}
 		}
 
 		private void ChangeLink_LinkClicked(object sender, LinkLabelLinkClickedEventArgs e)
 		{
-			if(ConnectWindow.ShowModal(this, DefaultPerforceSettings, ref ServerAndPortOverride, ref UserNameOverride, ServiceProvider))
+			if(ConnectWindow.ShowModal(this, _defaultPerforceSettings, ref _serverAndPortOverride, ref _userNameOverride, _serviceProvider))
 			{
 				UpdateServerLabel();
 			}
@@ -164,26 +164,26 @@ namespace UnrealGameSync
 
 		private void UpdateServerLabel()
 		{
-			ServerLabel.Text = OpenProjectWindow.GetServerLabelText(DefaultPerforceSettings, ServerAndPortOverride, UserNameOverride);
+			ServerLabel.Text = OpenProjectWindow.GetServerLabelText(_defaultPerforceSettings, _serverAndPortOverride, _userNameOverride);
 			ChangeLink.Location = new Point(ServerLabel.Right + 5, ChangeLink.Location.Y);
 		}
 
 		private void WorkspaceNameNewBtn_Click(object sender, EventArgs e)
 		{
-			string? WorkspaceName;
-			if(NewWorkspaceWindow.ShowModal(this, DefaultPerforceSettings, StreamName, WorkspaceNameTextBox.Text, ServiceProvider, out WorkspaceName))
+			string? workspaceName;
+			if(NewWorkspaceWindow.ShowModal(this, _defaultPerforceSettings, _streamName, WorkspaceNameTextBox.Text, _serviceProvider, out workspaceName))
 			{
-				WorkspaceNameTextBox.Text = WorkspaceName;
+				WorkspaceNameTextBox.Text = workspaceName;
 				UpdateOkButton();
 			}
 		}
 
 		private void WorkspaceNameBrowseBtn_Click(object sender, EventArgs e)
 		{
-			string? WorkspaceName = WorkspaceNameTextBox.Text;
-			if(SelectWorkspaceWindow.ShowModal(this, DefaultPerforceSettings, WorkspaceName, ServiceProvider, out WorkspaceName))
+			string? workspaceName = WorkspaceNameTextBox.Text;
+			if(SelectWorkspaceWindow.ShowModal(this, _defaultPerforceSettings, workspaceName, _serviceProvider, out workspaceName))
 			{
-				WorkspaceNameTextBox.Text = WorkspaceName;
+				WorkspaceNameTextBox.Text = workspaceName;
 				UpdateOkButton();
 			}
 		}
@@ -193,42 +193,42 @@ namespace UnrealGameSync
 			OkBtn.Enabled = (WorkspaceNameTextBox.Text.Length > 0);
 		}
 
-		public static bool ValidateWorkspace(IWin32Window Owner, IPerforceSettings Perforce, string WorkspaceName, string StreamName, IServiceProvider ServiceProvider, [NotNullWhen(true)] out WorkspaceInfo? SelectedWorkspaceInfo)
+		public static bool ValidateWorkspace(IWin32Window owner, IPerforceSettings perforce, string workspaceName, string streamName, IServiceProvider serviceProvider, [NotNullWhen(true)] out WorkspaceInfo? selectedWorkspaceInfo)
 		{
-			ValidateWorkspaceTask ValidateWorkspace = new ValidateWorkspaceTask(WorkspaceName, StreamName);
+			ValidateWorkspaceTask validateWorkspace = new ValidateWorkspaceTask(workspaceName, streamName);
 
-			ModalTask? Task = PerforceModalTask.Execute(Owner, "Checking workspace", "Checking workspace, please wait...", Perforce, ValidateWorkspace.RunAsync, ServiceProvider.GetRequiredService<ILogger<ValidateWorkspaceTask>>());
-			if (Task != null && Task.Succeeded)
+			ModalTask? task = PerforceModalTask.Execute(owner, "Checking workspace", "Checking workspace, please wait...", perforce, validateWorkspace.RunAsync, serviceProvider.GetRequiredService<ILogger<ValidateWorkspaceTask>>());
+			if (task != null && task.Succeeded)
 			{
-				if (ValidateWorkspace.bRequiresStreamSwitch)
+				if (validateWorkspace.RequiresStreamSwitch)
 				{
-					string Message;
-					if (ValidateWorkspace.bHasOpenFiles)
+					string message;
+					if (validateWorkspace.HasOpenFiles)
 					{
-						Message = String.Format("You have files open for edit in this workspace. If you switch this workspace to {0}, you will not be able to submit them until you switch back.\n\nContinue switching streams?", StreamName);
+						message = String.Format("You have files open for edit in this workspace. If you switch this workspace to {0}, you will not be able to submit them until you switch back.\n\nContinue switching streams?", streamName);
 					}
 					else
 					{
-						Message = String.Format("Switch this workspace to {0}?", StreamName);
+						message = String.Format("Switch this workspace to {0}?", streamName);
 					}
-					if (MessageBox.Show(Message, "Switch Streams", MessageBoxButtons.YesNo) != DialogResult.Yes)
+					if (MessageBox.Show(message, "Switch Streams", MessageBoxButtons.YesNo) != DialogResult.Yes)
 					{
-						SelectedWorkspaceInfo = null;
+						selectedWorkspaceInfo = null;
 						return false;
 					}
 				}
 
-				SelectedWorkspaceInfo = new WorkspaceInfo(Perforce.ServerAndPort, Perforce.UserName, ValidateWorkspace.WorkspaceName, ValidateWorkspace.bRequiresStreamSwitch);
+				selectedWorkspaceInfo = new WorkspaceInfo(perforce.ServerAndPort, perforce.UserName, validateWorkspace.WorkspaceName, validateWorkspace.RequiresStreamSwitch);
 				return true;
 			}
 
-			SelectedWorkspaceInfo = null;
+			selectedWorkspaceInfo = null;
 			return false;
 		}
 
 		private void OkBtn_Click(object sender, EventArgs e)
 		{
-			if(ValidateWorkspace(this, Perforce, WorkspaceNameTextBox.Text, StreamName, ServiceProvider, out SelectedWorkspaceInfo))
+			if(ValidateWorkspace(this, Perforce, WorkspaceNameTextBox.Text, _streamName, _serviceProvider, out _selectedWorkspaceInfo))
 			{
 				DialogResult = DialogResult.OK;
 				Close();
diff --git a/Engine/Source/Programs/UnrealGameSync/UnrealGameSync/Forms/BuildStepWindow.cs b/Engine/Source/Programs/UnrealGameSync/UnrealGameSync/Forms/BuildStepWindow.cs
--- a/Engine/Source/Programs/UnrealGameSync/UnrealGameSync/Forms/BuildStepWindow.cs
+++ b/Engine/Source/Programs/UnrealGameSync/UnrealGameSync/Forms/BuildStepWindow.cs
@@ -16,18 +16,18 @@ namespace UnrealGameSync
 {
 	partial class BuildStepWindow : Form
 	{
-		BuildStep Step;
-		List<string> TargetNames;
-		DirectoryReference BaseDirectory;
-		IReadOnlyDictionary<string, string> Variables;
-		VariablesWindow? VariablesWindow;
+		BuildStep _step;
+		List<string> _targetNames;
+		DirectoryReference _baseDirectory;
+		IReadOnlyDictionary<string, string> _variables;
+		VariablesWindow? _variablesWindow;
 
-		public BuildStepWindow(BuildStep InTask, List<string> InTargetNames, DirectoryReference InBaseDirectory, IReadOnlyDictionary<string, string> InVariables)
+		public BuildStepWindow(BuildStep inTask, List<string> inTargetNames, DirectoryReference inBaseDirectory, IReadOnlyDictionary<string, string> inVariables)
 		{
-			Step = InTask;
-			TargetNames = InTargetNames;
-			BaseDirectory = InBaseDirectory;
-			Variables = InVariables;
+			_step = inTask;
+			_targetNames = inTargetNames;
+			_baseDirectory = inBaseDirectory;
+			_variables = inVariables;
 
 			InitializeComponent();
 
@@ -37,33 +37,33 @@ namespace UnrealGameSync
 
 		private void BuildTaskWindow_Load(object sender, EventArgs e)
 		{
-			string DefaultTargetName = TargetNames.FirstOrDefault(x => !x.EndsWith("Editor") && !x.EndsWith("Client") && !x.EndsWith("Server")) ?? ((TargetNames.Count > 0)? TargetNames[0] : "");
-			CompileTargetComboBox.Items.AddRange(TargetNames.ToArray());
-			CompileTargetComboBox.Text = String.IsNullOrEmpty(Step.Target)? DefaultTargetName : Step.Target;
-			StatusPanelLinkTextBox.Text = Step.StatusPanelLink;
+			string defaultTargetName = _targetNames.FirstOrDefault(x => !x.EndsWith("Editor") && !x.EndsWith("Client") && !x.EndsWith("Server")) ?? ((_targetNames.Count > 0)? _targetNames[0] : "");
+			CompileTargetComboBox.Items.AddRange(_targetNames.ToArray());
+			CompileTargetComboBox.Text = String.IsNullOrEmpty(_step.Target)? defaultTargetName : _step.Target;
+			StatusPanelLinkTextBox.Text = _step.StatusPanelLink;
 
-			DescriptionTextBox.Text = Step.Description;
-			StatusTextTextBox.Text = Step.StatusText;
-			DurationTextBox.Text = Step.EstimatedDuration.ToString();
+			DescriptionTextBox.Text = _step.Description;
+			StatusTextTextBox.Text = _step.StatusText;
+			DurationTextBox.Text = _step.EstimatedDuration.ToString();
 
-			switch(Step.Type)
+			switch(_step.Type)
 			{
 				case BuildStepType.Compile:
 					CompileRadioButton.Checked = true;
-					CompilePlatformComboBox.Text = Step.Platform;
-					CompileConfigComboBox.Text = Step.Configuration;
-					CompileArgumentsTextBox.Text = Step.Arguments;
+					CompilePlatformComboBox.Text = _step.Platform;
+					CompileConfigComboBox.Text = _step.Configuration;
+					CompileArgumentsTextBox.Text = _step.Arguments;
 					break;
 				case BuildStepType.Cook:
 					CookRadioButton.Checked = true;
-					CookFileNameTextBox.Text = Step.FileName;
+					CookFileNameTextBox.Text = _step.FileName;
 					break;
 				case BuildStepType.Other:
 					OtherRadioButton.Checked = true;
-					OtherFileNameTextBox.Text = Step.FileName;
-					OtherArgumentsTextBox.Text = Step.Arguments;
-					OtherUseLogWindowCheckBox.Checked = Step.bUseLogWindow;
-					OtherWorkingDirTextBox.Text = Step.WorkingDir;
+					OtherFileNameTextBox.Text = _step.FileName;
+					OtherArgumentsTextBox.Text = _step.Arguments;
+					OtherUseLogWindowCheckBox.Checked = _step.UseLogWindow;
+					OtherWorkingDirTextBox.Text = _step.WorkingDir;
 					break;
 			}
 
@@ -82,23 +82,23 @@ namespace UnrealGameSync
 
 		public void UpdateType()
 		{
-			bool bIsCompile = CompileRadioButton.Checked;
-			CompileTargetComboBox.Enabled = bIsCompile;
-			CompilePlatformComboBox.Enabled = bIsCompile;
-			CompileConfigComboBox.Enabled = bIsCompile;
-			CompileArgumentsTextBox.Enabled = bIsCompile;
+			bool isCompile = CompileRadioButton.Checked;
+			CompileTargetComboBox.Enabled = isCompile;
+			CompilePlatformComboBox.Enabled = isCompile;
+			CompileConfigComboBox.Enabled = isCompile;
+			CompileArgumentsTextBox.Enabled = isCompile;
 
-			bool bIsCook = CookRadioButton.Checked;
-			CookFileNameTextBox.Enabled = bIsCook;
-			CookFileNameButton.Enabled = bIsCook;
+			bool isCook = CookRadioButton.Checked;
+			CookFileNameTextBox.Enabled = isCook;
+			CookFileNameButton.Enabled = isCook;
 
-			bool bIsOther = OtherRadioButton.Checked;
-			OtherFileNameTextBox.Enabled = bIsOther;
-			OtherFileNameButton.Enabled = bIsOther;
-			OtherWorkingDirTextBox.Enabled = bIsOther;
-			OtherWorkingDirButton.Enabled = bIsOther;
-			OtherArgumentsTextBox.Enabled = bIsOther;
-			OtherUseLogWindowCheckBox.Enabled = bIsOther;
+			bool isOther = OtherRadioButton.Checked;
+			OtherFileNameTextBox.Enabled = isOther;
+			OtherFileNameButton.Enabled = isOther;
+			OtherWorkingDirTextBox.Enabled = isOther;
+			OtherWorkingDirButton.Enabled = isOther;
+			OtherArgumentsTextBox.Enabled = isOther;
+			OtherUseLogWindowCheckBox.Enabled = isOther;
 		}
 
 		private void CompileRadioButton_CheckedChanged(object sender, EventArgs e)
@@ -118,45 +118,45 @@ namespace UnrealGameSync
 
 		private void OkButton_Click(object sender, EventArgs e)
 		{
-			Step.Description = DescriptionTextBox.Text;
-			Step.StatusText = StatusTextTextBox.Text;
-			Step.StatusPanelLink = StatusPanelLinkTextBox.Text;
+			_step.Description = DescriptionTextBox.Text;
+			_step.StatusText = StatusTextTextBox.Text;
+			_step.StatusPanelLink = StatusPanelLinkTextBox.Text;
 
-			if (!int.TryParse(DurationTextBox.Text, out Step.EstimatedDuration))
+			if (!int.TryParse(DurationTextBox.Text, out _step.EstimatedDuration))
 			{
-				Step.EstimatedDuration = 1;
+				_step.EstimatedDuration = 1;
 			}
 
 			if(CompileRadioButton.Checked)
 			{
-				Step.Type = BuildStepType.Compile;
-				Step.Target = CompileTargetComboBox.Text;
-				Step.Platform = CompilePlatformComboBox.Text;
-				Step.Configuration = CompileConfigComboBox.Text;
-				Step.FileName = null;
-				Step.Arguments = CompileArgumentsTextBox.Text;
-				Step.bUseLogWindow = true;
+				_step.Type = BuildStepType.Compile;
+				_step.Target = CompileTargetComboBox.Text;
+				_step.Platform = CompilePlatformComboBox.Text;
+				_step.Configuration = CompileConfigComboBox.Text;
+				_step.FileName = null;
+				_step.Arguments = CompileArgumentsTextBox.Text;
+				_step.UseLogWindow = true;
 			}
 			else if(CookRadioButton.Checked)
 			{
-				Step.Type = BuildStepType.Cook;
-				Step.Target = null;
-				Step.Platform = null;
-				Step.Configuration = null;
-				Step.FileName = CookFileNameTextBox.Text;
-				Step.Arguments = null;
-				Step.bUseLogWindow = true;
+				_step.Type = BuildStepType.Cook;
+				_step.Target = null;
+				_step.Platform = null;
+				_step.Configuration = null;
+				_step.FileName = CookFileNameTextBox.Text;
+				_step.Arguments = null;
+				_step.UseLogWindow = true;
 			}
 			else
 			{
-				Step.Type = BuildStepType.Other;
-				Step.Target = null;
-				Step.Platform = null;
-				Step.Configuration = null;
-				Step.FileName = OtherFileNameTextBox.Text;
-				Step.WorkingDir = OtherWorkingDirTextBox.Text;
-				Step.Arguments = OtherArgumentsTextBox.Text;
-				Step.bUseLogWindow = OtherUseLogWindowCheckBox.Checked;
+				_step.Type = BuildStepType.Other;
+				_step.Target = null;
+				_step.Platform = null;
+				_step.Configuration = null;
+				_step.FileName = OtherFileNameTextBox.Text;
+				_step.WorkingDir = OtherWorkingDirTextBox.Text;
+				_step.Arguments = OtherArgumentsTextBox.Text;
+				_step.UseLogWindow = OtherUseLogWindowCheckBox.Checked;
 			}
 
 			DialogResult = DialogResult.OK;
@@ -171,108 +171,108 @@ namespace UnrealGameSync
 
 		private void CookFileNameButton_Click(object sender, EventArgs e)
 		{
-			OpenFileDialog Dialog = new OpenFileDialog();
-			Dialog.Filter = "Cook/Launch Profiles (*.ulp2)|*.ulp2";
-			Dialog.FileName = AddBaseDirectory(CookFileNameTextBox.Text);
-			if(Dialog.ShowDialog() == DialogResult.OK)
+			OpenFileDialog dialog = new OpenFileDialog();
+			dialog.Filter = "Cook/Launch Profiles (*.ulp2)|*.ulp2";
+			dialog.FileName = AddBaseDirectory(CookFileNameTextBox.Text);
+			if(dialog.ShowDialog() == DialogResult.OK)
 			{
-				CookFileNameTextBox.Text = RemoveBaseDirectory(Dialog.FileName);
+				CookFileNameTextBox.Text = RemoveBaseDirectory(dialog.FileName);
 			}
 		}
 
 		private void OtherFileNameButton_Click(object sender, EventArgs e)
 		{
-			OpenFileDialog Dialog = new OpenFileDialog();
-			Dialog.Filter = "Executable Files (*.exe,*.bat)|*.exe;*.bat|All files (*.*)|*.*";
-			Dialog.FileName = AddBaseDirectory(OtherFileNameTextBox.Text);
-			if(Dialog.ShowDialog() == DialogResult.OK)
+			OpenFileDialog dialog = new OpenFileDialog();
+			dialog.Filter = "Executable Files (*.exe,*.bat)|*.exe;*.bat|All files (*.*)|*.*";
+			dialog.FileName = AddBaseDirectory(OtherFileNameTextBox.Text);
+			if(dialog.ShowDialog() == DialogResult.OK)
 			{
-				OtherFileNameTextBox.Text = RemoveBaseDirectory(Dialog.FileName);
+				OtherFileNameTextBox.Text = RemoveBaseDirectory(dialog.FileName);
 			}
 		}
 
 		private void OtherWorkingDirButton_Click(object sender, EventArgs e)
 		{
-			FolderBrowserDialog Dialog = new FolderBrowserDialog();
-			Dialog.SelectedPath = OtherWorkingDirTextBox.Text;
-			if (Dialog.ShowDialog() == DialogResult.OK)
+			FolderBrowserDialog dialog = new FolderBrowserDialog();
+			dialog.SelectedPath = OtherWorkingDirTextBox.Text;
+			if (dialog.ShowDialog() == DialogResult.OK)
 			{
-				OtherWorkingDirTextBox.Text = Dialog.SelectedPath;
+				OtherWorkingDirTextBox.Text = dialog.SelectedPath;
 			}
 		}
 
-		private string AddBaseDirectory(string FileName)
+		private string AddBaseDirectory(string fileName)
 		{
-			if(FileName.Contains("$("))
+			if(fileName.Contains("$("))
 			{
 				return "";
 			}
-			else if(Path.IsPathRooted(FileName))
+			else if(Path.IsPathRooted(fileName))
 			{
-				return FileName;
+				return fileName;
 			}
 			else
 			{
-				return FileReference.Combine(BaseDirectory, FileName).FullName;
+				return FileReference.Combine(_baseDirectory, fileName).FullName;
 			}
 		}
 
-		private string RemoveBaseDirectory(string FileName)
+		private string RemoveBaseDirectory(string fileName)
 		{
-			FileReference FullFileName = new FileReference(FileName);
-			if(FullFileName.IsUnderDirectory(BaseDirectory))
+			FileReference fullFileName = new FileReference(fileName);
+			if(fullFileName.IsUnderDirectory(_baseDirectory))
 			{
-				return FullFileName.MakeRelativeTo(BaseDirectory);
+				return fullFileName.MakeRelativeTo(_baseDirectory);
 			}
 			else
 			{
-				return FileName;
+				return fileName;
 			}
 		}
 
 		private void OnClosedVariablesWindow(object sender, EventArgs e)
 		{
-			VariablesWindow = null;
+			_variablesWindow = null;
 		}
 
-		private void InsertVariable(string Name)
+		private void InsertVariable(string name)
 		{
-			IContainerControl Container = this;
-			if(Container != null)
+			IContainerControl container = this;
+			if(container != null)
 			{
 				for(;;)
 				{
-					IContainerControl? NextContainer = Container.ActiveControl as IContainerControl;
-					if(NextContainer == null)
+					IContainerControl? nextContainer = container.ActiveControl as IContainerControl;
+					if(nextContainer == null)
 					{
 						break;
 					}
-					Container = NextContainer;
+					container = nextContainer;
 				}
 
-				TextBox? FocusTextBox = Container.ActiveControl as TextBox;
-				if(FocusTextBox != null)
+				TextBox? focusTextBox = container.ActiveControl as TextBox;
+				if(focusTextBox != null)
 				{
-					FocusTextBox.SelectedText = Name;
+					focusTextBox.SelectedText = name;
 				}
 			}
 		}
 
 		private void VariablesButton_Click(object sender, EventArgs e)
 		{
-			if(VariablesWindow == null)
+			if(_variablesWindow == null)
 			{
-				VariablesWindow = new VariablesWindow(Variables);
-				VariablesWindow.OnInsertVariable += InsertVariable;
-				VariablesWindow.Location = new Point(Bounds.Right + 20, Bounds.Top);
-				VariablesWindow.Size = new Size(VariablesWindow.Size.Width, Size.Height);
-				VariablesWindow.FormClosed += OnClosedVariablesWindow;
-				VariablesWindow.Show(this);
+				_variablesWindow = new VariablesWindow(_variables);
+				_variablesWindow.OnInsertVariable += InsertVariable;
+				_variablesWindow.Location = new Point(Bounds.Right + 20, Bounds.Top);
+				_variablesWindow.Size = new Size(_variablesWindow.Size.Width, Size.Height);
+				_variablesWindow.FormClosed += OnClosedVariablesWindow;
+				_variablesWindow.Show(this);
 			}
 			else
 			{
-				VariablesWindow.Close();
-				VariablesWindow = null;
+				_variablesWindow.Close();
+				_variablesWindow = null;
 			}
 		}
 	}
diff --git a/Engine/Source/Programs/UnrealGameSync/UnrealGameSync/Forms/ChangelistWindow.cs b/Engine/Source/Programs/UnrealGameSync/UnrealGameSync/Forms/ChangelistWindow.cs
--- a/Engine/Source/Programs/UnrealGameSync/UnrealGameSync/Forms/ChangelistWindow.cs
+++ b/Engine/Source/Programs/UnrealGameSync/UnrealGameSync/Forms/ChangelistWindow.cs
@@ -16,26 +16,26 @@ namespace UnrealGameSync
 	{
 		public int ChangeNumber;
 
-		public ChangelistWindow(int InChangeNumber)
+		public ChangelistWindow(int inChangeNumber)
 		{
 			InitializeComponent();
 
 			this.MaximumSize = new System.Drawing.Size(32768, Height);
 			this.MinimumSize = new System.Drawing.Size(Width, Height);
 
-			ChangeNumber = InChangeNumber;
+			ChangeNumber = inChangeNumber;
 			ChangeNumberTextBox.Text = (ChangeNumber > 0)? ChangeNumber.ToString() : "";
 		}
 
 		private void OkBtn_Click(object sender, EventArgs e)
 		{
-			int NewChangeNumber;
-			if(int.TryParse(ChangeNumberTextBox.Text.Trim(), out NewChangeNumber))
+			int newChangeNumber;
+			if(int.TryParse(ChangeNumberTextBox.Text.Trim(), out newChangeNumber))
 			{
-				int DeltaChangeNumber = NewChangeNumber - ChangeNumber;
-				if((DeltaChangeNumber > -10000 && DeltaChangeNumber < 1000) || ChangeNumber <= 0 || MessageBox.Show(String.Format("Changelist {0} is {1} changes away from your current sync. Was that intended?", NewChangeNumber, Math.Abs(DeltaChangeNumber)), "Changelist is far away", MessageBoxButtons.YesNo) == DialogResult.Yes)
+				int deltaChangeNumber = newChangeNumber - ChangeNumber;
+				if((deltaChangeNumber > -10000 && deltaChangeNumber < 1000) || ChangeNumber <= 0 || MessageBox.Show(String.Format("Changelist {0} is {1} changes away from your current sync. Was that intended?", newChangeNumber, Math.Abs(deltaChangeNumber)), "Changelist is far away", MessageBoxButtons.YesNo) == DialogResult.Yes)
 				{
-					ChangeNumber = NewChangeNumber;
+					ChangeNumber = newChangeNumber;
 					DialogResult = DialogResult.OK;
 					Close();
 				}
diff --git a/Engine/Source/Programs/UnrealGameSync/UnrealGameSync/Forms/CleanWorkspaceWindow.cs b/Engine/Source/Programs/UnrealGameSync/UnrealGameSync/Forms/CleanWorkspaceWindow.cs
--- a/Engine/Source/Programs/UnrealGameSync/UnrealGameSync/Forms/CleanWorkspaceWindow.cs
+++ b/Engine/Source/Programs/UnrealGameSync/UnrealGameSync/Forms/CleanWorkspaceWindow.cs
@@ -56,8 +56,8 @@ namespace UnrealGameSync
 			CheckBoxState.CheckedNormal 
 		};
 
-		IPerforceSettings PerforceSettings;
-		FolderToClean RootFolderToClean;
+		IPerforceSettings _perforceSettings;
+		FolderToClean _rootFolderToClean;
 
 		static readonly string[] SafeToDeleteFolders =
 		{
@@ -101,214 +101,214 @@ namespace UnrealGameSync
 			".csproj.references",
 		};
 
-		IReadOnlyList<string> ExtraSafeToDeleteFolders;
-		IReadOnlyList<string> ExtraSafeToDeleteExtensions;
-		ILogger Logger;
+		IReadOnlyList<string> _extraSafeToDeleteFolders;
+		IReadOnlyList<string> _extraSafeToDeleteExtensions;
+		ILogger _logger;
 
 		[DllImport("Shell32.dll", EntryPoint = "ExtractIconExW", CharSet = CharSet.Unicode, ExactSpelling = true, CallingConvention = CallingConvention.StdCall)]
 		private static extern int ExtractIconEx(string sFile, int iIndex, IntPtr piLargeVersion, out IntPtr piSmallVersion, int amountIcons);
 
-		private CleanWorkspaceWindow(IPerforceSettings PerforceSettings, FolderToClean RootFolderToClean, string[] ExtraSafeToDeleteFolders, string[] ExtraSafeToDeleteExtensions, ILogger<CleanWorkspaceWindow> Logger)
+		private CleanWorkspaceWindow(IPerforceSettings perforceSettings, FolderToClean rootFolderToClean, string[] extraSafeToDeleteFolders, string[] extraSafeToDeleteExtensions, ILogger<CleanWorkspaceWindow> logger)
 		{
-			this.PerforceSettings = PerforceSettings;
-			this.RootFolderToClean = RootFolderToClean;
-			this.ExtraSafeToDeleteFolders = ExtraSafeToDeleteFolders.Select(x => x.Trim().Replace('\\', '/').Trim('/')).Where(x => x.Length > 0).Select(x => String.Format("/{0}/", x.ToLowerInvariant())).ToArray();
-			this.ExtraSafeToDeleteExtensions = ExtraSafeToDeleteExtensions.Select(x => x.Trim().ToLowerInvariant()).Where(x => x.Length > 0).ToArray();
-			this.Logger = Logger;
+			this._perforceSettings = perforceSettings;
+			this._rootFolderToClean = rootFolderToClean;
+			this._extraSafeToDeleteFolders = extraSafeToDeleteFolders.Select(x => x.Trim().Replace('\\', '/').Trim('/')).Where(x => x.Length > 0).Select(x => String.Format("/{0}/", x.ToLowerInvariant())).ToArray();
+			this._extraSafeToDeleteExtensions = extraSafeToDeleteExtensions.Select(x => x.Trim().ToLowerInvariant()).Where(x => x.Length > 0).ToArray();
+			this._logger = logger;
 
 			InitializeComponent();
 		}
 
-		public static void DoClean(IWin32Window Owner, IPerforceSettings PerforceSettings, DirectoryReference LocalRootPath, string ClientRootPath, IReadOnlyList<string> SyncPaths, string[] ExtraSafeToDeleteFolders, string[] ExtraSafeToDeleteExtensions, ILogger<CleanWorkspaceWindow> Logger)
+		public static void DoClean(IWin32Window owner, IPerforceSettings perforceSettings, DirectoryReference localRootPath, string clientRootPath, IReadOnlyList<string> syncPaths, string[] extraSafeToDeleteFolders, string[] extraSafeToDeleteExtensions, ILogger<CleanWorkspaceWindow> logger)
 		{
 			// Figure out which folders to clean
-			FolderToClean RootFolderToClean = new FolderToClean(LocalRootPath.ToDirectoryInfo());
-			using(FindFoldersToCleanTask QueryWorkspace = new FindFoldersToCleanTask(PerforceSettings, RootFolderToClean, ClientRootPath, SyncPaths, Logger))
+			FolderToClean rootFolderToClean = new FolderToClean(localRootPath.ToDirectoryInfo());
+			using(FindFoldersToCleanTask queryWorkspace = new FindFoldersToCleanTask(perforceSettings, rootFolderToClean, clientRootPath, syncPaths, logger))
 			{
-				ModalTask? Result = ModalTask.Execute(Owner, "Clean Workspace", "Querying files in Perforce, please wait...", x => QueryWorkspace.RunAsync(x), ModalTaskFlags.None);
-				if (Result == null || !Result.Succeeded)
+				ModalTask? result = ModalTask.Execute(owner, "Clean Workspace", "Querying files in Perforce, please wait...", x => queryWorkspace.RunAsync(x), ModalTaskFlags.None);
+				if (result == null || !result.Succeeded)
 				{
 					return;
 				}
 			}
 
 			// If there's nothing to delete, don't bother displaying the dialog at all
-			if(RootFolderToClean.FilesToDelete.Count == 0 && RootFolderToClean.NameToSubFolder.Count == 0)
+			if(rootFolderToClean.FilesToDelete.Count == 0 && rootFolderToClean.NameToSubFolder.Count == 0)
 			{
 				MessageBox.Show("You have no local files which are not in Perforce.", "Workspace Clean", MessageBoxButtons.OK);
 				return;
 			}
 
 			// Populate the tree
-			CleanWorkspaceWindow CleanWorkspace = new CleanWorkspaceWindow(PerforceSettings, RootFolderToClean, ExtraSafeToDeleteFolders, ExtraSafeToDeleteExtensions, Logger);
-			CleanWorkspace.ShowDialog();
+			CleanWorkspaceWindow cleanWorkspace = new CleanWorkspaceWindow(perforceSettings, rootFolderToClean, extraSafeToDeleteFolders, extraSafeToDeleteExtensions, logger);
+			cleanWorkspace.ShowDialog();
 		}
 
 		private void CleanWorkspaceWindow_Load(object sender, EventArgs e)
 		{
-			IntPtr FolderIconPtr;
-			ExtractIconEx("imageres.dll", 3, IntPtr.Zero, out FolderIconPtr, 1);
+			IntPtr folderIconPtr;
+			ExtractIconEx("imageres.dll", 3, IntPtr.Zero, out folderIconPtr, 1);
 
-			IntPtr FileIconPtr;
-			ExtractIconEx("imageres.dll", 2, IntPtr.Zero, out FileIconPtr, 1);
+			IntPtr fileIconPtr;
+			ExtractIconEx("imageres.dll", 2, IntPtr.Zero, out fileIconPtr, 1);
 
-			Icon[] Icons = new Icon[]{ Icon.FromHandle(FolderIconPtr), Icon.FromHandle(FileIconPtr) };
+			Icon[] icons = new Icon[]{ Icon.FromHandle(folderIconPtr), Icon.FromHandle(fileIconPtr) };
 
-			Size LargestIconSize = Size.Empty;
-			foreach(Icon Icon in Icons)
+			Size largestIconSize = Size.Empty;
+			foreach(Icon icon in icons)
 			{
-				LargestIconSize = new Size(Math.Max(LargestIconSize.Width, Icon.Width), Math.Max(LargestIconSize.Height, Icon.Height));
+				largestIconSize = new Size(Math.Max(largestIconSize.Width, icon.Width), Math.Max(largestIconSize.Height, icon.Height));
 			}
 
-			Size LargestCheckBoxSize = Size.Empty;
-			using(Graphics Graphics = Graphics.FromHwnd(IntPtr.Zero))
+			Size largestCheckBoxSize = Size.Empty;
+			using(Graphics graphics = Graphics.FromHwnd(IntPtr.Zero))
 			{
-				foreach(CheckBoxState State in CheckBoxStates)
+				foreach(CheckBoxState state in CheckBoxStates)
 				{
-					Size CheckBoxSize = CheckBoxRenderer.GetGlyphSize(Graphics, State);
-					LargestCheckBoxSize = new Size(Math.Max(LargestCheckBoxSize.Width, CheckBoxSize.Width), Math.Max(LargestCheckBoxSize.Height, CheckBoxSize.Height));
+					Size checkBoxSize = CheckBoxRenderer.GetGlyphSize(graphics, state);
+					largestCheckBoxSize = new Size(Math.Max(largestCheckBoxSize.Width, checkBoxSize.Width), Math.Max(largestCheckBoxSize.Height, checkBoxSize.Height));
 				}
 			}
 
-			Size ImageSize = new Size(LargestCheckBoxSize.Width + LargestIconSize.Width, Math.Max(LargestIconSize.Height, LargestCheckBoxSize.Height));
+			Size imageSize = new Size(largestCheckBoxSize.Width + largestIconSize.Width, Math.Max(largestIconSize.Height, largestCheckBoxSize.Height));
 
-			Bitmap TypeImageListBitmap = new Bitmap(Icons.Length * 3 * ImageSize.Width, ImageSize.Height, System.Drawing.Imaging.PixelFormat.Format32bppArgb);
-			using(Graphics Graphics = Graphics.FromImage(TypeImageListBitmap))
+			Bitmap typeImageListBitmap = new Bitmap(icons.Length * 3 * imageSize.Width, imageSize.Height, System.Drawing.Imaging.PixelFormat.Format32bppArgb);
+			using(Graphics graphics = Graphics.FromImage(typeImageListBitmap))
 			{
-				int MinX = 0;
-				for(int IconIdx = 0; IconIdx < Icons.Length; IconIdx++)
+				int minX = 0;
+				for(int iconIdx = 0; iconIdx < icons.Length; iconIdx++)
 				{
-					for(int StateIdx = 0; StateIdx < 3; StateIdx++)
+					for(int stateIdx = 0; stateIdx < 3; stateIdx++)
 					{
-						Size CheckBoxSize = CheckBoxRenderer.GetGlyphSize(Graphics, CheckBoxStates[StateIdx]);
-						CheckBoxRenderer.DrawCheckBox(Graphics, new Point(MinX + (LargestCheckBoxSize.Width - CheckBoxSize.Width) / 2, (LargestCheckBoxSize.Height - CheckBoxSize.Height) / 2), CheckBoxStates[StateIdx]);
+						Size checkBoxSize = CheckBoxRenderer.GetGlyphSize(graphics, CheckBoxStates[stateIdx]);
+						CheckBoxRenderer.DrawCheckBox(graphics, new Point(minX + (largestCheckBoxSize.Width - checkBoxSize.Width) / 2, (largestCheckBoxSize.Height - checkBoxSize.Height) / 2), CheckBoxStates[stateIdx]);
 
-						Size IconSize = Icons[IconIdx].Size;
-						Graphics.DrawIcon(Icons[IconIdx], MinX + LargestCheckBoxSize.Width + (LargestIconSize.Width - IconSize.Width) / 2, (LargestIconSize.Height - IconSize.Height) / 2);
+						Size iconSize = icons[iconIdx].Size;
+						graphics.DrawIcon(icons[iconIdx], minX + largestCheckBoxSize.Width + (largestIconSize.Width - iconSize.Width) / 2, (largestIconSize.Height - iconSize.Height) / 2);
 
-						MinX += ImageSize.Width;
+						minX += imageSize.Width;
 					}
 				}
 			}
 
-			ImageList TypeImageList = new ImageList();
-			TypeImageList.ImageSize = ImageSize;
-			TypeImageList.ColorDepth = ColorDepth.Depth32Bit;
-			TypeImageList.Images.AddStrip(TypeImageListBitmap);
-			TreeView.ImageList = TypeImageList;
+			ImageList typeImageList = new ImageList();
+			typeImageList.ImageSize = imageSize;
+			typeImageList.ColorDepth = ColorDepth.Depth32Bit;
+			typeImageList.Images.AddStrip(typeImageListBitmap);
+			TreeView.ImageList = typeImageList;
 
-			TreeNode Node = BuildTreeViewStructure(RootFolderToClean, "/", false, 0);
-			Node.Text = RootFolderToClean.Directory.FullName;
-			TreeView.Nodes.Add(Node);
+			TreeNode node = BuildTreeViewStructure(_rootFolderToClean, "/", false, 0);
+			node.Text = _rootFolderToClean.Directory.FullName;
+			TreeView.Nodes.Add(node);
 		}
 
-		private TreeNode BuildTreeViewStructure(FolderToClean Folder, string FolderPath, bool bParentFolderSelected, int Depth)
+		private TreeNode BuildTreeViewStructure(FolderToClean folder, string folderPath, bool parentFolderSelected, int depth)
 		{
-			bool bSelectFolder = bParentFolderSelected || IsSafeToDeleteFolder(FolderPath) || Folder.bEmptyLeaf;
+			bool selectFolder = parentFolderSelected || IsSafeToDeleteFolder(folderPath) || folder.EmptyLeaf;
 
-			TreeNodeData FolderNodeData = new TreeNodeData();
-			FolderNodeData.Folder = Folder;
+			TreeNodeData folderNodeData = new TreeNodeData();
+			folderNodeData.Folder = folder;
 
-			TreeNode FolderNode = new TreeNode();
-			FolderNode.Text = Folder.Name;
-			FolderNode.Tag = FolderNodeData;
+			TreeNode folderNode = new TreeNode();
+			folderNode.Text = folder.Name;
+			folderNode.Tag = folderNodeData;
 
-			foreach(FolderToClean SubFolder in Folder.NameToSubFolder.OrderBy(x => x.Key).Select(x => x.Value))
+			foreach(FolderToClean subFolder in folder.NameToSubFolder.OrderBy(x => x.Key).Select(x => x.Value))
 			{
-				TreeNode ChildNode = BuildTreeViewStructure(SubFolder, FolderPath + SubFolder.Name.ToLowerInvariant() + "/", bSelectFolder, Depth + 1);
-				FolderNode.Nodes.Add(ChildNode);
-
-				TreeNodeData ChildNodeData = (TreeNodeData)ChildNode.Tag;
-				FolderNodeData.NumFiles += ChildNodeData.NumFiles;
-				FolderNodeData.NumSelectedFiles += ChildNodeData.NumSelectedFiles;
-				FolderNodeData.NumEmptySelectedFiles += ChildNodeData.NumEmptySelectedFiles;
-				FolderNodeData.NumMissingSelectedFiles += ChildNodeData.NumMissingSelectedFiles;
-				FolderNodeData.NumDefaultSelectedFiles += ChildNodeData.NumDefaultSelectedFiles;
+				TreeNode childNode = BuildTreeViewStructure(subFolder, folderPath + subFolder.Name.ToLowerInvariant() + "/", selectFolder, depth + 1);
+				folderNode.Nodes.Add(childNode);
+
+				TreeNodeData childNodeData = (TreeNodeData)childNode.Tag;
+				folderNodeData.NumFiles += childNodeData.NumFiles;
+				folderNodeData.NumSelectedFiles += childNodeData.NumSelectedFiles;
+				folderNodeData.NumEmptySelectedFiles += childNodeData.NumEmptySelectedFiles;
+				folderNodeData.NumMissingSelectedFiles += childNodeData.NumMissingSelectedFiles;
+				folderNodeData.NumDefaultSelectedFiles += childNodeData.NumDefaultSelectedFiles;
 			}
 
-			foreach(FileInfo File in Folder.FilesToSync.OrderBy(x => x.Name))
+			foreach(FileInfo file in folder.FilesToSync.OrderBy(x => x.Name))
 			{
-				TreeNodeData FileNodeData = new TreeNodeData();
-				FileNodeData.Action = TreeNodeAction.Sync;
-				FileNodeData.File = File;
-				FileNodeData.NumFiles = 1;
-				FileNodeData.NumSelectedFiles = 1;
-				FileNodeData.NumEmptySelectedFiles = 0;
-				FileNodeData.NumMissingSelectedFiles = 1;
-				FileNodeData.NumDefaultSelectedFiles = FileNodeData.NumSelectedFiles;
-
-				TreeNode FileNode = new TreeNode();
-				FileNode.Text = File.Name + " (sync)";
-				FileNode.Tag = FileNodeData;
-				FolderNode.Nodes.Add(FileNode);
-
-				UpdateImage(FileNode);
-
-				FolderNodeData.NumFiles++;
-				FolderNodeData.NumSelectedFiles += FileNodeData.NumSelectedFiles;
-				FolderNodeData.NumEmptySelectedFiles += FileNodeData.NumEmptySelectedFiles;
-				FolderNodeData.NumMissingSelectedFiles += FileNodeData.NumMissingSelectedFiles;
-				FolderNodeData.NumDefaultSelectedFiles += FileNodeData.NumDefaultSelectedFiles;
+				TreeNodeData fileNodeData = new TreeNodeData();
+				fileNodeData.Action = TreeNodeAction.Sync;
+				fileNodeData.File = file;
+				fileNodeData.NumFiles = 1;
+				fileNodeData.NumSelectedFiles = 1;
+				fileNodeData.NumEmptySelectedFiles = 0;
+				fileNodeData.NumMissingSelectedFiles = 1;
+				fileNodeData.NumDefaultSelectedFiles = fileNodeData.NumSelectedFiles;
+
+				TreeNode fileNode = new TreeNode();
+				fileNode.Text = file.Name + " (sync)";
+				fileNode.Tag = fileNodeData;
+				folderNode.Nodes.Add(fileNode);
+
+				UpdateImage(fileNode);
+
+				folderNodeData.NumFiles++;
+				folderNodeData.NumSelectedFiles += fileNodeData.NumSelectedFiles;
+				folderNodeData.NumEmptySelectedFiles += fileNodeData.NumEmptySelectedFiles;
+				folderNodeData.NumMissingSelectedFiles += fileNodeData.NumMissingSelectedFiles;
+				folderNodeData.NumDefaultSelectedFiles += fileNodeData.NumDefaultSelectedFiles;
 			}
 
-			foreach(FileInfo File in Folder.FilesToDelete.OrderBy(x => x.Name))
+			foreach(FileInfo file in folder.FilesToDelete.OrderBy(x => x.Name))
 			{
-				string Name = File.Name.ToLowerInvariant();
-
-				bool bSelectFile = bSelectFolder || IsSafeToDeleteFile(FolderPath, File.Name.ToLowerInvariant());
-
-				TreeNodeData FileNodeData = new TreeNodeData();
-				FileNodeData.Action = TreeNodeAction.Delete;
-				FileNodeData.File = File;
-				FileNodeData.NumFiles = 1;
-				FileNodeData.NumSelectedFiles = bSelectFile? 1 : 0;
-				FileNodeData.NumEmptySelectedFiles = 0;
-				FileNodeData.NumMissingSelectedFiles = 0;
-				FileNodeData.NumDefaultSelectedFiles = FileNodeData.NumSelectedFiles;
-
-				TreeNode FileNode = new TreeNode();
-				FileNode.Text = File.Name;
-				FileNode.Tag = FileNodeData;
-				FolderNode.Nodes.Add(FileNode);
-
-				UpdateImage(FileNode);
-
-				FolderNodeData.NumFiles++;
-				FolderNodeData.NumSelectedFiles += FileNodeData.NumSelectedFiles;
-				FolderNodeData.NumEmptySelectedFiles += FileNodeData.NumEmptySelectedFiles;
-				FolderNodeData.NumMissingSelectedFiles += FileNodeData.NumMissingSelectedFiles;
-				FolderNodeData.NumDefaultSelectedFiles += FileNodeData.NumDefaultSelectedFiles;
+				string name = file.Name.ToLowerInvariant();
+
+				bool selectFile = selectFolder || IsSafeToDeleteFile(folderPath, file.Name.ToLowerInvariant());
+
+				TreeNodeData fileNodeData = new TreeNodeData();
+				fileNodeData.Action = TreeNodeAction.Delete;
+				fileNodeData.File = file;
+				fileNodeData.NumFiles = 1;
+				fileNodeData.NumSelectedFiles = selectFile? 1 : 0;
+				fileNodeData.NumEmptySelectedFiles = 0;
+				fileNodeData.NumMissingSelectedFiles = 0;
+				fileNodeData.NumDefaultSelectedFiles = fileNodeData.NumSelectedFiles;
+
+				TreeNode fileNode = new TreeNode();
+				fileNode.Text = file.Name;
+				fileNode.Tag = fileNodeData;
+				folderNode.Nodes.Add(fileNode);
+
+				UpdateImage(fileNode);
+
+				folderNodeData.NumFiles++;
+				folderNodeData.NumSelectedFiles += fileNodeData.NumSelectedFiles;
+				folderNodeData.NumEmptySelectedFiles += fileNodeData.NumEmptySelectedFiles;
+				folderNodeData.NumMissingSelectedFiles += fileNodeData.NumMissingSelectedFiles;
+				folderNodeData.NumDefaultSelectedFiles += fileNodeData.NumDefaultSelectedFiles;
 			}
 
-			if(FolderNodeData.Folder.bEmptyLeaf)
+			if(folderNodeData.Folder.EmptyLeaf)
 			{
-				FolderNodeData.NumFiles++;
-				FolderNodeData.NumSelectedFiles++;
-				FolderNodeData.NumEmptySelectedFiles++;
-				FolderNodeData.NumDefaultSelectedFiles++;
+				folderNodeData.NumFiles++;
+				folderNodeData.NumSelectedFiles++;
+				folderNodeData.NumEmptySelectedFiles++;
+				folderNodeData.NumDefaultSelectedFiles++;
 			}
 
-			if(FolderNodeData.NumSelectedFiles > 0 && !FolderNodeData.Folder.bEmptyAfterClean && Depth < 2)
+			if(folderNodeData.NumSelectedFiles > 0 && !folderNodeData.Folder.EmptyAfterClean && depth < 2)
 			{
-				FolderNode.Expand();
+				folderNode.Expand();
 			}
 			else
 			{
-				FolderNode.Collapse();
+				folderNode.Collapse();
 			}
 
-			UpdateImage(FolderNode);
-			return FolderNode;
+			UpdateImage(folderNode);
+			return folderNode;
 		}
 
-		private bool IsSafeToDeleteFolder(string FolderPath)
+		private bool IsSafeToDeleteFolder(string folderPath)
 		{
-			return SafeToDeleteFolders.Any(x => FolderPath.EndsWith(x)) || ExtraSafeToDeleteFolders.Any(x => FolderPath.EndsWith(x));
+			return SafeToDeleteFolders.Any(x => folderPath.EndsWith(x)) || _extraSafeToDeleteFolders.Any(x => folderPath.EndsWith(x));
 		}
 
-		private bool IsSafeToDeleteFile(string FolderPath, string Name)
+		private bool IsSafeToDeleteFile(string folderPath, string name)
 		{
-			return SafeToDeleteExtensions.Any(x => Name.EndsWith(x)) || ExtraSafeToDeleteExtensions.Any(x => Name.EndsWith(x));
+			return SafeToDeleteExtensions.Any(x => name.EndsWith(x)) || _extraSafeToDeleteExtensions.Any(x => name.EndsWith(x));
 		}
 
 		private void TreeView_DrawNode(object sender, DrawTreeNodeEventArgs e)
@@ -318,7 +318,7 @@ namespace UnrealGameSync
 
 		private void TreeView_NodeMouseClick(object sender, TreeNodeMouseClickEventArgs e)
 		{
-			TreeNode Node = e.Node;
+			TreeNode node = e.Node;
 			if(e.Button == System.Windows.Forms.MouseButtons.Right)
 			{
 				TreeView.SelectedNode = e.Node;
@@ -327,192 +327,192 @@ namespace UnrealGameSync
 			}
 			else if(e.X >= e.Node.Bounds.Left - 32 && e.X < e.Node.Bounds.Left - 16)
 			{
-				TreeNodeData NodeData = (TreeNodeData)Node.Tag;
-				SetSelected(Node, (NodeData.NumSelectedFiles == 0)? SelectionType.All : SelectionType.None);
+				TreeNodeData nodeData = (TreeNodeData)node.Tag;
+				SetSelected(node, (nodeData.NumSelectedFiles == 0)? SelectionType.All : SelectionType.None);
 			}
 		}
 
-		private void SetSelected(TreeNode ParentNode, SelectionType Type)
+		private void SetSelected(TreeNode parentNode, SelectionType type)
 		{
-			TreeNodeData ParentNodeData = (TreeNodeData)ParentNode.Tag;
+			TreeNodeData parentNodeData = (TreeNodeData)parentNode.Tag;
 
-			int PrevNumSelectedFiles = ParentNodeData.NumSelectedFiles;
-			SetSelectedOnChildren(ParentNode, Type);
+			int prevNumSelectedFiles = parentNodeData.NumSelectedFiles;
+			SetSelectedOnChildren(parentNode, type);
 
-			int DeltaNumSelectedFiles = ParentNodeData.NumSelectedFiles - PrevNumSelectedFiles;
-			if(DeltaNumSelectedFiles != 0)
+			int deltaNumSelectedFiles = parentNodeData.NumSelectedFiles - prevNumSelectedFiles;
+			if(deltaNumSelectedFiles != 0)
 			{
-				for(TreeNode NextParentNode = ParentNode.Parent; NextParentNode != null; NextParentNode = NextParentNode.Parent)
+				for(TreeNode nextParentNode = parentNode.Parent; nextParentNode != null; nextParentNode = nextParentNode.Parent)
 				{
-					TreeNodeData NextParentNodeData = (TreeNodeData)NextParentNode.Tag;
-					NextParentNodeData.NumSelectedFiles += DeltaNumSelectedFiles;
-					UpdateImage(NextParentNode);
+					TreeNodeData nextParentNodeData = (TreeNodeData)nextParentNode.Tag;
+					nextParentNodeData.NumSelectedFiles += deltaNumSelectedFiles;
+					UpdateImage(nextParentNode);
 				}
 			}
 		}
 
-		private void SetSelectedOnChildren(TreeNode ParentNode, SelectionType Type)
+		private void SetSelectedOnChildren(TreeNode parentNode, SelectionType type)
 		{
-			TreeNodeData ParentNodeData = (TreeNodeData)ParentNode.Tag;
+			TreeNodeData parentNodeData = (TreeNodeData)parentNode.Tag;
 
-			int NewNumSelectedFiles = 0;
-			switch(Type)
+			int newNumSelectedFiles = 0;
+			switch(type)
 			{
 				case SelectionType.All:
-					NewNumSelectedFiles = ParentNodeData.NumFiles;
+					newNumSelectedFiles = parentNodeData.NumFiles;
 					break;
 				case SelectionType.Empty:
-					NewNumSelectedFiles = ParentNodeData.NumEmptySelectedFiles;
+					newNumSelectedFiles = parentNodeData.NumEmptySelectedFiles;
 					break;
 				case SelectionType.Missing:
-					NewNumSelectedFiles = ParentNodeData.NumMissingSelectedFiles;
+					newNumSelectedFiles = parentNodeData.NumMissingSelectedFiles;
 					break;
 				case SelectionType.SafeToDelete:
-					NewNumSelectedFiles = ParentNodeData.NumDefaultSelectedFiles;
+					newNumSelectedFiles = parentNodeData.NumDefaultSelectedFiles;
 					break;
 				case SelectionType.None:
-					NewNumSelectedFiles = 0;
+					newNumSelectedFiles = 0;
 					break;
 			}
 
-			if(NewNumSelectedFiles != ParentNodeData.NumSelectedFiles)
+			if(newNumSelectedFiles != parentNodeData.NumSelectedFiles)
 			{
-				foreach(TreeNode? ChildNode in ParentNode.Nodes)
+				foreach(TreeNode? childNode in parentNode.Nodes)
 				{
-					if (ChildNode != null)
+					if (childNode != null)
 					{
-						SetSelectedOnChildren(ChildNode, Type);
+						SetSelectedOnChildren(childNode, type);
 					}
 				}
-				ParentNodeData.NumSelectedFiles = NewNumSelectedFiles;
-				UpdateImage(ParentNode);
+				parentNodeData.NumSelectedFiles = newNumSelectedFiles;
+				UpdateImage(parentNode);
 			}
 		}
 
-		private void UpdateImage(TreeNode Node)
+		private void UpdateImage(TreeNode node)
 		{
-			TreeNodeData NodeData = (TreeNodeData)Node.Tag;
-			int ImageIndex = (NodeData.Folder != null)? 0 : 3;
-			ImageIndex += (NodeData.NumSelectedFiles == 0)? 0 : (NodeData.NumSelectedFiles < NodeData.NumFiles || (NodeData.Folder != null && !NodeData.Folder.bEmptyAfterClean))? 1 : 2;
-			Node.ImageIndex = ImageIndex;
-			Node.SelectedImageIndex = ImageIndex;
+			TreeNodeData nodeData = (TreeNodeData)node.Tag;
+			int imageIndex = (nodeData.Folder != null)? 0 : 3;
+			imageIndex += (nodeData.NumSelectedFiles == 0)? 0 : (nodeData.NumSelectedFiles < nodeData.NumFiles || (nodeData.Folder != null && !nodeData.Folder.EmptyAfterClean))? 1 : 2;
+			node.ImageIndex = imageIndex;
+			node.SelectedImageIndex = imageIndex;
 		}
 
 		private void CleanBtn_Click(object sender, EventArgs e)
 		{
-			List<FileInfo> FilesToSync = new List<FileInfo>();
-			List<FileInfo> FilesToDelete = new List<FileInfo>();
-			List<DirectoryInfo> DirectoriesToDelete = new List<DirectoryInfo>();
-			foreach(TreeNode? RootNode in TreeView.Nodes)
+			List<FileInfo> filesToSync = new List<FileInfo>();
+			List<FileInfo> filesToDelete = new List<FileInfo>();
+			List<DirectoryInfo> directoriesToDelete = new List<DirectoryInfo>();
+			foreach(TreeNode? rootNode in TreeView.Nodes)
 			{
-				if (RootNode != null)
+				if (rootNode != null)
 				{
-					FindSelection(RootNode, FilesToSync, FilesToDelete, DirectoriesToDelete);
+					FindSelection(rootNode, filesToSync, filesToDelete, directoriesToDelete);
 				}
 			}
 
-			ModalTask? Result = ModalTask.Execute(this, "Clean Workspace", "Cleaning files, please wait...", x => DeleteFilesTask.RunAsync(PerforceSettings, FilesToSync, FilesToDelete, DirectoriesToDelete, Logger, x), ModalTaskFlags.Quiet);
-			if(Result != null && Result.Failed)
+			ModalTask? result = ModalTask.Execute(this, "Clean Workspace", "Cleaning files, please wait...", x => DeleteFilesTask.RunAsync(_perforceSettings, filesToSync, filesToDelete, directoriesToDelete, _logger, x), ModalTaskFlags.Quiet);
+			if(result != null && result.Failed)
 			{
-				FailedToDeleteWindow FailedToDelete = new FailedToDeleteWindow();
-				FailedToDelete.FileList.Text = Result.Error;
-				FailedToDelete.FileList.SelectionStart = 0;
-				FailedToDelete.FileList.SelectionLength = 0;
-				FailedToDelete.ShowDialog();
+				FailedToDeleteWindow failedToDelete = new FailedToDeleteWindow();
+				failedToDelete.FileList.Text = result.Error;
+				failedToDelete.FileList.SelectionStart = 0;
+				failedToDelete.FileList.SelectionLength = 0;
+				failedToDelete.ShowDialog();
 			}
 		}
 
-		private void FindSelection(TreeNode Node, List<FileInfo> FilesToSync, List<FileInfo> FilesToDelete, List<DirectoryInfo> DirectoriesToDelete)
+		private void FindSelection(TreeNode node, List<FileInfo> filesToSync, List<FileInfo> filesToDelete, List<DirectoryInfo> directoriesToDelete)
 		{
-			TreeNodeData NodeData = (TreeNodeData)Node.Tag;
-			if(NodeData.File != null)
+			TreeNodeData nodeData = (TreeNodeData)node.Tag;
+			if(nodeData.File != null)
 			{
-				if(NodeData.NumSelectedFiles > 0)
+				if(nodeData.NumSelectedFiles > 0)
 				{
-					if(NodeData.Action == TreeNodeAction.Delete)
+					if(nodeData.Action == TreeNodeAction.Delete)
 					{
-						FilesToDelete.Add(NodeData.File);
+						filesToDelete.Add(nodeData.File);
 					}
 					else
 					{
-						FilesToSync.Add(NodeData.File);
+						filesToSync.Add(nodeData.File);
 					}
 				}
 			}
 			else
 			{
-				foreach(TreeNode? ChildNode in Node.Nodes)
+				foreach(TreeNode? childNode in node.Nodes)
 				{
-					if (ChildNode != null)
+					if (childNode != null)
 					{
-						FindSelection(ChildNode, FilesToSync, FilesToDelete, DirectoriesToDelete);
+						FindSelection(childNode, filesToSync, filesToDelete, directoriesToDelete);
 					}
 				}
-				if(NodeData.Folder != null && NodeData.Folder.bEmptyAfterClean && NodeData.NumSelectedFiles == NodeData.NumFiles)
+				if(nodeData.Folder != null && nodeData.Folder.EmptyAfterClean && nodeData.NumSelectedFiles == nodeData.NumFiles)
 				{
-					DirectoriesToDelete.Add(NodeData.Folder.Directory);
+					directoriesToDelete.Add(nodeData.Folder.Directory);
 				}
 			}
 		}
 
 		private void SelectAllBtn_Click(object sender, EventArgs e)
 		{
-			foreach(TreeNode? Node in TreeView.Nodes)
+			foreach(TreeNode? node in TreeView.Nodes)
 			{
-				if (Node != null)
+				if (node != null)
 				{
-					SetSelected(Node, SelectionType.All);
+					SetSelected(node, SelectionType.All);
 				}
 			}
 		}
 
 		private void SelectMissingBtn_Click(object sender, EventArgs e)
 		{
-			foreach(TreeNode? Node in TreeView.Nodes)
+			foreach(TreeNode? node in TreeView.Nodes)
 			{
-				if (Node != null)
+				if (node != null)
 				{
-					SetSelected(Node, SelectionType.Missing);
+					SetSelected(node, SelectionType.Missing);
 				}
 			}
 		}
 
 		private void FolderContextMenu_SelectAll_Click(object sender, EventArgs e)
 		{
-			TreeNode Node = (TreeNode)FolderContextMenu.Tag;
-			SetSelected(Node, SelectionType.All);
+			TreeNode node = (TreeNode)FolderContextMenu.Tag;
+			SetSelected(node, SelectionType.All);
 		}
 
 		private void FolderContextMenu_SelectSafeToDelete_Click(object sender, EventArgs e)
 		{
-			TreeNode Node = (TreeNode)FolderContextMenu.Tag;
-			SetSelected(Node, SelectionType.SafeToDelete);
+			TreeNode node = (TreeNode)FolderContextMenu.Tag;
+			SetSelected(node, SelectionType.SafeToDelete);
 		}
 
 		private void FolderContextMenu_SelectEmptyFolder_Click(object sender, EventArgs e)
 		{
-			TreeNode Node = (TreeNode)FolderContextMenu.Tag;
-			SetSelected(Node, SelectionType.Empty);
+			TreeNode node = (TreeNode)FolderContextMenu.Tag;
+			SetSelected(node, SelectionType.Empty);
 		}
 
 		private void FolderContextMenu_SelectNone_Click(object sender, EventArgs e)
 		{
-			TreeNode Node = (TreeNode)FolderContextMenu.Tag;
-			SetSelected(Node, SelectionType.None);
+			TreeNode node = (TreeNode)FolderContextMenu.Tag;
+			SetSelected(node, SelectionType.None);
 		}
 
 		private void FolderContextMenu_OpenWithExplorer_Click(object sender, EventArgs e)
 		{
-			TreeNode Node = (TreeNode)FolderContextMenu.Tag;
-			TreeNodeData NodeData = (TreeNodeData)Node.Tag;
+			TreeNode node = (TreeNode)FolderContextMenu.Tag;
+			TreeNodeData nodeData = (TreeNodeData)node.Tag;
 
-			if (NodeData.Folder != null)
+			if (nodeData.Folder != null)
 			{
-				Process.Start("explorer.exe", String.Format("\"{0}\"", NodeData.Folder.Directory.FullName));
+				Process.Start("explorer.exe", String.Format("\"{0}\"", nodeData.Folder.Directory.FullName));
 			}
-			else if (NodeData.File != null)
+			else if (nodeData.File != null)
 			{
-				Process.Start("explorer.exe", String.Format("\"{0}\"", NodeData.File.Directory.FullName));
+				Process.Start("explorer.exe", String.Format("\"{0}\"", nodeData.File.Directory.FullName));
 			}
 		}
 	}
diff --git a/Engine/Source/Programs/UnrealGameSync/UnrealGameSync/Forms/ClobberWindow.cs b/Engine/Source/Programs/UnrealGameSync/UnrealGameSync/Forms/ClobberWindow.cs
--- a/Engine/Source/Programs/UnrealGameSync/UnrealGameSync/Forms/ClobberWindow.cs
+++ b/Engine/Source/Programs/UnrealGameSync/UnrealGameSync/Forms/ClobberWindow.cs
@@ -15,32 +15,32 @@ namespace UnrealGameSync
 {
 	public partial class ClobberWindow : Form
 	{
-		Dictionary<string, bool> FilesToClobber;
+		Dictionary<string, bool> _filesToClobber;
 
-		public ClobberWindow(Dictionary<string, bool> InFilesToClobber, HashSet<string> InUncontrolledFiles)
+		public ClobberWindow(Dictionary<string, bool> inFilesToClobber, HashSet<string> inUncontrolledFiles)
 		{
-			bool bUncontrolledChangeFound = false;
+			bool uncontrolledChangeFound = false;
 
 			InitializeComponent();
 
-			FilesToClobber = InFilesToClobber;
+			_filesToClobber = inFilesToClobber;
 
-			foreach(KeyValuePair<string, bool> FileToClobber in FilesToClobber)
+			foreach(KeyValuePair<string, bool> fileToClobber in _filesToClobber)
 			{
-				ListViewItem Item = new ListViewItem(Path.GetFileName(FileToClobber.Key));
-				Item.Tag = FileToClobber.Key;
-				Item.Checked = FileToClobber.Value;
-				Item.SubItems.Add(Path.GetDirectoryName(FileToClobber.Key));
-				FileList.Items.Add(Item);
+				ListViewItem item = new ListViewItem(Path.GetFileName(fileToClobber.Key));
+				item.Tag = fileToClobber.Key;
+				item.Checked = fileToClobber.Value;
+				item.SubItems.Add(Path.GetDirectoryName(fileToClobber.Key));
+				FileList.Items.Add(item);
 
-				if (InUncontrolledFiles.Contains(FileToClobber.Key.Replace("\\", "/")))
+				if (inUncontrolledFiles.Contains(fileToClobber.Key.Replace("\\", "/")))
 				{
-					bUncontrolledChangeFound = true;
-					Item.ForeColor = Color.Red;
+					uncontrolledChangeFound = true;
+					item.ForeColor = Color.Red;
 				}
 			}
 
-			if (bUncontrolledChangeFound)
+			if (uncontrolledChangeFound)
 			{
 				// Updates the string to inform the user to take special care with Uncontrolled Changes
 				this.label1.Text = "The following files are writable in your workspace." + Environment.NewLine +
@@ -51,22 +51,22 @@ namespace UnrealGameSync
 
 		private void UncheckAll_Click(object sender, EventArgs e)
 		{
-			foreach (ListViewItem? Item in FileList.Items)
+			foreach (ListViewItem? item in FileList.Items)
 			{
-				if (Item != null)
+				if (item != null)
 				{
-					Item.Checked = false;
+					item.Checked = false;
 				}
 			}
 		}
 
 		private void ContinueButton_Click(object sender, EventArgs e)
 		{
-			foreach(ListViewItem? Item in FileList.Items)
+			foreach(ListViewItem? item in FileList.Items)
 			{
-				if (Item != null)
+				if (item != null)
 				{
-					FilesToClobber[(string)Item.Tag] = Item.Checked;
+					_filesToClobber[(string)item.Tag] = item.Checked;
 				}
 			}
 		}
diff --git a/Engine/Source/Programs/UnrealGameSync/UnrealGameSync/Forms/ConnectWindow.cs b/Engine/Source/Programs/UnrealGameSync/UnrealGameSync/Forms/ConnectWindow.cs
--- a/Engine/Source/Programs/UnrealGameSync/UnrealGameSync/Forms/ConnectWindow.cs
+++ b/Engine/Source/Programs/UnrealGameSync/UnrealGameSync/Forms/ConnectWindow.cs
@@ -17,43 +17,43 @@ namespace UnrealGameSync
 {
 	partial class ConnectWindow : Form
 	{
-		IPerforceSettings DefaultPerforceSettings;
-		string? ServerAndPortOverride;
-		string? UserNameOverride;
-		IServiceProvider ServiceProvider;
+		IPerforceSettings _defaultPerforceSettings;
+		string? _serverAndPortOverride;
+		string? _userNameOverride;
+		IServiceProvider _serviceProvider;
 
-		private ConnectWindow(IPerforceSettings DefaultPerforceSettings, string? ServerAndPortOverride, string? UserNameOverride, IServiceProvider ServiceProvider)
+		private ConnectWindow(IPerforceSettings defaultPerforceSettings, string? serverAndPortOverride, string? userNameOverride, IServiceProvider serviceProvider)
 		{
 			InitializeComponent();
 
-			this.DefaultPerforceSettings = DefaultPerforceSettings;
-			this.ServiceProvider = ServiceProvider;
+			this._defaultPerforceSettings = defaultPerforceSettings;
+			this._serviceProvider = serviceProvider;
 
-			if(!String.IsNullOrWhiteSpace(ServerAndPortOverride))
+			if(!String.IsNullOrWhiteSpace(serverAndPortOverride))
 			{
-				this.ServerAndPortOverride = ServerAndPortOverride.Trim();
+				this._serverAndPortOverride = serverAndPortOverride.Trim();
 			}
-			if(!String.IsNullOrEmpty(UserNameOverride))
+			if(!String.IsNullOrEmpty(userNameOverride))
 			{
-				this.UserNameOverride = UserNameOverride.Trim();
+				this._userNameOverride = userNameOverride.Trim();
 			}
 
-			ServerAndPortTextBox.CueBanner = DefaultPerforceSettings.ServerAndPort;
-			ServerAndPortTextBox.Text = this.ServerAndPortOverride ?? DefaultPerforceSettings.ServerAndPort;
-			UserNameTextBox.CueBanner = DefaultPerforceSettings.UserName;
-			UserNameTextBox.Text = this.UserNameOverride ?? DefaultPerforceSettings.UserName;
-			UseDefaultConnectionSettings.Checked = this.ServerAndPortOverride == null && this.UserNameOverride == null;
+			ServerAndPortTextBox.CueBanner = defaultPerforceSettings.ServerAndPort;
+			ServerAndPortTextBox.Text = this._serverAndPortOverride ?? defaultPerforceSettings.ServerAndPort;
+			UserNameTextBox.CueBanner = defaultPerforceSettings.UserName;
+			UserNameTextBox.Text = this._userNameOverride ?? defaultPerforceSettings.UserName;
+			UseDefaultConnectionSettings.Checked = this._serverAndPortOverride == null && this._userNameOverride == null;
 
 			UpdateEnabledControls();
 		}
 
-		public static bool ShowModal(IWin32Window Owner, IPerforceSettings DefaultSettings, ref string? ServerAndPortOverride, ref string? UserNameOverride, IServiceProvider ServiceProvider)
+		public static bool ShowModal(IWin32Window owner, IPerforceSettings defaultSettings, ref string? serverAndPortOverride, ref string? userNameOverride, IServiceProvider serviceProvider)
 		{
-			ConnectWindow Connect = new ConnectWindow(DefaultSettings, ServerAndPortOverride, UserNameOverride, ServiceProvider);
-			if(Connect.ShowDialog(Owner) == DialogResult.OK)
+			ConnectWindow connect = new ConnectWindow(defaultSettings, serverAndPortOverride, userNameOverride, serviceProvider);
+			if(connect.ShowDialog(owner) == DialogResult.OK)
 			{
-				ServerAndPortOverride = Connect.ServerAndPortOverride;
-				UserNameOverride = Connect.UserNameOverride;
+				serverAndPortOverride = connect._serverAndPortOverride;
+				userNameOverride = connect._userNameOverride;
 				return true;
 			}
 			return false;
@@ -69,21 +69,21 @@ namespace UnrealGameSync
 		{
 			if(UseDefaultConnectionSettings.Checked)
 			{
-				ServerAndPortOverride = null;
-				UserNameOverride = null;
+				_serverAndPortOverride = null;
+				_userNameOverride = null;
 			}
 			else
 			{
-				ServerAndPortOverride = ServerAndPortTextBox.Text.Trim();
-				if(ServerAndPortOverride.Length == 0)
+				_serverAndPortOverride = ServerAndPortTextBox.Text.Trim();
+				if(_serverAndPortOverride.Length == 0)
 				{
-					ServerAndPortOverride = DefaultPerforceSettings.ServerAndPort;
+					_serverAndPortOverride = _defaultPerforceSettings.ServerAndPort;
 				}
 
-				UserNameOverride = UserNameTextBox.Text.Trim();
-				if(UserNameOverride.Length == 0)
+				_userNameOverride = UserNameTextBox.Text.Trim();
+				if(_userNameOverride.Length == 0)
 				{
-					UserNameOverride = DefaultPerforceSettings.UserName;
+					_userNameOverride = _defaultPerforceSettings.UserName;
 				}
 			}
 
@@ -98,21 +98,21 @@ namespace UnrealGameSync
 
 		private void UpdateEnabledControls()
 		{
-			bool bUseDefaultSettings = UseDefaultConnectionSettings.Checked;
+			bool useDefaultSettings = UseDefaultConnectionSettings.Checked;
 
-			ServerAndPortLabel.Enabled = !bUseDefaultSettings;
-			ServerAndPortTextBox.Enabled = !bUseDefaultSettings;
+			ServerAndPortLabel.Enabled = !useDefaultSettings;
+			ServerAndPortTextBox.Enabled = !useDefaultSettings;
 
-			UserNameLabel.Enabled = !bUseDefaultSettings;
-			UserNameTextBox.Enabled = !bUseDefaultSettings;
+			UserNameLabel.Enabled = !useDefaultSettings;
+			UserNameTextBox.Enabled = !useDefaultSettings;
 		}
 
 		private void BrowseUserBtn_Click(object sender, EventArgs e)
 		{
-			string? NewUserName;
-			if(SelectUserWindow.ShowModal(this, new PerforceSettings(DefaultPerforceSettings) { UserName = UserNameTextBox.Text, ServerAndPort = ServerAndPortTextBox.Text }, ServiceProvider, out NewUserName))
+			string? newUserName;
+			if(SelectUserWindow.ShowModal(this, new PerforceSettings(_defaultPerforceSettings) { UserName = UserNameTextBox.Text, ServerAndPort = ServerAndPortTextBox.Text }, _serviceProvider, out newUserName))
 			{
-				UserNameTextBox.Text = NewUserName;
+				UserNameTextBox.Text = newUserName;
 			}
 		}
 	}
diff --git a/Engine/Source/Programs/UnrealGameSync/UnrealGameSync/Forms/DeleteWindow.cs b/Engine/Source/Programs/UnrealGameSync/UnrealGameSync/Forms/DeleteWindow.cs
--- a/Engine/Source/Programs/UnrealGameSync/UnrealGameSync/Forms/DeleteWindow.cs
+++ b/Engine/Source/Programs/UnrealGameSync/UnrealGameSync/Forms/DeleteWindow.cs
@@ -15,44 +15,44 @@ namespace UnrealGameSync
 {
 	public partial class DeleteWindow : Form
 	{
-		Dictionary<string, bool> FilesToDelete;
+		Dictionary<string, bool> _filesToDelete;
 
-		public DeleteWindow(Dictionary<string, bool> InFilesToDelete)
+		public DeleteWindow(Dictionary<string, bool> inFilesToDelete)
 		{
 			InitializeComponent();
 
-			FilesToDelete = InFilesToDelete;
+			_filesToDelete = inFilesToDelete;
 
-			foreach(KeyValuePair<string, bool> FileToDelete in FilesToDelete)
+			foreach(KeyValuePair<string, bool> fileToDelete in _filesToDelete)
 			{
-				ListViewItem Item = new ListViewItem(FileToDelete.Key);
-				Item.Tag = FileToDelete.Key;
-				Item.Checked = FileToDelete.Value;
-				FileList.Items.Add(Item);
+				ListViewItem item = new ListViewItem(fileToDelete.Key);
+				item.Tag = fileToDelete.Key;
+				item.Checked = fileToDelete.Value;
+				FileList.Items.Add(item);
 			}
 		}
 
 		private void UncheckAll_Click(object sender, EventArgs e)
 		{
-			foreach(ListViewItem? Item in FileList.Items)
+			foreach(ListViewItem? item in FileList.Items)
 			{
-				Item!.Checked = false;
+				item!.Checked = false;
 			}
 		}
 
 		private void CheckAll_Click(object sender, EventArgs e)
 		{
-			foreach(ListViewItem? Item in FileList.Items)
+			foreach(ListViewItem? item in FileList.Items)
 			{
-				Item!.Checked = true;
+				item!.Checked = true;
 			}
 		}
 
 		private void ContinueButton_Click(object sender, EventArgs e)
 		{
-			foreach(ListViewItem? Item in FileList.Items)
+			foreach(ListViewItem? item in FileList.Items)
 			{
-				FilesToDelete[(string)Item!.Tag] = Item.Checked;
+				_filesToDelete[(string)item!.Tag] = item.Checked;
 			}
 		}
 	}
diff --git a/Engine/Source/Programs/UnrealGameSync/UnrealGameSync/Forms/DiagnosticsWindow.cs b/Engine/Source/Programs/UnrealGameSync/UnrealGameSync/Forms/DiagnosticsWindow.cs
--- a/Engine/Source/Programs/UnrealGameSync/UnrealGameSync/Forms/DiagnosticsWindow.cs
+++ b/Engine/Source/Programs/UnrealGameSync/UnrealGameSync/Forms/DiagnosticsWindow.cs
@@ -13,94 +13,94 @@ namespace UnrealGameSync
 {
 	public partial class DiagnosticsWindow : Form
 	{
-		DirectoryReference AppDataFolder;
-		DirectoryReference WorkspaceDataFolder;
-		List<FileReference> ExtraFiles;
+		DirectoryReference _appDataFolder;
+		DirectoryReference _workspaceDataFolder;
+		List<FileReference> _extraFiles;
 
-		public DiagnosticsWindow(DirectoryReference InAppDataFolder, DirectoryReference InWorkspaceDataFolder, string InDiagnosticsText, IEnumerable<FileReference> InExtraFiles)
+		public DiagnosticsWindow(DirectoryReference inAppDataFolder, DirectoryReference inWorkspaceDataFolder, string inDiagnosticsText, IEnumerable<FileReference> inExtraFiles)
 		{
 			InitializeComponent();
 
-			AppDataFolder = InAppDataFolder;
-			WorkspaceDataFolder = InWorkspaceDataFolder;
+			_appDataFolder = inAppDataFolder;
+			_workspaceDataFolder = inWorkspaceDataFolder;
 
-			DiagnosticsTextBox.Text = InDiagnosticsText.Replace("\n", "\r\n");
-			ExtraFiles = InExtraFiles.ToList();
+			DiagnosticsTextBox.Text = inDiagnosticsText.Replace("\n", "\r\n");
+			_extraFiles = inExtraFiles.ToList();
 		}
 
 		private void ViewApplicationDataButton_Click(object sender, EventArgs e)
 		{
-			Process.Start("explorer.exe", AppDataFolder.FullName);
+			Process.Start("explorer.exe", _appDataFolder.FullName);
 		}
 
 		private void ViewWorkspaceDataButton_Click(object sender, EventArgs e)
 		{
-			Process.Start("explorer.exe", WorkspaceDataFolder.FullName);
+			Process.Start("explorer.exe", _workspaceDataFolder.FullName);
 		}
 
 		private void SaveButton_Click(object sender, EventArgs e)
 		{
-			SaveFileDialog Dialog = new SaveFileDialog();
-			Dialog.Filter = "Zip Files (*.zip)|*.zip|AllFiles (*.*)|*.*";
-			Dialog.InitialDirectory = Environment.GetFolderPath(Environment.SpecialFolder.Desktop);
-			Dialog.FileName = Path.Combine(Dialog.InitialDirectory, "UGS-Diagnostics.zip");
-			if(Dialog.ShowDialog() == DialogResult.OK)
+			SaveFileDialog dialog = new SaveFileDialog();
+			dialog.Filter = "Zip Files (*.zip)|*.zip|AllFiles (*.*)|*.*";
+			dialog.InitialDirectory = Environment.GetFolderPath(Environment.SpecialFolder.Desktop);
+			dialog.FileName = Path.Combine(dialog.InitialDirectory, "UGS-Diagnostics.zip");
+			if(dialog.ShowDialog() == DialogResult.OK)
 			{
-				FileReference DiagnosticsFileName = FileReference.Combine(AppDataFolder, "Diagnostics.txt");
+				FileReference diagnosticsFileName = FileReference.Combine(_appDataFolder, "Diagnostics.txt");
 				try
 				{
-					FileReference.WriteAllLines(DiagnosticsFileName, DiagnosticsTextBox.Lines);
+					FileReference.WriteAllLines(diagnosticsFileName, DiagnosticsTextBox.Lines);
 				}
-				catch(Exception Ex)
+				catch(Exception ex)
 				{
-					MessageBox.Show(String.Format("Couldn't write to '{0}'\n\n{1}", DiagnosticsFileName, Ex.ToString()));
+					MessageBox.Show(String.Format("Couldn't write to '{0}'\n\n{1}", diagnosticsFileName, ex.ToString()));
 					return;
 				}
 
-				string ZipFileName = Dialog.FileName;
+				string zipFileName = dialog.FileName;
 				try
 				{
-					using (ZipArchive Zip = new ZipArchive(File.OpenWrite(ZipFileName), ZipArchiveMode.Create))
+					using (ZipArchive zip = new ZipArchive(File.OpenWrite(zipFileName), ZipArchiveMode.Create))
 					{
-						AddFilesToZip(Zip, AppDataFolder, "App/");
-						AddFilesToZip(Zip, WorkspaceDataFolder, "Workspace/");
+						AddFilesToZip(zip, _appDataFolder, "App/");
+						AddFilesToZip(zip, _workspaceDataFolder, "Workspace/");
 
-						foreach (FileReference ExtraFile in ExtraFiles)
+						foreach (FileReference extraFile in _extraFiles)
 						{
-							if(FileReference.Exists(ExtraFile))
+							if(FileReference.Exists(extraFile))
 							{
-								using (FileStream InputStream = FileReference.Open(ExtraFile, FileMode.Open, FileAccess.Read, FileShare.ReadWrite))
+								using (FileStream inputStream = FileReference.Open(extraFile, FileMode.Open, FileAccess.Read, FileShare.ReadWrite))
 								{
-									ZipArchiveEntry Entry = Zip.CreateEntry(ExtraFile.FullName.Replace(":", "").Replace('\\', '/'));
-									using (Stream OutputStream = Entry.Open())
+									ZipArchiveEntry entry = zip.CreateEntry(extraFile.FullName.Replace(":", "").Replace('\\', '/'));
+									using (Stream outputStream = entry.Open())
 									{
-										InputStream.CopyTo(OutputStream);
+										inputStream.CopyTo(outputStream);
 									}
 								}
 							}
 						}
 					}
 				}
-				catch(Exception Ex)
+				catch(Exception ex)
 				{
-					MessageBox.Show(String.Format("Couldn't save '{0}'\n\n{1}", ZipFileName, Ex.ToString()));
+					MessageBox.Show(String.Format("Couldn't save '{0}'\n\n{1}", zipFileName, ex.ToString()));
 					return;
 				}
 			}
 		}
 
-		private static void AddFilesToZip(ZipArchive Zip, DirectoryReference DataFolder, string RelativeDir)
+		private static void AddFilesToZip(ZipArchive zip, DirectoryReference dataFolder, string relativeDir)
 		{
-			foreach (FileReference FileName in DirectoryReference.EnumerateFiles(DataFolder))
+			foreach (FileReference fileName in DirectoryReference.EnumerateFiles(dataFolder))
 			{
-				if (!FileName.HasExtension(".exe") && !FileName.HasExtension(".dll"))
+				if (!fileName.HasExtension(".exe") && !fileName.HasExtension(".dll"))
 				{
-					using (FileStream InputStream = FileReference.Open(FileName, FileMode.Open, FileAccess.Read, FileShare.ReadWrite))
+					using (FileStream inputStream = FileReference.Open(fileName, FileMode.Open, FileAccess.Read, FileShare.ReadWrite))
 					{
-						ZipArchiveEntry Entry = Zip.CreateEntry(RelativeDir + FileName.MakeRelativeTo(DataFolder).Replace('\\', '/'));
-						using (Stream OutputStream = Entry.Open())
+						ZipArchiveEntry entry = zip.CreateEntry(relativeDir + fileName.MakeRelativeTo(dataFolder).Replace('\\', '/'));
+						using (Stream outputStream = entry.Open())
 						{
-							InputStream.CopyTo(OutputStream);
+							inputStream.CopyTo(outputStream);
 						}
 					}
 				}
