UGS: Reformat to standard C# style conventions.

diff --git a/Engine/Source/Programs/UnrealGameSync/UnrealGameSync/Forms/IssueAlertWindow.cs b/Engine/Source/Programs/UnrealGameSync/UnrealGameSync/Forms/IssueAlertWindow.cs
--- a/Engine/Source/Programs/UnrealGameSync/UnrealGameSync/Forms/IssueAlertWindow.cs
+++ b/Engine/Source/Programs/UnrealGameSync/UnrealGameSync/Forms/IssueAlertWindow.cs
@@ -12,29 +12,29 @@ namespace UnrealGameSync
 {
 	partial class IssueAlertWindow : Form
 	{
-		static readonly IntPtr HWND_TOPMOST = new IntPtr(-1);
+		static readonly IntPtr HwndTopmost = new IntPtr(-1);
 
-		const uint SWP_NOMOVE = 0x0002;
-		const uint SWP_NOSIZE = 0x0001;
-		const uint SWP_NOACTIVATE = 0x0010;
+		const uint SwpNomove = 0x0002;
+		const uint SwpNosize = 0x0001;
+		const uint SwpNoactivate = 0x0010;
 
 		[DllImport("user32.dll", SetLastError=true)]
-		static extern bool SetWindowPos(IntPtr hWnd, IntPtr hWndInsertAfter, int X, int Y, int Width, int Height, uint Flags);
+		static extern bool SetWindowPos(IntPtr hWnd, IntPtr hWndInsertAfter, int x, int y, int width, int height, uint flags);
 
 		public readonly IssueMonitor IssueMonitor;
 
-		public bool bIsWarning;
+		public bool isWarning;
 
-		public bool? bStrongAlert;
+		public bool? strongAlert;
 
-		public IssueAlertWindow(IssueMonitor IssueMonitor, IssueData Issue, IssueAlertReason Reason)
+		public IssueAlertWindow(IssueMonitor issueMonitor, IssueData issue, IssueAlertReason reason)
 		{
-			this.IssueMonitor = IssueMonitor;
-			this.Issue = Issue;
+			this.IssueMonitor = issueMonitor;
+			this.Issue = issue;
 
 			InitializeComponent();
 
-			SetIssue(Issue, Reason);
+			SetIssue(issue, reason);
 		}
 
 		public IssueData Issue
@@ -53,107 +53,107 @@ namespace UnrealGameSync
 		{
 			base.OnPaintBackground(e);
 
-			if(bStrongAlert ?? false)
+			if(strongAlert ?? false)
 			{
-				Color StripeColor = bIsWarning? Color.FromArgb(216, 167, 64) : Color.FromArgb(200, 74, 49);//214, 69, 64);
-				using (Brush StripeBrush = new SolidBrush(StripeColor))
+				Color stripeColor = isWarning? Color.FromArgb(216, 167, 64) : Color.FromArgb(200, 74, 49);//214, 69, 64);
+				using (Brush stripeBrush = new SolidBrush(stripeColor))
 				{
-					e.Graphics.FillRectangle(StripeBrush, 0, 0, Bounds.Width, Bounds.Height);
+					e.Graphics.FillRectangle(stripeBrush, 0, 0, Bounds.Width, Bounds.Height);
 				}
-				using (Pen Pen = new Pen(Color.FromArgb(255, 255, 255)))
+				using (Pen pen = new Pen(Color.FromArgb(255, 255, 255)))
 				{
-					e.Graphics.DrawRectangle(Pen, 0, 0, Bounds.Width - 1, Bounds.Height - 1);
+					e.Graphics.DrawRectangle(pen, 0, 0, Bounds.Width - 1, Bounds.Height - 1);
 				}
 			}
 			else
 			{
-				Color StripeColor = bIsWarning? Color.FromArgb(216, 167, 64) : Color.FromArgb(214, 69, 64);
+				Color stripeColor = isWarning? Color.FromArgb(216, 167, 64) : Color.FromArgb(214, 69, 64);
 
-				Color BackgroundColor = Color.FromArgb(241, 236, 236);
-				using(Brush SolidBrush = new SolidBrush(BackgroundColor))
+				Color backgroundColor = Color.FromArgb(241, 236, 236);
+				using(Brush solidBrush = new SolidBrush(backgroundColor))
 				{
 	//				e.Graphics.FillRectangle(SolidBrush, 0, 0, Bounds.Width - 1, Bounds.Height - 1);
 				}
 
-				Color BorderColor = StripeColor; 
-				using(Pen Pen = new Pen(Color.FromArgb(128, 128, 128)))
+				Color borderColor = stripeColor; 
+				using(Pen pen = new Pen(Color.FromArgb(128, 128, 128)))
 				{
-					e.Graphics.DrawRectangle(Pen, 0, 0, Bounds.Width - 1, Bounds.Height - 1);
+					e.Graphics.DrawRectangle(pen, 0, 0, Bounds.Width - 1, Bounds.Height - 1);
 				}
 
-				using(Brush StripeBrush = new SolidBrush(StripeColor))
+				using(Brush stripeBrush = new SolidBrush(stripeColor))
 				{
-					e.Graphics.FillRectangle(StripeBrush, 0, 0, /*6*/10 * e.Graphics.DpiX / 96.0f, Bounds.Height);
+					e.Graphics.FillRectangle(stripeBrush, 0, 0, /*6*/10 * e.Graphics.DpiX / 96.0f, Bounds.Height);
 				}
 			}
 		}
 
-		public void SetIssue(IssueData NewIssue, IssueAlertReason NewReason)
+		public void SetIssue(IssueData newIssue, IssueAlertReason newReason)
 		{
-			bool bNewStrongAlert = false;
+			bool newStrongAlert = false;
 
-			StringBuilder OwnerTextBuilder = new StringBuilder();
-			if(NewIssue.Owner == null)
+			StringBuilder ownerTextBuilder = new StringBuilder();
+			if(newIssue.Owner == null)
 			{
-				OwnerTextBuilder.Append("Currently unassigned.");
+				ownerTextBuilder.Append("Currently unassigned.");
 			}
 			else
 			{
-				if(String.Compare(NewIssue.Owner, IssueMonitor.UserName, StringComparison.OrdinalIgnoreCase) == 0)
+				if(String.Compare(newIssue.Owner, IssueMonitor.UserName, StringComparison.OrdinalIgnoreCase) == 0)
 				{
-					if(NewIssue.NominatedBy != null)
+					if(newIssue.NominatedBy != null)
 					{
-						OwnerTextBuilder.AppendFormat("You have been nominated to fix this issue by {0}.", Utility.FormatUserName(NewIssue.NominatedBy));
+						ownerTextBuilder.AppendFormat("You have been nominated to fix this issue by {0}.", Utility.FormatUserName(newIssue.NominatedBy));
 					}
 					else
 					{
-						OwnerTextBuilder.AppendFormat("Assigned to {0}.", Utility.FormatUserName(NewIssue.Owner));
+						ownerTextBuilder.AppendFormat("Assigned to {0}.", Utility.FormatUserName(newIssue.Owner));
 					}
-					bNewStrongAlert = true;
+					newStrongAlert = true;
 				}
 				else
 				{
-					OwnerTextBuilder.AppendFormat("Assigned to {0}", Utility.FormatUserName(NewIssue.Owner));
-					if(NewIssue.NominatedBy != null)
+					ownerTextBuilder.AppendFormat("Assigned to {0}", Utility.FormatUserName(newIssue.Owner));
+					if(newIssue.NominatedBy != null)
 					{
-						OwnerTextBuilder.AppendFormat(" by {0}", Utility.FormatUserName(NewIssue.NominatedBy));
+						ownerTextBuilder.AppendFormat(" by {0}", Utility.FormatUserName(newIssue.NominatedBy));
 					}
-					if(!NewIssue.AcknowledgedAt.HasValue && (NewReason & IssueAlertReason.UnacknowledgedTimer) != 0)
+					if(!newIssue.AcknowledgedAt.HasValue && (newReason & IssueAlertReason.UnacknowledgedTimer) != 0)
 					{
-						OwnerTextBuilder.Append(" (not acknowledged)");
+						ownerTextBuilder.Append(" (not acknowledged)");
 					}
-					OwnerTextBuilder.Append(".");
+					ownerTextBuilder.Append(".");
 				}
 			}
-			OwnerTextBuilder.AppendFormat(" Open for {0}.", Utility.FormatDurationMinutes((int)(NewIssue.RetrievedAt - NewIssue.CreatedAt).TotalMinutes));
-			string OwnerText = OwnerTextBuilder.ToString();
+			ownerTextBuilder.AppendFormat(" Open for {0}.", Utility.FormatDurationMinutes((int)(newIssue.RetrievedAt - newIssue.CreatedAt).TotalMinutes));
+			string ownerText = ownerTextBuilder.ToString();
 
-			bool bNewIsWarning = Issue.bIsWarning;
+			bool newIsWarning = Issue.IsWarning;
 
-			Issue = NewIssue;
+			Issue = newIssue;
 
-			string Summary = NewIssue.Summary;
+			string summary = newIssue.Summary;
 
-			int MaxLength = 128;
-			if(Summary.Length > MaxLength)
+			int maxLength = 128;
+			if(summary.Length > maxLength)
 			{
-				Summary = Summary.Substring(0, MaxLength).TrimEnd() + "...";
+				summary = summary.Substring(0, maxLength).TrimEnd() + "...";
 			}
 
-			if(Summary != SummaryLabel.Text || OwnerText != OwnerLabel.Text || Reason != NewReason || bIsWarning != bNewIsWarning || bStrongAlert != bNewStrongAlert)
+			if(summary != SummaryLabel.Text || ownerText != OwnerLabel.Text || Reason != newReason || isWarning != newIsWarning || strongAlert != newStrongAlert)
 			{
-				Rectangle PrevBounds = Bounds;
+				Rectangle prevBounds = Bounds;
 				SuspendLayout();
 
-				SummaryLabel.Text = Summary;
-				OwnerLabel.Text = OwnerText;
+				SummaryLabel.Text = summary;
+				OwnerLabel.Text = ownerText;
 
-				bool bForceUpdateButtons = false;
-				if(bStrongAlert != bNewStrongAlert)
+				bool forceUpdateButtons = false;
+				if(strongAlert != newStrongAlert)
 				{
-					bStrongAlert = bNewStrongAlert;
+					strongAlert = newStrongAlert;
 
-					if(bNewStrongAlert)
+					if(newStrongAlert)
 					{
 						SummaryLabel.ForeColor = Color.FromArgb(255, 255, 255);
 						SummaryLabel.LinkColor = Color.FromArgb(255, 255, 255);
@@ -172,48 +172,48 @@ namespace UnrealGameSync
 						LatestBuildLinkLabel.LinkColor = Color.FromArgb(16, 102, 192);
 					}
 
-					bForceUpdateButtons = true;
+					forceUpdateButtons = true;
 				}
 
-				if (bIsWarning != bNewIsWarning)
+				if (isWarning != newIsWarning)
 				{
-					bIsWarning = bNewIsWarning;
+					isWarning = newIsWarning;
 				}
 
-				if(Reason != NewReason || bForceUpdateButtons)
+				if(Reason != newReason || forceUpdateButtons)
 				{
-					Reason = NewReason;
+					Reason = newReason;
 
-					List<Button> Buttons = new List<Button>();
-					Buttons.Add(DetailsBtn);
-					if((NewReason & IssueAlertReason.Owner) != 0)
+					List<Button> buttons = new List<Button>();
+					buttons.Add(DetailsBtn);
+					if((newReason & IssueAlertReason.Owner) != 0)
 					{
 						AcceptBtn.Text = "Acknowledge";
-						Buttons.Add(AcceptBtn);
+						buttons.Add(AcceptBtn);
 					}
-					else if((NewReason & IssueAlertReason.Normal) != 0)
+					else if((newReason & IssueAlertReason.Normal) != 0)
 					{
 						AcceptBtn.Text = "Will Fix";
-						Buttons.Add(AcceptBtn);
+						buttons.Add(AcceptBtn);
 						DeclineBtn.Text = "Not Me";
-						Buttons.Add(DeclineBtn);
+						buttons.Add(DeclineBtn);
 					}
 					else
 					{
 						DeclineBtn.Text = "Dismiss";
-						Buttons.Add(DeclineBtn);
+						buttons.Add(DeclineBtn);
 					}
 
-					tableLayoutPanel3.ColumnCount = Buttons.Count;
+					tableLayoutPanel3.ColumnCount = buttons.Count;
 					tableLayoutPanel3.Controls.Clear();
-					for(int Idx = 0; Idx < Buttons.Count; Idx++)
+					for(int idx = 0; idx < buttons.Count; idx++)
 					{
-						tableLayoutPanel3.Controls.Add(Buttons[Idx], Idx, 0);
+						tableLayoutPanel3.Controls.Add(buttons[idx], idx, 0);
 					}
 				}
 
 				ResumeLayout(true);
-				Location = new Point(PrevBounds.Right - Bounds.Width, PrevBounds.Y);
+				Location = new Point(prevBounds.Right - Bounds.Width, prevBounds.Y);
 				Invalidate();
 			}
 		}
@@ -223,7 +223,7 @@ namespace UnrealGameSync
 			base.OnHandleCreated(e);
 
 			// Manually make the window topmost, so that we can pass SWP_NOACTIVATE
-			SetWindowPos(Handle, HWND_TOPMOST, 0, 0, 0, 0, SWP_NOMOVE | SWP_NOSIZE | SWP_NOACTIVATE);
+			SetWindowPos(Handle, HwndTopmost, 0, 0, 0, 0, SwpNomove | SwpNosize | SwpNoactivate);
 		}
 
 		protected override bool ShowWithoutActivation 
@@ -243,8 +243,8 @@ namespace UnrealGameSync
 
 		public void LaunchUrl()
 		{
-			string Url = Issue.BuildUrl;
-			if(String.IsNullOrEmpty(Url))
+			string url = Issue.BuildUrl;
+			if(String.IsNullOrEmpty(url))
 			{
 				MessageBox.Show("No additional information is available");
 			}
@@ -252,11 +252,11 @@ namespace UnrealGameSync
 			{
 				try
 				{
-					Utility.OpenUrl(Url);
+					Utility.OpenUrl(url);
 				}
-				catch (Exception Ex)
+				catch (Exception ex)
 				{
-					MessageBox.Show(String.Format("Unable to launch '{0}' (Error: {1})", Url, Ex.Message));
+					MessageBox.Show(String.Format("Unable to launch '{0}' (Error: {1})", url, ex.Message));
 				}
 			}
 		}
diff --git a/Engine/Source/Programs/UnrealGameSync/UnrealGameSync/Forms/IssueBrowserWindow.cs b/Engine/Source/Programs/UnrealGameSync/UnrealGameSync/Forms/IssueBrowserWindow.cs
--- a/Engine/Source/Programs/UnrealGameSync/UnrealGameSync/Forms/IssueBrowserWindow.cs
+++ b/Engine/Source/Programs/UnrealGameSync/UnrealGameSync/Forms/IssueBrowserWindow.cs
@@ -21,51 +21,51 @@ namespace UnrealGameSync
 {
 	partial class IssueBrowserWindow : Form
 	{
-		IssueMonitor IssueMonitor;
-		IPerforceSettings PerforceSettings;
-		TimeSpan? ServerTimeOffset;
-		IServiceProvider ServiceProvider;
-		ILogger Logger;
-		string? CurrentStream;
-		int MaxResults = 0;
-		int PendingMaxResults = 0;
-		string? FilterName;
-		Dictionary<string, Func<IssueData, bool>> CustomFilters = new Dictionary<string, Func<IssueData, bool>>();
-		List<IssueData> Issues = new List<IssueData>();
-		SynchronizationContext MainThreadSynchronizationContext;
-		Task? BackgroundTask;
-		bool bDisposed;
-
-		public IssueBrowserWindow(IssueMonitor IssueMonitor, IPerforceSettings PerforceSettings, TimeSpan? ServerTimeOffset, IServiceProvider ServiceProvider, string? CurrentStream, Dictionary<string, Func<IssueData, bool>> CustomFilters, string? FilterName)
+		IssueMonitor _issueMonitor;
+		IPerforceSettings _perforceSettings;
+		TimeSpan? _serverTimeOffset;
+		IServiceProvider _serviceProvider;
+		ILogger _logger;
+		string? _currentStream;
+		int _maxResults = 0;
+		int _pendingMaxResults = 0;
+		string? _filterName;
+		Dictionary<string, Func<IssueData, bool>> _customFilters = new Dictionary<string, Func<IssueData, bool>>();
+		List<IssueData> _issues = new List<IssueData>();
+		SynchronizationContext _mainThreadSynchronizationContext;
+		Task? _backgroundTask;
+		bool _disposed;
+
+		public IssueBrowserWindow(IssueMonitor issueMonitor, IPerforceSettings perforceSettings, TimeSpan? serverTimeOffset, IServiceProvider serviceProvider, string? currentStream, Dictionary<string, Func<IssueData, bool>> customFilters, string? filterName)
 		{
-			this.IssueMonitor = IssueMonitor;
-			this.PerforceSettings = PerforceSettings;
-			this.ServerTimeOffset = ServerTimeOffset;
-			this.ServiceProvider = ServiceProvider;
-			this.Logger = ServiceProvider.GetRequiredService<ILogger<IssueBrowserWindow>>();
-			this.CurrentStream = CurrentStream;
-			this.FilterName = FilterName;
-			this.CustomFilters = CustomFilters;
-			this.MainThreadSynchronizationContext = SynchronizationContext.Current!;
-
-			IssueMonitor.AddRef();
+			this._issueMonitor = issueMonitor;
+			this._perforceSettings = perforceSettings;
+			this._serverTimeOffset = serverTimeOffset;
+			this._serviceProvider = serviceProvider;
+			this._logger = serviceProvider.GetRequiredService<ILogger<IssueBrowserWindow>>();
+			this._currentStream = currentStream;
+			this._filterName = filterName;
+			this._customFilters = customFilters;
+			this._mainThreadSynchronizationContext = SynchronizationContext.Current!;
+
+			issueMonitor.AddRef();
 
 			InitializeComponent();
 
-			using (Graphics Graphics = Graphics.FromHwnd(IntPtr.Zero))
+			using (Graphics graphics = Graphics.FromHwnd(IntPtr.Zero))
 			{
-				float DpiScaleX = Graphics.DpiX / 96.0f;
-				foreach (ColumnHeader? Column in IssueListView.Columns)
+				float dpiScaleX = graphics.DpiX / 96.0f;
+				foreach (ColumnHeader? column in IssueListView.Columns)
 				{
-					if (Column != null)
+					if (column != null)
 					{
-						Column.Width = (int)(Column.Width * DpiScaleX);
+						column.Width = (int)(column.Width * dpiScaleX);
 					}
 				}
 			}
 
-			System.Reflection.PropertyInfo DoubleBufferedProperty = typeof(Control).GetProperty("DoubleBuffered", System.Reflection.BindingFlags.NonPublic | System.Reflection.BindingFlags.Instance)!;
-			DoubleBufferedProperty.SetValue(IssueListView, true, null);
+			System.Reflection.PropertyInfo doubleBufferedProperty = typeof(Control).GetProperty("DoubleBuffered", System.Reflection.BindingFlags.NonPublic | System.Reflection.BindingFlags.Instance)!;
+			doubleBufferedProperty.SetValue(IssueListView, true, null);
 		}
 
 		protected override void Dispose(bool disposing)
@@ -75,9 +75,9 @@ namespace UnrealGameSync
 				components.Dispose();
 			}
 
-			IssueMonitor.Release();
+			_issueMonitor.Release();
 
-			bDisposed = true;
+			_disposed = true;
 
 			base.Dispose(disposing);
 		}
@@ -89,59 +89,59 @@ namespace UnrealGameSync
 
 		private void FetchMoreResults()
 		{
-			PendingMaxResults = MaxResults + 100;
+			_pendingMaxResults = _maxResults + 100;
 			CheckToStartBackgroundThread();
 		}
 
-		async Task FetchIssues(int NewMaxResults, CancellationToken CancellationToken)
+		async Task FetchIssues(int newMaxResults, CancellationToken cancellationToken)
 		{
 			try
 			{
-				SortedDictionary<int, IssueData> NewSortedIssues = new SortedDictionary<int, IssueData>();
+				SortedDictionary<int, IssueData> newSortedIssues = new SortedDictionary<int, IssueData>();
 
-				List<IssueData> OpenIssues = await RESTApi.GetAsync<List<IssueData>>($"{IssueMonitor.ApiUrl}/api/issues", CancellationToken);
-				foreach (IssueData OpenIssue in OpenIssues)
+				List<IssueData> openIssues = await RestApi.GetAsync<List<IssueData>>($"{_issueMonitor.ApiUrl}/api/issues", cancellationToken);
+				foreach (IssueData openIssue in openIssues)
 				{
-					NewSortedIssues[(int)OpenIssue.Id] = OpenIssue;
+					newSortedIssues[(int)openIssue.Id] = openIssue;
 				}
 
-				List<IssueData> ResolvedIssues = await RESTApi.GetAsync<List<IssueData>>($"{IssueMonitor.ApiUrl}/api/issues?includeresolved=true&maxresults={NewMaxResults}", CancellationToken);
-				foreach (IssueData ResolvedIssue in ResolvedIssues)
+				List<IssueData> resolvedIssues = await RestApi.GetAsync<List<IssueData>>($"{_issueMonitor.ApiUrl}/api/issues?includeresolved=true&maxresults={newMaxResults}", cancellationToken);
+				foreach (IssueData resolvedIssue in resolvedIssues)
 				{
-					NewSortedIssues[(int)ResolvedIssue.Id] = ResolvedIssue;
+					newSortedIssues[(int)resolvedIssue.Id] = resolvedIssue;
 				}
 
-				List<IssueData> NewIssues = NewSortedIssues.Values.Reverse().ToList();
-				MainThreadSynchronizationContext.Post((o) => { if (!bDisposed) { FetchIssuesSuccess(NewMaxResults, NewIssues); } }, null);
+				List<IssueData> newIssues = newSortedIssues.Values.Reverse().ToList();
+				_mainThreadSynchronizationContext.Post((o) => { if (!_disposed) { FetchIssuesSuccess(newMaxResults, newIssues); } }, null);
 			}
-			catch(Exception Ex)
+			catch(Exception ex)
 			{
-				MainThreadSynchronizationContext.Post((o) => { if (!bDisposed) { FetchIssuesFailure(Ex); } }, null);
+				_mainThreadSynchronizationContext.Post((o) => { if (!_disposed) { FetchIssuesFailure(ex); } }, null);
 			}
 		}
 
-		void FetchIssuesSuccess(int NewMaxResults, List<IssueData> NewIssues)
+		void FetchIssuesSuccess(int newMaxResults, List<IssueData> newIssues)
 		{
-			MaxResults = NewMaxResults;
+			_maxResults = newMaxResults;
 
 			// Update the list of project names
-			Issues = NewIssues;
+			_issues = newIssues;
 			UpdateIssueList();
 
-			BackgroundTask = null;
+			_backgroundTask = null;
 			CheckToStartBackgroundThread();
 		}
 
-		void FetchIssuesFailure(Exception Ex)
+		void FetchIssuesFailure(Exception ex)
 		{
-			StatusLabel.Text = String.Format("Unable to fetch issues ({0})", Ex.Message);
-			BackgroundTask = null;
+			StatusLabel.Text = String.Format("Unable to fetch issues ({0})", ex.Message);
+			_backgroundTask = null;
 			CheckToStartBackgroundThread();
 		}
 
 		void CheckToStartBackgroundThread()
 		{
-			if(PendingMaxResults != MaxResults)
+			if(_pendingMaxResults != _maxResults)
 			{
 				StartBackgroundThread();
 			}
@@ -149,134 +149,134 @@ namespace UnrealGameSync
 
 		void StartBackgroundThread()
 		{
-			if(BackgroundTask == null)
+			if(_backgroundTask == null)
 			{
-				int NewMaxResultsCopy = PendingMaxResults;
-				BackgroundTask = Task.Run(() => FetchIssues(NewMaxResultsCopy, CancellationToken.None));
+				int newMaxResultsCopy = _pendingMaxResults;
+				_backgroundTask = Task.Run(() => FetchIssues(newMaxResultsCopy, CancellationToken.None));
 			}
 		}
 
 		void UpdateIssueList()
 		{
 			// Get the time at midnight
-			DateTime Now = DateTime.Now;
-			DateTime Midnight = (Now - Now.TimeOfDay).ToUniversalTime();
+			DateTime now = DateTime.Now;
+			DateTime midnight = (now - now.TimeOfDay).ToUniversalTime();
 
 			// Get the regex for the selected filter
-			Func<IssueData, bool>? Filter;
-			if (String.IsNullOrEmpty(FilterName))
+			Func<IssueData, bool>? filter;
+			if (String.IsNullOrEmpty(_filterName))
 			{
-				Filter = x => true;
+				filter = x => true;
 			}
-			else if (!CustomFilters.TryGetValue(FilterName, out Filter))
+			else if (!_customFilters.TryGetValue(_filterName, out filter))
 			{
-				Filter = x => x.Streams == null || x.Streams.Any(y => String.Equals(y, FilterName, StringComparison.OrdinalIgnoreCase));
+				filter = x => x.Streams == null || x.Streams.Any(y => String.Equals(y, _filterName, StringComparison.OrdinalIgnoreCase));
 			}
 
 			// Update the table
-			int ItemIdx = 0;
+			int itemIdx = 0;
 			IssueListView.BeginUpdate();
-			foreach(IssueData Issue in Issues)
+			foreach(IssueData issue in _issues)
 			{
-				if(Filter(Issue))
+				if(filter(issue))
 				{
 					for(;;)
 					{
-						if(ItemIdx == IssueListView.Items.Count)
+						if(itemIdx == IssueListView.Items.Count)
 						{
-							IssueList_InsertItem(ItemIdx, Issue, Midnight);
+							IssueList_InsertItem(itemIdx, issue, midnight);
 							break;
 						}
 
-						ListViewItem ExistingItem = IssueListView.Items[ItemIdx];
-						IssueData ExistingIssue = (IssueData)ExistingItem.Tag;
-						if(ExistingIssue == null || ExistingIssue.Id < Issue.Id)
+						ListViewItem existingItem = IssueListView.Items[itemIdx];
+						IssueData existingIssue = (IssueData)existingItem.Tag;
+						if(existingIssue == null || existingIssue.Id < issue.Id)
 						{
-							IssueList_InsertItem(ItemIdx, Issue, Midnight);
+							IssueList_InsertItem(itemIdx, issue, midnight);
 							break;
 						}
-						else if(ExistingIssue.Id == Issue.Id)
+						else if(existingIssue.Id == issue.Id)
 						{
-							IssueList_UpdateItem(ExistingItem, Issue, Midnight);
+							IssueList_UpdateItem(existingItem, issue, midnight);
 							break;
 						}
 						else
 						{
-							IssueListView.Items.RemoveAt(ItemIdx);
+							IssueListView.Items.RemoveAt(itemIdx);
 							continue;
 						}
 					}
-					ItemIdx++;
+					itemIdx++;
 				}
 			}
-			while(ItemIdx < IssueListView.Items.Count)
+			while(itemIdx < IssueListView.Items.Count)
 			{
-				IssueListView.Items.RemoveAt(ItemIdx);
+				IssueListView.Items.RemoveAt(itemIdx);
 			}
 			IssueListView.EndUpdate();
 
 			// Update the maximum number of results
-			string FilterText = "";
-			if(!String.IsNullOrEmpty(FilterName))
+			string filterText = "";
+			if(!String.IsNullOrEmpty(_filterName))
 			{
-				FilterText = String.Format(" matching filter '{0}'", FilterName);
+				filterText = String.Format(" matching filter '{0}'", _filterName);
 			}
-			StatusLabel.Text = (IssueListView.Items.Count == Issues.Count)? String.Format("Showing {0} results{1}.", Issues.Count, FilterText) : String.Format("Showing {0}/{1} results{2}.", IssueListView.Items.Count, Issues.Count, FilterText);
+			StatusLabel.Text = (IssueListView.Items.Count == _issues.Count)? String.Format("Showing {0} results{1}.", _issues.Count, filterText) : String.Format("Showing {0}/{1} results{2}.", IssueListView.Items.Count, _issues.Count, filterText);
 		}
 
-		void IssueList_InsertItem(int ItemIdx, IssueData Issue, DateTime Midnight)
+		void IssueList_InsertItem(int itemIdx, IssueData issue, DateTime midnight)
 		{
-			ListViewItem Item = new ListViewItem("");
-			for(int Idx = 0; Idx < IssueListView.Columns.Count - 1; Idx++)
+			ListViewItem item = new ListViewItem("");
+			for(int idx = 0; idx < IssueListView.Columns.Count - 1; idx++)
 			{
-				Item.SubItems.Add("");
+				item.SubItems.Add("");
 			}
-			Item.Tag = Issue;
-			IssueList_UpdateItem(Item, Issue, Midnight);
-			IssueListView.Items.Insert(ItemIdx, Item);
+			item.Tag = issue;
+			IssueList_UpdateItem(item, issue, midnight);
+			IssueListView.Items.Insert(itemIdx, item);
 		}
 
-		void IssueList_UpdateItem(ListViewItem Item, IssueData Issue, DateTime Midnight)
+		void IssueList_UpdateItem(ListViewItem item, IssueData issue, DateTime midnight)
 		{
-			Item.SubItems[IdHeader.Index].Text = Issue.Id.ToString();
-			Item.SubItems[CreatedHeader.Index].Text = FormatIssueDateTime(Issue.CreatedAt.ToLocalTime(), Midnight);
-			Item.SubItems[ResolvedHeader.Index].Text = Issue.ResolvedAt.HasValue ? FormatIssueDateTime(Issue.ResolvedAt.Value.ToLocalTime(), Midnight) : "Unresolved";
-			Item.SubItems[TimeToFixHeader.Index].Text = Issue.ResolvedAt.HasValue ? Utility.FormatDurationMinutes(Issue.ResolvedAt.Value - Issue.CreatedAt) : "-";
-			Item.SubItems[OwnerHeader.Index].Text = (Issue.Owner == null) ? "-" : Utility.FormatUserName(Issue.Owner);
-			Item.SubItems[DescriptionHeader.Index].Text = Issue.Summary;
+			item.SubItems[IdHeader.Index].Text = issue.Id.ToString();
+			item.SubItems[CreatedHeader.Index].Text = FormatIssueDateTime(issue.CreatedAt.ToLocalTime(), midnight);
+			item.SubItems[ResolvedHeader.Index].Text = issue.ResolvedAt.HasValue ? FormatIssueDateTime(issue.ResolvedAt.Value.ToLocalTime(), midnight) : "Unresolved";
+			item.SubItems[TimeToFixHeader.Index].Text = issue.ResolvedAt.HasValue ? Utility.FormatDurationMinutes(issue.ResolvedAt.Value - issue.CreatedAt) : "-";
+			item.SubItems[OwnerHeader.Index].Text = (issue.Owner == null) ? "-" : Utility.FormatUserName(issue.Owner);
+			item.SubItems[DescriptionHeader.Index].Text = issue.Summary;
 		}
 
-		static string FormatIssueDateTime(DateTime DateTime, DateTime Midnight)
+		static string FormatIssueDateTime(DateTime dateTime, DateTime midnight)
 		{
-			if(DateTime > Midnight)
+			if(dateTime > midnight)
 			{
-				return DateTime.ToShortTimeString();
+				return dateTime.ToShortTimeString();
 			}
 			else
 			{
-				return DateTime.ToShortDateString();
+				return dateTime.ToShortDateString();
 			}
 		}
 
-		static List<IssueBrowserWindow> ExistingWindows = new List<IssueBrowserWindow>();
+		static List<IssueBrowserWindow> _existingWindows = new List<IssueBrowserWindow>();
 
-		public static void Show(Form Owner, IssueMonitor IssueMonitor, IPerforceSettings PerforceSettings, TimeSpan? ServerTimeOffset, IServiceProvider ServiceProvider, string? CurrentStream, Dictionary<string, Func<IssueData, bool>> CustomFilters, string? DefaultFilter)
+		public static void Show(Form owner, IssueMonitor issueMonitor, IPerforceSettings perforceSettings, TimeSpan? serverTimeOffset, IServiceProvider serviceProvider, string? currentStream, Dictionary<string, Func<IssueData, bool>> customFilters, string? defaultFilter)
 		{
-			IssueBrowserWindow Window = ExistingWindows.FirstOrDefault(x => x.IssueMonitor == IssueMonitor);
-			if(Window == null)
+			IssueBrowserWindow window = _existingWindows.FirstOrDefault(x => x._issueMonitor == issueMonitor);
+			if(window == null)
 			{
-				Window = new IssueBrowserWindow(IssueMonitor, PerforceSettings, ServerTimeOffset, ServiceProvider, CurrentStream, CustomFilters, DefaultFilter);
-				Window.Owner = Owner;
-				Window.StartPosition = FormStartPosition.Manual;
-				Window.Location = new Point(Owner.Location.X + (Owner.Width - Window.Width) / 2, Owner.Location.Y + (Owner.Height - Window.Height) / 2);
-				Window.Show(Owner);
-
-				ExistingWindows.Add(Window);
-				Window.FormClosed += (E, S) => ExistingWindows.Remove(Window);
+				window = new IssueBrowserWindow(issueMonitor, perforceSettings, serverTimeOffset, serviceProvider, currentStream, customFilters, defaultFilter);
+				window.Owner = owner;
+				window.StartPosition = FormStartPosition.Manual;
+				window.Location = new Point(owner.Location.X + (owner.Width - window.Width) / 2, owner.Location.Y + (owner.Height - window.Height) / 2);
+				window.Show(owner);
+
+				_existingWindows.Add(window);
+				window.FormClosed += (e, s) => _existingWindows.Remove(window);
 			}
 			else
 			{
-				Window.Activate();
+				window.Activate();
 			}
 		}
 
@@ -293,53 +293,53 @@ namespace UnrealGameSync
 
 		private void IssueListView_MouseDoubleClick(object sender, MouseEventArgs e)
 		{
-			ListViewHitTestInfo HitTest = IssueListView.HitTest(e.Location);
-			if(HitTest.Item != null)
+			ListViewHitTestInfo hitTest = IssueListView.HitTest(e.Location);
+			if(hitTest.Item != null)
 			{
-				IssueData Issue = (IssueData)HitTest.Item.Tag;
-				ShowIssue(Issue);
+				IssueData issue = (IssueData)hitTest.Item.Tag;
+				ShowIssue(issue);
 			}
 		}
 
 		private void FilterBtn_Click(object sender, EventArgs e)
 		{
-			int SeparatorIdx = FilterMenu.Items.IndexOf(FilterMenu_Separator);
-			while(FilterMenu.Items.Count > SeparatorIdx + 1)
+			int separatorIdx = FilterMenu.Items.IndexOf(FilterMenu_Separator);
+			while(FilterMenu.Items.Count > separatorIdx + 1)
 			{
-				FilterMenu.Items.RemoveAt(SeparatorIdx + 1);
+				FilterMenu.Items.RemoveAt(separatorIdx + 1);
 			}
 
-			FilterMenu_ShowAll.Checked = (FilterName == null);
+			FilterMenu_ShowAll.Checked = (_filterName == null);
 
-			if (CustomFilters.Count > 0)
+			if (_customFilters.Count > 0)
 			{
-				foreach (KeyValuePair<string, Func<IssueData, bool>> CustomFilter in CustomFilters.OrderBy(x => x.Key))
+				foreach (KeyValuePair<string, Func<IssueData, bool>> customFilter in _customFilters.OrderBy(x => x.Key))
 				{
-					ToolStripMenuItem Item = new ToolStripMenuItem(CustomFilter.Key);
-					Item.Checked = (FilterName == CustomFilter.Key);
-					Item.Click += (S, E) => { FilterName = CustomFilter.Key; UpdateIssueList(); };
-					FilterMenu.Items.Add(Item);
+					ToolStripMenuItem item = new ToolStripMenuItem(customFilter.Key);
+					item.Checked = (_filterName == customFilter.Key);
+					item.Click += (s, e) => { _filterName = customFilter.Key; UpdateIssueList(); };
+					FilterMenu.Items.Add(item);
 				}
 			}
 
-			HashSet<string> Streams = new HashSet<string>(StringComparer.OrdinalIgnoreCase);
-			foreach (IssueData Issue in Issues)
+			HashSet<string> streams = new HashSet<string>(StringComparer.OrdinalIgnoreCase);
+			foreach (IssueData issue in _issues)
 			{
-				if (Issue.Streams != null)
+				if (issue.Streams != null)
 				{
-					Streams.UnionWith(Issue.Streams);
+					streams.UnionWith(issue.Streams);
 				}
 			}
 
-			if (Streams.Count > 0)
+			if (streams.Count > 0)
 			{
 				FilterMenu.Items.Add(new ToolStripSeparator());
-				foreach (string Stream in Streams.OrderBy(x => x))
+				foreach (string stream in streams.OrderBy(x => x))
 				{
-					ToolStripMenuItem Item = new ToolStripMenuItem(Stream);
-					Item.Checked = (FilterName == Stream);
-					Item.Click += (S, E) => { FilterName = Stream; UpdateIssueList(); };
-					FilterMenu.Items.Add(Item);
+					ToolStripMenuItem item = new ToolStripMenuItem(stream);
+					item.Checked = (_filterName == stream);
+					item.Click += (s, e) => { _filterName = stream; UpdateIssueList(); };
+					FilterMenu.Items.Add(item);
 				}
 			}
 
@@ -348,13 +348,13 @@ namespace UnrealGameSync
 
 		private void FilterMenu_ShowAll_Click(object sender, EventArgs e)
 		{
-			FilterName = null;
+			_filterName = null;
 			UpdateIssueList();
 		}
 
-		private void ShowIssue(IssueData Issue)
+		private void ShowIssue(IssueData issue)
 		{
-			IssueDetailsWindow.Show(Owner, IssueMonitor, PerforceSettings, ServerTimeOffset, Issue, ServiceProvider, CurrentStream);
+			IssueDetailsWindow.Show(Owner, _issueMonitor, _perforceSettings, _serverTimeOffset, issue, _serviceProvider, _currentStream);
 		}
 
 		private void IssueListView_DrawColumnHeader(object sender, DrawListViewColumnHeaderEventArgs e)
@@ -374,25 +374,25 @@ namespace UnrealGameSync
 			}
 			else if (e.Item.Tag is IssueData)
 			{
-				IssueData Issue = (IssueData)e.Item.Tag;
+				IssueData issue = (IssueData)e.Item.Tag;
 
-				Color BackgroundColor;
-				if (Issue.ResolvedAt.HasValue)
+				Color backgroundColor;
+				if (issue.ResolvedAt.HasValue)
 				{
-					BackgroundColor = SystemColors.Window;//Color.FromArgb(248, 254, 246);
+					backgroundColor = SystemColors.Window;//Color.FromArgb(248, 254, 246);
 				}
-				else if(Issue.FixChange > 0)
+				else if(issue.FixChange > 0)
 				{
-					BackgroundColor = Color.FromArgb(245, 245, 245);
+					backgroundColor = Color.FromArgb(245, 245, 245);
 				}
 				else
 				{
-					BackgroundColor = Color.FromArgb(254, 248, 246);
+					backgroundColor = Color.FromArgb(254, 248, 246);
 				}
 
-				using (SolidBrush Brush = new SolidBrush(BackgroundColor))
+				using (SolidBrush brush = new SolidBrush(backgroundColor))
 				{
-					e.Graphics.FillRectangle(Brush, e.Bounds);
+					e.Graphics.FillRectangle(brush, e.Bounds);
 				}
 			}
 			else
@@ -403,10 +403,10 @@ namespace UnrealGameSync
 
 		private void IssueListView_DrawSubItem(object sender, DrawListViewSubItemEventArgs e)
 		{
-			IssueData Issue = (IssueData)e.Item.Tag;
+			IssueData issue = (IssueData)e.Item.Tag;
 			if (e.ColumnIndex == IconHeader.Index)
 			{
-				if(!Issue.ResolvedAt.HasValue && Issue.FixChange == 0)
+				if(!issue.ResolvedAt.HasValue && issue.FixChange == 0)
 				{
 					IssueListView.DrawIcon(e.Graphics, e.Bounds, WorkspaceControl.BadBuildIcon);
 				}
