UGS: Reformat to standard C# style conventions.

diff --git a/Engine/Source/Programs/UnrealGameSync/UnrealGameSync/Forms/IssueDetailsWindow.cs b/Engine/Source/Programs/UnrealGameSync/UnrealGameSync/Forms/IssueDetailsWindow.cs
--- a/Engine/Source/Programs/UnrealGameSync/UnrealGameSync/Forms/IssueDetailsWindow.cs
+++ b/Engine/Source/Programs/UnrealGameSync/UnrealGameSync/Forms/IssueDetailsWindow.cs
@@ -31,30 +31,30 @@ namespace UnrealGameSync
 			public IssueBuildOutcome Outcome { get; }
 			public IReadOnlyList<IssueBuildData> Builds { get; }
 
-			public BuildGroup(string JobName, string JobUrl, int Change, IssueBuildOutcome Outcome, IReadOnlyList<IssueBuildData> Builds)
+			public BuildGroup(string jobName, string jobUrl, int change, IssueBuildOutcome outcome, IReadOnlyList<IssueBuildData> builds)
 			{
-				this.JobName = JobName;
-				this.JobUrl = JobUrl;
-				this.Change = Change;
-				this.Outcome = Outcome;
-				this.Builds = Builds;
+				this.JobName = jobName;
+				this.JobUrl = jobUrl;
+				this.Change = change;
+				this.Outcome = outcome;
+				this.Builds = builds;
 			}
 		}
 
 		class PerforceChangeRange
 		{
-			public bool bExpanded;
+			public bool Expanded;
 			public int MinChange { get; }
 			public int MaxChange { get; }
 			public List<ChangesRecord>? Changes;
 			public string? ErrorMessage;
 			public BuildGroup BuildGroup { get; }
 
-			public PerforceChangeRange(int MinChange, int MaxChange, BuildGroup BuildGroup)
+			public PerforceChangeRange(int minChange, int maxChange, BuildGroup buildGroup)
 			{
-				this.MinChange = MinChange;
-				this.MaxChange = MaxChange;
-				this.BuildGroup = BuildGroup;
+				this.MinChange = minChange;
+				this.MaxChange = maxChange;
+				this.BuildGroup = buildGroup;
 			}
 		}
 
@@ -62,149 +62,149 @@ namespace UnrealGameSync
 		{
 			public DescribeRecord DescribeRecord;
 
-			public PerforceChangeDetailsWithDescribeRecord(DescribeRecord DescribeRecord)
-				: base(DescribeRecord)
+			public PerforceChangeDetailsWithDescribeRecord(DescribeRecord describeRecord)
+				: base(describeRecord)
 			{
-				this.DescribeRecord = DescribeRecord;
+				this.DescribeRecord = describeRecord;
 			}
 		}
 
 		class PerforceWorkerTask : IDisposable
 		{
-			readonly IPerforceSettings PerforceSettings;
-			readonly string Filter;
-			Action<PerforceChangeRange>? OnUpdateChanges;
-			Action<ChangesRecord>? OnUpdateChangeMetadata;
-			readonly ILogger Logger;
-			readonly object LockObject = new object();
-			readonly SynchronizationContext SynchronizationContext;
-			Task? WorkerTask;
-			CancellationTokenSource CancellationSource;
-			AsyncEvent RefreshEvent;
-			List<PerforceChangeRange> Requests = new List<PerforceChangeRange>();
-			Dictionary<int, PerforceChangeDetailsWithDescribeRecord> ChangeNumberToDetails = new Dictionary<int, PerforceChangeDetailsWithDescribeRecord>();
-
-			public PerforceWorkerTask(IPerforceSettings PerforceSettings, string Filter, Action<PerforceChangeRange>? OnUpdateChanges, Action<ChangesRecord>? OnUpdateChangeMetadata, ILogger Logger)
-			{
-				this.PerforceSettings = PerforceSettings;
-				this.Filter = Filter;
-				this.OnUpdateChanges = OnUpdateChanges;
-				this.OnUpdateChangeMetadata = OnUpdateChangeMetadata;
-				this.Logger = Logger;
-				this.SynchronizationContext = SynchronizationContext.Current!;
-
-				RefreshEvent = new AsyncEvent();
-				CancellationSource = new CancellationTokenSource();
-				WorkerTask = Task.Run(() => RunAsync(CancellationSource.Token));
-			}
-
-			public void AddRequest(PerforceChangeRange Range)
-			{
-				lock(LockObject)
+			readonly IPerforceSettings _perforceSettings;
+			readonly string _filter;
+			Action<PerforceChangeRange>? _onUpdateChanges;
+			Action<ChangesRecord>? _onUpdateChangeMetadata;
+			readonly ILogger _logger;
+			readonly object _lockObject = new object();
+			readonly SynchronizationContext _synchronizationContext;
+			Task? _workerTask;
+			CancellationTokenSource _cancellationSource;
+			AsyncEvent _refreshEvent;
+			List<PerforceChangeRange> _requests = new List<PerforceChangeRange>();
+			Dictionary<int, PerforceChangeDetailsWithDescribeRecord> _changeNumberToDetails = new Dictionary<int, PerforceChangeDetailsWithDescribeRecord>();
+
+			public PerforceWorkerTask(IPerforceSettings perforceSettings, string filter, Action<PerforceChangeRange>? onUpdateChanges, Action<ChangesRecord>? onUpdateChangeMetadata, ILogger logger)
+			{
+				this._perforceSettings = perforceSettings;
+				this._filter = filter;
+				this._onUpdateChanges = onUpdateChanges;
+				this._onUpdateChangeMetadata = onUpdateChangeMetadata;
+				this._logger = logger;
+				this._synchronizationContext = SynchronizationContext.Current!;
+
+				_refreshEvent = new AsyncEvent();
+				_cancellationSource = new CancellationTokenSource();
+				_workerTask = Task.Run(() => RunAsync(_cancellationSource.Token));
+			}
+
+			public void AddRequest(PerforceChangeRange range)
+			{
+				lock(_lockObject)
 				{
-					Requests.Add(Range);
+					_requests.Add(range);
 				}
-				RefreshEvent.Set();
+				_refreshEvent.Set();
 			}
 
 			public void Dispose()
 			{
-				if (WorkerTask != null)
+				if (_workerTask != null)
 				{
-					OnUpdateChanges = null;
-					OnUpdateChangeMetadata = null;
+					_onUpdateChanges = null;
+					_onUpdateChangeMetadata = null;
 
-					CancellationSource.Cancel();
+					_cancellationSource.Cancel();
 
-					WorkerTask.ContinueWith(x => CancellationSource.Dispose());
-					WorkerTask = null;
+					_workerTask.ContinueWith(x => _cancellationSource.Dispose());
+					_workerTask = null;
 				}
 			}
 
-			async Task RunAsync(CancellationToken CancellationToken)
+			async Task RunAsync(CancellationToken cancellationToken)
 			{
-				using IPerforceConnection Perforce = await PerforceConnection.CreateAsync(PerforceSettings, Logger);
+				using IPerforceConnection perforce = await PerforceConnection.CreateAsync(_perforceSettings, _logger);
 
-				List<PerforceChangeRange> CompletedRequests = new List<PerforceChangeRange>(); 
-				while(!CancellationToken.IsCancellationRequested)
+				List<PerforceChangeRange> completedRequests = new List<PerforceChangeRange>(); 
+				while(!cancellationToken.IsCancellationRequested)
 				{
-					Task RefreshTask = RefreshEvent.Task;
+					Task refreshTask = _refreshEvent.Task;
 
 					// Check if there's a request in the queue
-					PerforceChangeRange? NextRequest = null;
-					lock(LockObject)
+					PerforceChangeRange? nextRequest = null;
+					lock(_lockObject)
 					{
-						if(Requests.Count > 0)
+						if(_requests.Count > 0)
 						{
-							NextRequest = Requests[0];
-							Requests.RemoveAt(0);
+							nextRequest = _requests[0];
+							_requests.RemoveAt(0);
 						}
 					}
 
 					// Process the request
-					if(NextRequest != null)
+					if(nextRequest != null)
 					{
-						string RangeFilter = String.Format("{0}@{1},{2}", Filter, NextRequest.MinChange, (NextRequest.MaxChange == -1)? "now" : NextRequest.MaxChange.ToString());
+						string rangeFilter = String.Format("{0}@{1},{2}", _filter, nextRequest.MinChange, (nextRequest.MaxChange == -1)? "now" : nextRequest.MaxChange.ToString());
 
-						PerforceResponseList<ChangesRecord> NewChanges = await Perforce.TryGetChangesAsync(ChangesOptions.None, -1, ChangeStatus.Submitted, RangeFilter, CancellationToken);
-						if(NewChanges.Succeeded)
+						PerforceResponseList<ChangesRecord> newChanges = await perforce.TryGetChangesAsync(ChangesOptions.None, -1, ChangeStatus.Submitted, rangeFilter, cancellationToken);
+						if(newChanges.Succeeded)
 						{
-							NextRequest.Changes = NewChanges.Data;
-							CompletedRequests.Add(NextRequest);
+							nextRequest.Changes = newChanges.Data;
+							completedRequests.Add(nextRequest);
 						}
 						else
 						{
-							NextRequest.ErrorMessage = "Unable to fetch changes. Check your connection settings and try again.";
+							nextRequest.ErrorMessage = "Unable to fetch changes. Check your connection settings and try again.";
 						}
 
-						SynchronizationContext.Post(x => OnUpdateChanges?.Invoke(NextRequest), null);
+						_synchronizationContext.Post(x => _onUpdateChanges?.Invoke(nextRequest), null);
 						continue;
 					}
 
 					// Figure out which changes to fetch
-					List<ChangesRecord> DescribeChanges;
-					lock(LockObject)
+					List<ChangesRecord> describeChanges;
+					lock(_lockObject)
 					{
-						DescribeChanges = CompletedRequests.SelectMany(x => x.Changes).Where(x => !ChangeNumberToDetails.ContainsKey(x.Number)).ToList();
+						describeChanges = completedRequests.SelectMany(x => x.Changes).Where(x => !_changeNumberToDetails.ContainsKey(x.Number)).ToList();
 					}
 
 					// Fetch info on each individual change
-					foreach(ChangesRecord DescribeChange in DescribeChanges)
+					foreach(ChangesRecord describeChange in describeChanges)
 					{
-						lock(LockObject)
+						lock(_lockObject)
 						{
-							if(CancellationToken.IsCancellationRequested || Requests.Count > 0)
+							if(cancellationToken.IsCancellationRequested || _requests.Count > 0)
 							{
 								break;
 							}
 						}
 
-						PerforceResponse<DescribeRecord> Response = await Perforce.TryDescribeAsync(DescribeChange.Number);
-						if (Response.Succeeded)
+						PerforceResponse<DescribeRecord> response = await perforce.TryDescribeAsync(describeChange.Number);
+						if (response.Succeeded)
 						{
-							DescribeRecord Record = Response.Data;
-							lock(LockObject)
+							DescribeRecord record = response.Data;
+							lock(_lockObject)
 							{
-								ChangeNumberToDetails[Record.Number] = new PerforceChangeDetailsWithDescribeRecord(Record);
+								_changeNumberToDetails[record.Number] = new PerforceChangeDetailsWithDescribeRecord(record);
 							}
 						}
 
-						SynchronizationContext.Post(x => OnUpdateChangeMetadata?.Invoke(DescribeChange), null);
+						_synchronizationContext.Post(x => _onUpdateChangeMetadata?.Invoke(describeChange), null);
 					}
 
 					// Wait for something to change
-					using (IDisposable _ = CancellationToken.Register(() => RefreshEvent.Set()))
+					using (IDisposable _ = cancellationToken.Register(() => _refreshEvent.Set()))
 					{
-						await RefreshTask;
+						await refreshTask;
 					}
 				}
 			}
 
-			public bool TryGetChangeDetails(int ChangeNumber, [NotNullWhen(true)] out PerforceChangeDetailsWithDescribeRecord? Details)
+			public bool TryGetChangeDetails(int changeNumber, [NotNullWhen(true)] out PerforceChangeDetailsWithDescribeRecord? details)
 			{
-				lock(LockObject)
+				lock(_lockObject)
 				{
-					return ChangeNumberToDetails.TryGetValue(ChangeNumber, out Details);
+					return _changeNumberToDetails.TryGetValue(changeNumber, out details);
 				}
 			}
 		}
@@ -214,124 +214,124 @@ namespace UnrealGameSync
 			public string Label;
 			public Color Color;
 
-			public BadgeInfo(string Label, Color Color)
+			public BadgeInfo(string label, Color color)
 			{
-				this.Label = Label;
-				this.Color = Color;
+				this.Label = label;
+				this.Color = color;
 			}
 		}
 
 		class ExpandRangeStatusElement : StatusElement
 		{
-			string Text;
-			Action LinkAction;
+			string _text;
+			Action _linkAction;
 
-			public ExpandRangeStatusElement(string Text, Action InLinkAction)
+			public ExpandRangeStatusElement(string text, Action inLinkAction)
 			{
-				this.Text = Text;
-				this.LinkAction = InLinkAction;
+				this._text = text;
+				this._linkAction = inLinkAction;
 				Cursor = NativeCursors.Hand;
 			}
 
-			public override void OnClick(Point Location)
+			public override void OnClick(Point location)
 			{
-				LinkAction();
+				_linkAction();
 			}
 
-			public override Size Measure(Graphics Graphics, StatusElementResources Resources)
+			public override Size Measure(Graphics graphics, StatusElementResources resources)
 			{
-				return TextRenderer.MeasureText(Graphics, Text, Resources.FindOrAddFont(FontStyle.Regular), new Size(int.MaxValue, int.MaxValue), TextFormatFlags.NoPadding);
+				return TextRenderer.MeasureText(graphics, _text, resources.FindOrAddFont(FontStyle.Regular), new Size(int.MaxValue, int.MaxValue), TextFormatFlags.NoPadding);
 			}
 
-			public override void Draw(Graphics Graphics, StatusElementResources Resources)
+			public override void Draw(Graphics graphics, StatusElementResources resources)
 			{
-				Color TextColor = Color.Gray;
-				FontStyle Style = FontStyle.Italic;
-				if(bMouseDown)
+				Color textColor = Color.Gray;
+				FontStyle style = FontStyle.Italic;
+				if(MouseDown)
 				{
-					TextColor = Color.FromArgb(TextColor.B / 2, TextColor.G / 2, TextColor.R);
-					Style |= FontStyle.Underline;
+					textColor = Color.FromArgb(textColor.B / 2, textColor.G / 2, textColor.R);
+					style |= FontStyle.Underline;
 				}
-				else if(bMouseOver)
+				else if(MouseOver)
 				{
-					TextColor = Color.FromArgb(TextColor.B, TextColor.G, TextColor.R);
-					Style |= FontStyle.Underline;
+					textColor = Color.FromArgb(textColor.B, textColor.G, textColor.R);
+					style |= FontStyle.Underline;
 				}
-				TextRenderer.DrawText(Graphics, Text, Resources.FindOrAddFont(Style), Bounds.Location, TextColor, TextFormatFlags.NoPadding);
+				TextRenderer.DrawText(graphics, _text, resources.FindOrAddFont(style), Bounds.Location, textColor, TextFormatFlags.NoPadding);
 			}
 		}
 
-		IssueMonitor IssueMonitor;
-		IssueData Issue;
-		List<IssueBuildData> IssueBuilds;
-		List<IssueDiagnosticData> Diagnostics;
-		IPerforceSettings PerforceSettings;
-		TimeSpan? ServerTimeOffset;
-		PerforceWorkerTask? PerforceWorker;
-		IServiceProvider ServiceProvider;
-		SynchronizationContext MainThreadSynchronizationContext;
-		string? SelectedStream;
-		List<PerforceChangeRange> SelectedStreamRanges = new List<PerforceChangeRange>();
-		bool bIsDisposing;
-		Font? BoldFont;
-		ChangesRecord? ContextMenuChange;
-		string? LastOwner;
-		string? LastDetailsText;
-		System.Windows.Forms.Timer? UpdateTimer;
-		StatusElementResources StatusElementResources;
-
-		IssueDetailsWindow(IssueMonitor IssueMonitor, IssueData Issue, List<IssueBuildData> IssueBuilds, List<IssueDiagnosticData> Diagnostics, IPerforceSettings PerforceSettings, TimeSpan? ServerTimeOffset, IServiceProvider ServiceProvider, string? CurrentStream)
+		IssueMonitor _issueMonitor;
+		IssueData _issue;
+		List<IssueBuildData> _issueBuilds;
+		List<IssueDiagnosticData> _diagnostics;
+		IPerforceSettings _perforceSettings;
+		TimeSpan? _serverTimeOffset;
+		PerforceWorkerTask? _perforceWorker;
+		IServiceProvider _serviceProvider;
+		SynchronizationContext _mainThreadSynchronizationContext;
+		string? _selectedStream;
+		List<PerforceChangeRange> _selectedStreamRanges = new List<PerforceChangeRange>();
+		bool _isDisposing;
+		Font? _boldFont;
+		ChangesRecord? _contextMenuChange;
+		string? _lastOwner;
+		string? _lastDetailsText;
+		System.Windows.Forms.Timer? _updateTimer;
+		StatusElementResources _statusElementResources;
+
+		IssueDetailsWindow(IssueMonitor issueMonitor, IssueData issue, List<IssueBuildData> issueBuilds, List<IssueDiagnosticData> diagnostics, IPerforceSettings perforceSettings, TimeSpan? serverTimeOffset, IServiceProvider serviceProvider, string? currentStream)
 		{
-			this.IssueMonitor = IssueMonitor;
-			this.Issue = Issue;
-			this.IssueBuilds = IssueBuilds;
-			this.Diagnostics = Diagnostics;
-			this.PerforceSettings = PerforceSettings;
-			this.ServerTimeOffset = ServerTimeOffset;
-			this.ServiceProvider = ServiceProvider;
+			this._issueMonitor = issueMonitor;
+			this._issue = issue;
+			this._issueBuilds = issueBuilds;
+			this._diagnostics = diagnostics;
+			this._perforceSettings = perforceSettings;
+			this._serverTimeOffset = serverTimeOffset;
+			this._serviceProvider = serviceProvider;
 
-			IssueMonitor.AddRef();
+			issueMonitor.AddRef();
 
-			MainThreadSynchronizationContext = SynchronizationContext.Current!;
+			_mainThreadSynchronizationContext = SynchronizationContext.Current!;
 
 			InitializeComponent();
 
-			this.Text = String.Format("Issue {0}", Issue.Id);
-			this.StatusElementResources = new StatusElementResources(BuildListView.Font);
+			this.Text = String.Format("Issue {0}", issue.Id);
+			this._statusElementResources = new StatusElementResources(BuildListView.Font);
 			base.Disposed += IssueDetailsWindow_Disposed;
 
-			IssueMonitor.OnIssuesChanged += OnUpdateIssuesAsync;
-			IssueMonitor.StartTracking(Issue.Id);
+			issueMonitor.OnIssuesChanged += OnUpdateIssuesAsync;
+			issueMonitor.StartTracking(issue.Id);
 
 			BuildListView.SmallImageList = new ImageList(){ ImageSize = new Size(1, 20) };
 
-			System.Reflection.PropertyInfo DoubleBufferedProperty = typeof(Control).GetProperty("DoubleBuffered", System.Reflection.BindingFlags.NonPublic | System.Reflection.BindingFlags.Instance)!;
-			DoubleBufferedProperty.SetValue(BuildListView, true, null); 
+			System.Reflection.PropertyInfo doubleBufferedProperty = typeof(Control).GetProperty("DoubleBuffered", System.Reflection.BindingFlags.NonPublic | System.Reflection.BindingFlags.Instance)!;
+			doubleBufferedProperty.SetValue(BuildListView, true, null); 
 
-			using(Graphics Graphics = Graphics.FromHwnd(IntPtr.Zero))
+			using(Graphics graphics = Graphics.FromHwnd(IntPtr.Zero))
 			{
-				float DpiScaleX = Graphics.DpiX / 96.0f;
-				foreach(ColumnHeader? Column in BuildListView.Columns)
+				float dpiScaleX = graphics.DpiX / 96.0f;
+				foreach(ColumnHeader? column in BuildListView.Columns)
 				{
-					if (Column != null)
+					if (column != null)
 					{
-						Column.Width = (int)(Column.Width * DpiScaleX);
+						column.Width = (int)(column.Width * dpiScaleX);
 					}
 				}
 			}
 
-			int SelectIdx = 0;
-			foreach(string Stream in IssueBuilds.Select(x => x.Stream).Distinct().OrderBy(x => x))
+			int selectIdx = 0;
+			foreach(string stream in issueBuilds.Select(x => x.Stream).Distinct().OrderBy(x => x))
 			{
-				StreamComboBox.Items.Add(Stream);
-				if(Stream == CurrentStream)
+				StreamComboBox.Items.Add(stream);
+				if(stream == currentStream)
 				{
-					SelectIdx = StreamComboBox.Items.Count - 1;
+					selectIdx = StreamComboBox.Items.Count - 1;
 				}
 			}
 			if(StreamComboBox.Items.Count > 0)
 			{
-				StreamComboBox.SelectedIndex = SelectIdx;
+				StreamComboBox.SelectedIndex = selectIdx;
 			}
 
 			FilterTypeComboBox.SelectedIndex = 0;
@@ -348,12 +348,12 @@ namespace UnrealGameSync
 				components.Dispose();
 			}
 
-			if (IssueMonitor != null)
+			if (_issueMonitor != null)
 			{
-				IssueMonitor.StopTracking(Issue.Id);
-				IssueMonitor.OnIssuesChanged -= OnUpdateIssuesAsync;
-				IssueMonitor.Release();
-				IssueMonitor = null!;
+				_issueMonitor.StopTracking(_issue.Id);
+				_issueMonitor.OnIssuesChanged -= OnUpdateIssuesAsync;
+				_issueMonitor.Release();
+				_issueMonitor = null!;
 			}
 
 			base.Dispose(disposing);
@@ -361,25 +361,25 @@ namespace UnrealGameSync
 
 		private void StartUpdateTimer()
 		{
-			if(UpdateTimer == null)
+			if(_updateTimer == null)
 			{
-				UpdateTimer = new System.Windows.Forms.Timer();
-				UpdateTimer.Interval = 100;
-				UpdateTimer.Tick += UpdateTimer_Tick;
-				UpdateTimer.Start();
+				_updateTimer = new System.Windows.Forms.Timer();
+				_updateTimer.Interval = 100;
+				_updateTimer.Tick += UpdateTimer_Tick;
+				_updateTimer.Start();
 
-				components.Add(UpdateTimer);
+				components.Add(_updateTimer);
 			}
 		}
 
 		private void StopUpdateTimer()
 		{
-			if(UpdateTimer != null)
+			if(_updateTimer != null)
 			{
-				components.Remove(UpdateTimer);
+				components.Remove(_updateTimer);
 
-				UpdateTimer.Dispose();
-				UpdateTimer = null;
+				_updateTimer.Dispose();
+				_updateTimer = null;
 			}
 		}
 
@@ -389,233 +389,233 @@ namespace UnrealGameSync
 			StopUpdateTimer();
 		}
 
-		void UpdateSummaryTextIfChanged(Label Label, string NewText)
+		void UpdateSummaryTextIfChanged(Label label, string newText)
 		{
-			if (Label.Text != NewText)
+			if (label.Text != newText)
 			{
-				Label.Text = NewText;
+				label.Text = newText;
 			}
 		}
 
-		void UpdateSummaryTextIfChanged(TextBox TextBox, string NewText)
+		void UpdateSummaryTextIfChanged(TextBox textBox, string newText)
 		{
-			if(TextBox.Text != NewText)
+			if(textBox.Text != newText)
 			{
-				TextBox.Text = NewText;
-				TextBox.SelectionLength = 0;
-				TextBox.SelectionStart = NewText.Length;
+				textBox.Text = newText;
+				textBox.SelectionLength = 0;
+				textBox.SelectionStart = newText.Length;
 			}
 		}
 
-		static void AppendEscapedRtf(StringBuilder Result, string Text)
+		static void AppendEscapedRtf(StringBuilder result, string text)
 		{
-			for (int Idx = 0; Idx < Text.Length; Idx++)
+			for (int idx = 0; idx < text.Length; idx++)
 			{
-				char Character = Text[Idx];
-				if(Character == '\n')
+				char character = text[idx];
+				if(character == '\n')
 				{
-					Result.Append(@"\line");
+					result.Append(@"\line");
 				}
-				else if(Character >= 0x20 && Character <= 0x7f)
+				else if(character >= 0x20 && character <= 0x7f)
 				{
-					if(Character == '\\' || Character == '{' || Character == '}')
+					if(character == '\\' || character == '{' || character == '}')
 					{
-						Result.Append('\\');
+						result.Append('\\');
 					}
-					Result.Append(Character);
+					result.Append(character);
 				}
 				else
 				{
-					Result.AppendFormat("\\u{0}?", (int)Character);
+					result.AppendFormat("\\u{0}?", (int)character);
 				}
 			}
 		}
 
-		static void AppendHyperlink(StringBuilder RichText, string Label, string Url)
+		static void AppendHyperlink(StringBuilder richText, string label, string url)
 		{
-			RichText.Append(@"{\field");
-			RichText.Append(@"{\*\fldinst");
-			RichText.AppendFormat("{{ HYPERLINK \"{0}\" }}", Url);
-			RichText.Append(@"}");
-			RichText.Append(@"{\fldrslt ");
-			AppendEscapedRtf(RichText, Label);
-			RichText.Append(@"}");
-			RichText.Append(@"}");
+			richText.Append(@"{\field");
+			richText.Append(@"{\*\fldinst");
+			richText.AppendFormat("{{ HYPERLINK \"{0}\" }}", url);
+			richText.Append(@"}");
+			richText.Append(@"{\fldrslt ");
+			AppendEscapedRtf(richText, label);
+			richText.Append(@"}");
+			richText.Append(@"}");
 		}
 
-		static string CreateRichTextErrors(IssueData Issue, List<IssueBuildData> IssueBuilds, List<IssueDiagnosticData> Diagnostics)
+		static string CreateRichTextErrors(IssueData issue, List<IssueBuildData> issueBuilds, List<IssueDiagnosticData> diagnostics)
 		{
-			StringBuilder RichText = new StringBuilder();
+			StringBuilder richText = new StringBuilder();
 
-			RichText.AppendLine(@"{\rtf1\ansi");
-			RichText.AppendLine(@"{\fonttbl{\f0\fnil\fcharset0 Arial;}{\f1\fnil\fcharset0 Courier New;}{\f2\fnil\fcharset0 Calibri;}}");
-			RichText.AppendLine(@"{\colortbl;\red192\green80\blue77;\red0\green0\blue0;}");
+			richText.AppendLine(@"{\rtf1\ansi");
+			richText.AppendLine(@"{\fonttbl{\f0\fnil\fcharset0 Arial;}{\f1\fnil\fcharset0 Courier New;}{\f2\fnil\fcharset0 Calibri;}}");
+			richText.AppendLine(@"{\colortbl;\red192\green80\blue77;\red0\green0\blue0;}");
 
-			bool bFirst = true;
-			foreach (IGrouping<long, IssueDiagnosticData> Group in Diagnostics.GroupBy(x => x.BuildId ?? -1))
+			bool first = true;
+			foreach (IGrouping<long, IssueDiagnosticData> group in diagnostics.GroupBy(x => x.BuildId ?? -1))
 			{
 				// Step 'Foo'
-				IssueBuildData Build = IssueBuilds.FirstOrDefault(x => x.Id == Group.Key);
-				if (Build != null)
+				IssueBuildData build = issueBuilds.FirstOrDefault(x => x.Id == group.Key);
+				if (build != null)
 				{
-					RichText.Append(@"\pard");   // Paragraph default
-					RichText.Append(@"\cf2");    // Foreground color
-					RichText.Append(@"\b1");     // Bold
-					RichText.Append(@"\f0");     // Font
-					RichText.Append(@"\fs16");   // Font size
-					if (bFirst)
+					richText.Append(@"\pard");   // Paragraph default
+					richText.Append(@"\cf2");    // Foreground color
+					richText.Append(@"\b1");     // Bold
+					richText.Append(@"\f0");     // Font
+					richText.Append(@"\fs16");   // Font size
+					if (first)
 					{
-						RichText.Append(@"\sb100"); // Space before
+						richText.Append(@"\sb100"); // Space before
 					}
 					else
 					{
-						RichText.Append(@"\sb300"); // Space before
+						richText.Append(@"\sb300"); // Space before
 					}
-					RichText.Append(@" In step '\ul1");
-					AppendHyperlink(RichText, Build.JobStepName, Build.JobStepUrl);
-					RichText.Append(@"\ul0':");
+					richText.Append(@" In step '\ul1");
+					AppendHyperlink(richText, build.JobStepName, build.JobStepUrl);
+					richText.Append(@"\ul0':");
 
-					RichText.AppendLine(@"\par");
+					richText.AppendLine(@"\par");
 				}
 
-				IssueDiagnosticData[] DiagnosticsArray = Group.ToArray();
-				for(int Idx = 0; Idx < DiagnosticsArray.Length; Idx++)
+				IssueDiagnosticData[] diagnosticsArray = group.ToArray();
+				for(int idx = 0; idx < diagnosticsArray.Length; idx++)
 				{
-					IssueDiagnosticData Diagnostic = DiagnosticsArray[Idx];
+					IssueDiagnosticData diagnostic = diagnosticsArray[idx];
 
 					// Error X/Y:
-					RichText.Append(@"\pard");   // Paragraph default
-					RichText.Append(@"\cf1");    // Foreground color
-					RichText.Append(@"\b1");     // Bold
-					RichText.Append(@"\f0");     // Font
-					RichText.Append(@"\fs16");   // Font size
-					RichText.Append(@"\fi50");   // First line indent
-					RichText.Append(@"\li50");   // Other line indent
-					RichText.Append(@"\sb100");  // Space before
-					RichText.Append(@"\sa50");   // Space after
-					RichText.Append(@" ");
-
-					RichText.Append(@"\ul1");
-					AppendHyperlink(RichText, String.Format("Error {0}/{1}", Idx + 1, DiagnosticsArray.Length), Diagnostic.Url);
-					RichText.Append(@"\ul0");
-
-					RichText.AppendLine(@"\par");
+					richText.Append(@"\pard");   // Paragraph default
+					richText.Append(@"\cf1");    // Foreground color
+					richText.Append(@"\b1");     // Bold
+					richText.Append(@"\f0");     // Font
+					richText.Append(@"\fs16");   // Font size
+					richText.Append(@"\fi50");   // First line indent
+					richText.Append(@"\li50");   // Other line indent
+					richText.Append(@"\sb100");  // Space before
+					richText.Append(@"\sa50");   // Space after
+					richText.Append(@" ");
+
+					richText.Append(@"\ul1");
+					AppendHyperlink(richText, String.Format("Error {0}/{1}", idx + 1, diagnosticsArray.Length), diagnostic.Url);
+					richText.Append(@"\ul0");
+
+					richText.AppendLine(@"\par");
 
 					// Error text
-					foreach(string Line in Diagnostic.Message.TrimEnd().Split('\n'))
+					foreach(string line in diagnostic.Message.TrimEnd().Split('\n'))
 					{
-						RichText.Append(@"\pard");   // Paragraph default
-						RichText.Append(@"\cf0");    // Foreground color
-						RichText.Append(@"\b0");     // Bold
-						RichText.Append(@"\f1");     // Font
-						RichText.Append(@"\fs16");   // Font size 16
-						RichText.Append(@"\fi150");  // First line indent
-						RichText.Append(@"\li150");  // Other line indent
-						AppendEscapedRtf(RichText, Line);
-						RichText.Append(@"\par");
+						richText.Append(@"\pard");   // Paragraph default
+						richText.Append(@"\cf0");    // Foreground color
+						richText.Append(@"\b0");     // Bold
+						richText.Append(@"\f1");     // Font
+						richText.Append(@"\fs16");   // Font size 16
+						richText.Append(@"\fi150");  // First line indent
+						richText.Append(@"\li150");  // Other line indent
+						AppendEscapedRtf(richText, line);
+						richText.Append(@"\par");
 					}
 				}
 
-				bFirst = false;
+				first = false;
 			}
 
-			RichText.AppendLine("}");
-			return RichText.ToString();
+			richText.AppendLine("}");
+			return richText.ToString();
 		}
 
 		void UpdateCurrentIssue()
 		{
-			UpdateSummaryTextIfChanged(SummaryTextBox, Issue.Summary.ToString());
+			UpdateSummaryTextIfChanged(SummaryTextBox, _issue.Summary.ToString());
 
-			IssueBuildData FirstFailingBuild = IssueBuilds.FirstOrDefault(x => x.ErrorUrl != null);
-			BuildLinkLabel.Text = (FirstFailingBuild != null)? FirstFailingBuild.JobName : "Unknown";
+			IssueBuildData firstFailingBuild = _issueBuilds.FirstOrDefault(x => x.ErrorUrl != null);
+			BuildLinkLabel.Text = (firstFailingBuild != null)? firstFailingBuild.JobName : "Unknown";
 
-			StringBuilder Status = new StringBuilder();
-			if(IssueMonitor.HasPendingUpdate())
+			StringBuilder status = new StringBuilder();
+			if(_issueMonitor.HasPendingUpdate())
 			{
-				Status.Append("Updating...");
+				status.Append("Updating...");
 			}
-			else if(Issue.FixChange != 0)
+			else if(_issue.FixChange != 0)
 			{
-				if(Issue.FixChange < 0)
+				if(_issue.FixChange < 0)
 				{
-					if(Issue.ResolvedAt.HasValue)
+					if(_issue.ResolvedAt.HasValue)
 					{
-						Status.AppendFormat("Closed as systemic issue.", Issue.FixChange);
+						status.AppendFormat("Closed as systemic issue.", _issue.FixChange);
 					}
 					else
 					{
-						Status.AppendFormat("Fixed as systemic issue (pending verification).", Issue.FixChange);
+						status.AppendFormat("Fixed as systemic issue (pending verification).", _issue.FixChange);
 					}
 				}
 				else
 				{
-					if(Issue.ResolvedAt.HasValue)
+					if(_issue.ResolvedAt.HasValue)
 					{
-						Status.AppendFormat("Closed. Fixed in CL {0}.", Issue.FixChange);
+						status.AppendFormat("Closed. Fixed in CL {0}.", _issue.FixChange);
 					}
 					else
 					{
-						Status.AppendFormat("Fixed in CL {0} (pending verification)", Issue.FixChange);
+						status.AppendFormat("Fixed in CL {0} (pending verification)", _issue.FixChange);
 					}
 				}
 			}
-			else if(Issue.ResolvedAt.HasValue)
+			else if(_issue.ResolvedAt.HasValue)
 			{
-				Status.Append("Resolved");
+				status.Append("Resolved");
 			}
-			else if(Issue.Owner == null)
+			else if(_issue.Owner == null)
 			{
-				Status.Append("Currently unassigned");
+				status.Append("Currently unassigned");
 			}
 			else
 			{
-				Status.Append(Utility.FormatUserName(Issue.Owner));
-				if(Issue.NominatedBy != null)
+				status.Append(Utility.FormatUserName(_issue.Owner));
+				if(_issue.NominatedBy != null)
 				{
-					Status.AppendFormat(" nominated by {0}", Utility.FormatUserName(Issue.NominatedBy));
+					status.AppendFormat(" nominated by {0}", Utility.FormatUserName(_issue.NominatedBy));
 				}
-				if(Issue.AcknowledgedAt.HasValue)
+				if(_issue.AcknowledgedAt.HasValue)
 				{
-					Status.AppendFormat(" (acknowledged {0})", Utility.FormatRecentDateTime(Issue.AcknowledgedAt.Value.ToLocalTime()));
+					status.AppendFormat(" (acknowledged {0})", Utility.FormatRecentDateTime(_issue.AcknowledgedAt.Value.ToLocalTime()));
 				}
 				else
 				{
-					Status.Append(" (not acknowledged)");
+					status.Append(" (not acknowledged)");
 				}
 			}
-			UpdateSummaryTextIfChanged(StatusTextBox, Status.ToString());
+			UpdateSummaryTextIfChanged(StatusTextBox, status.ToString());
 
-			StringBuilder OpenSince = new StringBuilder();
-			OpenSince.Append(Utility.FormatRecentDateTime(Issue.CreatedAt.ToLocalTime()));
-			if(OpenSince.Length > 0)
+			StringBuilder openSince = new StringBuilder();
+			openSince.Append(Utility.FormatRecentDateTime(_issue.CreatedAt.ToLocalTime()));
+			if(openSince.Length > 0)
 			{
-				OpenSince[0] = Char.ToUpper(OpenSince[0]);
+				openSince[0] = Char.ToUpper(openSince[0]);
 			}
-			OpenSince.AppendFormat(" ({0})", Utility.FormatDurationMinutes(Issue.RetrievedAt - Issue.CreatedAt));
-			UpdateSummaryTextIfChanged(OpenSinceTextBox, OpenSince.ToString());
+			openSince.AppendFormat(" ({0})", Utility.FormatDurationMinutes(_issue.RetrievedAt - _issue.CreatedAt));
+			UpdateSummaryTextIfChanged(OpenSinceTextBox, openSince.ToString());
 
-			if(LastOwner != Issue.Owner)
+			if(_lastOwner != _issue.Owner)
 			{
-				LastOwner = Issue.Owner;
+				_lastOwner = _issue.Owner;
 				BuildListView.Invalidate();
 			}
 
-			UpdateSummaryTextIfChanged(StepNamesTextBox, String.Join(", ", IssueBuilds.Select(x => x.JobStepName).Distinct().OrderBy(x => x)));
-			UpdateSummaryTextIfChanged(StreamNamesTextBox, String.Join(", ", IssueBuilds.Select(x => x.Stream).Distinct().OrderBy(x => x)));
+			UpdateSummaryTextIfChanged(StepNamesTextBox, String.Join(", ", _issueBuilds.Select(x => x.JobStepName).Distinct().OrderBy(x => x)));
+			UpdateSummaryTextIfChanged(StreamNamesTextBox, String.Join(", ", _issueBuilds.Select(x => x.Stream).Distinct().OrderBy(x => x)));
 
-			string RtfText = CreateRichTextErrors(Issue, IssueBuilds, Diagnostics);
-			if (LastDetailsText != RtfText)
+			string rtfText = CreateRichTextErrors(_issue, _issueBuilds, _diagnostics);
+			if (_lastDetailsText != rtfText)
 			{
-				using (MemoryStream Stream = new MemoryStream(Encoding.UTF8.GetBytes(RtfText), false))
+				using (MemoryStream stream = new MemoryStream(Encoding.UTF8.GetBytes(rtfText), false))
 				{
-					DetailsTextBox.LoadFile(Stream, RichTextBoxStreamType.RichText);
+					DetailsTextBox.LoadFile(stream, RichTextBoxStreamType.RichText);
 					DetailsTextBox.Select(0, 0);
 				}
-				LastDetailsText = RtfText;
+				_lastDetailsText = rtfText;
 			}
 
-			if(Issue.FixChange == 0)
+			if(_issue.FixChange == 0)
 			{
 				MarkFixedBtn.Text = "Mark Fixed...";
 			}
@@ -627,52 +627,52 @@ namespace UnrealGameSync
 
 		private void IssueDetailsWindow_Disposed(object? sender, EventArgs e)
 		{
-			bIsDisposing = true;
+			_isDisposing = true;
 			DestroyWorker();
 		}
 
 		void FetchAllBuildChanges()
 		{
-			foreach(PerforceChangeRange Range in SelectedStreamRanges)
+			foreach(PerforceChangeRange range in _selectedStreamRanges)
 			{
-				FetchBuildChanges(Range);
+				FetchBuildChanges(range);
 			}
 		}
 
-		void FetchBuildChanges(PerforceChangeRange Range)
+		void FetchBuildChanges(PerforceChangeRange range)
 		{
-			if(!Range.bExpanded)
+			if(!range.Expanded)
 			{
-				Range.bExpanded = true;
-				PerforceWorker!.AddRequest(Range);
+				range.Expanded = true;
+				_perforceWorker!.AddRequest(range);
 				UpdateBuildList();
 			}
 		}
 
-		bool FilterMatch(Regex FilterRegex, ChangesRecord Summary)
+		bool FilterMatch(Regex filterRegex, ChangesRecord summary)
 		{
-			if(FilterRegex.IsMatch(Summary.User))
+			if(filterRegex.IsMatch(summary.User))
 			{
 				return true;
 			}
-			if(FilterRegex.IsMatch(Summary.Description))
+			if(filterRegex.IsMatch(summary.Description))
 			{
 				return true;
 			}
 			return false;
 		}
 
-		bool FilterMatch(Regex FilterRegex, DescribeRecord DescribeRecord)
+		bool FilterMatch(Regex filterRegex, DescribeRecord describeRecord)
 		{
-			if(FilterRegex.IsMatch(DescribeRecord.User))
+			if(filterRegex.IsMatch(describeRecord.User))
 			{
 				return true;
 			}
-			if(FilterRegex.IsMatch(DescribeRecord.Description))
+			if(filterRegex.IsMatch(describeRecord.Description))
 			{
 				return true;
 			}
-			if(DescribeRecord.Files.Any(x => FilterRegex.IsMatch(x.DepotFile)))
+			if(describeRecord.Files.Any(x => filterRegex.IsMatch(x.DepotFile)))
 			{
 				return true;
 			}
@@ -681,124 +681,124 @@ namespace UnrealGameSync
 
 		void UpdateBuildList()
 		{
-			int NumNewItems = 0;
+			int numNewItems = 0;
 			BuildListView.BeginUpdate();
 
 			// Capture the initial selection
-			object? PrevSelection = null;
+			object? prevSelection = null;
 			if(BuildListView.SelectedItems.Count > 0)
 			{
-				PrevSelection = BuildListView.SelectedItems[0].Tag;
+				prevSelection = BuildListView.SelectedItems[0].Tag;
 			}
 
 			// Get all the search terms
-			string[] FilterTerms = FilterTextBox.Text.Split(new char[]{ ' ' }, StringSplitOptions.RemoveEmptyEntries);
+			string[] filterTerms = FilterTextBox.Text.Split(new char[]{ ' ' }, StringSplitOptions.RemoveEmptyEntries);
 
 			// Build a regex from each filter term
-			List<Regex> FilterRegexes = new List<Regex>();
-			foreach(string FilterTerm in FilterTerms)
+			List<Regex> filterRegexes = new List<Regex>();
+			foreach(string filterTerm in filterTerms)
 			{
-				string RegexText = Regex.Escape(FilterTerm);
-				RegexText = RegexText.Replace("\\?", ".");
-				RegexText = RegexText.Replace("\\*", "[^\\\\/]*");
-				RegexText = RegexText.Replace("\\.\\.\\.", ".*");
-				FilterRegexes.Add(new Regex(RegexText, RegexOptions.IgnoreCase));
+				string regexText = Regex.Escape(filterTerm);
+				regexText = regexText.Replace("\\?", ".");
+				regexText = regexText.Replace("\\*", "[^\\\\/]*");
+				regexText = regexText.Replace("\\.\\.\\.", ".*");
+				filterRegexes.Add(new Regex(regexText, RegexOptions.IgnoreCase));
 			}
 
 			// Get the filter type
-			bool bOnlyShowCodeChanges = (FilterTypeComboBox.SelectedIndex == 1);
-			bool bOnlyShowContentChanges = (FilterTypeComboBox.SelectedIndex == 2);
+			bool onlyShowCodeChanges = (FilterTypeComboBox.SelectedIndex == 1);
+			bool onlyShowContentChanges = (FilterTypeComboBox.SelectedIndex == 2);
 
 			// Create rows for all the ranges
-			foreach(PerforceChangeRange Range in SelectedStreamRanges)
+			foreach(PerforceChangeRange range in _selectedStreamRanges)
 			{
-				if(Range.bExpanded)
+				if(range.Expanded)
 				{
-					if(Range.Changes == null)
+					if(range.Changes == null)
 					{
-						ListViewItem FetchingItem = new ListViewItem("");
-						FetchingItem.Tag = Range;
-						FetchingItem.SubItems.Add("");
-						FetchingItem.SubItems.Add("");
-
-						StatusLineListViewWidget FetchingWidget = new StatusLineListViewWidget(FetchingItem, StatusElementResources);
-						FetchingWidget.HorizontalAlignment = HorizontalAlignment.Left;
-						if(Range.ErrorMessage != null)
+						ListViewItem fetchingItem = new ListViewItem("");
+						fetchingItem.Tag = range;
+						fetchingItem.SubItems.Add("");
+						fetchingItem.SubItems.Add("");
+
+						StatusLineListViewWidget fetchingWidget = new StatusLineListViewWidget(fetchingItem, _statusElementResources);
+						fetchingWidget.HorizontalAlignment = HorizontalAlignment.Left;
+						if(range.ErrorMessage != null)
 						{
-							FetchingWidget.Line.AddText(Range.ErrorMessage, Color.Gray, FontStyle.Italic);
+							fetchingWidget.Line.AddText(range.ErrorMessage, Color.Gray, FontStyle.Italic);
 						}
 						else
 						{
-							FetchingWidget.Line.AddText("Fetching changes, please wait...", Color.Gray, FontStyle.Italic);
+							fetchingWidget.Line.AddText("Fetching changes, please wait...", Color.Gray, FontStyle.Italic);
 						}
 			
-						FetchingItem.SubItems.Add(new ListViewItem.ListViewSubItem(FetchingItem, ""){ Tag = FetchingWidget });
-						FetchingItem.SubItems.Add(new ListViewItem.ListViewSubItem(FetchingItem, ""){ Tag = FetchingWidget });
+						fetchingItem.SubItems.Add(new ListViewItem.ListViewSubItem(fetchingItem, ""){ Tag = fetchingWidget });
+						fetchingItem.SubItems.Add(new ListViewItem.ListViewSubItem(fetchingItem, ""){ Tag = fetchingWidget });
 
-						BuildListView.Items.Insert(NumNewItems++, FetchingItem);
+						BuildListView.Items.Insert(numNewItems++, fetchingItem);
 					}
 					else
 					{
-						foreach(ChangesRecord Change in Range.Changes)
+						foreach(ChangesRecord change in range.Changes)
 						{
-							PerforceChangeDetailsWithDescribeRecord? Details;
-							PerforceWorker!.TryGetChangeDetails(Change.Number, out Details);
+							PerforceChangeDetailsWithDescribeRecord? details;
+							_perforceWorker!.TryGetChangeDetails(change.Number, out details);
 
-							if(FilterRegexes.Count > 0 || bOnlyShowCodeChanges || bOnlyShowContentChanges)
+							if(filterRegexes.Count > 0 || onlyShowCodeChanges || onlyShowContentChanges)
 							{
-								if(Details == null)
+								if(details == null)
 								{
-									if(bOnlyShowCodeChanges || bOnlyShowCodeChanges)
+									if(onlyShowCodeChanges || onlyShowCodeChanges)
 									{
 										continue;
 									}
-									if(FilterRegexes.Any(x => !FilterMatch(x, Change)))
+									if(filterRegexes.Any(x => !FilterMatch(x, change)))
 									{
 										continue;
 									}
 								}
 								else
 								{
-									if(bOnlyShowCodeChanges && !Details.bContainsCode)
+									if(onlyShowCodeChanges && !details.ContainsCode)
 									{
 										continue;
 									}
-									if(bOnlyShowContentChanges && !Details.bContainsContent)
+									if(onlyShowContentChanges && !details.ContainsContent)
 									{
 										continue;
 									}
-									if(FilterRegexes.Any(x => !FilterMatch(x, Details.DescribeRecord)))
+									if(filterRegexes.Any(x => !FilterMatch(x, details.DescribeRecord)))
 									{
 										continue;
 									}
 								}
 							}
 
-							ListViewItem Item = new ListViewItem("");
-							Item.Tag = Change;
+							ListViewItem item = new ListViewItem("");
+							item.Tag = change;
 
-							StatusLineListViewWidget TypeWidget = new StatusLineListViewWidget(Item, StatusElementResources);
-							UpdateChangeTypeWidget(TypeWidget, Details);
-							Item.SubItems.Add(new ListViewItem.ListViewSubItem(Item, "") { Tag = TypeWidget });
+							StatusLineListViewWidget typeWidget = new StatusLineListViewWidget(item, _statusElementResources);
+							UpdateChangeTypeWidget(typeWidget, details);
+							item.SubItems.Add(new ListViewItem.ListViewSubItem(item, "") { Tag = typeWidget });
 
-							Item.SubItems.Add(Change.Number.ToString());
+							item.SubItems.Add(change.Number.ToString());
 
-							DateTime DisplayTime;
-							if (!ServerTimeOffset.HasValue)
+							DateTime displayTime;
+							if (!_serverTimeOffset.HasValue)
 							{
-								DisplayTime = Change.Time.ToLocalTime();
+								displayTime = change.Time.ToLocalTime();
 							}
 							else
 							{
-								DisplayTime = new DateTime(Change.Time.Ticks + ServerTimeOffset.Value.Ticks, DateTimeKind.Local);
+								displayTime = new DateTime(change.Time.Ticks + _serverTimeOffset.Value.Ticks, DateTimeKind.Local);
 							}
 
-							Item.SubItems.Add(DisplayTime.ToString("h\\.mmtt"));
+							item.SubItems.Add(displayTime.ToString("h\\.mmtt"));
 
-							Item.SubItems.Add(WorkspaceControl.FormatUserName(Change.User));
+							item.SubItems.Add(WorkspaceControl.FormatUserName(change.User));
 
-							Item.SubItems.Add(Change.Description);
-							BuildListView.Items.Insert(NumNewItems++, Item);
+							item.SubItems.Add(change.Description);
+							BuildListView.Items.Insert(numNewItems++, item);
 						}
 					}
 				}
@@ -826,33 +826,33 @@ namespace UnrealGameSync
 					BuildListView.Items.Insert(NumNewItems++, RangeItem);
 				}
 				*/
-				ListViewItem BuildItem = new ListViewItem("");
-				BuildItem.Tag = Range.BuildGroup;
-				BuildItem.SubItems.Add("");
-				BuildItem.SubItems.Add(Range.BuildGroup.Change.ToString());
-				BuildItem.SubItems.Add("");
+				ListViewItem buildItem = new ListViewItem("");
+				buildItem.Tag = range.BuildGroup;
+				buildItem.SubItems.Add("");
+				buildItem.SubItems.Add(range.BuildGroup.Change.ToString());
+				buildItem.SubItems.Add("");
 
-				StatusLineListViewWidget BuildWidget = new StatusLineListViewWidget(BuildItem, StatusElementResources);
-				BuildWidget.HorizontalAlignment = HorizontalAlignment.Left;
+				StatusLineListViewWidget buildWidget = new StatusLineListViewWidget(buildItem, _statusElementResources);
+				buildWidget.HorizontalAlignment = HorizontalAlignment.Left;
 
-				BuildWidget.Line.AddLink(Range.BuildGroup.JobName, FontStyle.Underline, () => Utility.OpenUrl(Range.BuildGroup.JobUrl));
-				BuildItem.SubItems.Add(new ListViewItem.ListViewSubItem(BuildItem, ""){ Tag = BuildWidget });
-				BuildItem.SubItems.Add(new ListViewItem.ListViewSubItem(BuildItem, ""){ Tag = BuildWidget });
+				buildWidget.Line.AddLink(range.BuildGroup.JobName, FontStyle.Underline, () => Utility.OpenUrl(range.BuildGroup.JobUrl));
+				buildItem.SubItems.Add(new ListViewItem.ListViewSubItem(buildItem, ""){ Tag = buildWidget });
+				buildItem.SubItems.Add(new ListViewItem.ListViewSubItem(buildItem, ""){ Tag = buildWidget });
 
-				BuildListView.Items.Insert(NumNewItems++, BuildItem);
+				BuildListView.Items.Insert(numNewItems++, buildItem);
 			}
 			
 			// Re-select the original item
-			for(int Idx = 0; Idx < BuildListView.Items.Count; Idx++)
+			for(int idx = 0; idx < BuildListView.Items.Count; idx++)
 			{
-				if(BuildListView.Items[Idx].Tag == PrevSelection)
+				if(BuildListView.Items[idx].Tag == prevSelection)
 				{
-					BuildListView.Items[Idx].Selected = true;
+					BuildListView.Items[idx].Selected = true;
 				}
 			}
 
 			// Remove all the items we no longer need
-			while(BuildListView.Items.Count > NumNewItems)
+			while(BuildListView.Items.Count > numNewItems)
 			{
 				BuildListView.Items.RemoveAt(BuildListView.Items.Count - 1);
 			}
@@ -862,90 +862,90 @@ namespace UnrealGameSync
 			BuildListView.Invalidate();
 		}
 
-		BuildGroup? SelectedBuildGroup;
+		BuildGroup? _selectedBuildGroup;
 
-		void ShowJobContextMenu(Point Point, BuildGroup BuildGroup)
+		void ShowJobContextMenu(Point point, BuildGroup buildGroup)
 		{
-			SelectedBuildGroup = BuildGroup;
+			_selectedBuildGroup = buildGroup;
 
-			int MinIndex = JobContextMenu.Items.IndexOf(JobContextMenu_StepSeparatorMin) + 1;
-			while(JobContextMenu.Items.Count > MinIndex)
+			int minIndex = JobContextMenu.Items.IndexOf(JobContextMenu_StepSeparatorMin) + 1;
+			while(JobContextMenu.Items.Count > minIndex)
 			{
-				JobContextMenu.Items.RemoveAt(MinIndex);
+				JobContextMenu.Items.RemoveAt(minIndex);
 			}
 
-			JobContextMenu_ViewJob.Text = String.Format("View Job: {0}", BuildGroup.JobName);
+			JobContextMenu_ViewJob.Text = String.Format("View Job: {0}", buildGroup.JobName);
 
-			foreach (IssueBuildData Build in BuildGroup.Builds.OrderBy(x => x.JobStepName))
+			foreach (IssueBuildData build in buildGroup.Builds.OrderBy(x => x.JobStepName))
 			{
-				ToolStripMenuItem MenuItem = new ToolStripMenuItem(String.Format("View Step: {0}", Build.JobStepName));
-				MenuItem.Click += (S, E) => Utility.OpenUrl(Build.JobStepUrl);
-				JobContextMenu.Items.Insert(MinIndex++, MenuItem);
+				ToolStripMenuItem menuItem = new ToolStripMenuItem(String.Format("View Step: {0}", build.JobStepName));
+				menuItem.Click += (s, e) => Utility.OpenUrl(build.JobStepUrl);
+				JobContextMenu.Items.Insert(minIndex++, menuItem);
 			}
 
-			JobContextMenu.Show(BuildListView, Point, ToolStripDropDownDirection.BelowRight);
+			JobContextMenu.Show(BuildListView, point, ToolStripDropDownDirection.BelowRight);
 		}
 
-		void UpdateChangeTypeWidget(StatusLineListViewWidget TypeWidget, PerforceChangeDetails? Details)
+		void UpdateChangeTypeWidget(StatusLineListViewWidget typeWidget, PerforceChangeDetails? details)
 		{
-			TypeWidget.Line.Clear();
-			if(Details == null)
+			typeWidget.Line.Clear();
+			if(details == null)
 			{
-				TypeWidget.Line.AddBadge("Unknown", Color.FromArgb(192, 192, 192), null);
+				typeWidget.Line.AddBadge("Unknown", Color.FromArgb(192, 192, 192), null);
 			}
 			else
 			{
-				if(Details.bContainsCode)
+				if(details.ContainsCode)
 				{
-					TypeWidget.Line.AddBadge("Code", Color.FromArgb(116, 185, 255), null);
+					typeWidget.Line.AddBadge("Code", Color.FromArgb(116, 185, 255), null);
 				}
-				if(Details.bContainsContent)
+				if(details.ContainsContent)
 				{
-					TypeWidget.Line.AddBadge("Content", Color.FromArgb(162, 155, 255), null);
+					typeWidget.Line.AddBadge("Content", Color.FromArgb(162, 155, 255), null);
 				}
 			}
 		}
 
 		void CreateWorker()
 		{
-			string? NewSelectedStream = StreamComboBox.SelectedItem as string;
-			if(SelectedStream != NewSelectedStream)
+			string? newSelectedStream = StreamComboBox.SelectedItem as string;
+			if(_selectedStream != newSelectedStream)
 			{
 				DestroyWorker();
 			
-				SelectedStream = NewSelectedStream;
+				_selectedStream = newSelectedStream;
 
 				BuildListView.BeginUpdate();
 				BuildListView.Items.Clear();
 
-				if(SelectedStream != null)
+				if(_selectedStream != null)
 				{
-					SelectedStreamRanges = new List<PerforceChangeRange>();
+					_selectedStreamRanges = new List<PerforceChangeRange>();
 
-					int MaxChange = -1;
-					foreach(IGrouping<string, IssueBuildData> Group in IssueBuilds.Where(x => x.Stream == SelectedStream).OrderByDescending(x => x.Change).ThenByDescending(x => x.JobUrl).GroupBy(x => x.JobUrl))
+					int maxChange = -1;
+					foreach(IGrouping<string, IssueBuildData> group in _issueBuilds.Where(x => x.Stream == _selectedStream).OrderByDescending(x => x.Change).ThenByDescending(x => x.JobUrl).GroupBy(x => x.JobUrl))
 					{
-						IssueBuildOutcome Outcome = Group.Any(x => x.Outcome == IssueBuildOutcome.Error) ? IssueBuildOutcome.Error : Group.Any(x => x.Outcome == IssueBuildOutcome.Warning) ? IssueBuildOutcome.Warning : IssueBuildOutcome.Success;
-						BuildGroup BuildGroup = new BuildGroup(Group.First().JobName, Group.Key, Group.First().Change, Outcome, Group.ToList());
+						IssueBuildOutcome outcome = group.Any(x => x.Outcome == IssueBuildOutcome.Error) ? IssueBuildOutcome.Error : group.Any(x => x.Outcome == IssueBuildOutcome.Warning) ? IssueBuildOutcome.Warning : IssueBuildOutcome.Success;
+						BuildGroup buildGroup = new BuildGroup(group.First().JobName, group.Key, group.First().Change, outcome, group.ToList());
 
-						PerforceChangeRange Range = new PerforceChangeRange(BuildGroup.Change + 1, MaxChange, BuildGroup);
-						SelectedStreamRanges.Add(Range);
+						PerforceChangeRange range = new PerforceChangeRange(buildGroup.Change + 1, maxChange, buildGroup);
+						_selectedStreamRanges.Add(range);
 
-						MaxChange = BuildGroup.Change;
+						maxChange = buildGroup.Change;
 					}
 
-					PerforceWorker = new PerforceWorkerTask(PerforceSettings, String.Format("{0}/...", SelectedStream), OnRequestComplete, UpdateChangeMetadata, ServiceProvider.GetRequiredService<ILogger<PerforceWorkerTask>>());
+					_perforceWorker = new PerforceWorkerTask(_perforceSettings, String.Format("{0}/...", _selectedStream), OnRequestComplete, UpdateChangeMetadata, _serviceProvider.GetRequiredService<ILogger<PerforceWorkerTask>>());
 
 					UpdateBuildList();
 
-					for(int Idx = 0; Idx + 2 < SelectedStreamRanges.Count; Idx++)
+					for(int idx = 0; idx + 2 < _selectedStreamRanges.Count; idx++)
 					{
-						if(SelectedStreamRanges[Idx].BuildGroup.Outcome != SelectedStreamRanges[Idx + 1].BuildGroup.Outcome)
+						if(_selectedStreamRanges[idx].BuildGroup.Outcome != _selectedStreamRanges[idx + 1].BuildGroup.Outcome)
 						{
-							FetchBuildChanges(SelectedStreamRanges[Idx + 1]);
+							FetchBuildChanges(_selectedStreamRanges[idx + 1]);
 						}
 					}
-					FetchBuildChanges(SelectedStreamRanges[SelectedStreamRanges.Count - 1]);
+					FetchBuildChanges(_selectedStreamRanges[_selectedStreamRanges.Count - 1]);
 				}
 
 				BuildListView.EndUpdate();
@@ -955,12 +955,12 @@ namespace UnrealGameSync
 
 		void OnUpdateIssues()
 		{
-			List<IssueData> NewIssues = IssueMonitor.GetIssues();
-			foreach(IssueData NewIssue in NewIssues)
+			List<IssueData> newIssues = _issueMonitor.GetIssues();
+			foreach(IssueData newIssue in newIssues)
 			{
-				if(NewIssue.Id == Issue.Id)
+				if(newIssue.Id == _issue.Id)
 				{
-					Issue = NewIssue;
+					_issue = newIssue;
 					UpdateCurrentIssue();
 					break;
 				}
@@ -969,24 +969,24 @@ namespace UnrealGameSync
 
 		private void OnUpdateIssuesAsync()
 		{
-			MainThreadSynchronizationContext.Post((o) => { if(!bIsDisposing){ OnUpdateIssues(); } }, null);
+			_mainThreadSynchronizationContext.Post((o) => { if(!_isDisposing){ OnUpdateIssues(); } }, null);
 		}
 
 		void DestroyWorker()
 		{
-			if(PerforceWorker != null)
+			if(_perforceWorker != null)
 			{
-				PerforceWorker.Dispose();
-				PerforceWorker = null!;
+				_perforceWorker.Dispose();
+				_perforceWorker = null!;
 			}
 		}
 
-		void OnRequestComplete(PerforceChangeRange Range)
+		void OnRequestComplete(PerforceChangeRange range)
 		{
 			UpdateBuildList();
 		}
 
-		void UpdateChangeMetadata(ChangesRecord Change)
+		void UpdateChangeMetadata(ChangesRecord change)
 		{
 			if(FilterTextBox.Text.Length > 0 || FilterTypeComboBox.SelectedIndex != 0)
 			{
@@ -994,88 +994,88 @@ namespace UnrealGameSync
 			}
 			else
 			{
-				foreach(ListViewItem? Item in BuildListView.Items)
+				foreach(ListViewItem? item in BuildListView.Items)
 				{
-					if(Item != null && Item.Tag == Change)
+					if(item != null && item.Tag == change)
 					{
-						PerforceChangeDetailsWithDescribeRecord? Details;
-						PerforceWorker!.TryGetChangeDetails(Change.Number, out Details);
+						PerforceChangeDetailsWithDescribeRecord? details;
+						_perforceWorker!.TryGetChangeDetails(change.Number, out details);
 
-						StatusLineListViewWidget TypeWidget = (StatusLineListViewWidget)Item.SubItems[TypeHeader.Index].Tag;
-						UpdateChangeTypeWidget(TypeWidget, Details);
+						StatusLineListViewWidget typeWidget = (StatusLineListViewWidget)item.SubItems[TypeHeader.Index].Tag;
+						UpdateChangeTypeWidget(typeWidget, details);
 
-						BuildListView.RedrawItems(Item.Index, Item.Index, true);
+						BuildListView.RedrawItems(item.Index, item.Index, true);
 						break;
 					}
 				}
 			}
 		}
 
-		static List<IssueDetailsWindow> ExistingWindows = new List<IssueDetailsWindow>();
+		static List<IssueDetailsWindow> _existingWindows = new List<IssueDetailsWindow>();
 
 		class UpdateIssueDetailsTask
 		{
-			string ApiUrl;
-			long IssueId;
-			List<IssueDiagnosticData> Diagnostics;
+			string _apiUrl;
+			long _issueId;
+			List<IssueDiagnosticData> _diagnostics;
 
-			public UpdateIssueDetailsTask(string ApiUrl, long IssueId, List<IssueDiagnosticData> Diagnostics)
+			public UpdateIssueDetailsTask(string apiUrl, long issueId, List<IssueDiagnosticData> diagnostics)
 			{
-				this.ApiUrl = ApiUrl;
-				this.IssueId = IssueId;
-				this.Diagnostics = Diagnostics;
+				this._apiUrl = apiUrl;
+				this._issueId = issueId;
+				this._diagnostics = diagnostics;
 			}
 
-			public async Task RunAsync(CancellationToken CancellationToken)
+			public async Task RunAsync(CancellationToken cancellationToken)
 			{
-				Diagnostics.AddRange(await RESTApi.GetAsync<List<IssueDiagnosticData>>($"{ApiUrl}/api/issues/{IssueId}/diagnostics", CancellationToken));
+				_diagnostics.AddRange(await RestApi.GetAsync<List<IssueDiagnosticData>>($"{_apiUrl}/api/issues/{_issueId}/diagnostics", cancellationToken));
 			}
 		}
 
-		public static void Show(Form Owner, IssueMonitor IssueMonitor, IPerforceSettings PerforceSettings, TimeSpan? ServerTimeOffset, IssueData Issue, IServiceProvider ServiceProvider, string? CurrentStream)
+		public static void Show(Form owner, IssueMonitor issueMonitor, IPerforceSettings perforceSettings, TimeSpan? serverTimeOffset, IssueData issue, IServiceProvider serviceProvider, string? currentStream)
 		{
-			Func<CancellationToken, Task<List<IssueBuildData>>> Func = x => RESTApi.GetAsync<List<IssueBuildData>>($"{IssueMonitor.ApiUrl}/api/issues/{Issue.Id}/builds", x);
-			ModalTask<List<IssueBuildData>>? IssueBuilds = ModalTask.Execute(Owner, "Querying issue", "Querying issue, please wait...", Func);
-			if (IssueBuilds != null && IssueBuilds.Succeeded)
+			Func<CancellationToken, Task<List<IssueBuildData>>> func = x => RestApi.GetAsync<List<IssueBuildData>>($"{issueMonitor.ApiUrl}/api/issues/{issue.Id}/builds", x);
+			ModalTask<List<IssueBuildData>>? issueBuilds = ModalTask.Execute(owner, "Querying issue", "Querying issue, please wait...", func);
+			if (issueBuilds != null && issueBuilds.Succeeded)
 			{
-				Show(Owner, IssueMonitor, PerforceSettings, ServerTimeOffset, Issue, IssueBuilds.Result, ServiceProvider, CurrentStream);
+				Show(owner, issueMonitor, perforceSettings, serverTimeOffset, issue, issueBuilds.Result, serviceProvider, currentStream);
 			}
 		}
 
-		public static void Show(Form Owner, IssueMonitor IssueMonitor, IPerforceSettings PerforceSettings, TimeSpan? ServerTimeOffset, IssueData Issue, List<IssueBuildData> IssueBuilds, IServiceProvider ServiceProvider, string? CurrentStream)
+		public static void Show(Form owner, IssueMonitor issueMonitor, IPerforceSettings perforceSettings, TimeSpan? serverTimeOffset, IssueData issue, List<IssueBuildData> issueBuilds, IServiceProvider serviceProvider, string? currentStream)
 		{
-			IssueDetailsWindow Window = ExistingWindows.FirstOrDefault(x => x.IssueMonitor == IssueMonitor && x.Issue.Id == Issue.Id);
-			if(Window == null)
+			IssueDetailsWindow window = _existingWindows.FirstOrDefault(x => x._issueMonitor == issueMonitor && x._issue.Id == issue.Id);
+			if(window == null)
 			{
-				List<IssueDiagnosticData> Diagnostics = new List<IssueDiagnosticData>();
+				List<IssueDiagnosticData> diagnostics = new List<IssueDiagnosticData>();
 
-				UpdateIssueDetailsTask Task = new UpdateIssueDetailsTask(IssueMonitor.ApiUrl!, Issue.Id, Diagnostics);
-				ModalTask? TaskResult = ModalTask.Execute(Owner, "Fetching data", "Fetching data, please wait...", Task.RunAsync);
-				if (TaskResult == null || !TaskResult.Succeeded)
+				UpdateIssueDetailsTask task = new UpdateIssueDetailsTask(issueMonitor.ApiUrl!, issue.Id, diagnostics);
+				ModalTask? taskResult = ModalTask.Execute(owner, "Fetching data", "Fetching data, please wait...", task.RunAsync);
+				if (taskResult == null || !taskResult.Succeeded)
 				{
 					return;
 				}
 
-				Window = new IssueDetailsWindow(IssueMonitor, Issue, IssueBuilds, Diagnostics, PerforceSettings, ServerTimeOffset, ServiceProvider, CurrentStream);
-				Window.Owner = Owner;
-				if (Owner.Visible && Owner.WindowState != FormWindowState.Minimized)
+				window = new IssueDetailsWindow(issueMonitor, issue, issueBuilds, diagnostics, perforceSettings, serverTimeOffset, serviceProvider, currentStream);
+				window.Owner = owner;
+				if (owner.Visible && owner.WindowState != FormWindowState.Minimized)
 				{
-					Window.StartPosition = FormStartPosition.Manual;
-					Window.Location = new Point(Owner.Location.X + (Owner.Width - Window.Width) / 2, Owner.Location.Y + (Owner.Height - Window.Height) / 2);
+					window.StartPosition = FormStartPosition.Manual;
+					window.Location = new Point(owner.Location.X + (owner.Width - window.Width) / 2, owner.Location.Y + (owner.Height - window.Height) / 2);
 				}
 				else
 				{
-					Window.StartPosition = FormStartPosition.CenterScreen;
+					window.StartPosition = FormStartPosition.CenterScreen;
 				}
-				Window.Show();
+				window.Show();
 
-				ExistingWindows.Add(Window);
-				Window.FormClosed += (S, E) => ExistingWindows.Remove(Window);
+				_existingWindows.Add(window);
+				window.FormClosed += (s, e) => _existingWindows.Remove(window);
 			}
 			else
 			{
-				Window.Location = new Point(Owner.Location.X + (Owner.Width - Window.Width) / 2, Owner.Location.Y + (Owner.Height - Window.Height) / 2);
-				Window.Activate();
+				window.Location = new Point(owner.Location.X + (owner.Width - window.Width) / 2, owner.Location.Y + (owner.Height - window.Height) / 2);
+				window.Activate();
 			}
 		}
 
@@ -1108,29 +1108,29 @@ namespace UnrealGameSync
 			}
 			else if(e.Item.Tag is BuildGroup)
 			{
-				BuildGroup BuildGroup = (BuildGroup)e.Item.Tag;
+				BuildGroup buildGroup = (BuildGroup)e.Item.Tag;
 
-				Color BackgroundColor;
-				if(BuildGroup.Outcome == IssueBuildOutcome.Error)
+				Color backgroundColor;
+				if(buildGroup.Outcome == IssueBuildOutcome.Error)
 				{
-					BackgroundColor = Color.FromArgb(254, 248, 246);
+					backgroundColor = Color.FromArgb(254, 248, 246);
 				}
-				else if(BuildGroup.Outcome == IssueBuildOutcome.Warning)
+				else if(buildGroup.Outcome == IssueBuildOutcome.Warning)
 				{
-					BackgroundColor = Color.FromArgb(254, 254, 246);
+					backgroundColor = Color.FromArgb(254, 254, 246);
 				}
-				else if(BuildGroup.Outcome == IssueBuildOutcome.Success)
+				else if(buildGroup.Outcome == IssueBuildOutcome.Success)
 				{
-					BackgroundColor = Color.FromArgb(248, 254, 246);
+					backgroundColor = Color.FromArgb(248, 254, 246);
 				}
 				else
 				{
-					BackgroundColor = Color.FromArgb(245, 245, 245);
+					backgroundColor = Color.FromArgb(245, 245, 245);
 				}
 
-				using(SolidBrush Brush = new SolidBrush(BackgroundColor))
+				using(SolidBrush brush = new SolidBrush(backgroundColor))
 				{
-					e.Graphics.FillRectangle(Brush, e.Bounds);
+					e.Graphics.FillRectangle(brush, e.Bounds);
 				}
 			}
 			else
@@ -1141,9 +1141,9 @@ namespace UnrealGameSync
 
 		private void BuildListView_DrawSubItem(object sender, DrawListViewSubItemEventArgs e)
 		{
-			Font CurrentFont = this.Font;//BuildFont;(Change.Number == Workspace.PendingChangeNumber || Change.Number == Workspace.CurrentChangeNumber)? SelectedBuildFont : BuildFont;
+			Font currentFont = this.Font;//BuildFont;(Change.Number == Workspace.PendingChangeNumber || Change.Number == Workspace.CurrentChangeNumber)? SelectedBuildFont : BuildFont;
 
-			Color TextColor = SystemColors.WindowText;
+			Color textColor = SystemColors.WindowText;
 
 			if(e.SubItem.Tag is CustomListViewWidget)
 			{
@@ -1151,48 +1151,48 @@ namespace UnrealGameSync
 			}
 			else if(e.Item.Tag is ChangesRecord)
 			{
-				ChangesRecord Change = (ChangesRecord)e.Item.Tag;
+				ChangesRecord change = (ChangesRecord)e.Item.Tag;
 
-				Font ChangeFont = BuildListView.Font;
-				if(Issue.Owner != null && String.Compare(Change.User, Issue.Owner, StringComparison.OrdinalIgnoreCase) == 0)
+				Font changeFont = BuildListView.Font;
+				if(_issue.Owner != null && String.Compare(change.User, _issue.Owner, StringComparison.OrdinalIgnoreCase) == 0)
 				{
-					ChangeFont = BoldFont!;
+					changeFont = _boldFont!;
 				}
 
 				if(e.ColumnIndex == ChangeHeader.Index)
 				{
-					TextRenderer.DrawText(e.Graphics, e.SubItem.Text, ChangeFont, e.Bounds, TextColor, TextFormatFlags.EndEllipsis | TextFormatFlags.SingleLine | TextFormatFlags.HorizontalCenter | TextFormatFlags.VerticalCenter | TextFormatFlags.NoPrefix);
+					TextRenderer.DrawText(e.Graphics, e.SubItem.Text, changeFont, e.Bounds, textColor, TextFormatFlags.EndEllipsis | TextFormatFlags.SingleLine | TextFormatFlags.HorizontalCenter | TextFormatFlags.VerticalCenter | TextFormatFlags.NoPrefix);
 				}
 				else if(e.ColumnIndex == AuthorHeader.Index)
 				{
-					TextRenderer.DrawText(e.Graphics, e.SubItem.Text, ChangeFont, e.Bounds, TextColor, TextFormatFlags.EndEllipsis | TextFormatFlags.SingleLine | TextFormatFlags.Left | TextFormatFlags.VerticalCenter | TextFormatFlags.NoPrefix);
+					TextRenderer.DrawText(e.Graphics, e.SubItem.Text, changeFont, e.Bounds, textColor, TextFormatFlags.EndEllipsis | TextFormatFlags.SingleLine | TextFormatFlags.Left | TextFormatFlags.VerticalCenter | TextFormatFlags.NoPrefix);
 				}
 				else if(e.ColumnIndex == DescriptionHeader.Index)
 				{
-					TextRenderer.DrawText(e.Graphics, e.SubItem.Text, ChangeFont, e.Bounds, TextColor, TextFormatFlags.EndEllipsis | TextFormatFlags.SingleLine | TextFormatFlags.Left | TextFormatFlags.VerticalCenter | TextFormatFlags.NoPrefix);
+					TextRenderer.DrawText(e.Graphics, e.SubItem.Text, changeFont, e.Bounds, textColor, TextFormatFlags.EndEllipsis | TextFormatFlags.SingleLine | TextFormatFlags.Left | TextFormatFlags.VerticalCenter | TextFormatFlags.NoPrefix);
 				}
 				else if(e.ColumnIndex == TimeHeader.Index)
 				{
-					TextRenderer.DrawText(e.Graphics, e.SubItem.Text, ChangeFont, e.Bounds, TextColor, TextFormatFlags.EndEllipsis | TextFormatFlags.SingleLine | TextFormatFlags.HorizontalCenter | TextFormatFlags.VerticalCenter | TextFormatFlags.NoPrefix);
+					TextRenderer.DrawText(e.Graphics, e.SubItem.Text, changeFont, e.Bounds, textColor, TextFormatFlags.EndEllipsis | TextFormatFlags.SingleLine | TextFormatFlags.HorizontalCenter | TextFormatFlags.VerticalCenter | TextFormatFlags.NoPrefix);
 				}
 			}
 			else if(e.Item.Tag is IssueBuildData)
 			{
-				Font BoldFont = BuildListView.Font;
+				Font boldFont = BuildListView.Font;
 
 //				TextColor = SystemColors.Window;
-				IssueBuildData BuildData = (IssueBuildData)e.Item.Tag;
+				IssueBuildData buildData = (IssueBuildData)e.Item.Tag;
 				if(e.ColumnIndex == IconHeader.Index)
 				{
-					if(BuildData.Outcome == IssueBuildOutcome.Success)
+					if(buildData.Outcome == IssueBuildOutcome.Success)
 					{
 						BuildListView.DrawIcon(e.Graphics, e.Bounds, WorkspaceControl.GoodBuildIcon);
 					}
-					else if(BuildData.Outcome == IssueBuildOutcome.Warning)
+					else if(buildData.Outcome == IssueBuildOutcome.Warning)
 					{
 						BuildListView.DrawIcon(e.Graphics, e.Bounds, WorkspaceControl.MixedBuildIcon);
 					}
-					else if(BuildData.Outcome == IssueBuildOutcome.Error)
+					else if(buildData.Outcome == IssueBuildOutcome.Error)
 					{
 						BuildListView.DrawIcon(e.Graphics, e.Bounds, WorkspaceControl.BadBuildIcon);
 					}
@@ -1203,81 +1203,81 @@ namespace UnrealGameSync
 				}
 				else if(e.ColumnIndex == ChangeHeader.Index)
 				{
-					TextRenderer.DrawText(e.Graphics, BuildData.Change.ToString(), BoldFont, e.Bounds, TextColor, TextFormatFlags.EndEllipsis | TextFormatFlags.SingleLine | TextFormatFlags.HorizontalCenter | TextFormatFlags.VerticalCenter | TextFormatFlags.NoPrefix);
+					TextRenderer.DrawText(e.Graphics, buildData.Change.ToString(), boldFont, e.Bounds, textColor, TextFormatFlags.EndEllipsis | TextFormatFlags.SingleLine | TextFormatFlags.HorizontalCenter | TextFormatFlags.VerticalCenter | TextFormatFlags.NoPrefix);
 				}
 				else if(e.ColumnIndex == TypeHeader.Index)
 				{
-					string Status;
-					if(BuildData.Outcome == IssueBuildOutcome.Success)
+					string status;
+					if(buildData.Outcome == IssueBuildOutcome.Success)
 					{
-						Status = "Succeeded";
+						status = "Succeeded";
 					}
-					else if(BuildData.Outcome == IssueBuildOutcome.Warning)
+					else if(buildData.Outcome == IssueBuildOutcome.Warning)
 					{
-						Status = "Warning";
+						status = "Warning";
 					}
-					else if(BuildData.Outcome == IssueBuildOutcome.Error)
+					else if(buildData.Outcome == IssueBuildOutcome.Error)
 					{
-						Status = "Failed";
+						status = "Failed";
 					}
 					else
 					{
-						Status = "Pending";
+						status = "Pending";
 					}
-					TextRenderer.DrawText(e.Graphics, Status, BoldFont, e.Bounds, TextColor, TextFormatFlags.EndEllipsis | TextFormatFlags.SingleLine | TextFormatFlags.HorizontalCenter | TextFormatFlags.VerticalCenter | TextFormatFlags.NoPrefix);
+					TextRenderer.DrawText(e.Graphics, status, boldFont, e.Bounds, textColor, TextFormatFlags.EndEllipsis | TextFormatFlags.SingleLine | TextFormatFlags.HorizontalCenter | TextFormatFlags.VerticalCenter | TextFormatFlags.NoPrefix);
 				}
 				else if(e.ColumnIndex == AuthorHeader.Index)
 				{
-					Rectangle Bounds = new Rectangle(e.Bounds.X, e.Bounds.Y, e.Bounds.Width + e.Item.SubItems[e.ColumnIndex].Bounds.Width, e.Bounds.Height);
-					TextRenderer.DrawText(e.Graphics, BuildData.JobName, Font, Bounds, TextColor, TextFormatFlags.EndEllipsis | TextFormatFlags.SingleLine | TextFormatFlags.Left | TextFormatFlags.VerticalCenter | TextFormatFlags.NoPrefix);
+					Rectangle bounds = new Rectangle(e.Bounds.X, e.Bounds.Y, e.Bounds.Width + e.Item.SubItems[e.ColumnIndex].Bounds.Width, e.Bounds.Height);
+					TextRenderer.DrawText(e.Graphics, buildData.JobName, Font, bounds, textColor, TextFormatFlags.EndEllipsis | TextFormatFlags.SingleLine | TextFormatFlags.Left | TextFormatFlags.VerticalCenter | TextFormatFlags.NoPrefix);
 				}
 			}
 			else
 			{
-				TextFormatFlags Flags = TextFormatFlags.EndEllipsis | TextFormatFlags.SingleLine | TextFormatFlags.VerticalCenter | TextFormatFlags.NoPrefix;
+				TextFormatFlags flags = TextFormatFlags.EndEllipsis | TextFormatFlags.SingleLine | TextFormatFlags.VerticalCenter | TextFormatFlags.NoPrefix;
 				if(BuildListView.Columns[e.ColumnIndex].TextAlign == HorizontalAlignment.Left)
 				{
-					Flags |= TextFormatFlags.Left;
+					flags |= TextFormatFlags.Left;
 				}
 				else if(BuildListView.Columns[e.ColumnIndex].TextAlign == HorizontalAlignment.Center)
 				{
-					Flags |= TextFormatFlags.HorizontalCenter;
+					flags |= TextFormatFlags.HorizontalCenter;
 				}
 				else
 				{
-					Flags |= TextFormatFlags.Right;
+					flags |= TextFormatFlags.Right;
 				}
-				TextRenderer.DrawText(e.Graphics, e.SubItem.Text, CurrentFont, e.Bounds, TextColor, Flags);
+				TextRenderer.DrawText(e.Graphics, e.SubItem.Text, currentFont, e.Bounds, textColor, flags);
 			}
 		}
 
 		private void BuildListView_FontChanged(object sender, EventArgs e)
 		{
-			if(BoldFont != null)
+			if(_boldFont != null)
 			{
-				BoldFont.Dispose();
+				_boldFont.Dispose();
 			}
-			BoldFont = new Font(BuildListView.Font.FontFamily, BuildListView.Font.SizeInPoints, FontStyle.Bold);
+			_boldFont = new Font(BuildListView.Font.FontFamily, BuildListView.Font.SizeInPoints, FontStyle.Bold);
 
-			if(StatusElementResources != null)
+			if(_statusElementResources != null)
 			{
-				StatusElementResources.Dispose();
+				_statusElementResources.Dispose();
 			}
-			StatusElementResources = new StatusElementResources(BuildListView.Font);
+			_statusElementResources = new StatusElementResources(BuildListView.Font);
 		}
 
-		private void BuildListView_MouseClick(object Sender, MouseEventArgs Args)
+		private void BuildListView_MouseClick(object sender, MouseEventArgs args)
 		{
-			if(Args.Button == MouseButtons.Right)
+			if(args.Button == MouseButtons.Right)
 			{
-				ListViewHitTestInfo HitTest = BuildListView.HitTest(Args.Location);
-				if(HitTest.Item != null && HitTest.Item.Tag != null)
+				ListViewHitTestInfo hitTest = BuildListView.HitTest(args.Location);
+				if(hitTest.Item != null && hitTest.Item.Tag != null)
 				{
-					ContextMenuChange = HitTest.Item.Tag as ChangesRecord; 
-					if(ContextMenuChange != null)
+					_contextMenuChange = hitTest.Item.Tag as ChangesRecord; 
+					if(_contextMenuChange != null)
 					{
-						BuildListContextMenu_Assign.Text = String.Format("Assign to {0}", ContextMenuChange.User);
-						BuildListContextMenu.Show(BuildListView, Args.Location);
+						BuildListContextMenu_Assign.Text = String.Format("Assign to {0}", _contextMenuChange.User);
+						BuildListContextMenu.Show(BuildListView, args.Location);
 					}
 				}
 			}
@@ -1285,52 +1285,52 @@ namespace UnrealGameSync
 
 		private void BuildListContextMenu_MoreInfo_Click(object sender, EventArgs e)
 		{
-			if (ContextMenuChange != null)
+			if (_contextMenuChange != null)
 			{
-				Program.SpawnP4VC(String.Format("{0} change {1}", PerforceSettings.GetArgumentsForExternalProgram(true), ContextMenuChange.Number));
+				Program.SpawnP4Vc(String.Format("{0} change {1}", _perforceSettings.GetArgumentsForExternalProgram(true), _contextMenuChange.Number));
 			}
 		}
 
 		private void BuildListContextMenu_Blame_Click(object sender, EventArgs e)
 		{
-			if(ContextMenuChange != null && ContextMenuChange.User != null)
+			if(_contextMenuChange != null && _contextMenuChange.User != null)
 			{
-				AssignToUser(ContextMenuChange.User);
+				AssignToUser(_contextMenuChange.User);
 			}
 		}
 
-		private void AssignToUser(string User)
+		private void AssignToUser(string user)
 		{
-			IssueUpdateData Update = new IssueUpdateData();
-			Update.Id = Issue.Id;
-			Update.Owner = User;
-			Update.FixChange = 0;
-			if(String.Compare(User, PerforceSettings.UserName, StringComparison.OrdinalIgnoreCase) == 0)
+			IssueUpdateData update = new IssueUpdateData();
+			update.Id = _issue.Id;
+			update.Owner = user;
+			update.FixChange = 0;
+			if(String.Compare(user, _perforceSettings.UserName, StringComparison.OrdinalIgnoreCase) == 0)
 			{
-				Update.NominatedBy = "";
-				Update.Acknowledged = true;
+				update.NominatedBy = "";
+				update.Acknowledged = true;
 			}
 			else
 			{
-				Update.NominatedBy = PerforceSettings.UserName;
-				Update.Acknowledged = false;
+				update.NominatedBy = _perforceSettings.UserName;
+				update.Acknowledged = false;
 			}
-			IssueMonitor.PostUpdate(Update);
+			_issueMonitor.PostUpdate(update);
 
 			UpdateCurrentIssue();
 		}
 
 		private void AssignToMeBtn_Click(object sender, EventArgs e)
 		{
-			AssignToUser(IssueMonitor.UserName);
+			AssignToUser(_issueMonitor.UserName);
 		}
 
 		private void AssignToOtherBtn_Click(object sender, EventArgs e)
 		{
-			string? SelectedUserName;
-			if(SelectUserWindow.ShowModal(this, PerforceSettings, ServiceProvider, out SelectedUserName))
+			string? selectedUserName;
+			if(SelectUserWindow.ShowModal(this, _perforceSettings, _serviceProvider, out selectedUserName))
 			{
-				AssignToUser(SelectedUserName);
+				AssignToUser(selectedUserName);
 			}
 		}
 
@@ -1347,44 +1347,44 @@ namespace UnrealGameSync
 
 		private void MarkFixedBtn_Click(object sender, EventArgs e)
 		{
-			int FixChangeNumber = Issue.FixChange;
-			if(FixChangeNumber == 0)
+			int fixChangeNumber = _issue.FixChange;
+			if(fixChangeNumber == 0)
 			{
-				if(IssueFixedWindow.ShowModal(this, PerforceSettings, ServiceProvider, ref FixChangeNumber))
+				if(IssueFixedWindow.ShowModal(this, _perforceSettings, _serviceProvider, ref fixChangeNumber))
 				{
-					IssueUpdateData Update = new IssueUpdateData();
-					Update.Id = Issue.Id;
-					Update.FixChange = FixChangeNumber;
-					IssueMonitor.PostUpdate(Update);
+					IssueUpdateData update = new IssueUpdateData();
+					update.Id = _issue.Id;
+					update.FixChange = fixChangeNumber;
+					_issueMonitor.PostUpdate(update);
 				}
 			}
 			else
 			{
-				IssueUpdateData Update = new IssueUpdateData();
-				Update.Id = Issue.Id;
-				Update.FixChange = 0;
-				IssueMonitor.PostUpdate(Update);
+				IssueUpdateData update = new IssueUpdateData();
+				update.Id = _issue.Id;
+				update.FixChange = 0;
+				_issueMonitor.PostUpdate(update);
 			}
 		}
 
 		private void DescriptionLabel_LinkClicked(object sender, LinkLabelLinkClickedEventArgs e)
 		{
-			IssueBuildData LastBuild = IssueBuilds.Where(x => x.Stream == SelectedStream).OrderByDescending(x => x.Change).ThenByDescending(x => x.ErrorUrl).FirstOrDefault();
-			if(LastBuild != null)
+			IssueBuildData lastBuild = _issueBuilds.Where(x => x.Stream == _selectedStream).OrderByDescending(x => x.Change).ThenByDescending(x => x.ErrorUrl).FirstOrDefault();
+			if(lastBuild != null)
 			{
-				Utility.OpenUrl(LastBuild.ErrorUrl);
+				Utility.OpenUrl(lastBuild.ErrorUrl);
 			}
 		}
 
 		private void JobContextMenu_ShowError_Click(object sender, EventArgs e)
 		{
-			if (SelectedBuildGroup != null)
+			if (_selectedBuildGroup != null)
 			{
-				foreach (IssueBuildData Build in SelectedBuildGroup.Builds)
+				foreach (IssueBuildData build in _selectedBuildGroup.Builds)
 				{
-					if (Build.ErrorUrl != null)
+					if (build.ErrorUrl != null)
 					{
-						Utility.OpenUrl(Build.ErrorUrl);
+						Utility.OpenUrl(build.ErrorUrl);
 						break;
 					}
 				}
@@ -1393,9 +1393,9 @@ namespace UnrealGameSync
 
 		private void JobContextMenu_ViewJob_Click(object sender, EventArgs e)
 		{
-			if (SelectedBuildGroup != null)
+			if (_selectedBuildGroup != null)
 			{
-				Utility.OpenUrl(SelectedBuildGroup.JobUrl);
+				Utility.OpenUrl(_selectedBuildGroup.JobUrl);
 			}
 		}
 
@@ -1403,13 +1403,13 @@ namespace UnrealGameSync
 		{
 			if((e.Button & MouseButtons.Right) != 0)
 			{
-				ListViewHitTestInfo HitTest = BuildListView.HitTest(e.Location);
-				if(HitTest.Item != null)
+				ListViewHitTestInfo hitTest = BuildListView.HitTest(e.Location);
+				if(hitTest.Item != null)
 				{
-					BuildGroup? Group = HitTest.Item.Tag as BuildGroup;
-					if(Group != null)
+					BuildGroup? group = hitTest.Item.Tag as BuildGroup;
+					if(group != null)
 					{
-						ShowJobContextMenu(e.Location, Group);
+						ShowJobContextMenu(e.Location, group);
 					}
 				}
 			}
@@ -1417,10 +1417,10 @@ namespace UnrealGameSync
 
 		private void BuildLinkLabel_LinkClicked(object sender, LinkLabelLinkClickedEventArgs e)
 		{
-			IssueBuildData Build = IssueBuilds.FirstOrDefault(x => x.ErrorUrl != null);
-			if (Build != null)
+			IssueBuildData build = _issueBuilds.FirstOrDefault(x => x.ErrorUrl != null);
+			if (build != null)
 			{
-				Utility.OpenUrl(Build.ErrorUrl);
+				Utility.OpenUrl(build.ErrorUrl);
 			}
 		}
 
diff --git a/Engine/Source/Programs/UnrealGameSync/UnrealGameSync/Forms/IssueFixedWindow.cs b/Engine/Source/Programs/UnrealGameSync/UnrealGameSync/Forms/IssueFixedWindow.cs
--- a/Engine/Source/Programs/UnrealGameSync/UnrealGameSync/Forms/IssueFixedWindow.cs
+++ b/Engine/Source/Programs/UnrealGameSync/UnrealGameSync/Forms/IssueFixedWindow.cs
@@ -23,27 +23,27 @@ namespace UnrealGameSync
 	{
 		class FindChangesWorker : Component
 		{
-			public delegate void OnCompleteDelegate(string UserName, List<DescribeRecord> Changes);
-
-			IPerforceSettings PerforceSettings;
-			SynchronizationContext MainThreadSyncContext;
-			AsyncEvent WakeEvent;
-			string? RequestedUserName;
-			CancellationTokenSource CancellationSource;
-			Task BackgroundTask;
-			OnCompleteDelegate? OnComplete;
-			ILogger Logger;
-
-			public FindChangesWorker(IPerforceSettings PerforceSettings, OnCompleteDelegate? OnComplete, ILogger Logger)
+			public delegate void OnCompleteDelegate(string userName, List<DescribeRecord> changes);
+
+			IPerforceSettings _perforceSettings;
+			SynchronizationContext _mainThreadSyncContext;
+			AsyncEvent _wakeEvent;
+			string? _requestedUserName;
+			CancellationTokenSource _cancellationSource;
+			Task _backgroundTask;
+			OnCompleteDelegate? _onComplete;
+			ILogger _logger;
+
+			public FindChangesWorker(IPerforceSettings perforceSettings, OnCompleteDelegate? onComplete, ILogger logger)
 			{
-				this.PerforceSettings = PerforceSettings;
-				this.MainThreadSyncContext = SynchronizationContext.Current!;
-				WakeEvent = new AsyncEvent();
-				this.OnComplete = OnComplete;
-				this.Logger = Logger;
-
-				CancellationSource = new CancellationTokenSource();
-				BackgroundTask = Task.Run(() => DoWork(CancellationSource.Token));
+				this._perforceSettings = perforceSettings;
+				this._mainThreadSyncContext = SynchronizationContext.Current!;
+				_wakeEvent = new AsyncEvent();
+				this._onComplete = onComplete;
+				this._logger = logger;
+
+				_cancellationSource = new CancellationTokenSource();
+				_backgroundTask = Task.Run(() => DoWork(_cancellationSource.Token));
 			}
 
 			public void Stop()
@@ -53,23 +53,23 @@ namespace UnrealGameSync
 
 			Task StopAsync()
 			{
-				Task StopTask = Task.CompletedTask;
-				if(BackgroundTask != null)
+				Task stopTask = Task.CompletedTask;
+				if(_backgroundTask != null)
 				{
-					OnComplete = null;
+					_onComplete = null;
 
-					CancellationSource.Cancel();
-					StopTask = BackgroundTask.ContinueWith(_ => CancellationSource.Dispose());
+					_cancellationSource.Cancel();
+					stopTask = _backgroundTask.ContinueWith(_ => _cancellationSource.Dispose());
 
-					BackgroundTask = null!;
+					_backgroundTask = null!;
 				}
-				return StopTask;
+				return stopTask;
 			}
 
-			public void FetchChanges(string UserName)
+			public void FetchChanges(string userName)
 			{
-				RequestedUserName = UserName;
-				WakeEvent.Set();
+				_requestedUserName = userName;
+				_wakeEvent.Set();
 			}
 
 			protected override void Dispose(bool disposing)
@@ -79,55 +79,55 @@ namespace UnrealGameSync
 				Stop();
 			}
 
-			public async Task DoWork(CancellationToken CancellationToken)
+			public async Task DoWork(CancellationToken cancellationToken)
 			{
-				using IPerforceConnection Perforce = await PerforceConnection.CreateAsync(PerforceSettings, Logger);
+				using IPerforceConnection perforce = await PerforceConnection.CreateAsync(_perforceSettings, _logger);
 
-				Task WakeTask = WakeEvent.Task;
-				Task CancelTask = Task.Delay(-1, CancellationToken);
+				Task wakeTask = _wakeEvent.Task;
+				Task cancelTask = Task.Delay(-1, cancellationToken);
 				for(;;)
 				{
-					await Task.WhenAny(WakeTask, CancelTask);
+					await Task.WhenAny(wakeTask, cancelTask);
 
-					if (CancellationToken.IsCancellationRequested)
+					if (cancellationToken.IsCancellationRequested)
 					{
 						break;
 					}
 
-					WakeTask = WakeEvent.Task;
+					wakeTask = _wakeEvent.Task;
 
-					string? UserName = RequestedUserName;
-					if (UserName != null)
+					string? userName = _requestedUserName;
+					if (userName != null)
 					{
-						List<ChangesRecord> Changes = await Perforce.GetChangesAsync(ChangesOptions.IncludeTimes, 100, ChangeStatus.Submitted, FileSpecList.Any, CancellationToken);
-						List<DescribeRecord> Descriptions = await Perforce.DescribeAsync(Changes.Select(x => x.Number).ToArray(), CancellationToken);
+						List<ChangesRecord> changes = await perforce.GetChangesAsync(ChangesOptions.IncludeTimes, 100, ChangeStatus.Submitted, FileSpecList.Any, cancellationToken);
+						List<DescribeRecord> descriptions = await perforce.DescribeAsync(changes.Select(x => x.Number).ToArray(), cancellationToken);
 
-						MainThreadSyncContext.Post(_ => OnComplete?.Invoke(UserName, Descriptions), null);
+						_mainThreadSyncContext.Post(_ => _onComplete?.Invoke(userName, descriptions), null);
 					}
 				}
 			}
 		}
 
-		IPerforceSettings PerforceSettings;
-		int ChangeNumber;
-		FindChangesWorker Worker;
-		IServiceProvider ServiceProvider;
+		IPerforceSettings _perforceSettings;
+		int _changeNumber;
+		FindChangesWorker _worker;
+		IServiceProvider _serviceProvider;
 	
-		public IssueFixedWindow(IPerforceSettings PerforceSettings, int InitialChangeNumber, IServiceProvider ServiceProvider)
+		public IssueFixedWindow(IPerforceSettings perforceSettings, int initialChangeNumber, IServiceProvider serviceProvider)
 		{
 			InitializeComponent();
 
-			this.PerforceSettings = PerforceSettings;
-			this.Worker = new FindChangesWorker(PerforceSettings, PopulateChanges, ServiceProvider.GetRequiredService<ILogger<FindChangesWorker>>());
-			this.ServiceProvider = ServiceProvider;
-			components!.Add(Worker);
+			this._perforceSettings = perforceSettings;
+			this._worker = new FindChangesWorker(perforceSettings, PopulateChanges, serviceProvider.GetRequiredService<ILogger<FindChangesWorker>>());
+			this._serviceProvider = serviceProvider;
+			components!.Add(_worker);
 
-			UserNameTextBox.Text = PerforceSettings.UserName;
+			UserNameTextBox.Text = perforceSettings.UserName;
 			UserNameTextBox.SelectionStart = UserNameTextBox.Text.Length;
 
-			if(InitialChangeNumber != 0)
+			if(initialChangeNumber != 0)
 			{
-				if(InitialChangeNumber < 0)
+				if(initialChangeNumber < 0)
 				{
 					SpecifyChangeRadioButton.Checked = false;
 					SystemicFixRadioButton.Checked = true;
@@ -135,7 +135,7 @@ namespace UnrealGameSync
 				else
 				{
 					SpecifyChangeRadioButton.Checked = true;
-					ChangeNumberTextBox.Text = InitialChangeNumber.ToString();
+					ChangeNumberTextBox.Text = initialChangeNumber.ToString();
 				}
 			}
 
@@ -159,21 +159,21 @@ namespace UnrealGameSync
 				ChangesListView.Select();
 			}
 
-			Worker.FetchChanges(UserNameTextBox.Text);
+			_worker.FetchChanges(UserNameTextBox.Text);
 		}
 
-		public static bool ShowModal(IWin32Window Owner, IPerforceSettings Perforce, IServiceProvider ServiceProvider, ref int FixChangeNumber)
+		public static bool ShowModal(IWin32Window owner, IPerforceSettings perforce, IServiceProvider serviceProvider, ref int fixChangeNumber)
 		{
-			using(IssueFixedWindow FixedWindow = new IssueFixedWindow(Perforce, FixChangeNumber, ServiceProvider))
+			using(IssueFixedWindow fixedWindow = new IssueFixedWindow(perforce, fixChangeNumber, serviceProvider))
 			{
-				if(FixedWindow.ShowDialog(Owner) == DialogResult.OK)
+				if(fixedWindow.ShowDialog(owner) == DialogResult.OK)
 				{
-					FixChangeNumber = FixedWindow.ChangeNumber;
+					fixChangeNumber = fixedWindow._changeNumber;
 					return true;
 				}
 				else
 				{
-					FixChangeNumber = 0;
+					fixChangeNumber = 0;
 					return false;
 				}
 			}
@@ -181,7 +181,7 @@ namespace UnrealGameSync
 
 		private void UpdateSelectedChangeAndClose()
 		{
-			if (TryGetSelectedChange(out ChangeNumber))
+			if (TryGetSelectedChange(out _changeNumber))
 			{
 				DialogResult = DialogResult.OK;
 				Close();
@@ -193,41 +193,41 @@ namespace UnrealGameSync
 			UpdateSelectedChangeAndClose();
 		}
 
-		private void PopulateChanges(string UserName, List<DescribeRecord> Changes)
+		private void PopulateChanges(string userName, List<DescribeRecord> changes)
 		{
 			if(!IsDisposed)
 			{
 				ChangesListView.BeginUpdate();
 				ChangesListView.Items.Clear();
-				if(Changes != null)
+				if(changes != null)
 				{
-					foreach(DescribeRecord Change in Changes)
+					foreach(DescribeRecord change in changes)
 					{
-						if(Change.Description != null && Change.Description.IndexOf("#ROBOMERGE-SOURCE", 0) == -1)
+						if(change.Description != null && change.Description.IndexOf("#ROBOMERGE-SOURCE", 0) == -1)
 						{
-							string Stream = "";
-							if(Change.Files.Count > 0)
+							string stream = "";
+							if(change.Files.Count > 0)
 							{
-								string DepotFile = Change.Files[0].DepotFile;
+								string depotFile = change.Files[0].DepotFile;
 
-								int Idx = 0;
-								for(int Count = 0; Idx < DepotFile.Length; Idx++)
+								int idx = 0;
+								for(int count = 0; idx < depotFile.Length; idx++)
 								{
-									if(DepotFile[Idx] == '/' && ++Count >= 4)
+									if(depotFile[idx] == '/' && ++count >= 4)
 									{
 										break;
 									}
 								}
 
-								Stream = DepotFile.Substring(0, Idx);
+								stream = depotFile.Substring(0, idx);
 							}
 
-							ListViewItem Item = new ListViewItem("");
-							Item.Tag = Change;
-							Item.SubItems.Add(Change.Number.ToString());
-							Item.SubItems.Add(Stream);
-							Item.SubItems.Add(Change.Description.Replace('\n', ' '));
-							ChangesListView.Items.Add(Item);
+							ListViewItem item = new ListViewItem("");
+							item.Tag = change;
+							item.SubItems.Add(change.Number.ToString());
+							item.SubItems.Add(stream);
+							item.SubItems.Add(change.Description.Replace('\n', ' '));
+							ChangesListView.Items.Add(item);
 						}
 					}
 				}
@@ -235,18 +235,18 @@ namespace UnrealGameSync
 			}
 		}
 
-		private void ChangesListView_MouseClick(object Sender, MouseEventArgs Args)
+		private void ChangesListView_MouseClick(object sender, MouseEventArgs args)
 		{
-			if(Args.Button == MouseButtons.Right)
+			if(args.Button == MouseButtons.Right)
 			{
-				ListViewHitTestInfo HitTest = ChangesListView.HitTest(Args.Location);
-				if(HitTest.Item != null && HitTest.Item.Tag != null)
+				ListViewHitTestInfo hitTest = ChangesListView.HitTest(args.Location);
+				if(hitTest.Item != null && hitTest.Item.Tag != null)
 				{
-					DescribeRecord? Record = HitTest.Item.Tag as DescribeRecord; 
-					if(Record != null)
+					DescribeRecord? record = hitTest.Item.Tag as DescribeRecord; 
+					if(record != null)
 					{
-						ChangesListContextMenu.Tag = Record;
-						ChangesListContextMenu.Show(ChangesListView, Args.Location);
+						ChangesListContextMenu.Tag = record;
+						ChangesListContextMenu.Show(ChangesListView, args.Location);
 					}
 				}
 			}
@@ -259,8 +259,8 @@ namespace UnrealGameSync
 
 		private void ChangesContextMenu_MoreInfo_Click(object sender, EventArgs e)
 		{
-			DescribeRecord Record = (DescribeRecord)ChangesListContextMenu.Tag;
-			Program.SpawnP4VC(String.Format("{0} change {1}", PerforceSettings.GetArgumentsForExternalProgram(true), Record.Number));
+			DescribeRecord record = (DescribeRecord)ChangesListContextMenu.Tag;
+			Program.SpawnP4Vc(String.Format("{0} change {1}", _perforceSettings.GetArgumentsForExternalProgram(true), record.Number));
 		}
 
 		private void ChangeNumberTextBox_Enter(object sender, EventArgs e)
@@ -297,28 +297,28 @@ namespace UnrealGameSync
 			UpdateOkButton();
 		}
 
-		private bool TryGetSelectedChange(out int ChangeNumber)
+		private bool TryGetSelectedChange(out int changeNumber)
 		{
 			if(SpecifyChangeRadioButton.Checked)
 			{
-				return int.TryParse(ChangeNumberTextBox.Text, out ChangeNumber);
+				return int.TryParse(ChangeNumberTextBox.Text, out changeNumber);
 			}
 			else if(SystemicFixRadioButton.Checked)
 			{
-				ChangeNumber = -1;
+				changeNumber = -1;
 				return true;
 			}
 			else
 			{
-				DescribeRecord? Change = (ChangesListView.SelectedItems.Count > 0)? ChangesListView.SelectedItems[0].Tag as DescribeRecord : null;
-				if(Change == null)
+				DescribeRecord? change = (ChangesListView.SelectedItems.Count > 0)? ChangesListView.SelectedItems[0].Tag as DescribeRecord : null;
+				if(change == null)
 				{
-					ChangeNumber = 0;
+					changeNumber = 0;
 					return false;
 				}
 				else
 				{
-					ChangeNumber = Change.Number;
+					changeNumber = change.Number;
 					return true;
 				}
 			}
@@ -326,8 +326,8 @@ namespace UnrealGameSync
 
 		private void UpdateOkButton()
 		{
-			int ChangeNumber;
-			OkBtn.Enabled = TryGetSelectedChange(out ChangeNumber);
+			int changeNumber;
+			OkBtn.Enabled = TryGetSelectedChange(out changeNumber);
 		}
 
 		private void SpecifyChangeRadioButton_CheckedChanged(object sender, EventArgs e)
@@ -362,16 +362,16 @@ namespace UnrealGameSync
 
 		private void UserBrowseBtn_Click(object sender, EventArgs e)
 		{
-			string? SelectedUserName;
-			if(SelectUserWindow.ShowModal(this, PerforceSettings, ServiceProvider, out SelectedUserName))
+			string? selectedUserName;
+			if(SelectUserWindow.ShowModal(this, _perforceSettings, _serviceProvider, out selectedUserName))
 			{
-				UserNameTextBox.Text = SelectedUserName;
+				UserNameTextBox.Text = selectedUserName;
 			}
 		}
 
 		private void UserNameTextBox_TextChanged(object sender, EventArgs e)
 		{
-			Worker.FetchChanges(UserNameTextBox.Text);
+			_worker.FetchChanges(UserNameTextBox.Text);
 		}
 	}
 }
diff --git a/Engine/Source/Programs/UnrealGameSync/UnrealGameSync/Forms/IssueSettingsWindow.cs b/Engine/Source/Programs/UnrealGameSync/UnrealGameSync/Forms/IssueSettingsWindow.cs
--- a/Engine/Source/Programs/UnrealGameSync/UnrealGameSync/Forms/IssueSettingsWindow.cs
+++ b/Engine/Source/Programs/UnrealGameSync/UnrealGameSync/Forms/IssueSettingsWindow.cs
@@ -15,28 +15,28 @@ namespace UnrealGameSync
 {
 	partial class IssueSettingsWindow : Form
 	{
-		UserSettings Settings;
-		ILogger Logger;
+		UserSettings _settings;
+		ILogger _logger;
 
-		public IssueSettingsWindow(UserSettings Settings, string CurrentProject, ILogger Logger)
+		public IssueSettingsWindow(UserSettings settings, string currentProject, ILogger logger)
 		{
-			this.Settings = Settings;
-			this.Logger = Logger;
+			this._settings = settings;
+			this._logger = logger;
 
 			InitializeComponent();
 
-			if (Settings.NotifyProjects.Count == 0)
+			if (settings.NotifyProjects.Count == 0)
 			{
 				NotifyProjectsCheckBox.Checked = false;
-				NotifyProjectsTextBox.Text = CurrentProject;
+				NotifyProjectsTextBox.Text = currentProject;
 			}
 			else
 			{
 				NotifyProjectsCheckBox.Checked = true;
-				NotifyProjectsTextBox.Text = String.Join(" ", Settings.NotifyProjects);
+				NotifyProjectsTextBox.Text = String.Join(" ", settings.NotifyProjects);
 			}
 
-			if(Settings.NotifyUnassignedMinutes < 0)
+			if(settings.NotifyUnassignedMinutes < 0)
 			{
 				NotifyUnassignedCheckBox.Checked = false;
 				NotifyUnassignedTextBox.Text = "5";
@@ -44,10 +44,10 @@ namespace UnrealGameSync
 			else
 			{
 				NotifyUnassignedCheckBox.Checked = true;
-				NotifyUnassignedTextBox.Text = Settings.NotifyUnassignedMinutes.ToString();
+				NotifyUnassignedTextBox.Text = settings.NotifyUnassignedMinutes.ToString();
 			}
 
-			if (Settings.NotifyUnacknowledgedMinutes < 0)
+			if (settings.NotifyUnacknowledgedMinutes < 0)
 			{
 				NotifyUnacknowledgedCheckBox.Checked = false;
 				NotifyUnacknowledgedTextBox.Text = "5";
@@ -55,10 +55,10 @@ namespace UnrealGameSync
 			else
 			{
 				NotifyUnacknowledgedCheckBox.Checked = true;
-				NotifyUnacknowledgedTextBox.Text = Settings.NotifyUnacknowledgedMinutes.ToString();
+				NotifyUnacknowledgedTextBox.Text = settings.NotifyUnacknowledgedMinutes.ToString();
 			}
 
-			if (Settings.NotifyUnresolvedMinutes < 0)
+			if (settings.NotifyUnresolvedMinutes < 0)
 			{
 				NotifyUnresolvedCheckBox.Checked = false;
 				NotifyUnresolvedTextBox.Text = "20";
@@ -66,7 +66,7 @@ namespace UnrealGameSync
 			else
 			{
 				NotifyUnresolvedCheckBox.Checked = true;
-				NotifyUnresolvedTextBox.Text = Settings.NotifyUnresolvedMinutes.ToString();
+				NotifyUnresolvedTextBox.Text = settings.NotifyUnresolvedMinutes.ToString();
 			}
 
 			UpdateEnabledTextBoxes();
@@ -102,53 +102,53 @@ namespace UnrealGameSync
 
 		private void OkBtn_Click(object sender, EventArgs e)
 		{
-			List<string> NewNotifyProjects = new List<string>();
+			List<string> newNotifyProjects = new List<string>();
 			if (NotifyProjectsCheckBox.Checked)
 			{
-				NewNotifyProjects.AddRange(NotifyProjectsTextBox.Text.Split(new char[] { ' ' }, StringSplitOptions.RemoveEmptyEntries));
+				newNotifyProjects.AddRange(NotifyProjectsTextBox.Text.Split(new char[] { ' ' }, StringSplitOptions.RemoveEmptyEntries));
 			}
 
-			int NewNotifyUnresolvedMinutes = -1;
+			int newNotifyUnresolvedMinutes = -1;
 			if(NotifyUnresolvedCheckBox.Checked)
 			{
-				ushort NewNotifyUnresolvedMinutesValue;
-				if(!ushort.TryParse(NotifyUnresolvedTextBox.Text, out NewNotifyUnresolvedMinutesValue))
+				ushort newNotifyUnresolvedMinutesValue;
+				if(!ushort.TryParse(NotifyUnresolvedTextBox.Text, out newNotifyUnresolvedMinutesValue))
 				{
 					MessageBox.Show("Invalid time");
 					return;
 				}
-				NewNotifyUnresolvedMinutes = NewNotifyUnresolvedMinutesValue;
+				newNotifyUnresolvedMinutes = newNotifyUnresolvedMinutesValue;
 			}
 
-			int NewNotifyUnacknowledgedMinutes = -1;
+			int newNotifyUnacknowledgedMinutes = -1;
 			if (NotifyUnacknowledgedCheckBox.Checked)
 			{
-				ushort NewNotifyUnacknowledgedMinutesValue;
-				if (!ushort.TryParse(NotifyUnacknowledgedTextBox.Text, out NewNotifyUnacknowledgedMinutesValue))
+				ushort newNotifyUnacknowledgedMinutesValue;
+				if (!ushort.TryParse(NotifyUnacknowledgedTextBox.Text, out newNotifyUnacknowledgedMinutesValue))
 				{
 					MessageBox.Show("Invalid time");
 					return;
 				}
-				NewNotifyUnacknowledgedMinutes = NewNotifyUnacknowledgedMinutesValue;
+				newNotifyUnacknowledgedMinutes = newNotifyUnacknowledgedMinutesValue;
 			}
 
-			int NewNotifyUnassignedMinutes = -1;
+			int newNotifyUnassignedMinutes = -1;
 			if(NotifyUnassignedCheckBox.Checked)
 			{
-				ushort NewNotifyUnassignedMinutesValue;
-				if(!ushort.TryParse(NotifyUnassignedTextBox.Text, out NewNotifyUnassignedMinutesValue))
+				ushort newNotifyUnassignedMinutesValue;
+				if(!ushort.TryParse(NotifyUnassignedTextBox.Text, out newNotifyUnassignedMinutesValue))
 				{
 					MessageBox.Show("Invalid time");
 					return;
 				}
-				NewNotifyUnassignedMinutes = NewNotifyUnassignedMinutesValue;
+				newNotifyUnassignedMinutes = newNotifyUnassignedMinutesValue;
 			}
 
-			Settings.NotifyProjects = NewNotifyProjects;
-			Settings.NotifyUnresolvedMinutes = NewNotifyUnresolvedMinutes;
-			Settings.NotifyUnacknowledgedMinutes = NewNotifyUnacknowledgedMinutes;
-			Settings.NotifyUnassignedMinutes = NewNotifyUnassignedMinutes;
-			Settings.Save(Logger);
+			_settings.NotifyProjects = newNotifyProjects;
+			_settings.NotifyUnresolvedMinutes = newNotifyUnresolvedMinutes;
+			_settings.NotifyUnacknowledgedMinutes = newNotifyUnacknowledgedMinutes;
+			_settings.NotifyUnassignedMinutes = newNotifyUnassignedMinutes;
+			_settings.Save(_logger);
 
 			DialogResult = DialogResult.OK;
 			Close();
