UGS: Reformat to standard C# style conventions.

diff --git a/Engine/Source/Programs/UnrealGameSync/UnrealGameSync/Forms/MainWindow.cs b/Engine/Source/Programs/UnrealGameSync/UnrealGameSync/Forms/MainWindow.cs
--- a/Engine/Source/Programs/UnrealGameSync/UnrealGameSync/Forms/MainWindow.cs
+++ b/Engine/Source/Programs/UnrealGameSync/UnrealGameSync/Forms/MainWindow.cs
@@ -60,9 +60,9 @@ namespace UnrealGameSync
 			public IssueMonitor IssueMonitor;
 			public int RefCount;
 
-			public WorkspaceIssueMonitor(IssueMonitor IssueMonitor)
+			public WorkspaceIssueMonitor(IssueMonitor issueMonitor)
 			{
-				this.IssueMonitor = IssueMonitor;
+				this.IssueMonitor = issueMonitor;
 			}
 		}
 
@@ -79,68 +79,68 @@ namespace UnrealGameSync
 		[DllImport("user32.dll")]
 		public static extern int SendMessage(IntPtr hWnd, Int32 wMsg, Int32 wParam, Int32 lParam);
 
-		private const int WM_SETREDRAW = 11; 
+		private const int WmSetredraw = 11; 
 
-		UpdateMonitor UpdateMonitor;
-		SynchronizationContext MainThreadSynchronizationContext;
-		List<IssueMonitor> DefaultIssueMonitors = new List<IssueMonitor>();
-		List<WorkspaceIssueMonitor> WorkspaceIssueMonitors = new List<WorkspaceIssueMonitor>();
+		UpdateMonitor _updateMonitor;
+		SynchronizationContext _mainThreadSynchronizationContext;
+		List<IssueMonitor> _defaultIssueMonitors = new List<IssueMonitor>();
+		List<WorkspaceIssueMonitor> _workspaceIssueMonitors = new List<WorkspaceIssueMonitor>();
 
-		string? ApiUrl;
-		DirectoryReference DataFolder;
-		DirectoryReference CacheFolder;
-		IPerforceSettings DefaultPerforceSettings;
-		IServiceProvider ServiceProvider;
-		ILogger Logger;
-		UserSettings Settings;
-		int TabMenu_TabIdx = -1;
-		int ChangingWorkspacesRefCount;
+		string? _apiUrl;
+		DirectoryReference _dataFolder;
+		DirectoryReference _cacheFolder;
+		IPerforceSettings _defaultPerforceSettings;
+		IServiceProvider _serviceProvider;
+		ILogger _logger;
+		UserSettings _settings;
+		int _tabMenuTabIdx = -1;
+		int _changingWorkspacesRefCount;
 
-		bool bAllowClose = false;
+		bool _allowClose = false;
 
-		bool bRestoreStateOnLoad;
+		bool _restoreStateOnLoad;
 
-		OIDCTokenManager? OIDCTokenManager;
+		OidcTokenManager? _oidcTokenManager;
 
-		System.Threading.Timer? ScheduleTimer;
-		System.Threading.Timer? ScheduleSettledTimer;
+		System.Threading.Timer? _scheduleTimer;
+		System.Threading.Timer? _scheduleSettledTimer;
 
-		string OriginalExecutableFileName;
-		bool bPreview;
+		string _originalExecutableFileName;
+		bool _preview;
 
-		IMainWindowTabPanel? CurrentTabPanel;
+		IMainWindowTabPanel? _currentTabPanel;
 
-		AutomationServer AutomationServer;
+		AutomationServer _automationServer;
 
-		bool bAllowCreatingHandle;
+		bool _allowCreatingHandle;
 
-		Rectangle PrimaryWorkArea;
-		List<IssueAlertWindow> AlertWindows = new List<IssueAlertWindow>();
+		Rectangle _primaryWorkArea;
+		List<IssueAlertWindow> _alertWindows = new List<IssueAlertWindow>();
 
 		public ToolUpdateMonitor ToolUpdateMonitor { get; private set; }
 
-		public MainWindow(UpdateMonitor InUpdateMonitor, string? InApiUrl, DirectoryReference InDataFolder, DirectoryReference InCacheFolder, bool bInRestoreStateOnLoad, string InOriginalExecutableFileName, bool bInPreview, List<(UserSelectedProjectSettings, ModalTask<OpenProjectInfo>)> StartupTasks, IPerforceSettings InDefaultPerforceSettings, IServiceProvider InServiceProvider, UserSettings InSettings, string? InUri, OIDCTokenManager? InOidcTokenManager)
+		public MainWindow(UpdateMonitor inUpdateMonitor, string? inApiUrl, DirectoryReference inDataFolder, DirectoryReference inCacheFolder, bool inRestoreStateOnLoad, string inOriginalExecutableFileName, bool inPreview, List<(UserSelectedProjectSettings, ModalTask<OpenProjectInfo>)> startupTasks, IPerforceSettings inDefaultPerforceSettings, IServiceProvider inServiceProvider, UserSettings inSettings, string? inUri, OidcTokenManager? inOidcTokenManager)
 		{
-			ServiceProvider = InServiceProvider;
-			Logger = ServiceProvider.GetRequiredService<ILogger<MainWindow>>();
+			_serviceProvider = inServiceProvider;
+			_logger = _serviceProvider.GetRequiredService<ILogger<MainWindow>>();
 
-			Logger.LogInformation("Opening Main Window for {NumProject} projects. Last Project {LastProject}", StartupTasks.Count, InSettings.LastProject);
+			_logger.LogInformation("Opening Main Window for {NumProject} projects. Last Project {LastProject}", startupTasks.Count, inSettings.LastProject);
 
 			InitializeComponent();
 
-			UpdateMonitor = InUpdateMonitor;
-			MainThreadSynchronizationContext = SynchronizationContext.Current!;
-			ApiUrl = InApiUrl;
-			DataFolder = InDataFolder;
-			CacheFolder = InCacheFolder;
-			bRestoreStateOnLoad = bInRestoreStateOnLoad;
-			OriginalExecutableFileName = InOriginalExecutableFileName;
-			bPreview = bInPreview;
-			DefaultPerforceSettings = InDefaultPerforceSettings;
-			ToolUpdateMonitor = new ToolUpdateMonitor(DefaultPerforceSettings, DataFolder, InSettings, ServiceProvider);
+			_updateMonitor = inUpdateMonitor;
+			_mainThreadSynchronizationContext = SynchronizationContext.Current!;
+			_apiUrl = inApiUrl;
+			_dataFolder = inDataFolder;
+			_cacheFolder = inCacheFolder;
+			_restoreStateOnLoad = inRestoreStateOnLoad;
+			_originalExecutableFileName = inOriginalExecutableFileName;
+			_preview = inPreview;
+			_defaultPerforceSettings = inDefaultPerforceSettings;
+			ToolUpdateMonitor = new ToolUpdateMonitor(_defaultPerforceSettings, _dataFolder, inSettings, _serviceProvider);
 
-			Settings = InSettings;
-			OIDCTokenManager = InOidcTokenManager;
+			_settings = inSettings;
+			_oidcTokenManager = inOidcTokenManager;
 
 			// While creating tab controls during startup, we need to prevent layout calls resulting in the window handle being created too early. Disable layout calls here.
 			SuspendLayout();
@@ -156,29 +156,29 @@ namespace UnrealGameSync
 
 			SetupDefaultControl();
 
-			int SelectTabIdx = -1;
-			foreach((UserSelectedProjectSettings Project, ModalTask<OpenProjectInfo> StartupTask) in StartupTasks)
+			int selectTabIdx = -1;
+			foreach((UserSelectedProjectSettings project, ModalTask<OpenProjectInfo> startupTask) in startupTasks)
 			{
-				int TabIdx = -1;
-				if (StartupTask.Succeeded)
+				int tabIdx = -1;
+				if (startupTask.Succeeded)
 				{
-					TabIdx = TryOpenProject(StartupTask.Result, -1, OpenProjectOptions.Quiet);
+					tabIdx = TryOpenProject(startupTask.Result, -1, OpenProjectOptions.Quiet);
 				}
-				else if(StartupTask.Failed)
+				else if(startupTask.Failed)
 				{
-					Logger.LogError("StartupProject Error: {Message}", StartupTask.Error);
-					CreateErrorPanel(-1, Project, StartupTask.Error);
+					_logger.LogError("StartupProject Error: {Message}", startupTask.Error);
+					CreateErrorPanel(-1, project, startupTask.Error);
 				}
 
-				if (TabIdx != -1 && Settings.LastProject != null && Project.Equals(Settings.LastProject))
+				if (tabIdx != -1 && _settings.LastProject != null && project.Equals(_settings.LastProject))
 				{
-					SelectTabIdx = TabIdx;
+					selectTabIdx = tabIdx;
 				}
 			}
 
-			if(SelectTabIdx != -1)
+			if(selectTabIdx != -1)
 			{
-				TabControl.SelectTab(SelectTabIdx);
+				TabControl.SelectTab(selectTabIdx);
 			}
 			else if(TabControl.GetTabCount() > 0)
 			{
@@ -187,266 +187,266 @@ namespace UnrealGameSync
 
 			StartScheduleTimer();
 
-			if(bPreview)
+			if(_preview)
 			{
 				Text += $" {Program.GetVersionString()} (UNSTABLE)";
 			}
 
-			ILogger<AutomationServer> AutomationLogger = ServiceProvider.GetRequiredService<ILogger<AutomationServer>>();
-			AutomationServer = new AutomationServer(Request => { MainThreadSynchronizationContext.Post(Obj => PostAutomationRequest(Request), null); }, InUri, AutomationLogger);
+			ILogger<AutomationServer> automationLogger = _serviceProvider.GetRequiredService<ILogger<AutomationServer>>();
+			_automationServer = new AutomationServer(request => { _mainThreadSynchronizationContext.Post(obj => PostAutomationRequest(request), null); }, inUri, automationLogger);
 
 			// Allow creating controls from now on
 			TabPanel.ResumeLayout(false);
 			ResumeLayout(false);
 
-			foreach (string DefaultIssueApiUrl in DeploymentSettings.DefaultIssueApiUrls)
+			foreach (string defaultIssueApiUrl in DeploymentSettings.DefaultIssueApiUrls)
 			{
-				DefaultIssueMonitors.Add(CreateIssueMonitor(DefaultIssueApiUrl, InDefaultPerforceSettings.UserName));
+				_defaultIssueMonitors.Add(CreateIssueMonitor(defaultIssueApiUrl, inDefaultPerforceSettings.UserName));
 			}
 
-			bAllowCreatingHandle = true;
+			_allowCreatingHandle = true;
 
-			foreach(WorkspaceIssueMonitor WorkspaceIssueMonitor in WorkspaceIssueMonitors)
+			foreach(WorkspaceIssueMonitor workspaceIssueMonitor in _workspaceIssueMonitors)
 			{
-				WorkspaceIssueMonitor.IssueMonitor.Start();
+				workspaceIssueMonitor.IssueMonitor.Start();
 			}
 
 			ToolUpdateMonitor.Start();
 		}
 
-		void PostAutomationRequest(AutomationRequest Request)
+		void PostAutomationRequest(AutomationRequest request)
 		{
 			try
 			{
 				if (!CanFocus)
 				{
-					Request.SetOutput(new AutomationRequestOutput(AutomationRequestResult.Busy));
+					request.SetOutput(new AutomationRequestOutput(AutomationRequestResult.Busy));
 				}
-				else if (Request.Input.Type == AutomationRequestType.SyncProject)
+				else if (request.Input.Type == AutomationRequestType.SyncProject)
 				{
-					AutomationRequestOutput? Output = StartAutomatedSync(Request, true);
-					if (Output != null)
+					AutomationRequestOutput? output = StartAutomatedSync(request, true);
+					if (output != null)
 					{
-						Request.SetOutput(Output);
+						request.SetOutput(output);
 					}
 				}
-				else if (Request.Input.Type == AutomationRequestType.FindProject)
+				else if (request.Input.Type == AutomationRequestType.FindProject)
 				{
-					AutomationRequestOutput Output = FindProject(Request);
-					Request.SetOutput(Output);
+					AutomationRequestOutput output = FindProject(request);
+					request.SetOutput(output);
 				}
-				else if (Request.Input.Type == AutomationRequestType.OpenProject)
+				else if (request.Input.Type == AutomationRequestType.OpenProject)
 				{
-					AutomationRequestOutput? Output = StartAutomatedSync(Request, false);
-					if (Output != null)
+					AutomationRequestOutput? output = StartAutomatedSync(request, false);
+					if (output != null)
 					{
-						Request.SetOutput(Output);
+						request.SetOutput(output);
 					}
 				}
-				else if (Request.Input.Type == AutomationRequestType.ExecCommand)
+				else if (request.Input.Type == AutomationRequestType.ExecCommand)
 				{
-					AutomationRequestOutput Output = StartExecCommand(Request);
-					Request.SetOutput(Output);
+					AutomationRequestOutput output = StartExecCommand(request);
+					request.SetOutput(output);
 				}
-				else if (Request.Input.Type == AutomationRequestType.OpenIssue)
+				else if (request.Input.Type == AutomationRequestType.OpenIssue)
 				{
-					AutomationRequestOutput Output = OpenIssue(Request);
-					Request.SetOutput(new AutomationRequestOutput(AutomationRequestResult.Ok));
+					AutomationRequestOutput output = OpenIssue(request);
+					request.SetOutput(new AutomationRequestOutput(AutomationRequestResult.Ok));
 				}
 				else
 				{
-					Request.SetOutput(new AutomationRequestOutput(AutomationRequestResult.Invalid));
+					request.SetOutput(new AutomationRequestOutput(AutomationRequestResult.Invalid));
 				}
 			}
-			catch(Exception Ex)
+			catch(Exception ex)
 			{
-				Logger.LogError(Ex, "Exception running automation request");
-				Request.SetOutput(new AutomationRequestOutput(AutomationRequestResult.Invalid));
+				_logger.LogError(ex, "Exception running automation request");
+				request.SetOutput(new AutomationRequestOutput(AutomationRequestResult.Invalid));
 			}
 		}
 
-		AutomationRequestOutput StartExecCommand(AutomationRequest Request)
+		AutomationRequestOutput StartExecCommand(AutomationRequest request)
 		{
-			BinaryReader Reader = new BinaryReader(new MemoryStream(Request.Input.Data));
-			string StreamName = Reader.ReadString();
-			int Changelist = Reader.ReadInt32();
-			string Command = Reader.ReadString();
-			string ProjectPath = Reader.ReadString();
+			BinaryReader reader = new BinaryReader(new MemoryStream(request.Input.Data));
+			string streamName = reader.ReadString();
+			int changelist = reader.ReadInt32();
+			string command = reader.ReadString();
+			string projectPath = reader.ReadString();
 
-			AutomatedBuildWindow.BuildInfo? BuildInfo;
-			if (!AutomatedBuildWindow.ShowModal(this, DefaultPerforceSettings, StreamName, ProjectPath, Changelist, Command.ToString(), Settings, ServiceProvider, out BuildInfo))
+			AutomatedBuildWindow.BuildInfo? buildInfo;
+			if (!AutomatedBuildWindow.ShowModal(this, _defaultPerforceSettings, streamName, projectPath, changelist, command.ToString(), _settings, _serviceProvider, out buildInfo))
 			{
 				return new AutomationRequestOutput(AutomationRequestResult.Canceled);
 			}
 
-			WorkspaceControl? Workspace;
-			if (!OpenWorkspaceForAutomation(BuildInfo.SelectedWorkspaceInfo, StreamName, BuildInfo.ProjectPath, out Workspace))
+			WorkspaceControl? workspace;
+			if (!OpenWorkspaceForAutomation(buildInfo.SelectedWorkspaceInfo, streamName, buildInfo.ProjectPath, out workspace))
 			{
 				return new AutomationRequestOutput(AutomationRequestResult.Error);
 			}
 
-			Workspace.AddStartupCallback((Control, bCancel) => StartExecCommandAfterStartup(Control, bCancel, BuildInfo.bSync? Changelist : -1, BuildInfo.ExecCommand));
+			workspace.AddStartupCallback((control, cancel) => StartExecCommandAfterStartup(control, cancel, buildInfo.Sync? changelist : -1, buildInfo.ExecCommand));
 			return new AutomationRequestOutput(AutomationRequestResult.Ok);
 		}
 
-		private void StartExecCommandAfterStartup(WorkspaceControl Workspace, bool bCancel, int Changelist, string Command)
+		private void StartExecCommandAfterStartup(WorkspaceControl workspace, bool cancel, int changelist, string command)
 		{
-			if (!bCancel)
+			if (!cancel)
 			{
-				if(Changelist == -1)
+				if(changelist == -1)
 				{
-					StartExecCommandAfterSync(Workspace, WorkspaceUpdateResult.Success, Command);
+					StartExecCommandAfterSync(workspace, WorkspaceUpdateResult.Success, command);
 				}
 				else
 				{
-					Workspace.SyncChange(Changelist, true, Result => StartExecCommandAfterSync(Workspace, Result, Command));
+					workspace.SyncChange(changelist, true, result => StartExecCommandAfterSync(workspace, result, command));
 				}
 			}
 		}
 
-		private void StartExecCommandAfterSync(WorkspaceControl Workspace, WorkspaceUpdateResult Result, string Command)
+		private void StartExecCommandAfterSync(WorkspaceControl workspace, WorkspaceUpdateResult result, string command)
 		{
-			if (Result == WorkspaceUpdateResult.Success && Command != null)
+			if (result == WorkspaceUpdateResult.Success && command != null)
 			{
-				string CmdExe = Environment.GetEnvironmentVariable("COMSPEC") ?? "C:\\Windows\\System32\\cmd.exe";
-				Workspace.ExecCommand("Run build command", "Running build command", CmdExe, String.Format("/c {0}", Command), Workspace.BranchDirectoryName.FullName, true);
+				string cmdExe = Environment.GetEnvironmentVariable("COMSPEC") ?? "C:\\Windows\\System32\\cmd.exe";
+				workspace.ExecCommand("Run build command", "Running build command", cmdExe, String.Format("/c {0}", command), workspace.BranchDirectoryName.FullName, true);
 			}
 		}
 
-		AutomationRequestOutput? StartAutomatedSync(AutomationRequest Request, bool bForceSync)
+		AutomationRequestOutput? StartAutomatedSync(AutomationRequest request, bool forceSync)
 		{
 			ShowAndActivate();
 
-			BinaryReader Reader = new BinaryReader(new MemoryStream(Request.Input.Data));
-			string StreamName = Reader.ReadString();
-			string ProjectPath = Reader.ReadString();
+			BinaryReader reader = new BinaryReader(new MemoryStream(request.Input.Data));
+			string streamName = reader.ReadString();
+			string projectPath = reader.ReadString();
 
-			AutomatedSyncWindow.WorkspaceInfo? WorkspaceInfo;
-			if(!AutomatedSyncWindow.ShowModal(this, DefaultPerforceSettings, StreamName, ProjectPath, out WorkspaceInfo, ServiceProvider))
+			AutomatedSyncWindow.WorkspaceInfo? workspaceInfo;
+			if(!AutomatedSyncWindow.ShowModal(this, _defaultPerforceSettings, streamName, projectPath, out workspaceInfo, _serviceProvider))
 			{
 				return new AutomationRequestOutput(AutomationRequestResult.Canceled);
 			}
 
-			WorkspaceControl? Workspace;
-			if(!OpenWorkspaceForAutomation(WorkspaceInfo, StreamName, ProjectPath, out Workspace))
+			WorkspaceControl? workspace;
+			if(!OpenWorkspaceForAutomation(workspaceInfo, streamName, projectPath, out workspace))
 			{
 				return new AutomationRequestOutput(AutomationRequestResult.Error);
 			}
 
-			if(!bForceSync && Workspace.CanLaunchEditor())
+			if(!forceSync && workspace.CanLaunchEditor())
 			{
-				return new AutomationRequestOutput(AutomationRequestResult.Ok, Encoding.UTF8.GetBytes(Workspace.SelectedFileName.FullName));
+				return new AutomationRequestOutput(AutomationRequestResult.Ok, Encoding.UTF8.GetBytes(workspace.SelectedFileName.FullName));
 			}
 
-			Workspace.AddStartupCallback((Control, bCancel) => StartAutomatedSyncAfterStartup(Control, bCancel, Request));
+			workspace.AddStartupCallback((control, cancel) => StartAutomatedSyncAfterStartup(control, cancel, request));
 			return null;
 		}
 
-		private bool OpenWorkspaceForAutomation(AutomatedSyncWindow.WorkspaceInfo WorkspaceInfo, string StreamName, string ProjectPath, [NotNullWhen(true)] out WorkspaceControl? OutWorkspace)
+		private bool OpenWorkspaceForAutomation(AutomatedSyncWindow.WorkspaceInfo workspaceInfo, string streamName, string projectPath, [NotNullWhen(true)] out WorkspaceControl? outWorkspace)
 		{
-			if (WorkspaceInfo.bRequiresStreamSwitch)
+			if (workspaceInfo.RequiresStreamSwitch)
 			{
 				// Close any tab containing this window
-				for (int ExistingTabIdx = 0; ExistingTabIdx < TabControl.GetTabCount(); ExistingTabIdx++)
+				for (int existingTabIdx = 0; existingTabIdx < TabControl.GetTabCount(); existingTabIdx++)
 				{
-					WorkspaceControl? ExistingWorkspace = TabControl.GetTabData(ExistingTabIdx) as WorkspaceControl;
-					if (ExistingWorkspace != null && ExistingWorkspace.ClientName.Equals(WorkspaceInfo.WorkspaceName))
+					WorkspaceControl? existingWorkspace = TabControl.GetTabData(existingTabIdx) as WorkspaceControl;
+					if (existingWorkspace != null && existingWorkspace.ClientName.Equals(workspaceInfo.WorkspaceName))
 					{
-						TabControl.RemoveTab(ExistingTabIdx);
+						TabControl.RemoveTab(existingTabIdx);
 						break;
 					}
 				}
 
 				// Switch the stream
-				Func<IPerforceConnection, CancellationToken, Task> SwitchTask = async (Connection, CancellationToken) =>
+				Func<IPerforceConnection, CancellationToken, Task> switchTask = async (connection, cancellationToken) =>
 				{
-					await Connection.SwitchClientToStreamAsync(StreamName, SwitchClientOptions.None, CancellationToken);
+					await connection.SwitchClientToStreamAsync(streamName, SwitchClientOptions.None, cancellationToken);
 				};
 
-				PerforceSettings Settings = new PerforceSettings(WorkspaceInfo.ServerAndPort, WorkspaceInfo.UserName);
-				Settings.ClientName = WorkspaceInfo.WorkspaceName;
+				PerforceSettings settings = new PerforceSettings(workspaceInfo.ServerAndPort, workspaceInfo.UserName);
+				settings.ClientName = workspaceInfo.WorkspaceName;
 
-				ModalTask? Result = PerforceModalTask.Execute(Owner, "Please wait", "Switching streams, please wait...", Settings, SwitchTask, Logger, ModalTaskFlags.Quiet);
-				if (Result == null || !Result.Succeeded)
+				ModalTask? result = PerforceModalTask.Execute(Owner, "Please wait", "Switching streams, please wait...", settings, switchTask, _logger, ModalTaskFlags.Quiet);
+				if (result == null || !result.Succeeded)
 				{
-					Logger.LogError("Unable to switch stream ({Message})", Result?.Error ?? "Operation cancelled");
-					OutWorkspace = null;
+					_logger.LogError("Unable to switch stream ({Message})", result?.Error ?? "Operation cancelled");
+					outWorkspace = null;
 					return false;
 				}
 			}
 
-			UserSelectedProjectSettings SelectedProject = new UserSelectedProjectSettings(WorkspaceInfo.ServerAndPort, WorkspaceInfo.UserName, UserSelectedProjectType.Client, String.Format("//{0}{1}", WorkspaceInfo.WorkspaceName, ProjectPath), null);
+			UserSelectedProjectSettings selectedProject = new UserSelectedProjectSettings(workspaceInfo.ServerAndPort, workspaceInfo.UserName, UserSelectedProjectType.Client, String.Format("//{0}{1}", workspaceInfo.WorkspaceName, projectPath), null);
 
-			int TabIdx = TryOpenProject(SelectedProject, -1, OpenProjectOptions.None);
-			if (TabIdx == -1)
+			int tabIdx = TryOpenProject(selectedProject, -1, OpenProjectOptions.None);
+			if (tabIdx == -1)
 			{
-				Logger.LogError("Unable to open project");
-				OutWorkspace = null;
+				_logger.LogError("Unable to open project");
+				outWorkspace = null;
 				return false;
 			}
 
-			WorkspaceControl? Workspace = TabControl.GetTabData(TabIdx) as WorkspaceControl;
-			if (Workspace == null)
+			WorkspaceControl? workspace = TabControl.GetTabData(tabIdx) as WorkspaceControl;
+			if (workspace == null)
 			{
-				Logger.LogError("Workspace was unable to open");
-				OutWorkspace = null;
+				_logger.LogError("Workspace was unable to open");
+				outWorkspace = null;
 				return false;
 			}
 
-			TabControl.SelectTab(TabIdx);
-			OutWorkspace = Workspace;
+			TabControl.SelectTab(tabIdx);
+			outWorkspace = workspace;
 			return true;
 		}
 
-		private void StartAutomatedSyncAfterStartup(WorkspaceControl Workspace, bool bCancel, AutomationRequest Request)
+		private void StartAutomatedSyncAfterStartup(WorkspaceControl workspace, bool cancel, AutomationRequest request)
 		{
-			if(bCancel)
+			if(cancel)
 			{
-				Request.SetOutput(new AutomationRequestOutput(AutomationRequestResult.Canceled));
+				request.SetOutput(new AutomationRequestOutput(AutomationRequestResult.Canceled));
 			}
 			else
 			{
-				Workspace.SyncLatestChange(Result => CompleteAutomatedSync(Result, Workspace.SelectedFileName, Request));
+				workspace.SyncLatestChange(result => CompleteAutomatedSync(result, workspace.SelectedFileName, request));
 			}
 		}
 
-		void CompleteAutomatedSync(WorkspaceUpdateResult Result, FileReference SelectedFileName, AutomationRequest Request)
+		void CompleteAutomatedSync(WorkspaceUpdateResult result, FileReference selectedFileName, AutomationRequest request)
 		{
-			if(Result == WorkspaceUpdateResult.Success)
+			if(result == WorkspaceUpdateResult.Success)
 			{
-				Request.SetOutput(new AutomationRequestOutput(AutomationRequestResult.Ok, Encoding.UTF8.GetBytes(SelectedFileName.FullName)));
+				request.SetOutput(new AutomationRequestOutput(AutomationRequestResult.Ok, Encoding.UTF8.GetBytes(selectedFileName.FullName)));
 			}
-			else if(Result == WorkspaceUpdateResult.Canceled)
+			else if(result == WorkspaceUpdateResult.Canceled)
 			{
-				Request.SetOutput(new AutomationRequestOutput(AutomationRequestResult.Canceled));
+				request.SetOutput(new AutomationRequestOutput(AutomationRequestResult.Canceled));
 			}
 			else
 			{
-				Request.SetOutput(new AutomationRequestOutput(AutomationRequestResult.Error));
+				request.SetOutput(new AutomationRequestOutput(AutomationRequestResult.Error));
 			}
 		}
 
-		AutomationRequestOutput FindProject(AutomationRequest Request)
+		AutomationRequestOutput FindProject(AutomationRequest request)
 		{
-			BinaryReader Reader = new BinaryReader(new MemoryStream(Request.Input.Data));
-			string StreamName = Reader.ReadString();
-			string ProjectPath = Reader.ReadString();
+			BinaryReader reader = new BinaryReader(new MemoryStream(request.Input.Data));
+			string streamName = reader.ReadString();
+			string projectPath = reader.ReadString();
 
-			for(int ExistingTabIdx = 0; ExistingTabIdx < TabControl.GetTabCount(); ExistingTabIdx++)
+			for(int existingTabIdx = 0; existingTabIdx < TabControl.GetTabCount(); existingTabIdx++)
 			{
-				WorkspaceControl? ExistingWorkspace = TabControl.GetTabData(ExistingTabIdx) as WorkspaceControl;
-				if(ExistingWorkspace != null && String.Compare(ExistingWorkspace.StreamName, StreamName, StringComparison.OrdinalIgnoreCase) == 0 && ExistingWorkspace.SelectedProject != null)
+				WorkspaceControl? existingWorkspace = TabControl.GetTabData(existingTabIdx) as WorkspaceControl;
+				if(existingWorkspace != null && String.Compare(existingWorkspace.StreamName, streamName, StringComparison.OrdinalIgnoreCase) == 0 && existingWorkspace.SelectedProject != null)
 				{
-					string? ClientPath = ExistingWorkspace.SelectedProject.ClientPath;
-					if(ClientPath != null && ClientPath.StartsWith("//"))
+					string? clientPath = existingWorkspace.SelectedProject.ClientPath;
+					if(clientPath != null && clientPath.StartsWith("//"))
 					{
-						int SlashIdx = ClientPath.IndexOf('/', 2);
-						if(SlashIdx != -1)
+						int slashIdx = clientPath.IndexOf('/', 2);
+						if(slashIdx != -1)
 						{
-							string ExistingProjectPath = ClientPath.Substring(SlashIdx);
-							if(String.Compare(ExistingProjectPath, ProjectPath, StringComparison.OrdinalIgnoreCase) == 0)
+							string existingProjectPath = clientPath.Substring(slashIdx);
+							if(String.Compare(existingProjectPath, projectPath, StringComparison.OrdinalIgnoreCase) == 0)
 							{
-								return new AutomationRequestOutput(AutomationRequestResult.Ok, Encoding.UTF8.GetBytes(ExistingWorkspace.SelectedFileName.FullName));
+								return new AutomationRequestOutput(AutomationRequestResult.Ok, Encoding.UTF8.GetBytes(existingWorkspace.SelectedFileName.FullName));
 							}
 						}
 					}
@@ -456,43 +456,43 @@ namespace UnrealGameSync
 			return new AutomationRequestOutput(AutomationRequestResult.NotFound);
 		}
 
-		AutomationRequestOutput OpenIssue(AutomationRequest Request)
+		AutomationRequestOutput OpenIssue(AutomationRequest request)
 		{
-			BinaryReader Reader = new BinaryReader(new MemoryStream(Request.Input.Data));
-			int IssueId = Reader.ReadInt32();
+			BinaryReader reader = new BinaryReader(new MemoryStream(request.Input.Data));
+			int issueId = reader.ReadInt32();
 
-			for (int ExistingTabIdx = 0; ExistingTabIdx < TabControl.GetTabCount(); ExistingTabIdx++)
+			for (int existingTabIdx = 0; existingTabIdx < TabControl.GetTabCount(); existingTabIdx++)
 			{
-				WorkspaceControl? ExistingWorkspace = TabControl.GetTabData(ExistingTabIdx) as WorkspaceControl;
-				if (ExistingWorkspace != null)
+				WorkspaceControl? existingWorkspace = TabControl.GetTabData(existingTabIdx) as WorkspaceControl;
+				if (existingWorkspace != null)
 				{
-					IssueMonitor IssueMonitor = ExistingWorkspace.GetIssueMonitor();
-					if (IssueMonitor != null && (DeploymentSettings.UrlHandleIssueApi == null || String.Compare(IssueMonitor.ApiUrl, DeploymentSettings.UrlHandleIssueApi, StringComparison.OrdinalIgnoreCase) == 0))
+					IssueMonitor issueMonitor = existingWorkspace.GetIssueMonitor();
+					if (issueMonitor != null && (DeploymentSettings.UrlHandleIssueApi == null || String.Compare(issueMonitor.ApiUrl, DeploymentSettings.UrlHandleIssueApi, StringComparison.OrdinalIgnoreCase) == 0))
 					{
-						IssueMonitor.AddRef();
+						issueMonitor.AddRef();
 						try
 						{
-							Func<CancellationToken, Task<IssueData>> Func = x => RESTApi.GetAsync<IssueData>($"{IssueMonitor.ApiUrl}/api/issues/{IssueId}", x);
-							ModalTask<IssueData>? IssueTask = ModalTask.Execute(this, "Finding Issue", "Querying issue data, please wait...", Func);
-							if (IssueTask == null)
+							Func<CancellationToken, Task<IssueData>> func = x => RestApi.GetAsync<IssueData>($"{issueMonitor.ApiUrl}/api/issues/{issueId}", x);
+							ModalTask<IssueData>? issueTask = ModalTask.Execute(this, "Finding Issue", "Querying issue data, please wait...", func);
+							if (issueTask == null)
 							{
-								Logger.LogInformation("Operation cancelled");
+								_logger.LogInformation("Operation cancelled");
 								return new AutomationRequestOutput(AutomationRequestResult.Canceled);
 							}
-							else if (IssueTask.Succeeded)
+							else if (issueTask.Succeeded)
 							{
-								ExistingWorkspace.ShowIssueDetails(IssueTask.Result);
+								existingWorkspace.ShowIssueDetails(issueTask.Result);
 								return new AutomationRequestOutput(AutomationRequestResult.Ok);
 							}
 							else
 							{
-								Logger.LogError(IssueTask?.Exception, "Unable to query issue {IssueId} from {ApiUrl}: {Error}", IssueId, IssueMonitor.ApiUrl, IssueTask?.Error);
+								_logger.LogError(issueTask?.Exception, "Unable to query issue {IssueId} from {ApiUrl}: {Error}", issueId, issueMonitor.ApiUrl, issueTask?.Error);
 								return new AutomationRequestOutput(AutomationRequestResult.Error);
 							}
 						}
 						finally
 						{
-							IssueMonitor.Release();
+							issueMonitor.Release();
 						}
 					}
 				}
@@ -505,54 +505,54 @@ namespace UnrealGameSync
 		{
 			base.OnHandleCreated(e);
 
-			Debug.Assert(bAllowCreatingHandle, "Window handle should not be created before constructor has run.");
+			Debug.Assert(_allowCreatingHandle, "Window handle should not be created before constructor has run.");
 		}
 
-		void TabControl_OnButtonClick(int ButtonIdx, Point Location, MouseButtons Buttons)
+		void TabControl_OnButtonClick(int buttonIdx, Point location, MouseButtons buttons)
 		{
-			if(ButtonIdx == 0)
+			if(buttonIdx == 0)
 			{
 				EditSelectedProject(TabControl.GetSelectedTabIndex());
 			}
 		}
 
-		void TabControl_OnTabClicked(object? TabData, Point Location, MouseButtons Buttons)
+		void TabControl_OnTabClicked(object? tabData, Point location, MouseButtons buttons)
 		{
-			if(Buttons == System.Windows.Forms.MouseButtons.Right)
+			if(buttons == System.Windows.Forms.MouseButtons.Right)
 			{
 				Activate();
 
-				int InsertIdx = 0;
+				int insertIdx = 0;
 
-				while(TabMenu_RecentProjects.DropDownItems[InsertIdx] != TabMenu_Recent_Separator)
+				while(TabMenu_RecentProjects.DropDownItems[insertIdx] != TabMenu_Recent_Separator)
 				{
-					TabMenu_RecentProjects.DropDownItems.RemoveAt(InsertIdx);
+					TabMenu_RecentProjects.DropDownItems.RemoveAt(insertIdx);
 				}
 
-				TabMenu_TabIdx = -1;
-				for(int Idx = 0; Idx < TabControl.GetTabCount(); Idx++)
+				_tabMenuTabIdx = -1;
+				for(int idx = 0; idx < TabControl.GetTabCount(); idx++)
 				{
-					if(TabControl.GetTabData(Idx) == TabData)
+					if(TabControl.GetTabData(idx) == tabData)
 					{
-						TabMenu_TabIdx = Idx;
+						_tabMenuTabIdx = idx;
 						break;
 					}
 				}
 
-				foreach(UserSelectedProjectSettings RecentProject in Settings.RecentProjects)
+				foreach(UserSelectedProjectSettings recentProject in _settings.RecentProjects)
 				{
-					ToolStripMenuItem Item = new ToolStripMenuItem(RecentProject.ToString(), null, new EventHandler((o, e) => TabMenu_OpenRecentProject_Click(RecentProject, TabMenu_TabIdx)));
-					TabMenu_RecentProjects.DropDownItems.Insert(InsertIdx, Item);
-					InsertIdx++;
+					ToolStripMenuItem item = new ToolStripMenuItem(recentProject.ToString(), null, new EventHandler((o, e) => TabMenu_OpenRecentProject_Click(recentProject, _tabMenuTabIdx)));
+					TabMenu_RecentProjects.DropDownItems.Insert(insertIdx, item);
+					insertIdx++;
 				}
 
-				TabMenu_RecentProjects.Visible = (Settings.RecentProjects.Count > 0);
+				TabMenu_RecentProjects.Visible = (_settings.RecentProjects.Count > 0);
 
-				TabMenu_TabNames_Stream.Checked = Settings.TabLabels == TabLabels.Stream;
-				TabMenu_TabNames_WorkspaceName.Checked = Settings.TabLabels == TabLabels.WorkspaceName;
-				TabMenu_TabNames_WorkspaceRoot.Checked = Settings.TabLabels == TabLabels.WorkspaceRoot;
-				TabMenu_TabNames_ProjectFile.Checked = Settings.TabLabels == TabLabels.ProjectFile;
-				TabMenu.Show(TabControl, Location);
+				TabMenu_TabNames_Stream.Checked = _settings.TabLabels == TabLabels.Stream;
+				TabMenu_TabNames_WorkspaceName.Checked = _settings.TabLabels == TabLabels.WorkspaceName;
+				TabMenu_TabNames_WorkspaceRoot.Checked = _settings.TabLabels == TabLabels.WorkspaceRoot;
+				TabMenu_TabNames_ProjectFile.Checked = _settings.TabLabels == TabLabels.ProjectFile;
+				TabMenu.Show(TabControl, location);
 
 				TabControl.LockHover();
 			}
@@ -563,22 +563,22 @@ namespace UnrealGameSync
 			SaveTabSettings();
 		}
 
-		void TabControl_OnTabClosed(object Data)
+		void TabControl_OnTabClosed(object data)
 		{
-			IMainWindowTabPanel TabPanel = (IMainWindowTabPanel)Data;
-			if(CurrentTabPanel == TabPanel)
+			IMainWindowTabPanel tabPanel = (IMainWindowTabPanel)data;
+			if(_currentTabPanel == tabPanel)
 			{
-				CurrentTabPanel = null;
+				_currentTabPanel = null;
 			}
-			TabPanel.Dispose();
+			tabPanel.Dispose();
 
 			SaveTabSettings();
 		}
 
-		bool TabControl_OnTabClosing(object TabData)
+		bool TabControl_OnTabClosing(object tabData)
 		{
-			IMainWindowTabPanel TabPanel = (IMainWindowTabPanel)TabData;
-			return TabPanel.CanClose();
+			IMainWindowTabPanel tabPanel = (IMainWindowTabPanel)tabData;
+			return tabPanel.CanClose();
 		}
 
 		/// <summary>
@@ -592,24 +592,24 @@ namespace UnrealGameSync
 				components.Dispose();
 			}
 
-			for(int Idx = 0; Idx < TabControl.GetTabCount(); Idx++)
+			for(int idx = 0; idx < TabControl.GetTabCount(); idx++)
 			{
-				((IMainWindowTabPanel)TabControl.GetTabData(Idx)).Dispose();
+				((IMainWindowTabPanel)TabControl.GetTabData(idx)).Dispose();
 			}
 
 			StopScheduleTimer();
 
-			foreach (IssueMonitor DefaultIssueMonitor in DefaultIssueMonitors)
+			foreach (IssueMonitor defaultIssueMonitor in _defaultIssueMonitors)
 			{
-				ReleaseIssueMonitor(DefaultIssueMonitor);
+				ReleaseIssueMonitor(defaultIssueMonitor);
 			}
-			DefaultIssueMonitors.Clear();
-			Debug.Assert(WorkspaceIssueMonitors.Count == 0);
+			_defaultIssueMonitors.Clear();
+			Debug.Assert(_workspaceIssueMonitors.Count == 0);
 
-			if(AutomationServer != null)
+			if(_automationServer != null)
 			{
-				AutomationServer.Dispose();
-				AutomationServer = null!;
+				_automationServer.Dispose();
+				_automationServer = null!;
 			}
 
 			if (ToolUpdateMonitor != null)
@@ -623,10 +623,10 @@ namespace UnrealGameSync
 
 		public bool ConfirmClose()
 		{
-			for (int Idx = 0; Idx < TabControl.GetTabCount(); Idx++)
+			for (int idx = 0; idx < TabControl.GetTabCount(); idx++)
 			{
-				IMainWindowTabPanel TabPanel = (IMainWindowTabPanel)TabControl.GetTabData(Idx);
-				if (!TabPanel.CanClose())
+				IMainWindowTabPanel tabPanel = (IMainWindowTabPanel)TabControl.GetTabData(idx);
+				if (!tabPanel.CanClose())
 				{
 					return false;
 				}
@@ -634,23 +634,23 @@ namespace UnrealGameSync
 			return true;
 		}
 
-		private void MainWindow_FormClosing(object Sender, FormClosingEventArgs EventArgs)
+		private void MainWindow_FormClosing(object sender, FormClosingEventArgs eventArgs)
 		{
-			if(!bAllowClose && Settings.bKeepInTray)
+			if(!_allowClose && _settings.KeepInTray)
 			{
 				Hide();
-				EventArgs.Cancel = true;
+				eventArgs.Cancel = true;
 			}
 			else
 			{
-				if (!bAllowClose)
+				if (!_allowClose)
 				{
-					for (int Idx = 0; Idx < TabControl.GetTabCount(); Idx++)
+					for (int idx = 0; idx < TabControl.GetTabCount(); idx++)
 					{
-						IMainWindowTabPanel TabPanel = (IMainWindowTabPanel)TabControl.GetTabData(Idx)!;
-						if (!TabPanel.CanClose())
+						IMainWindowTabPanel tabPanel = (IMainWindowTabPanel)TabControl.GetTabData(idx)!;
+						if (!tabPanel.CanClose())
 						{
-							EventArgs.Cancel = true;
+							eventArgs.Cancel = true;
 							return;
 						}
 					}
@@ -659,120 +659,120 @@ namespace UnrealGameSync
 				StopScheduleTimer();
 			}
 
-			Settings.bWindowVisible = Visible;
-			Settings.WindowState = WindowState.ToString();
+			_settings.WindowVisible = Visible;
+			_settings.WindowState = WindowState.ToString();
 			if(WindowState == FormWindowState.Normal)
 			{
-				Settings.WindowBounds = new Rectangle(Location, Size);
+				_settings.WindowBounds = new Rectangle(Location, Size);
 			}
 			else
 			{
-				Settings.WindowBounds = RestoreBounds;
+				_settings.WindowBounds = RestoreBounds;
 			}
 
-			Settings.Save(Logger);
+			_settings.Save(_logger);
 		}
 
 		private void SetupDefaultControl()
 		{
-			List<StatusLine> Lines = new List<StatusLine>();
+			List<StatusLine> lines = new List<StatusLine>();
 
-			StatusLine SummaryLine = new StatusLine();
-			SummaryLine.AddText("To get started, open an existing Unreal project file on your hard drive.");
-			Lines.Add(SummaryLine);
+			StatusLine summaryLine = new StatusLine();
+			summaryLine.AddText("To get started, open an existing Unreal project file on your hard drive.");
+			lines.Add(summaryLine);
 
-			StatusLine OpenLine = new StatusLine();
-			OpenLine.AddLink("Open project...", FontStyle.Bold | FontStyle.Underline, () => { OpenNewProject(); });
-			OpenLine.AddText("  |  ");
-			OpenLine.AddLink("Application settings...", FontStyle.Bold | FontStyle.Underline, () => { ModifyApplicationSettings(); });
-			Lines.Add(OpenLine);
+			StatusLine openLine = new StatusLine();
+			openLine.AddLink("Open project...", FontStyle.Bold | FontStyle.Underline, () => { OpenNewProject(); });
+			openLine.AddText("  |  ");
+			openLine.AddLink("Application settings...", FontStyle.Bold | FontStyle.Underline, () => { ModifyApplicationSettings(); });
+			lines.Add(openLine);
 
-			DefaultControl.Set(Lines, null, null, null);
+			DefaultControl.Set(lines, null, null, null);
 		}
 
-		private void CreateErrorPanel(int ReplaceTabIdx, UserSelectedProjectSettings Project, string Message)
+		private void CreateErrorPanel(int replaceTabIdx, UserSelectedProjectSettings project, string message)
 		{
-			Logger.LogError("{Error}", Message ?? "Unknown error");
+			_logger.LogError("{Error}", message ?? "Unknown error");
 
-			ErrorPanel ErrorPanel = new ErrorPanel(Project);
-			ErrorPanel.Parent = TabPanel;
-			ErrorPanel.BorderStyle = BorderStyle.FixedSingle;
-			ErrorPanel.BackColor = System.Drawing.Color.FromArgb(((int)(((byte)(250)))), ((int)(((byte)(250)))), ((int)(((byte)(250)))));
-			ErrorPanel.Location = new Point(0, 0);
-			ErrorPanel.Dock = DockStyle.Fill;
-			ErrorPanel.Hide();
+			ErrorPanel errorPanel = new ErrorPanel(project);
+			errorPanel.Parent = TabPanel;
+			errorPanel.BorderStyle = BorderStyle.FixedSingle;
+			errorPanel.BackColor = System.Drawing.Color.FromArgb(((int)(((byte)(250)))), ((int)(((byte)(250)))), ((int)(((byte)(250)))));
+			errorPanel.Location = new Point(0, 0);
+			errorPanel.Dock = DockStyle.Fill;
+			errorPanel.Hide();
 
-			string SummaryText = String.Format("Unable to open '{0}'.", Project.ToString());
+			string summaryText = String.Format("Unable to open '{0}'.", project.ToString());
 
-			int NewContentWidth = Math.Max(TextRenderer.MeasureText(SummaryText, ErrorPanel.Font).Width, 400);
-			if(!String.IsNullOrEmpty(Message))
+			int newContentWidth = Math.Max(TextRenderer.MeasureText(summaryText, errorPanel.Font).Width, 400);
+			if(!String.IsNullOrEmpty(message))
 			{
-				NewContentWidth = Math.Max(NewContentWidth, TextRenderer.MeasureText(Message, ErrorPanel.Font).Width);
+				newContentWidth = Math.Max(newContentWidth, TextRenderer.MeasureText(message, errorPanel.Font).Width);
 			}
 
-			ErrorPanel.SetContentWidth(NewContentWidth);
+			errorPanel.SetContentWidth(newContentWidth);
 
-			List<StatusLine> Lines = new List<StatusLine>();
+			List<StatusLine> lines = new List<StatusLine>();
 
-			StatusLine SummaryLine = new StatusLine();
-			SummaryLine.AddText(SummaryText);
-			Lines.Add(SummaryLine);
+			StatusLine summaryLine = new StatusLine();
+			summaryLine.AddText(summaryText);
+			lines.Add(summaryLine);
 
-			if(!String.IsNullOrEmpty(Message))
+			if(!String.IsNullOrEmpty(message))
 			{
-				Lines.Add(new StatusLine(){ LineHeight = 0.5f });
+				lines.Add(new StatusLine(){ LineHeight = 0.5f });
 
-				foreach(string MessageLine in Message.Split('\n'))
+				foreach(string messageLine in message.Split('\n'))
 				{
-					StatusLine ErrorLine = new StatusLine();
-					ErrorLine.AddText(MessageLine);
-					ErrorLine.LineHeight = 0.8f;
-					Lines.Add(ErrorLine);
+					StatusLine errorLine = new StatusLine();
+					errorLine.AddText(messageLine);
+					errorLine.LineHeight = 0.8f;
+					lines.Add(errorLine);
 				}
 			}
 
-			Lines.Add(new StatusLine(){ LineHeight = 0.5f });
+			lines.Add(new StatusLine(){ LineHeight = 0.5f });
 
-			StatusLine ActionLine = new StatusLine();
-			ActionLine.AddLink("Retry", FontStyle.Bold | FontStyle.Underline, () => { BeginInvoke(new MethodInvoker(() => { TryOpenProject(Project, TabControl.FindTabIndex(ErrorPanel)); })); });
-			ActionLine.AddText(" | ");
-			ActionLine.AddLink("Settings", FontStyle.Bold | FontStyle.Underline, () => { BeginInvoke(new MethodInvoker(() => { EditSelectedProject(ErrorPanel); })); });
-			ActionLine.AddText(" | ");
-			ActionLine.AddLink("Close", FontStyle.Bold | FontStyle.Underline, () => { BeginInvoke(new MethodInvoker(() => { TabControl.RemoveTab(TabControl.FindTabIndex(ErrorPanel)); })); });
-			Lines.Add(ActionLine);
+			StatusLine actionLine = new StatusLine();
+			actionLine.AddLink("Retry", FontStyle.Bold | FontStyle.Underline, () => { BeginInvoke(new MethodInvoker(() => { TryOpenProject(project, TabControl.FindTabIndex(errorPanel)); })); });
+			actionLine.AddText(" | ");
+			actionLine.AddLink("Settings", FontStyle.Bold | FontStyle.Underline, () => { BeginInvoke(new MethodInvoker(() => { EditSelectedProject(errorPanel); })); });
+			actionLine.AddText(" | ");
+			actionLine.AddLink("Close", FontStyle.Bold | FontStyle.Underline, () => { BeginInvoke(new MethodInvoker(() => { TabControl.RemoveTab(TabControl.FindTabIndex(errorPanel)); })); });
+			lines.Add(actionLine);
 
-			ErrorPanel.Set(Lines, null, null, null);
+			errorPanel.Set(lines, null, null, null);
 
-			string NewProjectName = "Unknown";
-			if(Project.Type == UserSelectedProjectType.Client && Project.ClientPath != null)
+			string newProjectName = "Unknown";
+			if(project.Type == UserSelectedProjectType.Client && project.ClientPath != null)
 			{
-				NewProjectName = Project.ClientPath.Substring(Project.ClientPath.LastIndexOf('/') + 1);
+				newProjectName = project.ClientPath.Substring(project.ClientPath.LastIndexOf('/') + 1);
 			}
-			if(Project.Type == UserSelectedProjectType.Local && Project.LocalPath != null)
+			if(project.Type == UserSelectedProjectType.Local && project.LocalPath != null)
 			{
-				NewProjectName = Project.LocalPath.Substring(Project.LocalPath.LastIndexOfAny(new char[]{ '/', '\\' }) + 1);
+				newProjectName = project.LocalPath.Substring(project.LocalPath.LastIndexOfAny(new char[]{ '/', '\\' }) + 1);
 			}
 
-			string NewTabName = String.Format("Error: {0}", NewProjectName);
-			if (ReplaceTabIdx == -1)
+			string newTabName = String.Format("Error: {0}", newProjectName);
+			if (replaceTabIdx == -1)
 			{
-				int TabIdx = TabControl.InsertTab(-1, NewTabName, ErrorPanel, ErrorPanel.TintColor);
-				TabControl.SelectTab(TabIdx);
+				int tabIdx = TabControl.InsertTab(-1, newTabName, errorPanel, errorPanel.TintColor);
+				TabControl.SelectTab(tabIdx);
 			}
 			else
 			{
-				TabControl.InsertTab(ReplaceTabIdx + 1, NewTabName, ErrorPanel, ErrorPanel.TintColor);
-				TabControl.RemoveTab(ReplaceTabIdx);
-				TabControl.SelectTab(ReplaceTabIdx);
+				TabControl.InsertTab(replaceTabIdx + 1, newTabName, errorPanel, errorPanel.TintColor);
+				TabControl.RemoveTab(replaceTabIdx);
+				TabControl.SelectTab(replaceTabIdx);
 			}
 
 			UpdateProgress();
 		}
 
 		[DllImport("user32.dll")]
-		private static extern int ShowWindow(IntPtr hWnd, uint Msg);
+		private static extern int ShowWindow(IntPtr hWnd, uint msg);
 
-		private const uint SW_RESTORE = 0x09;
+		private const uint SwRestore = 0x09;
 
 		public void ShowAndActivate()
 		{
@@ -781,12 +781,12 @@ namespace UnrealGameSync
 				Show();
 				if (WindowState == FormWindowState.Minimized)
 				{
-					ShowWindow(Handle, SW_RESTORE);
+					ShowWindow(Handle, SwRestore);
 				}
 				Activate();
 
-				Settings.bWindowVisible = Visible;
-				Settings.Save(Logger);
+				_settings.WindowVisible = Visible;
+				_settings.Save(_logger);
 			}
 		}
 
@@ -797,10 +797,10 @@ namespace UnrealGameSync
 				return false;
 			}
 
-			for (int TabIdx = 0; TabIdx < TabControl.GetTabCount(); TabIdx++)
+			for (int tabIdx = 0; tabIdx < TabControl.GetTabCount(); tabIdx++)
 			{
-				IMainWindowTabPanel TabPanel = (IMainWindowTabPanel)TabControl.GetTabData(TabIdx)!;
-				if(TabPanel.IsBusy())
+				IMainWindowTabPanel tabPanel = (IMainWindowTabPanel)TabControl.GetTabData(tabIdx)!;
+				if(tabPanel.IsBusy())
 				{
 					return false;
 				}
@@ -811,49 +811,49 @@ namespace UnrealGameSync
 
 		public bool CanSyncNow()
 		{
-			return CurrentTabPanel != null && CurrentTabPanel.CanSyncNow();
+			return _currentTabPanel != null && _currentTabPanel.CanSyncNow();
 		}
 
 		public bool CanLaunchEditor()
 		{
-			return CurrentTabPanel != null && CurrentTabPanel.CanLaunchEditor();
+			return _currentTabPanel != null && _currentTabPanel.CanLaunchEditor();
 		}
 
 		public void SyncLatestChange()
 		{
-			if(CurrentTabPanel != null)
+			if(_currentTabPanel != null)
 			{
-				CurrentTabPanel.SyncLatestChange();
+				_currentTabPanel.SyncLatestChange();
 			}
 		}
 
 		public void LaunchEditor()
 		{
-			if(CurrentTabPanel != null)
+			if(_currentTabPanel != null)
 			{
-				CurrentTabPanel.LaunchEditor();
+				_currentTabPanel.LaunchEditor();
 			}
 		}
 
 		public void ForceClose()
 		{
-			bAllowClose = true;
+			_allowClose = true;
 			Close();
 		}
 
 		private void MainWindow_Activated(object sender, EventArgs e)
 		{
-			if(CurrentTabPanel != null)
+			if(_currentTabPanel != null)
 			{
-				CurrentTabPanel.Activate();
+				_currentTabPanel.Activate();
 			}
 		}
 
 		private void MainWindow_Deactivate(object sender, EventArgs e)
 		{
-			if(CurrentTabPanel != null)
+			if(_currentTabPanel != null)
 			{
-				CurrentTabPanel.Deactivate();
+				_currentTabPanel.Deactivate();
 			}
 		}
 
@@ -861,28 +861,28 @@ namespace UnrealGameSync
 		{
 			StopScheduleTimer();
 
-			List<UserSelectedProjectSettings> OpenProjects = new List<UserSelectedProjectSettings>();
-			for(int TabIdx = 0; TabIdx < TabControl.GetTabCount(); TabIdx++)
+			List<UserSelectedProjectSettings> openProjects = new List<UserSelectedProjectSettings>();
+			for(int tabIdx = 0; tabIdx < TabControl.GetTabCount(); tabIdx++)
 			{
-				IMainWindowTabPanel TabPanel = (IMainWindowTabPanel)TabControl.GetTabData(TabIdx);
-				OpenProjects.Add(TabPanel.SelectedProject);
+				IMainWindowTabPanel tabPanel = (IMainWindowTabPanel)TabControl.GetTabData(tabIdx);
+				openProjects.Add(tabPanel.SelectedProject);
 			}
 
-			Dictionary<UserSelectedProjectSettings, List<LatestChangeType>> ProjectToLatestChangeTypes = new Dictionary<UserSelectedProjectSettings, List<LatestChangeType>>();
-			for (int Idx = 0; Idx < TabControl.GetTabCount(); Idx++)
+			Dictionary<UserSelectedProjectSettings, List<LatestChangeType>> projectToLatestChangeTypes = new Dictionary<UserSelectedProjectSettings, List<LatestChangeType>>();
+			for (int idx = 0; idx < TabControl.GetTabCount(); idx++)
 			{
-				WorkspaceControl? Workspace = TabControl.GetTabData(Idx) as WorkspaceControl;
-				if (Workspace != null)
+				WorkspaceControl? workspace = TabControl.GetTabData(idx) as WorkspaceControl;
+				if (workspace != null)
 				{
-					ProjectToLatestChangeTypes.Add(Workspace.SelectedProject, Workspace.GetCustomLatestChangeTypes());
+					projectToLatestChangeTypes.Add(workspace.SelectedProject, workspace.GetCustomLatestChangeTypes());
 				}
 			}
 
-			ScheduleWindow Schedule = new ScheduleWindow(Settings.bScheduleEnabled, Settings.ScheduleTime, Settings.ScheduleAnyOpenProject, Settings.ScheduleProjects, OpenProjects, ProjectToLatestChangeTypes);
-			if(Schedule.ShowDialog() == System.Windows.Forms.DialogResult.OK)
+			ScheduleWindow schedule = new ScheduleWindow(_settings.ScheduleEnabled, _settings.ScheduleTime, _settings.ScheduleAnyOpenProject, _settings.ScheduleProjects, openProjects, projectToLatestChangeTypes);
+			if(schedule.ShowDialog() == System.Windows.Forms.DialogResult.OK)
 			{
-				Schedule.CopySettings(out Settings.bScheduleEnabled, out Settings.ScheduleTime, out Settings.ScheduleAnyOpenProject, out Settings.ScheduleProjects);
-				Settings.Save(Logger);
+				schedule.CopySettings(out _settings.ScheduleEnabled, out _settings.ScheduleTime, out _settings.ScheduleAnyOpenProject, out _settings.ScheduleProjects);
+				_settings.Save(_logger);
 			}
 
 			StartScheduleTimer();
@@ -892,55 +892,55 @@ namespace UnrealGameSync
 		{
 			StopScheduleTimer();
 
-			if(Settings.bScheduleEnabled)
+			if(_settings.ScheduleEnabled)
 			{
-				DateTime CurrentTime = DateTime.Now;
-				Random Rnd = new Random();
+				DateTime currentTime = DateTime.Now;
+				Random rnd = new Random();
 
 				// add or subtract from the schedule time to distribute scheduled syncs over a little bit more time
 				// this avoids everyone hitting the p4 server at exactly the same time.
-				const int FudgeMinutes = 10;
-				TimeSpan FudgeTime = TimeSpan.FromMinutes(Rnd.Next(FudgeMinutes * -100, FudgeMinutes * 100) / 100.0);
-				DateTime NextScheduleTime = new DateTime(CurrentTime.Year, CurrentTime.Month, CurrentTime.Day, Settings.ScheduleTime.Hours, Settings.ScheduleTime.Minutes, Settings.ScheduleTime.Seconds);
-				NextScheduleTime += FudgeTime;
+				const int fudgeMinutes = 10;
+				TimeSpan fudgeTime = TimeSpan.FromMinutes(rnd.Next(fudgeMinutes * -100, fudgeMinutes * 100) / 100.0);
+				DateTime nextScheduleTime = new DateTime(currentTime.Year, currentTime.Month, currentTime.Day, _settings.ScheduleTime.Hours, _settings.ScheduleTime.Minutes, _settings.ScheduleTime.Seconds);
+				nextScheduleTime += fudgeTime;
 
-				if (NextScheduleTime < CurrentTime)
+				if (nextScheduleTime < currentTime)
 				{
-					NextScheduleTime = NextScheduleTime.AddDays(1.0);
+					nextScheduleTime = nextScheduleTime.AddDays(1.0);
 				}
 
-				TimeSpan IntervalToFirstTick = NextScheduleTime - CurrentTime;
-				ScheduleTimer = new System.Threading.Timer(x => MainThreadSynchronizationContext.Post((o) => { if(!IsDisposed){ ScheduleTimerElapsed(); } }, null), null, IntervalToFirstTick, TimeSpan.FromDays(1));
+				TimeSpan intervalToFirstTick = nextScheduleTime - currentTime;
+				_scheduleTimer = new System.Threading.Timer(x => _mainThreadSynchronizationContext.Post((o) => { if(!IsDisposed){ ScheduleTimerElapsed(); } }, null), null, intervalToFirstTick, TimeSpan.FromDays(1));
 
-				Logger.LogInformation("Schedule: Started ScheduleTimer for {Time} ({Time} remaining)", NextScheduleTime, IntervalToFirstTick);
+				_logger.LogInformation("Schedule: Started ScheduleTimer for {Time} ({Time} remaining)", nextScheduleTime, intervalToFirstTick);
 			}
 		}
 
 		private void StopScheduleTimer()
 		{
-			if(ScheduleTimer != null)
+			if(_scheduleTimer != null)
 			{
-				ScheduleTimer.Dispose();
-				ScheduleTimer = null;
-				Logger.LogInformation("Schedule: Stopped ScheduleTimer");
+				_scheduleTimer.Dispose();
+				_scheduleTimer = null;
+				_logger.LogInformation("Schedule: Stopped ScheduleTimer");
 			}
 			StopScheduleSettledTimer();
 		}
 
 		private void ScheduleTimerElapsed()
 		{
-			Logger.LogInformation("Schedule: Timer Elapsed");
+			_logger.LogInformation("Schedule: Timer Elapsed");
 
 			// Try to open any missing tabs. 
-			int NumInitialTabs = TabControl.GetTabCount();
-			foreach (UserSelectedProjectSettings ScheduledProject in Settings.ScheduleProjects)
+			int numInitialTabs = TabControl.GetTabCount();
+			foreach (UserSelectedProjectSettings scheduledProject in _settings.ScheduleProjects)
 			{
-				Logger.LogInformation("Schedule: Attempting to open {Project}", ScheduledProject);
-				TryOpenProject(ScheduledProject, -1, OpenProjectOptions.Quiet);
+				_logger.LogInformation("Schedule: Attempting to open {Project}", scheduledProject);
+				TryOpenProject(scheduledProject, -1, OpenProjectOptions.Quiet);
 			}
 
 			// If we did open something, leave it for a while to populate with data before trying to start the sync.
-			if(TabControl.GetTabCount() > NumInitialTabs)
+			if(TabControl.GetTabCount() > numInitialTabs)
 			{
 				StartScheduleSettledTimer();
 			}
@@ -953,231 +953,231 @@ namespace UnrealGameSync
 		private void StartScheduleSettledTimer()
 		{
 			StopScheduleSettledTimer();
-			ScheduleSettledTimer = new System.Threading.Timer(x => MainThreadSynchronizationContext.Post((o) => { if(!IsDisposed){ ScheduleSettledTimerElapsed(); } }, null), null, TimeSpan.FromSeconds(20.0), TimeSpan.FromMilliseconds(-1.0));
-			Logger.LogInformation("Schedule: Started ScheduleSettledTimer");
+			_scheduleSettledTimer = new System.Threading.Timer(x => _mainThreadSynchronizationContext.Post((o) => { if(!IsDisposed){ ScheduleSettledTimerElapsed(); } }, null), null, TimeSpan.FromSeconds(20.0), TimeSpan.FromMilliseconds(-1.0));
+			_logger.LogInformation("Schedule: Started ScheduleSettledTimer");
 		}
 
 		private void StopScheduleSettledTimer()
 		{
-			if(ScheduleSettledTimer != null)
+			if(_scheduleSettledTimer != null)
 			{
-				ScheduleSettledTimer.Dispose();
-				ScheduleSettledTimer = null;
+				_scheduleSettledTimer.Dispose();
+				_scheduleSettledTimer = null;
 
-				Logger.LogInformation("Schedule: Stopped ScheduleSettledTimer");
+				_logger.LogInformation("Schedule: Stopped ScheduleSettledTimer");
 			}
 		}
 
 		private void ScheduleSettledTimerElapsed()
 		{
-			Logger.LogInformation("Schedule: Starting Sync");
-			for(int Idx = 0; Idx < TabControl.GetTabCount(); Idx++)
+			_logger.LogInformation("Schedule: Starting Sync");
+			for(int idx = 0; idx < TabControl.GetTabCount(); idx++)
 			{
-				WorkspaceControl? Workspace = TabControl.GetTabData(Idx) as WorkspaceControl;
-				if(Workspace != null)
+				WorkspaceControl? workspace = TabControl.GetTabData(idx) as WorkspaceControl;
+				if(workspace != null)
 				{
-					Logger.LogInformation("Schedule: Considering {File}", Workspace.SelectedFileName);
-					if(Settings.ScheduleAnyOpenProject || Settings.ScheduleProjects.Any(x => x.LocalPath != null && x.LocalPath.Equals(Workspace.SelectedProject.LocalPath, StringComparison.OrdinalIgnoreCase)))
+					_logger.LogInformation("Schedule: Considering {File}", workspace.SelectedFileName);
+					if(_settings.ScheduleAnyOpenProject || _settings.ScheduleProjects.Any(x => x.LocalPath != null && x.LocalPath.Equals(workspace.SelectedProject.LocalPath, StringComparison.OrdinalIgnoreCase)))
 					{
-						Logger.LogInformation("Schedule: Starting Sync");
-						Workspace.ScheduleTimerElapsed();
+						_logger.LogInformation("Schedule: Starting Sync");
+						workspace.ScheduleTimerElapsed();
 					}
 				}
 			}
 		}
 
-		void TabControl_OnTabChanged(object? NewTabData)
+		void TabControl_OnTabChanged(object? newTabData)
 		{
 			if(IsHandleCreated)
 			{
-				SendMessage(Handle, WM_SETREDRAW, 0, 0);
+				SendMessage(Handle, WmSetredraw, 0, 0);
 			}
 
 			SuspendLayout();
 
-			if(CurrentTabPanel != null)
+			if(_currentTabPanel != null)
 			{
-				CurrentTabPanel.Deactivate();
-				CurrentTabPanel.Hide();
+				_currentTabPanel.Deactivate();
+				_currentTabPanel.Hide();
 			}
 
-			if(NewTabData == null)
+			if(newTabData == null)
 			{
-				CurrentTabPanel = null;
-				Settings.LastProject = null;
+				_currentTabPanel = null;
+				_settings.LastProject = null;
 				DefaultControl.Show();
 			}
 			else
 			{
-				CurrentTabPanel = (IMainWindowTabPanel)NewTabData;
-				Settings.LastProject = CurrentTabPanel.SelectedProject;
+				_currentTabPanel = (IMainWindowTabPanel)newTabData;
+				_settings.LastProject = _currentTabPanel.SelectedProject;
 				DefaultControl.Hide();
 			}
 
-			Settings.Save(Logger);
+			_settings.Save(_logger);
 
-			if(CurrentTabPanel != null)
+			if(_currentTabPanel != null)
 			{
-				CurrentTabPanel.Activate();
-				CurrentTabPanel.Show();
+				_currentTabPanel.Activate();
+				_currentTabPanel.Show();
 			}
 
 			ResumeLayout();
 
 			if(IsHandleCreated)
 			{
-				SendMessage(Handle, WM_SETREDRAW, 1, 0);
+				SendMessage(Handle, WmSetredraw, 1, 0);
 			}
 
 			Refresh();
 		}
 
-		public void RequestProjectChange(WorkspaceControl Workspace, UserSelectedProjectSettings Project, bool bModal)
+		public void RequestProjectChange(WorkspaceControl workspace, UserSelectedProjectSettings project, bool modal)
 		{
-			int TabIdx = TabControl.FindTabIndex(Workspace);
-			if(TabIdx != -1 && !Workspace.IsBusy() && CanFocus)
+			int tabIdx = TabControl.FindTabIndex(workspace);
+			if(tabIdx != -1 && !workspace.IsBusy() && CanFocus)
 			{
-				if(bModal)
+				if(modal)
 				{
-					TryOpenProject(Project, TabIdx);
+					TryOpenProject(project, tabIdx);
 				}
 				else
 				{
-					TryOpenProject(Project, TabIdx, OpenProjectOptions.Quiet);
+					TryOpenProject(project, tabIdx, OpenProjectOptions.Quiet);
 				}
 			}
 		}
 
 		public void OpenNewProject()
 		{
-			OpenProjectInfo? OpenProjectInfo = OpenProjectWindow.ShowModal(this, null, Settings, DataFolder, CacheFolder, DefaultPerforceSettings, ServiceProvider, Logger);
-			if(OpenProjectInfo != null)
+			OpenProjectInfo? openProjectInfo = OpenProjectWindow.ShowModal(this, null, _settings, _dataFolder, _cacheFolder, _defaultPerforceSettings, _serviceProvider, _logger);
+			if(openProjectInfo != null)
 			{
-				int NewTabIdx = TryOpenProject(OpenProjectInfo, -1, OpenProjectOptions.None);
-				if(NewTabIdx != -1)
+				int newTabIdx = TryOpenProject(openProjectInfo, -1, OpenProjectOptions.None);
+				if(newTabIdx != -1)
 				{
-					TabControl.SelectTab(NewTabIdx);
+					TabControl.SelectTab(newTabIdx);
 					SaveTabSettings();
-					UpdateRecentProjectsList(OpenProjectInfo.SelectedProject);
+					UpdateRecentProjectsList(openProjectInfo.SelectedProject);
 				}
 			}
 		}
 
-		void UpdateRecentProjectsList(UserSelectedProjectSettings DetectedProjectSettings)
+		void UpdateRecentProjectsList(UserSelectedProjectSettings detectedProjectSettings)
 		{
-			Settings.RecentProjects.RemoveAll(x => x.LocalPath != null && x.LocalPath.Equals(DetectedProjectSettings.LocalPath, StringComparison.OrdinalIgnoreCase));
-			Settings.RecentProjects.Insert(0, DetectedProjectSettings);
+			_settings.RecentProjects.RemoveAll(x => x.LocalPath != null && x.LocalPath.Equals(detectedProjectSettings.LocalPath, StringComparison.OrdinalIgnoreCase));
+			_settings.RecentProjects.Insert(0, detectedProjectSettings);
 
-			const int MaxRecentProjects = 10;
-			if (Settings.RecentProjects.Count > MaxRecentProjects)
+			const int maxRecentProjects = 10;
+			if (_settings.RecentProjects.Count > maxRecentProjects)
 			{
-				Settings.RecentProjects.RemoveRange(MaxRecentProjects, Settings.RecentProjects.Count - MaxRecentProjects);
+				_settings.RecentProjects.RemoveRange(maxRecentProjects, _settings.RecentProjects.Count - maxRecentProjects);
 			}
 
-			Settings.Save(Logger);
+			_settings.Save(_logger);
 		}
 
-		public void EditSelectedProject(int TabIdx)
+		public void EditSelectedProject(int tabIdx)
 		{
-			object TabData = TabControl.GetTabData(TabIdx);
-			if(TabData is WorkspaceControl)
+			object tabData = TabControl.GetTabData(tabIdx);
+			if(tabData is WorkspaceControl)
 			{
-				WorkspaceControl Workspace = (WorkspaceControl)TabData;
-				EditSelectedProject(TabIdx, Workspace.SelectedProject);
+				WorkspaceControl workspace = (WorkspaceControl)tabData;
+				EditSelectedProject(tabIdx, workspace.SelectedProject);
 			}
-			else if(TabData is ErrorPanel)
+			else if(tabData is ErrorPanel)
 			{
-				ErrorPanel Error = (ErrorPanel)TabData;
-				EditSelectedProject(TabIdx, Error.SelectedProject);
+				ErrorPanel error = (ErrorPanel)tabData;
+				EditSelectedProject(tabIdx, error.SelectedProject);
 			}
 		}
 
-		public void EditSelectedProject(WorkspaceControl Workspace)
+		public void EditSelectedProject(WorkspaceControl workspace)
 		{
-			int TabIdx = TabControl.FindTabIndex(Workspace);
-			if(TabIdx != -1)
+			int tabIdx = TabControl.FindTabIndex(workspace);
+			if(tabIdx != -1)
 			{
-				EditSelectedProject(TabIdx, Workspace.SelectedProject);
+				EditSelectedProject(tabIdx, workspace.SelectedProject);
 			}
 		}
 
-		public void EditSelectedProject(ErrorPanel Panel)
+		public void EditSelectedProject(ErrorPanel panel)
 		{
-			int TabIdx = TabControl.FindTabIndex(Panel);
-			if(TabIdx != -1)
+			int tabIdx = TabControl.FindTabIndex(panel);
+			if(tabIdx != -1)
 			{
-				EditSelectedProject(TabIdx, Panel.SelectedProject);
+				EditSelectedProject(tabIdx, panel.SelectedProject);
 			}
 		}
 
-		public void EditSelectedProject(int TabIdx, UserSelectedProjectSettings SelectedProject)
+		public void EditSelectedProject(int tabIdx, UserSelectedProjectSettings selectedProject)
 		{
-			OpenProjectInfo? OpenProjectInfo = OpenProjectWindow.ShowModal(this, SelectedProject, Settings, DataFolder, CacheFolder, DefaultPerforceSettings, ServiceProvider, Logger);
-			if(OpenProjectInfo != null)
+			OpenProjectInfo? openProjectInfo = OpenProjectWindow.ShowModal(this, selectedProject, _settings, _dataFolder, _cacheFolder, _defaultPerforceSettings, _serviceProvider, _logger);
+			if(openProjectInfo != null)
 			{
-				int NewTabIdx = TryOpenProject(OpenProjectInfo, TabIdx, OpenProjectOptions.None);
-				if(NewTabIdx != -1)
+				int newTabIdx = TryOpenProject(openProjectInfo, tabIdx, OpenProjectOptions.None);
+				if(newTabIdx != -1)
 				{
-					TabControl.SelectTab(NewTabIdx);
+					TabControl.SelectTab(newTabIdx);
 					SaveTabSettings();
-					UpdateRecentProjectsList(OpenProjectInfo.SelectedProject);
+					UpdateRecentProjectsList(openProjectInfo.SelectedProject);
 				}
 			}
 		}
 
-		int TryOpenProject(UserSelectedProjectSettings Project, int ReplaceTabIdx, OpenProjectOptions Options = OpenProjectOptions.None)
+		int TryOpenProject(UserSelectedProjectSettings project, int replaceTabIdx, OpenProjectOptions options = OpenProjectOptions.None)
 		{
-			ILogger<OpenProjectInfo> ProjectLogger = ServiceProvider.GetRequiredService<ILogger<OpenProjectInfo>>();
+			ILogger<OpenProjectInfo> projectLogger = _serviceProvider.GetRequiredService<ILogger<OpenProjectInfo>>();
 
-			ModalTaskFlags TaskFlags = ModalTaskFlags.None;
-			if((Options & OpenProjectOptions.Quiet) != 0)
+			ModalTaskFlags taskFlags = ModalTaskFlags.None;
+			if((options & OpenProjectOptions.Quiet) != 0)
 			{
-				TaskFlags |= ModalTaskFlags.Quiet;
+				taskFlags |= ModalTaskFlags.Quiet;
 			}
 
-			PerforceSettings PerforceSettings = Utility.OverridePerforceSettings(DefaultPerforceSettings, Project.ServerAndPort, Project.UserName);
+			PerforceSettings perforceSettings = Utility.OverridePerforceSettings(_defaultPerforceSettings, project.ServerAndPort, project.UserName);
 
-			ModalTask<OpenProjectInfo>? SettingsTask = PerforceModalTask.Execute(this, "Opening Project", "Opening project, please wait...", PerforceSettings, (p, c) => OpenProjectWindow.DetectSettingsAsync(p, Project, Settings, ProjectLogger, c), ProjectLogger, TaskFlags);
-			if (SettingsTask == null || SettingsTask.Failed)
+			ModalTask<OpenProjectInfo>? settingsTask = PerforceModalTask.Execute(this, "Opening Project", "Opening project, please wait...", perforceSettings, (p, c) => OpenProjectWindow.DetectSettingsAsync(p, project, _settings, projectLogger, c), projectLogger, taskFlags);
+			if (settingsTask == null || settingsTask.Failed)
 			{
-				if(SettingsTask != null) CreateErrorPanel(ReplaceTabIdx, Project, SettingsTask.Error);
+				if(settingsTask != null) CreateErrorPanel(replaceTabIdx, project, settingsTask.Error);
 				return -1;
 			}
 
-			return TryOpenProject(SettingsTask.Result, ReplaceTabIdx, Options);
+			return TryOpenProject(settingsTask.Result, replaceTabIdx, options);
 		}
 
-		int TryOpenProject(OpenProjectInfo OpenProjectInfo, int ReplaceTabIdx, OpenProjectOptions Options)
+		int TryOpenProject(OpenProjectInfo openProjectInfo, int replaceTabIdx, OpenProjectOptions options)
 		{
-			Logger.LogInformation("Trying to open project {Project}", OpenProjectInfo.ProjectInfo.ClientFileName);
+			_logger.LogInformation("Trying to open project {Project}", openProjectInfo.ProjectInfo.ClientFileName);
 
 			// Check that none of the other tabs already have it open
-			for(int TabIdx = 0; TabIdx < TabControl.GetTabCount(); TabIdx++)
+			for(int tabIdx = 0; tabIdx < TabControl.GetTabCount(); tabIdx++)
 			{
-				if(ReplaceTabIdx != TabIdx)
+				if(replaceTabIdx != tabIdx)
 				{
-					WorkspaceControl? Workspace = TabControl.GetTabData(TabIdx) as WorkspaceControl;
-					if(Workspace != null)
+					WorkspaceControl? workspace = TabControl.GetTabData(tabIdx) as WorkspaceControl;
+					if(workspace != null)
 					{
-						if(Workspace.SelectedFileName == OpenProjectInfo.ProjectInfo.LocalFileName)
+						if(workspace.SelectedFileName == openProjectInfo.ProjectInfo.LocalFileName)
 						{
-							Logger.LogInformation("  Already open in tab {TabIdx}", TabIdx);
-							if((Options & OpenProjectOptions.Quiet) == 0)
+							_logger.LogInformation("  Already open in tab {TabIdx}", tabIdx);
+							if((options & OpenProjectOptions.Quiet) == 0)
 							{
-								TabControl.SelectTab(TabIdx);
+								TabControl.SelectTab(tabIdx);
 							}
-							return TabIdx;
+							return tabIdx;
 						}
-						else if(OpenProjectInfo.ProjectInfo.LocalFileName.IsUnderDirectory(Workspace.BranchDirectoryName))
+						else if(openProjectInfo.ProjectInfo.LocalFileName.IsUnderDirectory(workspace.BranchDirectoryName))
 						{
-							if((Options & OpenProjectOptions.Quiet) == 0 && MessageBox.Show(String.Format("{0} is already open under {1}.\n\nWould you like to close it?", Workspace.SelectedFileName.GetFileNameWithoutExtension(), Workspace.BranchDirectoryName, OpenProjectInfo.ProjectInfo.LocalFileName.GetFileNameWithoutExtension()), "Branch already open", MessageBoxButtons.YesNo) == System.Windows.Forms.DialogResult.Yes)
+							if((options & OpenProjectOptions.Quiet) == 0 && MessageBox.Show(String.Format("{0} is already open under {1}.\n\nWould you like to close it?", workspace.SelectedFileName.GetFileNameWithoutExtension(), workspace.BranchDirectoryName, openProjectInfo.ProjectInfo.LocalFileName.GetFileNameWithoutExtension()), "Branch already open", MessageBoxButtons.YesNo) == System.Windows.Forms.DialogResult.Yes)
 							{
-								Logger.LogInformation("  Another project already open in this workspace, tab {TabIdx}. Replacing.", TabIdx);
-								TabControl.RemoveTab(TabIdx);
+								_logger.LogInformation("  Another project already open in this workspace, tab {TabIdx}. Replacing.", tabIdx);
+								TabControl.RemoveTab(tabIdx);
 							}
 							else
 							{
-								Logger.LogInformation("  Another project already open in this workspace, tab {TabIdx}. Aborting.", TabIdx);
+								_logger.LogInformation("  Another project already open in this workspace, tab {TabIdx}. Aborting.", tabIdx);
 								return -1;
 							}
 						}
@@ -1189,132 +1189,132 @@ namespace UnrealGameSync
 			DefaultControl.Hide();
 
 			// Remove the current tab. We need to ensure the workspace has been shut down before creating a new one with the same log files, etc...
-			if(ReplaceTabIdx != -1)
+			if(replaceTabIdx != -1)
 			{
-				WorkspaceControl? OldWorkspace = TabControl.GetTabData(ReplaceTabIdx) as WorkspaceControl;
-				if(OldWorkspace != null)
+				WorkspaceControl? oldWorkspace = TabControl.GetTabData(replaceTabIdx) as WorkspaceControl;
+				if(oldWorkspace != null)
 				{
-					OldWorkspace.Hide();
-					TabControl.SetTabData(ReplaceTabIdx, new ErrorPanel(OpenProjectInfo.SelectedProject));
-					OldWorkspace.Dispose();
+					oldWorkspace.Hide();
+					TabControl.SetTabData(replaceTabIdx, new ErrorPanel(openProjectInfo.SelectedProject));
+					oldWorkspace.Dispose();
 				}
 			}
 
 			// Now that we have the project settings, we can construct the tab
-			WorkspaceControl NewWorkspace = new WorkspaceControl(this, DataFolder, ApiUrl, OpenProjectInfo, ServiceProvider, Settings, OIDCTokenManager);
+			WorkspaceControl newWorkspace = new WorkspaceControl(this, _dataFolder, _apiUrl, openProjectInfo, _serviceProvider, _settings, _oidcTokenManager);
 			
-			NewWorkspace.Parent = TabPanel;
-			NewWorkspace.Dock = DockStyle.Fill;
-			NewWorkspace.Hide();
+			newWorkspace.Parent = TabPanel;
+			newWorkspace.Dock = DockStyle.Fill;
+			newWorkspace.Hide();
 
 			// Add the tab
-			string NewTabName = GetTabName(NewWorkspace);
-			if(ReplaceTabIdx == -1)
+			string newTabName = GetTabName(newWorkspace);
+			if(replaceTabIdx == -1)
 			{
-				int NewTabIdx = TabControl.InsertTab(-1, NewTabName, NewWorkspace, NewWorkspace.TintColor);
-				Logger.LogInformation("  Inserted tab {TabIdx}", NewTabIdx);
-				return NewTabIdx;
+				int newTabIdx = TabControl.InsertTab(-1, newTabName, newWorkspace, newWorkspace.TintColor);
+				_logger.LogInformation("  Inserted tab {TabIdx}", newTabIdx);
+				return newTabIdx;
 			}
 			else
 			{
-				Logger.LogInformation("  Replacing tab {TabIdx}", ReplaceTabIdx);
-				TabControl.InsertTab(ReplaceTabIdx + 1, NewTabName, NewWorkspace, NewWorkspace.TintColor);
-				TabControl.RemoveTab(ReplaceTabIdx);
-				return ReplaceTabIdx;
+				_logger.LogInformation("  Replacing tab {TabIdx}", replaceTabIdx);
+				TabControl.InsertTab(replaceTabIdx + 1, newTabName, newWorkspace, newWorkspace.TintColor);
+				TabControl.RemoveTab(replaceTabIdx);
+				return replaceTabIdx;
 			}
 		}
 
-		public void StreamChanged(WorkspaceControl Workspace)
+		public void StreamChanged(WorkspaceControl workspace)
 		{
-			MainThreadSynchronizationContext.Post((o) => { if(!IsDisposed) { StreamChangedCallback(Workspace); } }, null);
+			_mainThreadSynchronizationContext.Post((o) => { if(!IsDisposed) { StreamChangedCallback(workspace); } }, null);
 		}
 
-		public void StreamChangedCallback(WorkspaceControl Workspace)
+		public void StreamChangedCallback(WorkspaceControl workspace)
 		{
-			if(ChangingWorkspacesRefCount == 0)
+			if(_changingWorkspacesRefCount == 0)
 			{
-				ChangingWorkspacesRefCount++;
+				_changingWorkspacesRefCount++;
 
-				for(int Idx = 0; Idx < TabControl.GetTabCount(); Idx++)
+				for(int idx = 0; idx < TabControl.GetTabCount(); idx++)
 				{
-					if(TabControl.GetTabData(Idx) == Workspace)
+					if(TabControl.GetTabData(idx) == workspace)
 					{
-						UserSelectedProjectSettings Project = Workspace.SelectedProject;
-						if(TryOpenProject(Project, Idx) == -1)
+						UserSelectedProjectSettings project = workspace.SelectedProject;
+						if(TryOpenProject(project, idx) == -1)
 						{
-							TabControl.RemoveTab(Idx);
+							TabControl.RemoveTab(idx);
 						}
 						break;
 					}
 				}
 
-				ChangingWorkspacesRefCount--;
+				_changingWorkspacesRefCount--;
 			}
 		}
 
 		void SaveTabSettings()
 		{
-			Settings.OpenProjects.Clear();
-			for(int TabIdx = 0; TabIdx < TabControl.GetTabCount(); TabIdx++)
+			_settings.OpenProjects.Clear();
+			for(int tabIdx = 0; tabIdx < TabControl.GetTabCount(); tabIdx++)
 			{
-				IMainWindowTabPanel TabPanel = (IMainWindowTabPanel)TabControl.GetTabData(TabIdx)!;
-				Settings.OpenProjects.Add(TabPanel.SelectedProject);
+				IMainWindowTabPanel tabPanel = (IMainWindowTabPanel)TabControl.GetTabData(tabIdx)!;
+				_settings.OpenProjects.Add(tabPanel.SelectedProject);
 			}
-			Settings.Save(Logger);
+			_settings.Save(_logger);
 		}
 
-		void TabControl_OnNewTabClick(Point Location, MouseButtons Buttons)
+		void TabControl_OnNewTabClick(Point location, MouseButtons buttons)
 		{
-			if(Buttons == MouseButtons.Left)
+			if(buttons == MouseButtons.Left)
 			{
 				OpenNewProject();
 			}
 		}
 
-		string GetTabName(WorkspaceControl Workspace)
+		string GetTabName(WorkspaceControl workspace)
 		{
-			string TabName = "";
-			switch (Settings.TabLabels)
+			string tabName = "";
+			switch (_settings.TabLabels)
 			{
 				case TabLabels.Stream:
-					TabName = Workspace.StreamName ?? TabName;
+					tabName = workspace.StreamName ?? tabName;
 					break;
 				case TabLabels.ProjectFile:
-					TabName = Workspace.SelectedFileName.FullName;
+					tabName = workspace.SelectedFileName.FullName;
 					break;
 				case TabLabels.WorkspaceName:
-					TabName = Workspace.ClientName;
+					tabName = workspace.ClientName;
 					break;
 				case TabLabels.WorkspaceRoot:
-					TabName = Workspace.BranchDirectoryName.FullName;
+					tabName = workspace.BranchDirectoryName.FullName;
 					break;
 				default:
 					break;
 			}
 
 			// if this failes, return something sensible to avoid blank tabs
-			if (string.IsNullOrEmpty(TabName))
+			if (string.IsNullOrEmpty(tabName))
 			{
-				Logger.LogInformation("No TabName for {ClientName} for setting {TabSetting}. Defaulting to client name", Workspace.ClientName, Settings.TabLabels.ToString());
-				TabName = Workspace.ClientName;
+				_logger.LogInformation("No TabName for {ClientName} for setting {TabSetting}. Defaulting to client name", workspace.ClientName, _settings.TabLabels.ToString());
+				tabName = workspace.ClientName;
 			}
 
-			return TabName;
+			return tabName;
 		}
 
-		public void SetTabNames(TabLabels NewTabNames)
+		public void SetTabNames(TabLabels newTabNames)
 		{
-			if(Settings.TabLabels != NewTabNames)
+			if(_settings.TabLabels != newTabNames)
 			{
-				Settings.TabLabels = NewTabNames;
-				Settings.Save(Logger);
+				_settings.TabLabels = newTabNames;
+				_settings.Save(_logger);
 
-				for(int Idx = 0; Idx < TabControl.GetTabCount(); Idx++)
+				for(int idx = 0; idx < TabControl.GetTabCount(); idx++)
 				{
-					WorkspaceControl? Workspace = TabControl.GetTabData(Idx) as WorkspaceControl;
-					if(Workspace != null)
+					WorkspaceControl? workspace = TabControl.GetTabData(idx) as WorkspaceControl;
+					if(workspace != null)
 					{
-						TabControl.SetTabName(Idx, GetTabName(Workspace));
+						TabControl.SetTabName(idx, GetTabName(workspace));
 					}
 				}
 			}
@@ -1322,12 +1322,12 @@ namespace UnrealGameSync
 
 		private void TabMenu_OpenProject_Click(object sender, EventArgs e)
 		{
-			EditSelectedProject(TabMenu_TabIdx);
+			EditSelectedProject(_tabMenuTabIdx);
 		}
 
-		private void TabMenu_OpenRecentProject_Click(UserSelectedProjectSettings RecentProject, int TabIdx)
+		private void TabMenu_OpenRecentProject_Click(UserSelectedProjectSettings recentProject, int tabIdx)
 		{
-			TryOpenProject(RecentProject, TabIdx);
+			TryOpenProject(recentProject, tabIdx);
 			SaveTabSettings();
 		}
 
@@ -1353,8 +1353,8 @@ namespace UnrealGameSync
 
 		private void TabMenu_RecentProjects_ClearList_Click(object sender, EventArgs e)
 		{
-			Settings.RecentProjects.Clear();
-			Settings.Save(Logger);
+			_settings.RecentProjects.Clear();
+			_settings.Save(_logger);
 		}
 
 		private void TabMenu_Closed(object sender, ToolStripDropDownClosedEventArgs e)
@@ -1364,114 +1364,114 @@ namespace UnrealGameSync
 
 		private void RecentMenu_ClearList_Click(object sender, EventArgs e)
 		{
-			Settings.RecentProjects.Clear();
-			Settings.Save(Logger);
+			_settings.RecentProjects.Clear();
+			_settings.Save(_logger);
 		}
 
 		public void UpdateProgress()
 		{
-			TaskbarState State = TaskbarState.NoProgress;
-			float Progress = -1.0f;
+			TaskbarState state = TaskbarState.NoProgress;
+			float progress = -1.0f;
 
-			for(int Idx = 0; Idx < TabControl.GetTabCount(); Idx++)
+			for(int idx = 0; idx < TabControl.GetTabCount(); idx++)
 			{
-				IMainWindowTabPanel TabPanel = (IMainWindowTabPanel)TabControl.GetTabData(Idx)!;
+				IMainWindowTabPanel tabPanel = (IMainWindowTabPanel)TabControl.GetTabData(idx)!;
 
-				Tuple<TaskbarState, float> DesiredTaskbarState = TabPanel.DesiredTaskbarState;
-				if(DesiredTaskbarState.Item1 == TaskbarState.Error)
+				Tuple<TaskbarState, float> desiredTaskbarState = tabPanel.DesiredTaskbarState;
+				if(desiredTaskbarState.Item1 == TaskbarState.Error)
 				{
-					State = TaskbarState.Error;
-					TabControl.SetHighlight(Idx, Tuple.Create(Color.FromArgb(204, 64, 64), 1.0f));
+					state = TaskbarState.Error;
+					TabControl.SetHighlight(idx, Tuple.Create(Color.FromArgb(204, 64, 64), 1.0f));
 				}
-				else if(DesiredTaskbarState.Item1 == TaskbarState.Paused && State != TaskbarState.Error)
+				else if(desiredTaskbarState.Item1 == TaskbarState.Paused && state != TaskbarState.Error)
 				{
-					State = TaskbarState.Paused;
-					TabControl.SetHighlight(Idx, Tuple.Create(Color.FromArgb(255, 242, 0), 1.0f));
+					state = TaskbarState.Paused;
+					TabControl.SetHighlight(idx, Tuple.Create(Color.FromArgb(255, 242, 0), 1.0f));
 				}
-				else if(DesiredTaskbarState.Item1 == TaskbarState.Normal && State != TaskbarState.Error && State != TaskbarState.Paused)
+				else if(desiredTaskbarState.Item1 == TaskbarState.Normal && state != TaskbarState.Error && state != TaskbarState.Paused)
 				{
-					State = TaskbarState.Normal;
-					Progress = Math.Max(Progress, DesiredTaskbarState.Item2);
-					TabControl.SetHighlight(Idx, Tuple.Create(Color.FromArgb(28, 180, 64), DesiredTaskbarState.Item2));
+					state = TaskbarState.Normal;
+					progress = Math.Max(progress, desiredTaskbarState.Item2);
+					TabControl.SetHighlight(idx, Tuple.Create(Color.FromArgb(28, 180, 64), desiredTaskbarState.Item2));
 				}
 				else
 				{
-					TabControl.SetHighlight(Idx, null);
+					TabControl.SetHighlight(idx, null);
 				}
 			}
 
 			if(IsHandleCreated)
 			{
-				if(State == TaskbarState.Normal)
+				if(state == TaskbarState.Normal)
 				{
 					Taskbar.SetState(Handle, TaskbarState.Normal);
-					Taskbar.SetProgress(Handle, (ulong)(Progress * 1000.0f), 1000);
+					Taskbar.SetProgress(Handle, (ulong)(progress * 1000.0f), 1000);
 				}
 				else
 				{
-					Taskbar.SetState(Handle, State);
+					Taskbar.SetState(Handle, state);
 				}
 			}
 		}
 
 		public void UpdateTintColors()
 		{
-			for (int Idx = 0; Idx < TabControl.GetTabCount(); Idx++)
+			for (int idx = 0; idx < TabControl.GetTabCount(); idx++)
 			{
-				IMainWindowTabPanel TabPanel = (IMainWindowTabPanel)TabControl.GetTabData(Idx)!;
-				TabControl.SetTint(Idx, TabPanel.TintColor);
+				IMainWindowTabPanel tabPanel = (IMainWindowTabPanel)TabControl.GetTabData(idx)!;
+				TabControl.SetTint(idx, tabPanel.TintColor);
 			}
 		}
 
 		public void ModifyApplicationSettings()
 		{
-			bool? bRelaunchPreview = ApplicationSettingsWindow.ShowModal(this, DefaultPerforceSettings, bPreview, OriginalExecutableFileName, Settings, ToolUpdateMonitor, ServiceProvider.GetRequiredService<ILogger<ApplicationSettingsWindow>>());
-			if(bRelaunchPreview.HasValue)
+			bool? relaunchPreview = ApplicationSettingsWindow.ShowModal(this, _defaultPerforceSettings, _preview, _originalExecutableFileName, _settings, ToolUpdateMonitor, _serviceProvider.GetRequiredService<ILogger<ApplicationSettingsWindow>>());
+			if(relaunchPreview.HasValue)
 			{
-				UpdateMonitor.TriggerUpdate(UpdateType.UserInitiated, bRelaunchPreview);
+				_updateMonitor.TriggerUpdate(UpdateType.UserInitiated, relaunchPreview);
 			}
 
-			for (int Idx = 0; Idx < TabControl.GetTabCount(); Idx++)
+			for (int idx = 0; idx < TabControl.GetTabCount(); idx++)
 			{
-				IMainWindowTabPanel TabPanel = (IMainWindowTabPanel)TabControl.GetTabData(Idx)!;
-				TabPanel.UpdateSettings();
+				IMainWindowTabPanel tabPanel = (IMainWindowTabPanel)TabControl.GetTabData(idx)!;
+				tabPanel.UpdateSettings();
 			}
 		}
 
 		private void MainWindow_Load(object sender, EventArgs e)
 		{
-			if(Settings.WindowBounds != null)
+			if(_settings.WindowBounds != null)
 			{
-				Rectangle WindowBounds = Settings.WindowBounds.Value;
-				if(WindowBounds.Width > MinimumSize.Width && WindowBounds.Height > MinimumSize.Height)
+				Rectangle windowBounds = _settings.WindowBounds.Value;
+				if(windowBounds.Width > MinimumSize.Width && windowBounds.Height > MinimumSize.Height)
 				{
-					foreach (Screen Screen in Screen.AllScreens)
+					foreach (Screen screen in Screen.AllScreens)
 					{
-						if(WindowBounds.IntersectsWith(Screen.Bounds))
+						if(windowBounds.IntersectsWith(screen.Bounds))
 						{
-							Location = Settings.WindowBounds.Value.Location;
-							Size = Settings.WindowBounds.Value.Size;
+							Location = _settings.WindowBounds.Value.Location;
+							Size = _settings.WindowBounds.Value.Size;
 							break;
 						}
 					}
 				}
 			}
 
-			FormWindowState NewWindowState;
-			if (Enum.TryParse(Settings.WindowState, true, out NewWindowState))
+			FormWindowState newWindowState;
+			if (Enum.TryParse(_settings.WindowState, true, out newWindowState))
 			{
-				WindowState = NewWindowState;
+				WindowState = newWindowState;
 			}
 		}
 
-		bool ShowNotificationForIssue(IssueData Issue)
+		bool ShowNotificationForIssue(IssueData issue)
 		{
-			return Issue.Projects.Any(x => ShowNotificationsForProject(x));
+			return issue.Projects.Any(x => ShowNotificationsForProject(x));
 		}
 
-		bool ShowNotificationsForProject(string Project)
+		bool ShowNotificationsForProject(string project)
 		{
-			return String.IsNullOrEmpty(Project) || Settings.NotifyProjects.Count == 0 || Settings.NotifyProjects.Any(x => x.Equals(Project, StringComparison.OrdinalIgnoreCase));
+			return String.IsNullOrEmpty(project) || _settings.NotifyProjects.Count == 0 || _settings.NotifyProjects.Any(x => x.Equals(project, StringComparison.OrdinalIgnoreCase));
 		}
 
 		public void UpdateAlertWindows()
@@ -1481,88 +1481,88 @@ namespace UnrealGameSync
 				return;
 			}
 
-			HashSet<IssueData> AllIssues = new HashSet<IssueData>();
-			foreach(IssueMonitor IssueMonitor in WorkspaceIssueMonitors.Select(x => x.IssueMonitor))
+			HashSet<IssueData> allIssues = new HashSet<IssueData>();
+			foreach(IssueMonitor issueMonitor in _workspaceIssueMonitors.Select(x => x.IssueMonitor))
 			{
-				List<IssueData> Issues = IssueMonitor.GetIssues();
-				foreach(IssueData Issue in Issues)
+				List<IssueData> issues = issueMonitor.GetIssues();
+				foreach(IssueData issue in issues)
 				{
-					IssueAlertReason Reason = 0;
-					if(Issue.FixChange == 0 && !Issue.ResolvedAt.HasValue)
+					IssueAlertReason reason = 0;
+					if(issue.FixChange == 0 && !issue.ResolvedAt.HasValue)
 					{
-						if(Issue.Owner == null)
+						if(issue.Owner == null)
 						{
-							if(Issue.bNotify)
+							if(issue.Notify)
 							{
-								Reason |= IssueAlertReason.Normal;
+								reason |= IssueAlertReason.Normal;
 							}
-							if(ShowNotificationForIssue(Issue) && Settings.NotifyUnassignedMinutes >= 0 && Issue.RetrievedAt - Issue.CreatedAt >= TimeSpan.FromMinutes(Settings.NotifyUnassignedMinutes))
+							if(ShowNotificationForIssue(issue) && _settings.NotifyUnassignedMinutes >= 0 && issue.RetrievedAt - issue.CreatedAt >= TimeSpan.FromMinutes(_settings.NotifyUnassignedMinutes))
 							{
-								Reason |= IssueAlertReason.UnassignedTimer;
+								reason |= IssueAlertReason.UnassignedTimer;
 							}
 						}
-						else if(!Issue.AcknowledgedAt.HasValue)
+						else if(!issue.AcknowledgedAt.HasValue)
 						{
-							if(String.Compare(Issue.Owner, IssueMonitor.UserName, StringComparison.OrdinalIgnoreCase) == 0)
+							if(String.Compare(issue.Owner, issueMonitor.UserName, StringComparison.OrdinalIgnoreCase) == 0)
 							{
-								Reason |= IssueAlertReason.Owner;
+								reason |= IssueAlertReason.Owner;
 							}
-							else if(ShowNotificationForIssue(Issue) && Settings.NotifyUnacknowledgedMinutes >= 0 && Issue.RetrievedAt - Issue.CreatedAt >= TimeSpan.FromMinutes(Settings.NotifyUnacknowledgedMinutes))
+							else if(ShowNotificationForIssue(issue) && _settings.NotifyUnacknowledgedMinutes >= 0 && issue.RetrievedAt - issue.CreatedAt >= TimeSpan.FromMinutes(_settings.NotifyUnacknowledgedMinutes))
 							{
-								Reason |= IssueAlertReason.UnacknowledgedTimer;
+								reason |= IssueAlertReason.UnacknowledgedTimer;
 							}
 						}
-						if(ShowNotificationForIssue(Issue) && Settings.NotifyUnresolvedMinutes >= 0 && Issue.RetrievedAt - Issue.CreatedAt >= TimeSpan.FromMinutes(Settings.NotifyUnresolvedMinutes))
+						if(ShowNotificationForIssue(issue) && _settings.NotifyUnresolvedMinutes >= 0 && issue.RetrievedAt - issue.CreatedAt >= TimeSpan.FromMinutes(_settings.NotifyUnresolvedMinutes))
 						{
-							Reason |= IssueAlertReason.UnresolvedTimer;
+							reason |= IssueAlertReason.UnresolvedTimer;
 						}
 
-						IssueAlertReason PrevReason;
-						if(IssueMonitor.IssueIdToAlertReason.TryGetValue(Issue.Id, out PrevReason))
+						IssueAlertReason prevReason;
+						if(issueMonitor.IssueIdToAlertReason.TryGetValue(issue.Id, out prevReason))
 						{
-							Reason &= ~PrevReason;
+							reason &= ~prevReason;
 						}
 					}
 
-					IssueAlertWindow AlertWindow = AlertWindows.FirstOrDefault(x => x.IssueMonitor == IssueMonitor && x.Issue.Id == Issue.Id);
-					if(AlertWindow == null)
+					IssueAlertWindow alertWindow = _alertWindows.FirstOrDefault(x => x.IssueMonitor == issueMonitor && x.Issue.Id == issue.Id);
+					if(alertWindow == null)
 					{
-						if(Reason != 0)
+						if(reason != 0)
 						{
-							ShowAlertWindow(IssueMonitor, Issue, Reason);
+							ShowAlertWindow(issueMonitor, issue, reason);
 						}
 					}
 					else
 					{
-						if(Reason != 0)
+						if(reason != 0)
 						{
-							AlertWindow.SetIssue(Issue, Reason);
+							alertWindow.SetIssue(issue, reason);
 						}
 						else
 						{
-							CloseAlertWindow(AlertWindow);
+							CloseAlertWindow(alertWindow);
 						}
 					}
 				}
-				AllIssues.UnionWith(Issues);
+				allIssues.UnionWith(issues);
 			}
 
 			// Close any alert windows which don't have an active issues
-			for(int Idx = 0; Idx < AlertWindows.Count; Idx++)
+			for(int idx = 0; idx < _alertWindows.Count; idx++)
 			{
-				IssueAlertWindow AlertWindow = AlertWindows[Idx];
-				if(!AllIssues.Contains(AlertWindow.Issue))
+				IssueAlertWindow alertWindow = _alertWindows[idx];
+				if(!allIssues.Contains(alertWindow.Issue))
 				{
-					AlertWindow.IssueMonitor.IssueIdToAlertReason.Remove(AlertWindow.Issue.Id);
-					CloseAlertWindow(AlertWindow);
-					Idx--;
+					alertWindow.IssueMonitor.IssueIdToAlertReason.Remove(alertWindow.Issue.Id);
+					CloseAlertWindow(alertWindow);
+					idx--;
 				}
 			}
 		}
 
 		void IssueMonitor_OnUpdateAsync()
 		{
-			MainThreadSynchronizationContext.Post((o) => IssueMonitor_OnUpdate(), null);
+			_mainThreadSynchronizationContext.Post((o) => IssueMonitor_OnUpdate(), null);
 		}
 
 		void IssueMonitor_OnUpdate()
@@ -1570,100 +1570,100 @@ namespace UnrealGameSync
 			UpdateAlertWindows();
 		}
 
-		void ShowAlertWindow(IssueMonitor IssueMonitor, IssueData Issue, IssueAlertReason Reason)
+		void ShowAlertWindow(IssueMonitor issueMonitor, IssueData issue, IssueAlertReason reason)
 		{
-			IssueAlertWindow Alert = new IssueAlertWindow(IssueMonitor, Issue, Reason);
-			Alert.AcceptBtn.Click += (s, e) => AcceptIssue(Alert);
-			Alert.DeclineBtn.Click += (s, e) => DeclineIssue(Alert);
-			Alert.DetailsBtn.Click += (s, e) => ShowIssueDetails(Alert);
+			IssueAlertWindow alert = new IssueAlertWindow(issueMonitor, issue, reason);
+			alert.AcceptBtn.Click += (s, e) => AcceptIssue(alert);
+			alert.DeclineBtn.Click += (s, e) => DeclineIssue(alert);
+			alert.DetailsBtn.Click += (s, e) => ShowIssueDetails(alert);
 
 			SetAlertWindowPositions();
-			AlertWindows.Add(Alert);
-			SetAlertWindowPosition(AlertWindows.Count - 1);
+			_alertWindows.Add(alert);
+			SetAlertWindowPosition(_alertWindows.Count - 1);
 
-			Alert.Show(this);
+			alert.Show(this);
 
 			UpdateAlertPositionsTimer.Enabled = true;
 		}
 
-		void AcceptIssue(IssueAlertWindow Alert)
+		void AcceptIssue(IssueAlertWindow alert)
 		{
-			IssueData Issue = Alert.Issue;
+			IssueData issue = alert.Issue;
 
-			IssueAlertReason Reason;
-			Alert.IssueMonitor.IssueIdToAlertReason.TryGetValue(Issue.Id, out Reason);
-			Alert.IssueMonitor.IssueIdToAlertReason[Issue.Id] = Reason | Alert.Reason;
+			IssueAlertReason reason;
+			alert.IssueMonitor.IssueIdToAlertReason.TryGetValue(issue.Id, out reason);
+			alert.IssueMonitor.IssueIdToAlertReason[issue.Id] = reason | alert.Reason;
 
-			IssueUpdateData Update = new IssueUpdateData();
-			Update.Id = Issue.Id;
-			Update.Owner = Alert.IssueMonitor.UserName;
-			Update.NominatedBy = null;
-			Update.Acknowledged = true;
-			Alert.IssueMonitor.PostUpdate(Update);
+			IssueUpdateData update = new IssueUpdateData();
+			update.Id = issue.Id;
+			update.Owner = alert.IssueMonitor.UserName;
+			update.NominatedBy = null;
+			update.Acknowledged = true;
+			alert.IssueMonitor.PostUpdate(update);
 
-			CloseAlertWindow(Alert);
+			CloseAlertWindow(alert);
 		}
 
-		void DeclineIssue(IssueAlertWindow Alert)
+		void DeclineIssue(IssueAlertWindow alert)
 		{
-			IssueData Issue = Alert.Issue;
+			IssueData issue = alert.Issue;
 
-			IssueAlertReason Reason;
-			Alert.IssueMonitor.IssueIdToAlertReason.TryGetValue(Issue.Id, out Reason);
-			Alert.IssueMonitor.IssueIdToAlertReason[Issue.Id] = Reason | Alert.Reason;
+			IssueAlertReason reason;
+			alert.IssueMonitor.IssueIdToAlertReason.TryGetValue(issue.Id, out reason);
+			alert.IssueMonitor.IssueIdToAlertReason[issue.Id] = reason | alert.Reason;
 
-			CloseAlertWindow(Alert);
+			CloseAlertWindow(alert);
 		}
 
-		void ShowIssueDetails(IssueAlertWindow Alert)
+		void ShowIssueDetails(IssueAlertWindow alert)
 		{
-			for(int Idx = 0; Idx < TabControl.GetTabCount(); Idx++)
+			for(int idx = 0; idx < TabControl.GetTabCount(); idx++)
 			{
-				WorkspaceControl? Workspace = TabControl.GetTabData(Idx) as WorkspaceControl;
-				if(Workspace != null && Workspace.GetIssueMonitor() == Alert.IssueMonitor)
+				WorkspaceControl? workspace = TabControl.GetTabData(idx) as WorkspaceControl;
+				if(workspace != null && workspace.GetIssueMonitor() == alert.IssueMonitor)
 				{
-					Workspace.ShowIssueDetails(Alert.Issue);
+					workspace.ShowIssueDetails(alert.Issue);
 					break;
 				}
 			}
 		}
 
-		void CloseAlertWindow(IssueAlertWindow Alert)
+		void CloseAlertWindow(IssueAlertWindow alert)
 		{
-			IssueData Issue = Alert.Issue;
+			IssueData issue = alert.Issue;
 
-			Alert.Close();
-			Alert.Dispose();
+			alert.Close();
+			alert.Dispose();
 
-			AlertWindows.Remove(Alert);
+			_alertWindows.Remove(alert);
 
-			for(int Idx = 0; Idx < AlertWindows.Count; Idx++)
+			for(int idx = 0; idx < _alertWindows.Count; idx++)
 			{
-				SetAlertWindowPosition(Idx);
+				SetAlertWindowPosition(idx);
 			}
 
-			if(AlertWindows.Count == 0)
+			if(_alertWindows.Count == 0)
 			{
 				UpdateAlertPositionsTimer.Enabled = false;
 			}
 		}
 
-		private void SetAlertWindowPosition(int Idx)
+		private void SetAlertWindowPosition(int idx)
 		{
-			AlertWindows[Idx].Location = new Point(PrimaryWorkArea.Right - 40 - AlertWindows[Idx].Size.Width, PrimaryWorkArea.Height - 40 - (Idx + 1) * (AlertWindows[Idx].Size.Height + 15));
+			_alertWindows[idx].Location = new Point(_primaryWorkArea.Right - 40 - _alertWindows[idx].Size.Width, _primaryWorkArea.Height - 40 - (idx + 1) * (_alertWindows[idx].Size.Height + 15));
 		}
 
 		private void SetAlertWindowPositions()
 		{
 			if (Screen.PrimaryScreen != null)
 			{
-				Rectangle NewPrimaryWorkArea = Screen.PrimaryScreen.WorkingArea;
-				if (NewPrimaryWorkArea != PrimaryWorkArea)
+				Rectangle newPrimaryWorkArea = Screen.PrimaryScreen.WorkingArea;
+				if (newPrimaryWorkArea != _primaryWorkArea)
 				{
-					PrimaryWorkArea = NewPrimaryWorkArea;
-					for (int Idx = 0; Idx < AlertWindows.Count; Idx++)
+					_primaryWorkArea = newPrimaryWorkArea;
+					for (int idx = 0; idx < _alertWindows.Count; idx++)
 					{
-						SetAlertWindowPosition(Idx);
+						SetAlertWindowPosition(idx);
 					}
 				}
 			}
@@ -1674,55 +1674,55 @@ namespace UnrealGameSync
 			SetAlertWindowPositions();
 		}
 
-		public IssueMonitor CreateIssueMonitor(string? ApiUrl, string UserName)
+		public IssueMonitor CreateIssueMonitor(string? apiUrl, string userName)
 		{
-			WorkspaceIssueMonitor WorkspaceIssueMonitor = WorkspaceIssueMonitors.FirstOrDefault(x => String.Compare(x.IssueMonitor.ApiUrl, ApiUrl, StringComparison.OrdinalIgnoreCase) == 0 && String.Compare(x.IssueMonitor.UserName, UserName, StringComparison.OrdinalIgnoreCase) == 0);
-			if (WorkspaceIssueMonitor == null)
+			WorkspaceIssueMonitor workspaceIssueMonitor = _workspaceIssueMonitors.FirstOrDefault(x => String.Compare(x.IssueMonitor.ApiUrl, apiUrl, StringComparison.OrdinalIgnoreCase) == 0 && String.Compare(x.IssueMonitor.UserName, userName, StringComparison.OrdinalIgnoreCase) == 0);
+			if (workspaceIssueMonitor == null)
 			{
-				string ServerId = ApiUrl != null ? Regex.Replace(ApiUrl, @"^.*://", "") : "noserver";
-				ServerId = Regex.Replace(ServerId, "[^a-zA-Z.]", "+");
+				string serverId = apiUrl != null ? Regex.Replace(apiUrl, @"^.*://", "") : "noserver";
+				serverId = Regex.Replace(serverId, "[^a-zA-Z.]", "+");
 
-				FileReference LogFileName = FileReference.Combine(DataFolder, String.Format("IssueMonitor-{0}-{1}.log", ServerId, UserName));
+				FileReference logFileName = FileReference.Combine(_dataFolder, String.Format("IssueMonitor-{0}-{1}.log", serverId, userName));
 
-				IssueMonitor IssueMonitor = new IssueMonitor(ApiUrl, UserName, TimeSpan.FromSeconds(60.0), ServiceProvider);
-				IssueMonitor.OnIssuesChanged += IssueMonitor_OnUpdateAsync;
+				IssueMonitor issueMonitor = new IssueMonitor(apiUrl, userName, TimeSpan.FromSeconds(60.0), _serviceProvider);
+				issueMonitor.OnIssuesChanged += IssueMonitor_OnUpdateAsync;
 
-				WorkspaceIssueMonitor = new WorkspaceIssueMonitor(IssueMonitor);
-				WorkspaceIssueMonitors.Add(WorkspaceIssueMonitor);
+				workspaceIssueMonitor = new WorkspaceIssueMonitor(issueMonitor);
+				_workspaceIssueMonitors.Add(workspaceIssueMonitor);
 			}
 
-			WorkspaceIssueMonitor.RefCount++;
+			workspaceIssueMonitor.RefCount++;
 
-			if (WorkspaceIssueMonitor.RefCount == 1 && bAllowCreatingHandle)
+			if (workspaceIssueMonitor.RefCount == 1 && _allowCreatingHandle)
 			{
-				WorkspaceIssueMonitor.IssueMonitor.Start();
+				workspaceIssueMonitor.IssueMonitor.Start();
 			}
 
-			return WorkspaceIssueMonitor.IssueMonitor;
+			return workspaceIssueMonitor.IssueMonitor;
 		}
 
-		public void ReleaseIssueMonitor(IssueMonitor IssueMonitor)
+		public void ReleaseIssueMonitor(IssueMonitor issueMonitor)
 		{
-			int Index = WorkspaceIssueMonitors.FindIndex(x => x.IssueMonitor == IssueMonitor);
-			if(Index != -1)
+			int index = _workspaceIssueMonitors.FindIndex(x => x.IssueMonitor == issueMonitor);
+			if(index != -1)
 			{
-				WorkspaceIssueMonitor WorkspaceIssueMonitor = WorkspaceIssueMonitors[Index];
-				WorkspaceIssueMonitor.RefCount--;
+				WorkspaceIssueMonitor workspaceIssueMonitor = _workspaceIssueMonitors[index];
+				workspaceIssueMonitor.RefCount--;
 
-				if (WorkspaceIssueMonitor.RefCount == 0)
+				if (workspaceIssueMonitor.RefCount == 0)
 				{
-					for (int Idx = AlertWindows.Count - 1; Idx >= 0; Idx--)
+					for (int idx = _alertWindows.Count - 1; idx >= 0; idx--)
 					{
-						IssueAlertWindow AlertWindow = AlertWindows[Idx];
-						if (AlertWindow.IssueMonitor == IssueMonitor)
+						IssueAlertWindow alertWindow = _alertWindows[idx];
+						if (alertWindow.IssueMonitor == issueMonitor)
 						{
-							CloseAlertWindow(AlertWindow);
+							CloseAlertWindow(alertWindow);
 						}
 					}
-					IssueMonitor.OnIssuesChanged -= IssueMonitor_OnUpdateAsync;
-					IssueMonitor.Release();
+					issueMonitor.OnIssuesChanged -= IssueMonitor_OnUpdateAsync;
+					issueMonitor.Release();
 
-					WorkspaceIssueMonitors.RemoveAt(Index);
+					_workspaceIssueMonitors.RemoveAt(index);
 				}
 			}
 		}
diff --git a/Engine/Source/Programs/UnrealGameSync/UnrealGameSync/Forms/ModalTaskWindow.cs b/Engine/Source/Programs/UnrealGameSync/UnrealGameSync/Forms/ModalTaskWindow.cs
--- a/Engine/Source/Programs/UnrealGameSync/UnrealGameSync/Forms/ModalTaskWindow.cs
+++ b/Engine/Source/Programs/UnrealGameSync/UnrealGameSync/Forms/ModalTaskWindow.cs
@@ -16,18 +16,18 @@ namespace UnrealGameSync
 {
 	partial class ModalTaskWindow : Form
 	{
-		Task Task;
-		CancellationTokenSource CancellationSource;
+		Task _task;
+		CancellationTokenSource _cancellationSource;
 
-		public ModalTaskWindow(string InTitle, string InMessage, FormStartPosition InStartPosition, Task InTask, CancellationTokenSource InCancellationSource)//, Func<CancellationToken, Task> InTaskFunc)
+		public ModalTaskWindow(string inTitle, string inMessage, FormStartPosition inStartPosition, Task inTask, CancellationTokenSource inCancellationSource)//, Func<CancellationToken, Task> InTaskFunc)
 		{
 			InitializeComponent();
 
-			Text = InTitle;
-			MessageLabel.Text = InMessage;
-			StartPosition = InStartPosition;
-			Task = InTask;
-			CancellationSource = InCancellationSource;
+			Text = inTitle;
+			MessageLabel.Text = inMessage;
+			StartPosition = inStartPosition;
+			_task = inTask;
+			_cancellationSource = inCancellationSource;
 		}
 
 		public void ShowAndActivate()
@@ -38,14 +38,14 @@ namespace UnrealGameSync
 
 		private void ModalTaskWindow_Load(object sender, EventArgs e)
 		{
-			SynchronizationContext SyncContext = SynchronizationContext.Current!;
-			Task.ContinueWith(x => SyncContext.Post(y => Close(), null));
+			SynchronizationContext syncContext = SynchronizationContext.Current!;
+			_task.ContinueWith(x => syncContext.Post(y => Close(), null));
 		}
 
 		private void ModalTaskWindow_FormClosing(object sender, FormClosingEventArgs e)
 		{
-			CancellationSource.Cancel();
-			if (!Task.IsCompleted)
+			_cancellationSource.Cancel();
+			if (!_task.IsCompleted)
 			{
 				e.Cancel = true;
 			}
