UGS: Reformat to standard C# style conventions.

diff --git a/Engine/Source/Programs/UnrealGameSync/UnrealGameSync/Forms/ModifyBuildStepsWindow.cs b/Engine/Source/Programs/UnrealGameSync/UnrealGameSync/Forms/ModifyBuildStepsWindow.cs
--- a/Engine/Source/Programs/UnrealGameSync/UnrealGameSync/Forms/ModifyBuildStepsWindow.cs
+++ b/Engine/Source/Programs/UnrealGameSync/UnrealGameSync/Forms/ModifyBuildStepsWindow.cs
@@ -16,29 +16,29 @@ namespace UnrealGameSync
 {
 	partial class ModifyBuildStepsWindow : Form
 	{
-		List<string> TargetNames;
-		List<BuildStep> Steps;
-		HashSet<Guid> ProjectSteps;
-		DirectoryReference BaseDirectory;
-		ListViewItem.ListViewSubItem? MouseDownSubItem = null;
-		IReadOnlyDictionary<string, string> Variables;
+		List<string> _targetNames;
+		List<BuildStep> _steps;
+		HashSet<Guid> _projectSteps;
+		DirectoryReference _baseDirectory;
+		ListViewItem.ListViewSubItem? _mouseDownSubItem = null;
+		IReadOnlyDictionary<string, string> _variables;
 
-		public ModifyBuildStepsWindow(List<string> InTargetNames, List<BuildStep> InSteps, HashSet<Guid> InProjectSteps, DirectoryReference InBaseDirectory, IReadOnlyDictionary<string, string> InVariables)
+		public ModifyBuildStepsWindow(List<string> inTargetNames, List<BuildStep> inSteps, HashSet<Guid> inProjectSteps, DirectoryReference inBaseDirectory, IReadOnlyDictionary<string, string> inVariables)
 		{
-			TargetNames = InTargetNames;
-			Steps = InSteps;
-			ProjectSteps = InProjectSteps;
-			BaseDirectory = InBaseDirectory;
-			Variables = InVariables;
+			_targetNames = inTargetNames;
+			_steps = inSteps;
+			_projectSteps = inProjectSteps;
+			_baseDirectory = inBaseDirectory;
+			_variables = inVariables;
 
 			InitializeComponent();
 
-			using(Graphics Graphics = Graphics.FromHwnd(IntPtr.Zero))
+			using(Graphics graphics = Graphics.FromHwnd(IntPtr.Zero))
 			{
-				float DpiScaleX = Graphics.DpiX / 96.0f;
+				float dpiScaleX = graphics.DpiX / 96.0f;
 
-				NormalSyncColumn.Width = (int)(104 * DpiScaleX);
-				ScheduledSyncColumn.Width = (int)(104 * DpiScaleX);
+				NormalSyncColumn.Width = (int)(104 * dpiScaleX);
+				ScheduledSyncColumn.Width = (int)(104 * dpiScaleX);
 				DescriptionColumn.Width = BuildStepList.ClientSize.Width - NormalSyncColumn.Width - ScheduledSyncColumn.Width - 10;
 			}
 
@@ -47,44 +47,44 @@ namespace UnrealGameSync
 
 		private void ModifyBuildStepsWindow_Load(object sender, EventArgs e)
 		{
-			foreach(BuildStep Task in Steps)
+			foreach(BuildStep task in _steps)
 			{
-				AddTask(Task);
+				AddTask(task);
 			}
 			UpdateEnabledButtons();
 		}
 
-		void AddTask(BuildStep Task)
+		void AddTask(BuildStep task)
 		{
-			ListViewItem Item = new ListViewItem(Task.Description);
-			Item.Tag = Task;
-			Item.SubItems.Add(new ListViewItem.ListViewSubItem());
-			Item.SubItems.Add(new ListViewItem.ListViewSubItem());
-			BuildStepList.Items.Add(Item);
+			ListViewItem item = new ListViewItem(task.Description);
+			item.Tag = task;
+			item.SubItems.Add(new ListViewItem.ListViewSubItem());
+			item.SubItems.Add(new ListViewItem.ListViewSubItem());
+			BuildStepList.Items.Add(item);
 		}
 
 		private void NewStepButton_Click(object sender, EventArgs e)
 		{
-			BuildStep NewStep = new BuildStep(Guid.NewGuid(), BuildStepList.Items.Count, "Untitled Step", "Running Untitled Step...", 1, null, null, null, null, true);
-			NewStep.Description = "Untitled Task";
-			NewStep.EstimatedDuration = 1;
+			BuildStep newStep = new BuildStep(Guid.NewGuid(), BuildStepList.Items.Count, "Untitled Step", "Running Untitled Step...", 1, null, null, null, null, true);
+			newStep.Description = "Untitled Task";
+			newStep.EstimatedDuration = 1;
 
-			BuildStepWindow NewStepWindow = new BuildStepWindow(NewStep, TargetNames, BaseDirectory, Variables);
-			if(NewStepWindow.ShowDialog() == DialogResult.OK)
+			BuildStepWindow newStepWindow = new BuildStepWindow(newStep, _targetNames, _baseDirectory, _variables);
+			if(newStepWindow.ShowDialog() == DialogResult.OK)
 			{
-				AddTask(NewStep);
+				AddTask(newStep);
 			}
 		}
 
 		private void EditStepButton_Click(object sender, EventArgs e)
 		{
-			foreach(ListViewItem? Item in BuildStepList.SelectedItems)
+			foreach(ListViewItem? item in BuildStepList.SelectedItems)
 			{
-				if (Item != null)
+				if (item != null)
 				{
-					BuildStepWindow EditStep = new BuildStepWindow((BuildStep)Item.Tag, TargetNames, BaseDirectory, Variables);
-					EditStep.ShowDialog();
-					Item.Text = ((BuildStep)Item.Tag).Description;
+					BuildStepWindow editStep = new BuildStepWindow((BuildStep)item.Tag, _targetNames, _baseDirectory, _variables);
+					editStep.ShowDialog();
+					item.Text = ((BuildStep)item.Tag).Description;
 					break;
 				}
 			}
@@ -92,13 +92,13 @@ namespace UnrealGameSync
 
 		private void RemoveStepButton_Click(object sender, EventArgs e)
 		{
-			foreach(ListViewItem? Item in BuildStepList.SelectedItems)
+			foreach(ListViewItem? item in BuildStepList.SelectedItems)
 			{
-				if (Item != null)
+				if (item != null)
 				{
-					if (MessageBox.Show(String.Format("Remove the '{0}' step?", Item.Text), "Remove Step", MessageBoxButtons.YesNo) == DialogResult.Yes)
+					if (MessageBox.Show(String.Format("Remove the '{0}' step?", item.Text), "Remove Step", MessageBoxButtons.YesNo) == DialogResult.Yes)
 					{
-						BuildStepList.Items.Remove(Item);
+						BuildStepList.Items.Remove(item);
 					}
 					break;
 				}
@@ -107,12 +107,12 @@ namespace UnrealGameSync
 
 		private void CloseButton_Click(object sender, EventArgs e)
 		{
-			Steps.Clear();
-			foreach(ListViewItem? Item in BuildStepList.Items)
+			_steps.Clear();
+			foreach(ListViewItem? item in BuildStepList.Items)
 			{
-				if (Item != null)
+				if (item != null)
 				{
-					Steps.Add((BuildStep)Item.Tag);
+					_steps.Add((BuildStep)item.Tag);
 				}
 			}
 			Close();
@@ -120,14 +120,14 @@ namespace UnrealGameSync
 
 		private void UpdateEnabledButtons()
 		{
-			int SelectedIndex = (BuildStepList.SelectedIndices.Count > 0)? BuildStepList.SelectedIndices[0] : -1;
+			int selectedIndex = (BuildStepList.SelectedIndices.Count > 0)? BuildStepList.SelectedIndices[0] : -1;
 
-			bool bHasSelection = (SelectedIndex != -1);
-			EditStepButton.Enabled = bHasSelection;
-			RemoveStepButton.Enabled = bHasSelection && !ProjectSteps.Contains(((BuildStep)BuildStepList.SelectedItems[0].Tag).UniqueId);
+			bool hasSelection = (selectedIndex != -1);
+			EditStepButton.Enabled = hasSelection;
+			RemoveStepButton.Enabled = hasSelection && !_projectSteps.Contains(((BuildStep)BuildStepList.SelectedItems[0].Tag).UniqueId);
 
-			MoveUp.Enabled = (SelectedIndex >= 1);
-			MoveDown.Enabled = (SelectedIndex >= 0 && SelectedIndex < BuildStepList.Items.Count - 1);
+			MoveUp.Enabled = (selectedIndex >= 1);
+			MoveDown.Enabled = (selectedIndex >= 0 && selectedIndex < BuildStepList.Items.Count - 1);
 		}
 
 		private void BuildStepList_SelectedIndexChanged(object sender, EventArgs e)
@@ -149,85 +149,85 @@ namespace UnrealGameSync
 			}
 			else
 			{
-				BuildStep Task = (BuildStep)e.Item.Tag;
+				BuildStep task = (BuildStep)e.Item.Tag;
 
-				bool bEnabled;
+				bool enabled;
 				if(e.ColumnIndex == 1)
 				{
-					bEnabled = Task.bNormalSync;
+					enabled = task.NormalSync;
 				}
 				else
 				{
-					bEnabled = Task.bScheduledSync;
+					enabled = task.ScheduledSync;
 				}
 
-				bool Selected = BuildStepList.SelectedItems.Contains(e.Item);
-				e.Graphics.FillRectangle(Selected? SystemBrushes.Highlight : SystemBrushes.Window, e.Bounds);
+				bool selected = BuildStepList.SelectedItems.Contains(e.Item);
+				e.Graphics.FillRectangle(selected? SystemBrushes.Highlight : SystemBrushes.Window, e.Bounds);
 
-				CheckBoxState State;
-				if(bEnabled)
+				CheckBoxState state;
+				if(enabled)
 				{
-					State = (MouseDownSubItem == e.SubItem)? CheckBoxState.CheckedPressed : CheckBoxState.CheckedNormal;
+					state = (_mouseDownSubItem == e.SubItem)? CheckBoxState.CheckedPressed : CheckBoxState.CheckedNormal;
 				}
 				else
 				{
-					State = (MouseDownSubItem == e.SubItem)? CheckBoxState.UncheckedPressed : CheckBoxState.UncheckedNormal;
+					state = (_mouseDownSubItem == e.SubItem)? CheckBoxState.UncheckedPressed : CheckBoxState.UncheckedNormal;
 				}
 				
-				Size Size = CheckBoxRenderer.GetGlyphSize(e.Graphics, State);
-				CheckBoxRenderer.DrawCheckBox(e.Graphics, new Point(e.Bounds.Left + (e.Bounds.Width - Size.Width) / 2, e.Bounds.Top + (e.Bounds.Height - Size.Height) / 2), State);
+				Size size = CheckBoxRenderer.GetGlyphSize(e.Graphics, state);
+				CheckBoxRenderer.DrawCheckBox(e.Graphics, new Point(e.Bounds.Left + (e.Bounds.Width - size.Width) / 2, e.Bounds.Top + (e.Bounds.Height - size.Height) / 2), state);
 			}
 		}
 
 		private void BuildStepList_MouseDown(object sender, MouseEventArgs e)
 		{
-			ListViewHitTestInfo HitTest = BuildStepList.HitTest(e.X, e.Y);
-			MouseDownSubItem = HitTest.SubItem;
-			if(MouseDownSubItem != null)
+			ListViewHitTestInfo hitTest = BuildStepList.HitTest(e.X, e.Y);
+			_mouseDownSubItem = hitTest.SubItem;
+			if(_mouseDownSubItem != null)
 			{
-				BuildStepList.Invalidate(MouseDownSubItem.Bounds, true);
+				BuildStepList.Invalidate(_mouseDownSubItem.Bounds, true);
 			}
 		}
 
 		private void BuildStepList_MouseUp(object sender, MouseEventArgs e)
 		{
-			ListViewHitTestInfo HitTest = BuildStepList.HitTest(e.X, e.Y);
-			if(HitTest.Item != null && HitTest.SubItem != null)
+			ListViewHitTestInfo hitTest = BuildStepList.HitTest(e.X, e.Y);
+			if(hitTest.Item != null && hitTest.SubItem != null)
 			{
-				int ColumnIndex = HitTest.Item.SubItems.IndexOf(HitTest.SubItem);
-				if(ColumnIndex >= 1 && ColumnIndex <= 3)
+				int columnIndex = hitTest.Item.SubItems.IndexOf(hitTest.SubItem);
+				if(columnIndex >= 1 && columnIndex <= 3)
 				{
-					BuildStep Task = (BuildStep)HitTest.Item.Tag;
-					if(ColumnIndex == 1)
+					BuildStep task = (BuildStep)hitTest.Item.Tag;
+					if(columnIndex == 1)
 					{
-						Task.bNormalSync ^= true;
+						task.NormalSync ^= true;
 					}
 					else
 					{
-						Task.bScheduledSync ^= true;
+						task.ScheduledSync ^= true;
 					}
-					BuildStepList.Invalidate(HitTest.SubItem.Bounds);
+					BuildStepList.Invalidate(hitTest.SubItem.Bounds);
 				}
 			}
-			if(MouseDownSubItem != null)
+			if(_mouseDownSubItem != null)
 			{
-				BuildStepList.Invalidate(MouseDownSubItem.Bounds);
-				MouseDownSubItem = null;
+				BuildStepList.Invalidate(_mouseDownSubItem.Bounds);
+				_mouseDownSubItem = null;
 			}
 		}
 
 		private void MoveUp_Click(object sender, EventArgs e)
 		{
 			BuildStepList.BeginUpdate();
-			foreach(ListViewItem? Item in BuildStepList.SelectedItems)
+			foreach(ListViewItem? item in BuildStepList.SelectedItems)
 			{
-				if (Item != null)
+				if (item != null)
 				{
-					int Index = Item.Index;
-					if (Index > 0)
+					int index = item.Index;
+					if (index > 0)
 					{
-						BuildStepList.Items.RemoveAt(Index);
-						BuildStepList.Items.Insert(Index - 1, Item);
+						BuildStepList.Items.RemoveAt(index);
+						BuildStepList.Items.Insert(index - 1, item);
 					}
 					break;
 				}
@@ -239,15 +239,15 @@ namespace UnrealGameSync
 		private void MoveDown_Click(object sender, EventArgs e)
 		{
 			BuildStepList.BeginUpdate();
-			foreach(ListViewItem? Item in BuildStepList.SelectedItems)
+			foreach(ListViewItem? item in BuildStepList.SelectedItems)
 			{
-				if (Item != null)
+				if (item != null)
 				{
-					int Index = Item.Index;
-					if (Index < BuildStepList.Items.Count - 1)
+					int index = item.Index;
+					if (index < BuildStepList.Items.Count - 1)
 					{
-						BuildStepList.Items.RemoveAt(Index);
-						BuildStepList.Items.Insert(Index + 1, Item);
+						BuildStepList.Items.RemoveAt(index);
+						BuildStepList.Items.Insert(index + 1, item);
 					}
 					break;
 				}
@@ -258,14 +258,14 @@ namespace UnrealGameSync
 
 		private void ModifyBuildStepsWindow_FormClosed(object sender, FormClosedEventArgs e)
 		{
-			Steps.Clear();
-			foreach(ListViewItem? Item in BuildStepList.Items)
+			_steps.Clear();
+			foreach(ListViewItem? item in BuildStepList.Items)
 			{
-				if (Item != null)
+				if (item != null)
 				{
-					BuildStep Step = (BuildStep)Item.Tag;
-					Step.OrderIndex = Steps.Count;
-					Steps.Add(Step);
+					BuildStep step = (BuildStep)item.Tag;
+					step.OrderIndex = _steps.Count;
+					_steps.Add(step);
 				}
 			}
 		}
diff --git a/Engine/Source/Programs/UnrealGameSync/UnrealGameSync/Forms/NewWorkspaceWindow.cs b/Engine/Source/Programs/UnrealGameSync/UnrealGameSync/Forms/NewWorkspaceWindow.cs
--- a/Engine/Source/Programs/UnrealGameSync/UnrealGameSync/Forms/NewWorkspaceWindow.cs
+++ b/Engine/Source/Programs/UnrealGameSync/UnrealGameSync/Forms/NewWorkspaceWindow.cs
@@ -27,11 +27,11 @@ namespace UnrealGameSync
 			public string Stream { get; }
 			public DirectoryReference RootDir { get; }
 
-			public NewWorkspaceSettings(string Name, string Stream, DirectoryReference RootDir)
+			public NewWorkspaceSettings(string name, string stream, DirectoryReference rootDir)
 			{
-				this.Name = Name;
-				this.Stream = Stream;
-				this.RootDir = RootDir;
+				this.Name = name;
+				this.Stream = stream;
+				this.RootDir = rootDir;
 			}
 		}
 
@@ -41,195 +41,195 @@ namespace UnrealGameSync
 			public List<ClientsRecord> Clients;
 			public string? CurrentStream;
 
-			public FindWorkspaceSettingsTask(InfoRecord Info, List<ClientsRecord> Clients, string? CurrentStream)
+			public FindWorkspaceSettingsTask(InfoRecord info, List<ClientsRecord> clients, string? currentStream)
 			{
-				this.Info = Info;
-				this.Clients = Clients;
-				this.CurrentStream = CurrentStream;
+				this.Info = info;
+				this.Clients = clients;
+				this.CurrentStream = currentStream;
 			}
 
-			public static async Task<FindWorkspaceSettingsTask> RunAsync(IPerforceConnection Perforce, string CurrentWorkspaceName, CancellationToken CancellationToken)
+			public static async Task<FindWorkspaceSettingsTask> RunAsync(IPerforceConnection perforce, string currentWorkspaceName, CancellationToken cancellationToken)
 			{
-				InfoRecord Info = await Perforce.GetInfoAsync(InfoOptions.ShortOutput, CancellationToken);
-				List<ClientsRecord> Clients = await Perforce.GetClientsAsync(ClientsOptions.None, Perforce.Settings.UserName, CancellationToken);
+				InfoRecord info = await perforce.GetInfoAsync(InfoOptions.ShortOutput, cancellationToken);
+				List<ClientsRecord> clients = await perforce.GetClientsAsync(ClientsOptions.None, perforce.Settings.UserName, cancellationToken);
 
-				string? CurrentStream = null;
-				if(!String.IsNullOrEmpty(CurrentWorkspaceName))
+				string? currentStream = null;
+				if(!String.IsNullOrEmpty(currentWorkspaceName))
 				{
-					CurrentStream = await Perforce.GetCurrentStreamAsync(CancellationToken);
+					currentStream = await perforce.GetCurrentStreamAsync(cancellationToken);
 				}
 
-				return new FindWorkspaceSettingsTask(Info, Clients, CurrentStream);
+				return new FindWorkspaceSettingsTask(info, clients, currentStream);
 			}
 		}
 
 		static class NewWorkspaceTask
 		{
-			public static async Task RunAsync(IPerforceConnection Perforce, NewWorkspaceSettings Settings, string Owner, string HostName, CancellationToken CancellationToken)
+			public static async Task RunAsync(IPerforceConnection perforce, NewWorkspaceSettings settings, string owner, string hostName, CancellationToken cancellationToken)
 			{
-				PerforceResponseList<ClientsRecord> Response = await Perforce.TryGetClientsAsync(ClientsOptions.None, Settings.Name, -1, null, null, CancellationToken);
-				if (!Response.Succeeded)
+				PerforceResponseList<ClientsRecord> response = await perforce.TryGetClientsAsync(ClientsOptions.None, settings.Name, -1, null, null, cancellationToken);
+				if (!response.Succeeded)
 				{
-					throw new UserErrorException($"Unable to determine if client already exists.\n\n{Response[0]}");
+					throw new UserErrorException($"Unable to determine if client already exists.\n\n{response[0]}");
 				}
-				if (Response.Data.Count > 0)
+				if (response.Data.Count > 0)
 				{
-					throw new UserErrorException($"Client '{Settings.Name}' already exists.");
+					throw new UserErrorException($"Client '{settings.Name}' already exists.");
 				}
 
-				ClientRecord Client = new ClientRecord(Settings.Name, Owner, Settings.RootDir.FullName);
-				Client.Host = HostName;
-				Client.Stream = Settings.Stream;
-				Client.Options = ClientOptions.Rmdir;
-				await Perforce.CreateClientAsync(Client, CancellationToken);
+				ClientRecord client = new ClientRecord(settings.Name, owner, settings.RootDir.FullName);
+				client.Host = hostName;
+				client.Stream = settings.Stream;
+				client.Options = ClientOptions.Rmdir;
+				await perforce.CreateClientAsync(client, cancellationToken);
 			}
 		}
 
-		IPerforceSettings PerforceSettings;
-		InfoRecord Info;
-		List<ClientsRecord> Clients;
-		IServiceProvider ServiceProvider;
-		NewWorkspaceSettings? Settings;
-		DirectoryReference? DefaultRootPath;
+		IPerforceSettings _perforceSettings;
+		InfoRecord _info;
+		List<ClientsRecord> _clients;
+		IServiceProvider _serviceProvider;
+		NewWorkspaceSettings? _settings;
+		DirectoryReference? _defaultRootPath;
 
-		private NewWorkspaceWindow(IPerforceSettings PerforceSettings, string? ForceStream, string? DefaultStream, InfoRecord Info, List<ClientsRecord> Clients, IServiceProvider ServiceProvider)
+		private NewWorkspaceWindow(IPerforceSettings perforceSettings, string? forceStream, string? defaultStream, InfoRecord info, List<ClientsRecord> clients, IServiceProvider serviceProvider)
 		{
 			InitializeComponent();
 
-			this.PerforceSettings = PerforceSettings;
-			this.Info = Info;
-			this.Clients = Clients;
-			this.ServiceProvider = ServiceProvider;
+			this._perforceSettings = perforceSettings;
+			this._info = info;
+			this._clients = clients;
+			this._serviceProvider = serviceProvider;
 
-			Dictionary<DirectoryReference, int> RootPathToCount = new Dictionary<DirectoryReference, int>();
-			foreach(ClientsRecord Client in Clients)
+			Dictionary<DirectoryReference, int> rootPathToCount = new Dictionary<DirectoryReference, int>();
+			foreach(ClientsRecord client in clients)
 			{
-				if(Client.Host == null || String.Compare(Client.Host, Info.ClientHost, StringComparison.OrdinalIgnoreCase) == 0)
+				if(client.Host == null || String.Compare(client.Host, info.ClientHost, StringComparison.OrdinalIgnoreCase) == 0)
 				{
-					if(!String.IsNullOrEmpty(Client.Root) && Client.Root != ".")
+					if(!String.IsNullOrEmpty(client.Root) && client.Root != ".")
 					{
-						DirectoryReference? ParentDir;
+						DirectoryReference? parentDir;
 						try
 						{
-							ParentDir = new DirectoryReference(Client.Root);
+							parentDir = new DirectoryReference(client.Root);
 						}
 						catch
 						{
-							ParentDir = null;
+							parentDir = null;
 						}
 
-						if(ParentDir != null)
+						if(parentDir != null)
 						{
-							int Count;
-							RootPathToCount.TryGetValue(ParentDir, out Count);
-							RootPathToCount[ParentDir] = Count + 1;
+							int count;
+							rootPathToCount.TryGetValue(parentDir, out count);
+							rootPathToCount[parentDir] = count + 1;
 						}
 					}
 				}
 			}
 
-			int RootPathMaxCount = 0;
-			foreach(KeyValuePair<DirectoryReference, int> RootPathPair in RootPathToCount)
+			int rootPathMaxCount = 0;
+			foreach(KeyValuePair<DirectoryReference, int> rootPathPair in rootPathToCount)
 			{
-				if(RootPathPair.Value > RootPathMaxCount)
+				if(rootPathPair.Value > rootPathMaxCount)
 				{
-					DefaultRootPath = RootPathPair.Key;
-					RootPathMaxCount = RootPathPair.Value;
+					_defaultRootPath = rootPathPair.Key;
+					rootPathMaxCount = rootPathPair.Value;
 				}
 			}
 
-			if(ForceStream != null)
+			if(forceStream != null)
 			{
-				StreamTextBox.Text = ForceStream;
+				StreamTextBox.Text = forceStream;
 				StreamTextBox.Enabled = false;
 			}
 			else
 			{
-				StreamTextBox.Text = DefaultStream ?? "";
+				StreamTextBox.Text = defaultStream ?? "";
 				StreamTextBox.Enabled = true;
 			}
 			StreamTextBox.SelectionStart = StreamTextBox.Text.Length;
 			StreamTextBox.SelectionLength = 0;
 			StreamTextBox.Focus();
 
-			StreamBrowseBtn.Enabled = (ForceStream == null);
+			StreamBrowseBtn.Enabled = (forceStream == null);
 
 			UpdateOkButton();
 			UpdateNameCueBanner();
 			UpdateRootDirCueBanner();
 		}
 
-		public static bool ShowModal(IWin32Window Owner, IPerforceSettings PerforceSettings, string? ForceStreamName, string CurrentWorkspaceName, IServiceProvider ServiceProvider, [NotNullWhen(true)] out string? WorkspaceName)
+		public static bool ShowModal(IWin32Window owner, IPerforceSettings perforceSettings, string? forceStreamName, string currentWorkspaceName, IServiceProvider serviceProvider, [NotNullWhen(true)] out string? workspaceName)
 		{
-			ModalTask<FindWorkspaceSettingsTask>? Task = PerforceModalTask.Execute(Owner, "Checking settings", "Checking settings, please wait...", PerforceSettings, (p, c) => FindWorkspaceSettingsTask.RunAsync(p, CurrentWorkspaceName, c), ServiceProvider.GetRequiredService<ILogger<FindWorkspaceSettingsTask>>());
-			if (Task == null || !Task.Succeeded)
+			ModalTask<FindWorkspaceSettingsTask>? task = PerforceModalTask.Execute(owner, "Checking settings", "Checking settings, please wait...", perforceSettings, (p, c) => FindWorkspaceSettingsTask.RunAsync(p, currentWorkspaceName, c), serviceProvider.GetRequiredService<ILogger<FindWorkspaceSettingsTask>>());
+			if (task == null || !task.Succeeded)
 			{
-				WorkspaceName = null;
+				workspaceName = null;
 				return false;
 			}
 
-			NewWorkspaceWindow Window = new NewWorkspaceWindow(PerforceSettings, ForceStreamName, Task.Result.CurrentStream, Task.Result.Info, Task.Result.Clients, ServiceProvider);
-			if(Window.ShowDialog(Owner) == DialogResult.OK)
+			NewWorkspaceWindow window = new NewWorkspaceWindow(perforceSettings, forceStreamName, task.Result.CurrentStream, task.Result.Info, task.Result.Clients, serviceProvider);
+			if(window.ShowDialog(owner) == DialogResult.OK)
 			{
-				WorkspaceName = Window.Settings!.Name;
+				workspaceName = window._settings!.Name;
 				return true;
 			}
 			else
 			{
-				WorkspaceName = null;
+				workspaceName = null;
 				return false;
 			}
 		}
 
 		private void RootDirBrowseBtn_Click(object sender, EventArgs e)
 		{
-			FolderBrowserDialog Dialog = new FolderBrowserDialog();
-			Dialog.ShowNewFolderButton = true;
-			Dialog.SelectedPath = RootDirTextBox.Text;
-			if (Dialog.ShowDialog() == DialogResult.OK)
+			FolderBrowserDialog dialog = new FolderBrowserDialog();
+			dialog.ShowNewFolderButton = true;
+			dialog.SelectedPath = RootDirTextBox.Text;
+			if (dialog.ShowDialog() == DialogResult.OK)
 			{
-				RootDirTextBox.Text = Dialog.SelectedPath;
+				RootDirTextBox.Text = dialog.SelectedPath;
 				UpdateOkButton();
 			}
 		}
 
 		private string GetDefaultWorkspaceName()
 		{
-			string BaseName = Sanitize(String.Format("{0}_{1}_{2}", Info.UserName, Info.ClientHost, StreamTextBox.Text.Replace('/', '_').Trim('_'))).Trim('_');
+			string baseName = Sanitize(String.Format("{0}_{1}_{2}", _info.UserName, _info.ClientHost, StreamTextBox.Text.Replace('/', '_').Trim('_'))).Trim('_');
 
-			string Name = BaseName;
-			for(int Idx = 2; Clients.Any(x => x.Name != null && String.Compare(x.Name, Name, StringComparison.InvariantCultureIgnoreCase) == 0); Idx++)
+			string name = baseName;
+			for(int idx = 2; _clients.Any(x => x.Name != null && String.Compare(x.Name, name, StringComparison.InvariantCultureIgnoreCase) == 0); idx++)
 			{
-				Name = String.Format("{0}_{1}", BaseName, Idx);
+				name = String.Format("{0}_{1}", baseName, idx);
 			}
-			return Name;
+			return name;
 		}
 
 		private string GetDefaultWorkspaceRootDir()
 		{
-			string RootDir = "";
-			if(DefaultRootPath != null)
+			string rootDir = "";
+			if(_defaultRootPath != null)
 			{
-				string Suffix = String.Join("_", StreamTextBox.Text.Split(new char[]{ '/' }, StringSplitOptions.RemoveEmptyEntries).Select(x => Sanitize(x)).Where(x => x.Length > 0));
-				if(Suffix.Length > 0)
+				string suffix = String.Join("_", StreamTextBox.Text.Split(new char[]{ '/' }, StringSplitOptions.RemoveEmptyEntries).Select(x => Sanitize(x)).Where(x => x.Length > 0));
+				if(suffix.Length > 0)
 				{
-					RootDir = DirectoryReference.Combine(DefaultRootPath, Suffix).FullName;
+					rootDir = DirectoryReference.Combine(_defaultRootPath, suffix).FullName;
 				}
 			}
-			return RootDir;
+			return rootDir;
 		}
 
-		private string Sanitize(string Text)
+		private string Sanitize(string text)
 		{
-			StringBuilder Result = new StringBuilder();
-			for(int Idx = 0; Idx < Text.Length; Idx++)
+			StringBuilder result = new StringBuilder();
+			for(int idx = 0; idx < text.Length; idx++)
 			{
-				if(Char.IsLetterOrDigit(Text[Idx]) || Text[Idx] == '_' || Text[Idx] == '.' || Text[Idx] == '-')
+				if(Char.IsLetterOrDigit(text[idx]) || text[idx] == '_' || text[idx] == '.' || text[idx] == '-')
 				{
-					Result.Append(Text[Idx]);
+					result.Append(text[idx]);
 				}
 			}
-			return Result.ToString();
+			return result.ToString();
 		}
 
 		private void UpdateNameCueBanner()
@@ -244,53 +244,53 @@ namespace UnrealGameSync
 
 		private void UpdateOkButton()
 		{
-			NewWorkspaceSettings? Settings;
-			OkBtn.Enabled = TryGetWorkspaceSettings(out Settings);
+			NewWorkspaceSettings? settings;
+			OkBtn.Enabled = TryGetWorkspaceSettings(out settings);
 		}
 
-		private bool TryGetWorkspaceSettings([NotNullWhen(true)] out NewWorkspaceSettings? Settings)
+		private bool TryGetWorkspaceSettings([NotNullWhen(true)] out NewWorkspaceSettings? settings)
 		{
-			string NewWorkspaceName = NameTextBox.Text.Trim();
-			if(NewWorkspaceName.Length == 0)
+			string newWorkspaceName = NameTextBox.Text.Trim();
+			if(newWorkspaceName.Length == 0)
 			{
-				NewWorkspaceName = GetDefaultWorkspaceName();
-				if(NewWorkspaceName.Length == 0)
+				newWorkspaceName = GetDefaultWorkspaceName();
+				if(newWorkspaceName.Length == 0)
 				{
-					Settings = null;
+					settings = null;
 					return false;
 				}
 			}
 
-			string NewStream = StreamTextBox.Text.Trim();
-			if(!NewStream.StartsWith("//") || NewStream.IndexOf('/', 2) == -1)
+			string newStream = StreamTextBox.Text.Trim();
+			if(!newStream.StartsWith("//") || newStream.IndexOf('/', 2) == -1)
 			{
-				Settings = null;
+				settings = null;
 				return false;
 			}
 
-			string NewRootDir = RootDirTextBox.Text.Trim();
-			if(NewRootDir.Length == 0)
+			string newRootDir = RootDirTextBox.Text.Trim();
+			if(newRootDir.Length == 0)
 			{
-				NewRootDir = GetDefaultWorkspaceRootDir();
-				if(NewRootDir.Length == 0)
+				newRootDir = GetDefaultWorkspaceRootDir();
+				if(newRootDir.Length == 0)
 				{
-					Settings = null;
+					settings = null;
 					return false;
 				}
 			}
 
-			DirectoryReference NewRootDirRef;
+			DirectoryReference newRootDirRef;
 			try
 			{
-				NewRootDirRef = new DirectoryReference(NewRootDir);
+				newRootDirRef = new DirectoryReference(newRootDir);
 			}
 			catch
 			{
-				Settings = null;
+				settings = null;
 				return false;
 			}
 
-			Settings = new NewWorkspaceSettings(NewWorkspaceName, NewStream, NewRootDirRef);
+			settings = new NewWorkspaceSettings(newWorkspaceName, newStream, newRootDirRef);
 			return true;
 		}
 
@@ -313,21 +313,21 @@ namespace UnrealGameSync
 
 		private void OkBtn_Click(object sender, EventArgs e)
 		{
-			if(TryGetWorkspaceSettings(out Settings))
+			if(TryGetWorkspaceSettings(out _settings))
 			{
-				DirectoryInfo RootDir = Settings.RootDir.ToDirectoryInfo();
-				if(RootDir.Exists && RootDir.EnumerateFileSystemInfos().Any(x => x.Name != "." && x.Name != ".."))
+				DirectoryInfo rootDir = _settings.RootDir.ToDirectoryInfo();
+				if(rootDir.Exists && rootDir.EnumerateFileSystemInfos().Any(x => x.Name != "." && x.Name != ".."))
 				{
-					if(MessageBox.Show(this, String.Format("The directory '{0}' is not empty. Are you sure you want to create a workspace there?", RootDir.FullName), "Directory not empty", MessageBoxButtons.YesNo) != DialogResult.Yes)
+					if(MessageBox.Show(this, String.Format("The directory '{0}' is not empty. Are you sure you want to create a workspace there?", rootDir.FullName), "Directory not empty", MessageBoxButtons.YesNo) != DialogResult.Yes)
 					{
 						return;
 					}
 				}
 
-				ILogger Logger = ServiceProvider.GetRequiredService<ILogger<NewWorkspaceWindow>>();
+				ILogger logger = _serviceProvider.GetRequiredService<ILogger<NewWorkspaceWindow>>();
 
-				ModalTask? Result = PerforceModalTask.Execute(Owner, "Creating workspace", "Creating workspace, please wait...", PerforceSettings, (p, c) => NewWorkspaceTask.RunAsync(p, Settings, Info.UserName ?? "", Info.ClientHost ?? "", c), Logger);
-				if (Result == null || !Result.Succeeded)
+				ModalTask? result = PerforceModalTask.Execute(Owner, "Creating workspace", "Creating workspace, please wait...", _perforceSettings, (p, c) => NewWorkspaceTask.RunAsync(p, _settings, _info.UserName ?? "", _info.ClientHost ?? "", c), logger);
+				if (result == null || !result.Succeeded)
 				{
 					return;
 				}
@@ -345,10 +345,10 @@ namespace UnrealGameSync
 
 		private void StreamBrowseBtn_Click(object sender, EventArgs e)
 		{
-			string? StreamName = StreamTextBox.Text.Trim();
-			if(SelectStreamWindow.ShowModal(this, PerforceSettings, StreamName, ServiceProvider, out StreamName))
+			string? streamName = StreamTextBox.Text.Trim();
+			if(SelectStreamWindow.ShowModal(this, _perforceSettings, streamName, _serviceProvider, out streamName))
 			{
-				StreamTextBox.Text = StreamName;
+				StreamTextBox.Text = streamName;
 			}
 		}
 
diff --git a/Engine/Source/Programs/UnrealGameSync/UnrealGameSync/Forms/NotificationWindow.cs b/Engine/Source/Programs/UnrealGameSync/UnrealGameSync/Forms/NotificationWindow.cs
--- a/Engine/Source/Programs/UnrealGameSync/UnrealGameSync/Forms/NotificationWindow.cs
+++ b/Engine/Source/Programs/UnrealGameSync/UnrealGameSync/Forms/NotificationWindow.cs
@@ -23,111 +23,111 @@ namespace UnrealGameSync
 
 	public partial class NotificationWindow : Form
 	{
-		NotificationType Type;
+		NotificationType _type;
 
 		// Logo
-		Rectangle LogoBounds;
-		Bitmap? LogoBitmap;
+		Rectangle _logoBounds;
+		Bitmap? _logoBitmap;
 
 		// Caption
-		Rectangle CaptionBounds;
-		string CaptionText = "Caption Text";
-		Font? CaptionFont;
+		Rectangle _captionBounds;
+		string _captionText = "Caption Text";
+		Font? _captionFont;
 
 		// Message
-		Rectangle MessageBounds;
-		string MessageText = "Mesage Text";
+		Rectangle _messageBounds;
+		string _messageText = "Mesage Text";
 
 		// Close button
-		Rectangle CloseButtonBounds;
-		bool bHoverOverClose;
-		bool bMouseDownOverClose;
+		Rectangle _closeButtonBounds;
+		bool _hoverOverClose;
+		bool _mouseDownOverClose;
 
 		// Notifications
 		public Action? OnMoreInformation;
 		public Action? OnDismiss;
 
-		public NotificationWindow(Image InLogo)
+		public NotificationWindow(Image inLogo)
 		{
-			LogoBitmap = new Bitmap(InLogo);
+			_logoBitmap = new Bitmap(inLogo);
 			InitializeComponent();
 		}
 
-		public void Show(NotificationType InType, string InCaption, string InMessage)
+		public void Show(NotificationType inType, string inCaption, string inMessage)
 		{
-			Type = InType;
-			CaptionText = InCaption;
-			MessageText = InMessage;
+			_type = inType;
+			_captionText = inCaption;
+			_messageText = inMessage;
 
 			base.Show();
 
 			CalculateBounds();
 
-			Rectangle WorkingArea = Screen.PrimaryScreen.WorkingArea;
-			Location = new Point(WorkingArea.Right - Size.Width - 16, WorkingArea.Bottom - Size.Height - 16);
+			Rectangle workingArea = Screen.PrimaryScreen.WorkingArea;
+			Location = new Point(workingArea.Right - Size.Width - 16, workingArea.Bottom - Size.Height - 16);
 
 			BringToFront();
 			Invalidate();
 		}
 
-		protected override void Dispose(bool bDisposing)
+		protected override void Dispose(bool disposing)
 		{
-			if(bDisposing && LogoBitmap != null)
+			if(disposing && _logoBitmap != null)
 			{
-				LogoBitmap.Dispose();
-				LogoBitmap = null;
+				_logoBitmap.Dispose();
+				_logoBitmap = null;
 			}
-			if(bDisposing && components != null)
+			if(disposing && components != null)
 			{
 				components.Dispose();
 			}
-			base.Dispose(bDisposing);
+			base.Dispose(disposing);
 		}
 
 		protected override void OnClosed(EventArgs e)
 		{
 			base.OnClosed(e);
 
-			if(LogoBitmap != null)
+			if(_logoBitmap != null)
 			{
-				LogoBitmap.Dispose();
-				LogoBitmap = null;
+				_logoBitmap.Dispose();
+				_logoBitmap = null;
 			}
 		}
 
 		public void CalculateBounds()
 		{
-			int ReferenceSize8 = (int)(8 * Font.Height) / 16;
-			using(Graphics Graphics = CreateGraphics())
+			int referenceSize8 = (int)(8 * Font.Height) / 16;
+			using(Graphics graphics = CreateGraphics())
 			{
-				int GutterW = ReferenceSize8 + ReferenceSize8 / 2;
-				int GutterH = ReferenceSize8 + ReferenceSize8 / 2;
+				int gutterW = referenceSize8 + referenceSize8 / 2;
+				int gutterH = referenceSize8 + referenceSize8 / 2;
 
 				// Get the close button size
-				int CloseButtonSize = ReferenceSize8 * 2;
+				int closeButtonSize = referenceSize8 * 2;
 
 				// Space out the text
-				int CaptionY = GutterH;
-				int CaptionW = TextRenderer.MeasureText(Graphics, CaptionText, CaptionFont ?? Font, new Size(100 * ReferenceSize8, Font.Height), TextFormatFlags.Left | TextFormatFlags.Top | TextFormatFlags.SingleLine).Width;
-				int CaptionH = (CaptionFont ?? Font).Height;
+				int captionY = gutterH;
+				int captionW = TextRenderer.MeasureText(graphics, _captionText, _captionFont ?? Font, new Size(100 * referenceSize8, Font.Height), TextFormatFlags.Left | TextFormatFlags.Top | TextFormatFlags.SingleLine).Width;
+				int captionH = (_captionFont ?? Font).Height;
 
-				int MessageY = CaptionY + CaptionH + ReferenceSize8 / 2;
-				int MessageW = TextRenderer.MeasureText(Graphics, MessageText, Font, new Size(100 * ReferenceSize8, Font.Height), TextFormatFlags.Left | TextFormatFlags.Top | TextFormatFlags.SingleLine).Width;
-				int MessageH = Font.Height;
+				int messageY = captionY + captionH + referenceSize8 / 2;
+				int messageW = TextRenderer.MeasureText(graphics, _messageText, Font, new Size(100 * referenceSize8, Font.Height), TextFormatFlags.Left | TextFormatFlags.Top | TextFormatFlags.SingleLine).Width;
+				int messageH = Font.Height;
 
 				// Get the logo dimensions
-				int LogoH = MessageY + MessageH - GutterH;
-				int LogoW = (LogoBitmap == null)? 16 : (int)(LogoBitmap.Width * (float)LogoH / (float)LogoBitmap.Height);
+				int logoH = messageY + messageH - gutterH;
+				int logoW = (_logoBitmap == null)? 16 : (int)(_logoBitmap.Width * (float)logoH / (float)_logoBitmap.Height);
 
 				// Set the window size
-				int TextX = GutterW + LogoW + ReferenceSize8 * 2;
-				Size = new Size(TextX + Math.Max(MessageW, CaptionW) + CloseButtonSize + GutterW, MessageY + MessageH + GutterH);
+				int textX = gutterW + logoW + referenceSize8 * 2;
+				Size = new Size(textX + Math.Max(messageW, captionW) + closeButtonSize + gutterW, messageY + messageH + gutterH);
 
 				// Set the bounds of the individual elements
-				CloseButtonBounds = new Rectangle(Width - GutterW - CloseButtonSize, GutterH, CloseButtonSize, CloseButtonSize);
-				LogoBounds = new Rectangle(GutterW + ReferenceSize8, (Height - LogoH) / 2, LogoW, LogoH);
-				CaptionBounds = new Rectangle(TextX, CaptionY, CaptionW, CaptionH);
-				MessageBounds = new Rectangle(TextX, MessageY, MessageW, MessageH);
+				_closeButtonBounds = new Rectangle(Width - gutterW - closeButtonSize, gutterH, closeButtonSize, closeButtonSize);
+				_logoBounds = new Rectangle(gutterW + referenceSize8, (Height - logoH) / 2, logoW, logoH);
+				_captionBounds = new Rectangle(textX, captionY, captionW, captionH);
+				_messageBounds = new Rectangle(textX, messageY, messageW, messageH);
 			}
 		}
 
@@ -135,11 +135,11 @@ namespace UnrealGameSync
 		{
 			base.OnFontChanged(e);
 
-			if(CaptionFont != null)
+			if(_captionFont != null)
 			{
-				CaptionFont.Dispose();
+				_captionFont.Dispose();
 			}
-			CaptionFont = new Font(this.Font.Name, this.Font.Size * (12.0f / 9.0f), FontStyle.Regular);
+			_captionFont = new Font(this.Font.Name, this.Font.Size * (12.0f / 9.0f), FontStyle.Regular);
 
 			CalculateBounds();
 			Invalidate();
@@ -162,11 +162,11 @@ namespace UnrealGameSync
 		{
 			base.OnMouseMove(e);
 
-			bool bNewHoverOverClose = CloseButtonBounds.Contains(e.Location);
-			if(bHoverOverClose != bNewHoverOverClose)
+			bool newHoverOverClose = _closeButtonBounds.Contains(e.Location);
+			if(_hoverOverClose != newHoverOverClose)
 			{
-				bHoverOverClose = bNewHoverOverClose;
-				Invalidate(CloseButtonBounds);
+				_hoverOverClose = newHoverOverClose;
+				Invalidate(_closeButtonBounds);
 			}
 		}
 
@@ -176,8 +176,8 @@ namespace UnrealGameSync
 
 			if(e.Button == System.Windows.Forms.MouseButtons.Left)
 			{
-				bMouseDownOverClose = CloseButtonBounds.Contains(e.Location);
-				Invalidate(CloseButtonBounds);
+				_mouseDownOverClose = _closeButtonBounds.Contains(e.Location);
+				Invalidate(_closeButtonBounds);
 			}
 		}
 
@@ -187,10 +187,10 @@ namespace UnrealGameSync
 
 			if(e.Button == System.Windows.Forms.MouseButtons.Left)
 			{
-				if(bMouseDownOverClose)
+				if(_mouseDownOverClose)
 				{
-					bMouseDownOverClose = false;
-					Invalidate(CloseButtonBounds);
+					_mouseDownOverClose = false;
+					Invalidate(_closeButtonBounds);
 					Hide();
 					if(OnDismiss != null)
 					{
@@ -212,29 +212,29 @@ namespace UnrealGameSync
 		{
 			base.OnPaintBackground(e);
 
-			Color TintColor;
-			switch(Type)
+			Color tintColor;
+			switch(_type)
 			{
 				case NotificationType.Warning:
-					TintColor = Color.FromArgb(240, 240, 220);
+					tintColor = Color.FromArgb(240, 240, 220);
 					break;
 				case NotificationType.Error:
-					TintColor = Color.FromArgb(240, 220, 220);
+					tintColor = Color.FromArgb(240, 220, 220);
 					break;
 				default:
-					TintColor = Color.FromArgb(220, 220, 240);
+					tintColor = Color.FromArgb(220, 220, 240);
 					break;
 			}
 
-			using(Brush BackgroundBrush = new LinearGradientBrush(new Point(0, 0), new Point(0, Height), Color.White, TintColor))
+			using(Brush backgroundBrush = new LinearGradientBrush(new Point(0, 0), new Point(0, Height), Color.White, tintColor))
 			{
-				e.Graphics.FillRectangle(BackgroundBrush, ClientRectangle);
+				e.Graphics.FillRectangle(backgroundBrush, ClientRectangle);
 			}
 
 			e.Graphics.DrawRectangle(Pens.Black, new Rectangle(0, 0, Width - 1, Height - 1));
 
 			e.Graphics.InterpolationMode = InterpolationMode.HighQualityBicubic;
-			e.Graphics.DrawImage(LogoBitmap, LogoBounds);
+			e.Graphics.DrawImage(_logoBitmap, _logoBounds);
 		}
 
 		protected override void OnPaint(PaintEventArgs e)
@@ -244,33 +244,33 @@ namespace UnrealGameSync
 			e.Graphics.SmoothingMode = SmoothingMode.AntiAlias;
 
 			// Draw the close button
-			if(bHoverOverClose)
+			if(_hoverOverClose)
 			{
-				e.Graphics.FillRectangle(SystemBrushes.ActiveCaption, CloseButtonBounds);
+				e.Graphics.FillRectangle(SystemBrushes.ActiveCaption, _closeButtonBounds);
 			}
-			else if(bMouseDownOverClose)
+			else if(_mouseDownOverClose)
 			{
-				e.Graphics.FillRectangle(SystemBrushes.InactiveCaption, CloseButtonBounds);
+				e.Graphics.FillRectangle(SystemBrushes.InactiveCaption, _closeButtonBounds);
 			}
-			if(bHoverOverClose || bMouseDownOverClose)
+			if(_hoverOverClose || _mouseDownOverClose)
 			{
-				using(Pen BorderPen = new Pen(SystemColors.InactiveCaptionText, 1.0f))
+				using(Pen borderPen = new Pen(SystemColors.InactiveCaptionText, 1.0f))
 				{
-					e.Graphics.DrawRectangle(BorderPen, new Rectangle(CloseButtonBounds.X, CloseButtonBounds.Y, CloseButtonBounds.Width - 1, CloseButtonBounds.Height - 1));
+					e.Graphics.DrawRectangle(borderPen, new Rectangle(_closeButtonBounds.X, _closeButtonBounds.Y, _closeButtonBounds.Width - 1, _closeButtonBounds.Height - 1));
 				}
 			}
-			using(Pen CrossPen = new Pen(Brushes.Black, 2.0f))
+			using(Pen crossPen = new Pen(Brushes.Black, 2.0f))
 			{
-				float Offset = (CloseButtonBounds.Width - 1) / 4.0f;
-				e.Graphics.DrawLine(CrossPen, CloseButtonBounds.Left + Offset, CloseButtonBounds.Top + Offset, CloseButtonBounds.Right - 1 - Offset, CloseButtonBounds.Bottom - 1 - Offset);
-				e.Graphics.DrawLine(CrossPen, CloseButtonBounds.Left + Offset, CloseButtonBounds.Bottom - 1 - Offset, CloseButtonBounds.Right - 1 - Offset, CloseButtonBounds.Top + Offset);
+				float offset = (_closeButtonBounds.Width - 1) / 4.0f;
+				e.Graphics.DrawLine(crossPen, _closeButtonBounds.Left + offset, _closeButtonBounds.Top + offset, _closeButtonBounds.Right - 1 - offset, _closeButtonBounds.Bottom - 1 - offset);
+				e.Graphics.DrawLine(crossPen, _closeButtonBounds.Left + offset, _closeButtonBounds.Bottom - 1 - offset, _closeButtonBounds.Right - 1 - offset, _closeButtonBounds.Top + offset);
 			}
 
 			// Draw the caption
-			TextRenderer.DrawText(e.Graphics, CaptionText, CaptionFont, CaptionBounds, SystemColors.HotTrack, TextFormatFlags.Left);
+			TextRenderer.DrawText(e.Graphics, _captionText, _captionFont, _captionBounds, SystemColors.HotTrack, TextFormatFlags.Left);
 
 			// Draw the message
-			TextRenderer.DrawText(e.Graphics, MessageText, Font, MessageBounds, SystemColors.ControlText, TextFormatFlags.Left);
+			TextRenderer.DrawText(e.Graphics, _messageText, Font, _messageBounds, SystemColors.ControlText, TextFormatFlags.Left);
 		}
 	}
 }
diff --git a/Engine/Source/Programs/UnrealGameSync/UnrealGameSync/Forms/OIDCLoginWindow.cs b/Engine/Source/Programs/UnrealGameSync/UnrealGameSync/Forms/OIDCLoginWindow.cs
--- a/Engine/Source/Programs/UnrealGameSync/UnrealGameSync/Forms/OIDCLoginWindow.cs
+++ b/Engine/Source/Programs/UnrealGameSync/UnrealGameSync/Forms/OIDCLoginWindow.cs
@@ -7,23 +7,23 @@ using UnrealGameSync.Controls;
 
 namespace UnrealGameSync.Forms
 {
-	public partial class OIDCLoginWindow : Form
+	public partial class OidcLoginWindow : Form
 	{
-		public OIDCLoginWindow(OIDCTokenManager OidcManager)
+		public OidcLoginWindow(OidcTokenManager oidcManager)
 		{
 			InitializeComponent();
 
-			int LastYPosition = 10;
-			foreach (OIDCTokenManager.ProviderInfo Provider in OidcManager.Providers.Values)
+			int lastYPosition = 10;
+			foreach (OidcTokenManager.ProviderInfo provider in oidcManager.Providers.Values)
 			{
-				OIDCControl ServiceControl = new OIDCControl(OidcManager, Provider.Identifier, Provider.DisplayName)
+				OidcControl serviceControl = new OidcControl(oidcManager, provider.Identifier, provider.DisplayName)
 				{
 					Anchor = AnchorStyles.Top | AnchorStyles.Left | AnchorStyles.Right,
-					Location = new Point(10, LastYPosition)
+					Location = new Point(10, lastYPosition)
 				};
-				Controls.Add(ServiceControl);
+				Controls.Add(serviceControl);
 
-				LastYPosition = ServiceControl.Size.Height + ServiceControl.Location.Y + 10;
+				lastYPosition = serviceControl.Size.Height + serviceControl.Location.Y + 10;
 			}
 		}
 
diff --git a/Engine/Source/Programs/UnrealGameSync/UnrealGameSync/Forms/OpenProjectWindow.cs b/Engine/Source/Programs/UnrealGameSync/UnrealGameSync/Forms/OpenProjectWindow.cs
--- a/Engine/Source/Programs/UnrealGameSync/UnrealGameSync/Forms/OpenProjectWindow.cs
+++ b/Engine/Source/Programs/UnrealGameSync/UnrealGameSync/Forms/OpenProjectWindow.cs
@@ -23,55 +23,55 @@ namespace UnrealGameSync
 {
 	partial class OpenProjectWindow : Form
 	{
-		string? ServerAndPortOverride;
-		string? UserNameOverride;
-		OpenProjectInfo? OpenProjectInfo;
-		IPerforceSettings DefaultSettings;
-		IServiceProvider ServiceProvider;
-		UserSettings Settings;
-		ILogger Logger;
-
-		private OpenProjectWindow(UserSelectedProjectSettings? Project, UserSettings Settings, IPerforceSettings DefaultSettings, IServiceProvider ServiceProvider, ILogger Logger)
+		string? _serverAndPortOverride;
+		string? _userNameOverride;
+		OpenProjectInfo? _openProjectInfo;
+		IPerforceSettings _defaultSettings;
+		IServiceProvider _serviceProvider;
+		UserSettings _settings;
+		ILogger _logger;
+
+		private OpenProjectWindow(UserSelectedProjectSettings? project, UserSettings settings, IPerforceSettings defaultSettings, IServiceProvider serviceProvider, ILogger logger)
 		{
 			InitializeComponent();
 
-			this.Settings = Settings;
-			this.OpenProjectInfo = null;
-			this.DefaultSettings = DefaultSettings;
-			this.ServiceProvider = ServiceProvider;
-			this.Logger = Logger;
+			this._settings = settings;
+			this._openProjectInfo = null;
+			this._defaultSettings = defaultSettings;
+			this._serviceProvider = serviceProvider;
+			this._logger = logger;
 
-			if (Project == null)
+			if (project == null)
 			{
 				LocalFileRadioBtn.Checked = true;
 			}
 			else
 			{
-				if(!String.IsNullOrWhiteSpace(Project.ServerAndPort))
+				if(!String.IsNullOrWhiteSpace(project.ServerAndPort))
 				{
-					ServerAndPortOverride = Project.ServerAndPort;
+					_serverAndPortOverride = project.ServerAndPort;
 				}
-				if(!String.IsNullOrWhiteSpace(Project.UserName))
+				if(!String.IsNullOrWhiteSpace(project.UserName))
 				{
-					UserNameOverride = Project.UserName;
+					_userNameOverride = project.UserName;
 				}
 
-				if(Project.ClientPath != null && Project.ClientPath.StartsWith("//"))
+				if(project.ClientPath != null && project.ClientPath.StartsWith("//"))
 				{
-					int SlashIdx = Project.ClientPath.IndexOf('/', 2);
-					if(SlashIdx != -1)
+					int slashIdx = project.ClientPath.IndexOf('/', 2);
+					if(slashIdx != -1)
 					{
-						WorkspaceNameTextBox.Text = Project.ClientPath.Substring(2, SlashIdx - 2);
-						WorkspacePathTextBox.Text = Project.ClientPath.Substring(SlashIdx);
+						WorkspaceNameTextBox.Text = project.ClientPath.Substring(2, slashIdx - 2);
+						WorkspacePathTextBox.Text = project.ClientPath.Substring(slashIdx);
 					}
 				}
 
-				if(Project.LocalPath != null)
+				if(project.LocalPath != null)
 				{
-					LocalFileTextBox.Text = Project.LocalPath;
+					LocalFileTextBox.Text = project.LocalPath;
 				}
 
-				if(Project.Type == UserSelectedProjectType.Client)
+				if(project.Type == UserSelectedProjectType.Client)
 				{
 					WorkspaceRadioBtn.Checked = true;
 				}
@@ -89,15 +89,15 @@ namespace UnrealGameSync
 
 		private IPerforceSettings Perforce
 		{
-			get => Utility.OverridePerforceSettings(DefaultSettings, ServerAndPortOverride, UserNameOverride);
+			get => Utility.OverridePerforceSettings(_defaultSettings, _serverAndPortOverride, _userNameOverride);
 		}
 
-		public static OpenProjectInfo? ShowModal(IWin32Window Owner, UserSelectedProjectSettings? Project, UserSettings Settings, DirectoryReference DataFolder, DirectoryReference CacheFolder, IPerforceSettings DefaultPerforceSettings, IServiceProvider ServiceProvider, ILogger Logger)
+		public static OpenProjectInfo? ShowModal(IWin32Window owner, UserSelectedProjectSettings? project, UserSettings settings, DirectoryReference dataFolder, DirectoryReference cacheFolder, IPerforceSettings defaultPerforceSettings, IServiceProvider serviceProvider, ILogger logger)
 		{
-			OpenProjectWindow Window = new OpenProjectWindow(Project, Settings, DefaultPerforceSettings, ServiceProvider, Logger);
-			if(Window.ShowDialog(Owner) == DialogResult.OK)
+			OpenProjectWindow window = new OpenProjectWindow(project, settings, defaultPerforceSettings, serviceProvider, logger);
+			if(window.ShowDialog(owner) == DialogResult.OK)
 			{
-				return Window.OpenProjectInfo;
+				return window._openProjectInfo;
 			}
 			else
 			{
@@ -107,78 +107,78 @@ namespace UnrealGameSync
 
 		private void UpdateEnabledControls()
 		{
-			Color WorkspaceTextColor = WorkspaceRadioBtn.Checked? SystemColors.ControlText : SystemColors.GrayText;
-			WorkspaceNameLabel.ForeColor = WorkspaceTextColor;
-			WorkspaceNameTextBox.ForeColor = WorkspaceTextColor;
-			WorkspaceNameNewBtn.ForeColor = WorkspaceTextColor;
-			WorkspaceNameBrowseBtn.ForeColor = WorkspaceTextColor;
-			WorkspacePathLabel.ForeColor = WorkspaceTextColor;
-			WorkspacePathTextBox.ForeColor = WorkspaceTextColor;
-			WorkspacePathBrowseBtn.ForeColor = WorkspaceTextColor;
-
-			Color LocalFileTextColor = LocalFileRadioBtn.Checked? SystemColors.ControlText : SystemColors.GrayText;
-			LocalFileLabel.ForeColor = LocalFileTextColor;
-			LocalFileTextBox.ForeColor = LocalFileTextColor;
-			LocalFileBrowseBtn.ForeColor = LocalFileTextColor;
+			Color workspaceTextColor = WorkspaceRadioBtn.Checked? SystemColors.ControlText : SystemColors.GrayText;
+			WorkspaceNameLabel.ForeColor = workspaceTextColor;
+			WorkspaceNameTextBox.ForeColor = workspaceTextColor;
+			WorkspaceNameNewBtn.ForeColor = workspaceTextColor;
+			WorkspaceNameBrowseBtn.ForeColor = workspaceTextColor;
+			WorkspacePathLabel.ForeColor = workspaceTextColor;
+			WorkspacePathTextBox.ForeColor = workspaceTextColor;
+			WorkspacePathBrowseBtn.ForeColor = workspaceTextColor;
+
+			Color localFileTextColor = LocalFileRadioBtn.Checked? SystemColors.ControlText : SystemColors.GrayText;
+			LocalFileLabel.ForeColor = localFileTextColor;
+			LocalFileTextBox.ForeColor = localFileTextColor;
+			LocalFileBrowseBtn.ForeColor = localFileTextColor;
 
 			UpdateWorkspacePathBrowseButton();
 		}
 
-		public static string GetServerLabelText(IPerforceSettings DefaultSettings, string? ServerAndPort, string? UserName)
+		public static string GetServerLabelText(IPerforceSettings defaultSettings, string? serverAndPort, string? userName)
 		{
-			if(ServerAndPort == null && UserName == null)
+			if(serverAndPort == null && userName == null)
 			{
-				return String.Format("Using default connection settings (user '{0}' on server '{1}').", DefaultSettings.UserName, DefaultSettings.ServerAndPort);
+				return String.Format("Using default connection settings (user '{0}' on server '{1}').", defaultSettings.UserName, defaultSettings.ServerAndPort);
 			}
 			else
 			{
-				StringBuilder Text = new StringBuilder("Connecting as ");
-				if(UserName == null)
+				StringBuilder text = new StringBuilder("Connecting as ");
+				if(userName == null)
 				{
-					Text.Append("default user");
+					text.Append("default user");
 				}
 				else
 				{
-					Text.AppendFormat("user '{0}'", UserName);
+					text.AppendFormat("user '{0}'", userName);
 				}
-				Text.Append(" on ");
-				if(ServerAndPort == null)
+				text.Append(" on ");
+				if(serverAndPort == null)
 				{
-					Text.Append("default server.");
+					text.Append("default server.");
 				}
 				else
 				{
-					Text.AppendFormat("server '{0}'.", ServerAndPort);
+					text.AppendFormat("server '{0}'.", serverAndPort);
 				}
-				return Text.ToString();
+				return text.ToString();
 			}
 		}
 
 		private void UpdateServerLabel()
 		{
-			ServerLabel.Text = GetServerLabelText(DefaultSettings, ServerAndPortOverride, UserNameOverride);
+			ServerLabel.Text = GetServerLabelText(_defaultSettings, _serverAndPortOverride, _userNameOverride);
 		}
 
 		private void UpdateWorkspacePathBrowseButton()
 		{
-			string? WorkspaceName;
-			WorkspacePathBrowseBtn.Enabled = TryGetWorkspaceName(out WorkspaceName);
+			string? workspaceName;
+			WorkspacePathBrowseBtn.Enabled = TryGetWorkspaceName(out workspaceName);
 		}
 
 		private void UpdateOkButton()
 		{
-			string? ProjectPath;
-			OkBtn.Enabled = WorkspaceRadioBtn.Checked? TryGetClientPath(out ProjectPath) : TryGetLocalPath(out ProjectPath);
+			string? projectPath;
+			OkBtn.Enabled = WorkspaceRadioBtn.Checked? TryGetClientPath(out projectPath) : TryGetLocalPath(out projectPath);
 		}
 
 		private void WorkspaceNewBtn_Click(object sender, EventArgs e)
 		{
 			WorkspaceRadioBtn.Checked = true;
 			
-			string? WorkspaceName;
-			if(NewWorkspaceWindow.ShowModal(this, Perforce, null, WorkspaceNameTextBox.Text, ServiceProvider, out WorkspaceName))
+			string? workspaceName;
+			if(NewWorkspaceWindow.ShowModal(this, Perforce, null, WorkspaceNameTextBox.Text, _serviceProvider, out workspaceName))
 			{
-				WorkspaceNameTextBox.Text = WorkspaceName;
+				WorkspaceNameTextBox.Text = workspaceName;
 				UpdateOkButton();
 			}
 		}
@@ -187,10 +187,10 @@ namespace UnrealGameSync
 		{
 			WorkspaceRadioBtn.Checked = true;
 
-			string? WorkspaceName = WorkspaceNameTextBox.Text;
-			if(SelectWorkspaceWindow.ShowModal(this, Perforce, WorkspaceName, ServiceProvider, out WorkspaceName))
+			string? workspaceName = WorkspaceNameTextBox.Text;
+			if(SelectWorkspaceWindow.ShowModal(this, Perforce, workspaceName, _serviceProvider, out workspaceName))
 			{
-				WorkspaceNameTextBox.Text = WorkspaceName;
+				WorkspaceNameTextBox.Text = workspaceName;
 			}
 		}
 
@@ -198,125 +198,125 @@ namespace UnrealGameSync
 		{
 			WorkspaceRadioBtn.Checked = true;
 
-			string? WorkspaceName;
-			if(TryGetWorkspaceName(out WorkspaceName))
+			string? workspaceName;
+			if(TryGetWorkspaceName(out workspaceName))
 			{
-				string? WorkspacePath = WorkspacePathTextBox.Text.Trim();
-				if(SelectProjectFromWorkspaceWindow.ShowModal(this, Perforce, WorkspaceName, WorkspacePath, ServiceProvider, out WorkspacePath))
+				string? workspacePath = WorkspacePathTextBox.Text.Trim();
+				if(SelectProjectFromWorkspaceWindow.ShowModal(this, Perforce, workspaceName, workspacePath, _serviceProvider, out workspacePath))
 				{
-					WorkspacePathTextBox.Text = WorkspacePath;
+					WorkspacePathTextBox.Text = workspacePath;
 					UpdateOkButton();
 				}
 			}
 		}
 
-		private bool TryGetWorkspaceName([NotNullWhen(true)] out string? WorkspaceName)
+		private bool TryGetWorkspaceName([NotNullWhen(true)] out string? workspaceName)
 		{
-			string Text = WorkspaceNameTextBox.Text.Trim();
-			if(Text.Length == 0)
+			string text = WorkspaceNameTextBox.Text.Trim();
+			if(text.Length == 0)
 			{
-				WorkspaceName = null;
+				workspaceName = null;
 				return false;
 			}
 
-			WorkspaceName = Text;
+			workspaceName = text;
 			return true;
 		}
 
-		private bool TryGetClientPath([NotNullWhen(true)] out string? ClientPath)
+		private bool TryGetClientPath([NotNullWhen(true)] out string? clientPath)
 		{
-			string? WorkspaceName;
-			if(!TryGetWorkspaceName(out WorkspaceName))
+			string? workspaceName;
+			if(!TryGetWorkspaceName(out workspaceName))
 			{
-				ClientPath = null;
+				clientPath = null;
 				return false;
 			}
 
-			string WorkspacePath = WorkspacePathTextBox.Text.Trim();
-			if(WorkspacePath.Length == 0 || WorkspacePath[0] != '/')
+			string workspacePath = WorkspacePathTextBox.Text.Trim();
+			if(workspacePath.Length == 0 || workspacePath[0] != '/')
 			{
-				ClientPath = null;
+				clientPath = null;
 				return false;
 			}
 
-			ClientPath = String.Format("//{0}{1}", WorkspaceName, WorkspacePath);
+			clientPath = String.Format("//{0}{1}", workspaceName, workspacePath);
 			return true;
 		}
 
-		private bool TryGetLocalPath(out string? LocalPath)
+		private bool TryGetLocalPath(out string? localPath)
 		{
-			string LocalFile = LocalFileTextBox.Text.Trim();
-			if(LocalFile.Length == 0)
+			string localFile = LocalFileTextBox.Text.Trim();
+			if(localFile.Length == 0)
 			{
-				LocalPath = null;
+				localPath = null;
 				return false;
 			}
 
-			LocalPath = Path.GetFullPath(LocalFile);
+			localPath = Path.GetFullPath(localFile);
 			return true;
 		}
 
-		private bool TryGetSelectedProject([NotNullWhen(true)] out UserSelectedProjectSettings? Project)
+		private bool TryGetSelectedProject([NotNullWhen(true)] out UserSelectedProjectSettings? project)
 		{
 			if(WorkspaceRadioBtn.Checked)
 			{
-				string? ClientPath;
-				if(TryGetClientPath(out ClientPath))
+				string? clientPath;
+				if(TryGetClientPath(out clientPath))
 				{
-					Project = new UserSelectedProjectSettings(ServerAndPortOverride, UserNameOverride, UserSelectedProjectType.Client, ClientPath, null);
+					project = new UserSelectedProjectSettings(_serverAndPortOverride, _userNameOverride, UserSelectedProjectType.Client, clientPath, null);
 					return true;
 				}
 			}
 			else
 			{
-				string? LocalPath;
-				if(TryGetLocalPath(out LocalPath))
+				string? localPath;
+				if(TryGetLocalPath(out localPath))
 				{
-					Project = new UserSelectedProjectSettings(ServerAndPortOverride, UserNameOverride, UserSelectedProjectType.Local, null, LocalPath);
+					project = new UserSelectedProjectSettings(_serverAndPortOverride, _userNameOverride, UserSelectedProjectType.Local, null, localPath);
 					return true;
 				}
 			}
 
-			Project = null;
+			project = null;
 			return false;
 		}
 
 		private void OkBtn_Click(object sender, EventArgs e)
 		{
-			UserSelectedProjectSettings? SelectedProject;
-			if(TryGetSelectedProject(out SelectedProject))
+			UserSelectedProjectSettings? selectedProject;
+			if(TryGetSelectedProject(out selectedProject))
 			{
-				ILogger<OpenProjectInfo> Logger = ServiceProvider.GetRequiredService<ILogger<OpenProjectInfo>>();
+				ILogger<OpenProjectInfo> logger = _serviceProvider.GetRequiredService<ILogger<OpenProjectInfo>>();
 
-				PerforceSettings NewPerforceSettings = Utility.OverridePerforceSettings(Perforce, SelectedProject.ServerAndPort, SelectedProject.UserName);
+				PerforceSettings newPerforceSettings = Utility.OverridePerforceSettings(Perforce, selectedProject.ServerAndPort, selectedProject.UserName);
 
-				ModalTask<OpenProjectInfo>? NewOpenProjectInfo = PerforceModalTask.Execute(this, "Opening project", "Opening project, please wait...", NewPerforceSettings, (x, y) => DetectSettingsAsync(x, SelectedProject, Settings, Logger, y), Logger);
-				if (NewOpenProjectInfo != null && NewOpenProjectInfo.Succeeded)
+				ModalTask<OpenProjectInfo>? newOpenProjectInfo = PerforceModalTask.Execute(this, "Opening project", "Opening project, please wait...", newPerforceSettings, (x, y) => DetectSettingsAsync(x, selectedProject, _settings, logger, y), logger);
+				if (newOpenProjectInfo != null && newOpenProjectInfo.Succeeded)
 				{
-					OpenProjectInfo = NewOpenProjectInfo.Result;
+					_openProjectInfo = newOpenProjectInfo.Result;
 					DialogResult = DialogResult.OK;
 					Close();
 				}
 			}
 		}
 
-		public static async Task<OpenProjectInfo> DetectSettingsAsync(IPerforceConnection Perforce, UserSelectedProjectSettings SelectedProject, UserSettings UserSettings, ILogger<OpenProjectInfo> Logger, CancellationToken CancellationToken)
+		public static async Task<OpenProjectInfo> DetectSettingsAsync(IPerforceConnection perforce, UserSelectedProjectSettings selectedProject, UserSettings userSettings, ILogger<OpenProjectInfo> logger, CancellationToken cancellationToken)
 		{
-			OpenProjectInfo Settings = await OpenProjectInfo.CreateAsync(Perforce, SelectedProject, UserSettings, Logger, CancellationToken);
+			OpenProjectInfo settings = await OpenProjectInfo.CreateAsync(perforce, selectedProject, userSettings, logger, cancellationToken);
 			if (DeploymentSettings.OnDetectProjectSettings != null)
 			{
-				string? Message;
-				if (!DeploymentSettings.OnDetectProjectSettings(Settings, Logger, out Message))
+				string? message;
+				if (!DeploymentSettings.OnDetectProjectSettings(settings, logger, out message))
 				{
-					throw new UserErrorException(Message);
+					throw new UserErrorException(message);
 				}
 			}
-			return Settings;
+			return settings;
 		}
 
 		private void ChangeLink_LinkClicked(object sender, LinkLabelLinkClickedEventArgs e)
 		{
-			if(ConnectWindow.ShowModal(this, DefaultSettings, ref ServerAndPortOverride, ref UserNameOverride, ServiceProvider))
+			if(ConnectWindow.ShowModal(this, _defaultSettings, ref _serverAndPortOverride, ref _userNameOverride, _serviceProvider))
 			{
 				UpdateServerLabel();
 			}
@@ -326,29 +326,29 @@ namespace UnrealGameSync
 		{
 			LocalFileRadioBtn.Checked = true;
 
-			OpenFileDialog Dialog = new OpenFileDialog();
-			Dialog.Filter = "Project files (*.uproject)|*.uproject|Project directory lists (*.uprojectdirs)|*.uprojectdirs|All supported files (*.uproject;*.uprojectdirs)|*.uproject;*.uprojectdirs|All files (*.*)|*.*" ;
-			Dialog.FilterIndex = Settings.FilterIndex;
+			OpenFileDialog dialog = new OpenFileDialog();
+			dialog.Filter = "Project files (*.uproject)|*.uproject|Project directory lists (*.uprojectdirs)|*.uprojectdirs|All supported files (*.uproject;*.uprojectdirs)|*.uproject;*.uprojectdirs|All files (*.*)|*.*" ;
+			dialog.FilterIndex = _settings.FilterIndex;
 			
 			if(!String.IsNullOrEmpty(LocalFileTextBox.Text))
 			{
 				try
 				{
-					Dialog.InitialDirectory = Path.GetDirectoryName(LocalFileTextBox.Text);
+					dialog.InitialDirectory = Path.GetDirectoryName(LocalFileTextBox.Text);
 				}
 				catch
 				{
 				}
 			}
 
-			if (Dialog.ShowDialog(this) == DialogResult.OK)
+			if (dialog.ShowDialog(this) == DialogResult.OK)
 			{
-				string FullName = Path.GetFullPath(Dialog.FileName);
+				string fullName = Path.GetFullPath(dialog.FileName);
 
-				Settings.FilterIndex = Dialog.FilterIndex;
-				Settings.Save(Logger);
+				_settings.FilterIndex = dialog.FilterIndex;
+				_settings.Save(_logger);
 
-				LocalFileTextBox.Text = FullName;
+				LocalFileTextBox.Text = fullName;
 				UpdateOkButton();
 			}
 		}
diff --git a/Engine/Source/Programs/UnrealGameSync/UnrealGameSync/Forms/PasswordWindow.cs b/Engine/Source/Programs/UnrealGameSync/UnrealGameSync/Forms/PasswordWindow.cs
--- a/Engine/Source/Programs/UnrealGameSync/UnrealGameSync/Forms/PasswordWindow.cs
+++ b/Engine/Source/Programs/UnrealGameSync/UnrealGameSync/Forms/PasswordWindow.cs
@@ -16,14 +16,14 @@ namespace UnrealGameSync
 	{
 		public string Password;
 
-		public PasswordWindow(string Prompt, string Password)
+		public PasswordWindow(string prompt, string password)
 		{
 			InitializeComponent();
 
-			this.Password = Password;
+			this.Password = password;
 
-			PromptLabel.Text = Prompt;
-			PasswordTextBox.Text = Password;
+			PromptLabel.Text = prompt;
+			PasswordTextBox.Text = password;
 		}
 
 		private void OkBtn_Click(object sender, EventArgs e)
diff --git a/Engine/Source/Programs/UnrealGameSync/UnrealGameSync/Forms/PerforceSyncSettingsWindow.cs b/Engine/Source/Programs/UnrealGameSync/UnrealGameSync/Forms/PerforceSyncSettingsWindow.cs
--- a/Engine/Source/Programs/UnrealGameSync/UnrealGameSync/Forms/PerforceSyncSettingsWindow.cs
+++ b/Engine/Source/Programs/UnrealGameSync/UnrealGameSync/Forms/PerforceSyncSettingsWindow.cs
@@ -8,37 +8,37 @@ namespace UnrealGameSync
 {
 	partial class PerforceSyncSettingsWindow : Form
 	{
-		UserSettings Settings;
-		ILogger Logger;
+		UserSettings _settings;
+		ILogger _logger;
 
-		public PerforceSyncSettingsWindow(UserSettings Settings, ILogger Logger)
+		public PerforceSyncSettingsWindow(UserSettings settings, ILogger logger)
 		{
-			this.Settings = Settings;
-			this.Logger = Logger;
+			this._settings = settings;
+			this._logger = logger;
 
 			InitializeComponent();
 		}
 
 		private void PerforceSettingsWindow_Load(object sender, EventArgs e)
 		{
-			PerforceSyncOptions SyncOptions = Settings.SyncOptions;
-			numericUpDownNumRetries.Value = (SyncOptions.NumRetries > 0) ? SyncOptions.NumRetries : PerforceSyncOptions.DefaultNumRetries;
-			numericUpDownTcpBufferSize.Value = (SyncOptions.TcpBufferSize > 0) ? SyncOptions.TcpBufferSize / 1024 : PerforceSyncOptions.DefaultTcpBufferSize / 1024;
-			numericUpDownFileBufferSize.Value = (SyncOptions.FileBufferSize > 0) ? SyncOptions.FileBufferSize / 1024 : PerforceSyncOptions.DefaultFileBufferSize / 1024;
-			numericUpDownMaxCommandsPerBatch.Value = (SyncOptions.MaxCommandsPerBatch > 0) ? SyncOptions.MaxCommandsPerBatch : PerforceSyncOptions.DefaultMaxCommandsPerBatch;
-			numericUpDownMaxSizePerBatch.Value = (SyncOptions.MaxSizePerBatch > 0) ? SyncOptions.MaxSizePerBatch / 1024 / 1024 : PerforceSyncOptions.DefaultMaxSizePerBatch / 1024 / 1024;
-			numericUpDownRetriesOnSyncError.Value = (SyncOptions.NumSyncErrorRetries > 0) ? SyncOptions.NumSyncErrorRetries : PerforceSyncOptions.DefaultNumSyncErrorRetries;
+			PerforceSyncOptions syncOptions = _settings.SyncOptions;
+			numericUpDownNumRetries.Value = (syncOptions.NumRetries > 0) ? syncOptions.NumRetries : PerforceSyncOptions.DefaultNumRetries;
+			numericUpDownTcpBufferSize.Value = (syncOptions.TcpBufferSize > 0) ? syncOptions.TcpBufferSize / 1024 : PerforceSyncOptions.DefaultTcpBufferSize / 1024;
+			numericUpDownFileBufferSize.Value = (syncOptions.FileBufferSize > 0) ? syncOptions.FileBufferSize / 1024 : PerforceSyncOptions.DefaultFileBufferSize / 1024;
+			numericUpDownMaxCommandsPerBatch.Value = (syncOptions.MaxCommandsPerBatch > 0) ? syncOptions.MaxCommandsPerBatch : PerforceSyncOptions.DefaultMaxCommandsPerBatch;
+			numericUpDownMaxSizePerBatch.Value = (syncOptions.MaxSizePerBatch > 0) ? syncOptions.MaxSizePerBatch / 1024 / 1024 : PerforceSyncOptions.DefaultMaxSizePerBatch / 1024 / 1024;
+			numericUpDownRetriesOnSyncError.Value = (syncOptions.NumSyncErrorRetries > 0) ? syncOptions.NumSyncErrorRetries : PerforceSyncOptions.DefaultNumSyncErrorRetries;
 		}
 
 		private void OkButton_Click(object sender, EventArgs e)
 		{
-			Settings.SyncOptions.NumRetries = (int)numericUpDownNumRetries.Value;
-			Settings.SyncOptions.TcpBufferSize = (int)numericUpDownTcpBufferSize.Value * 1024;
-			Settings.SyncOptions.FileBufferSize = (int)numericUpDownFileBufferSize.Value * 1024;
-			Settings.SyncOptions.MaxCommandsPerBatch = (int)numericUpDownMaxCommandsPerBatch.Value;
-			Settings.SyncOptions.MaxSizePerBatch = (int)numericUpDownMaxSizePerBatch.Value * 1024 * 1024;
-			Settings.SyncOptions.NumSyncErrorRetries = (int)numericUpDownRetriesOnSyncError.Value;
-			Settings.Save(Logger);
+			_settings.SyncOptions.NumRetries = (int)numericUpDownNumRetries.Value;
+			_settings.SyncOptions.TcpBufferSize = (int)numericUpDownTcpBufferSize.Value * 1024;
+			_settings.SyncOptions.FileBufferSize = (int)numericUpDownFileBufferSize.Value * 1024;
+			_settings.SyncOptions.MaxCommandsPerBatch = (int)numericUpDownMaxCommandsPerBatch.Value;
+			_settings.SyncOptions.MaxSizePerBatch = (int)numericUpDownMaxSizePerBatch.Value * 1024 * 1024;
+			_settings.SyncOptions.NumSyncErrorRetries = (int)numericUpDownRetriesOnSyncError.Value;
+			_settings.Save(_logger);
 
 			DialogResult = System.Windows.Forms.DialogResult.OK;
 			Close();
@@ -52,7 +52,7 @@ namespace UnrealGameSync
 
 		private void ResetButton_Click(object sender, EventArgs e)
 		{
-			PerforceSyncOptions SyncOptions = Settings.SyncOptions;
+			PerforceSyncOptions syncOptions = _settings.SyncOptions;
 			numericUpDownNumRetries.Value = PerforceSyncOptions.DefaultNumRetries;
 			numericUpDownTcpBufferSize.Value = PerforceSyncOptions.DefaultTcpBufferSize / 1024;
 			numericUpDownFileBufferSize.Value = PerforceSyncOptions.DefaultFileBufferSize / 1024;
diff --git a/Engine/Source/Programs/UnrealGameSync/UnrealGameSync/Forms/ProgramsRunningWindow.cs b/Engine/Source/Programs/UnrealGameSync/UnrealGameSync/Forms/ProgramsRunningWindow.cs
--- a/Engine/Source/Programs/UnrealGameSync/UnrealGameSync/Forms/ProgramsRunningWindow.cs
+++ b/Engine/Source/Programs/UnrealGameSync/UnrealGameSync/Forms/ProgramsRunningWindow.cs
@@ -16,45 +16,45 @@ namespace UnrealGameSync
 {
 	public partial class ProgramsRunningWindow : Form
 	{
-		object SyncObject = new object();
-		FileReference[] Programs;
-		Func<FileReference[]> EnumeratePrograms;
-		ManualResetEvent? TerminateEvent;
-		Thread? BackgroundThread;
+		object _syncObject = new object();
+		FileReference[] _programs;
+		Func<FileReference[]> _enumeratePrograms;
+		ManualResetEvent? _terminateEvent;
+		Thread? _backgroundThread;
 
-		public ProgramsRunningWindow(Func<FileReference[]> EnumeratePrograms, FileReference[] Programs)
+		public ProgramsRunningWindow(Func<FileReference[]> enumeratePrograms, FileReference[] programs)
 		{
 			InitializeComponent();
 
-			this.Programs = Programs.OrderBy(x => x).ToArray();
-			this.EnumeratePrograms = EnumeratePrograms;
-			this.ProgramListBox.Items.AddRange(Programs);
+			this._programs = programs.OrderBy(x => x).ToArray();
+			this._enumeratePrograms = enumeratePrograms;
+			this.ProgramListBox.Items.AddRange(programs);
 		}
 
 		private void ProgramsRunningWindow_Load(object sender, EventArgs e)
 		{
-			TerminateEvent = new ManualResetEvent(false);
+			_terminateEvent = new ManualResetEvent(false);
 
-			BackgroundThread = new Thread(() => ExecuteBackgroundWork());
-			BackgroundThread.IsBackground = true;
-			BackgroundThread.Start();
+			_backgroundThread = new Thread(() => ExecuteBackgroundWork());
+			_backgroundThread.IsBackground = true;
+			_backgroundThread.Start();
 		}
 
 		private void ExecuteBackgroundWork()
 		{
 			for(;;)
 			{
-				if(TerminateEvent!.WaitOne(TimeSpan.FromSeconds(2.0)))
+				if(_terminateEvent!.WaitOne(TimeSpan.FromSeconds(2.0)))
 				{
 					break;
 				}
 
-				FileReference[] NewPrograms = EnumeratePrograms().OrderBy(x => x).ToArray();
-				lock(SyncObject)
+				FileReference[] newPrograms = _enumeratePrograms().OrderBy(x => x).ToArray();
+				lock(_syncObject)
 				{
-					if(!Enumerable.SequenceEqual(Programs, NewPrograms))
+					if(!Enumerable.SequenceEqual(_programs, newPrograms))
 					{
-						Programs = NewPrograms;
+						_programs = newPrograms;
 						BeginInvoke(new MethodInvoker(() => UpdatePrograms()));
 					}
 				}
@@ -64,9 +64,9 @@ namespace UnrealGameSync
 		private void UpdatePrograms()
 		{
 			ProgramListBox.Items.Clear();
-			ProgramListBox.Items.AddRange(Programs);
+			ProgramListBox.Items.AddRange(_programs);
 
-			if(Programs.Length == 0)
+			if(_programs.Length == 0)
 			{
 				DialogResult = DialogResult.OK;
 				Close();
@@ -75,15 +75,15 @@ namespace UnrealGameSync
 
 		private void ProgramsRunningWindow_FormClosed(object sender, FormClosedEventArgs e)
 		{
-			if(BackgroundThread != null)
+			if(_backgroundThread != null)
 			{
-				TerminateEvent!.Set();
+				_terminateEvent!.Set();
 
-				BackgroundThread.Join();
-				BackgroundThread = null;
+				_backgroundThread.Join();
+				_backgroundThread = null;
 
-				TerminateEvent.Dispose();
-				TerminateEvent = null;
+				_terminateEvent.Dispose();
+				_terminateEvent = null;
 			}
 		}
 	}
diff --git a/Engine/Source/Programs/UnrealGameSync/UnrealGameSync/Forms/ScheduleWindow.cs b/Engine/Source/Programs/UnrealGameSync/UnrealGameSync/Forms/ScheduleWindow.cs
--- a/Engine/Source/Programs/UnrealGameSync/UnrealGameSync/Forms/ScheduleWindow.cs
+++ b/Engine/Source/Programs/UnrealGameSync/UnrealGameSync/Forms/ScheduleWindow.cs
@@ -17,36 +17,36 @@ namespace UnrealGameSync
 	partial class ScheduleWindow : Form
 	{
 
-		Dictionary<UserSelectedProjectSettings, List<LatestChangeType>> ProjectToLatestChangeTypes;
+		Dictionary<UserSelectedProjectSettings, List<LatestChangeType>> _projectToLatestChangeTypes;
 
 		public ScheduleWindow(
-			bool InEnabled,
-			TimeSpan InTime,
-			bool AnyOpenProject,
-			IEnumerable<UserSelectedProjectSettings> ScheduledProjects,
-			IEnumerable<UserSelectedProjectSettings> OpenProjects,
-			Dictionary<UserSelectedProjectSettings, List<LatestChangeType>> InProjectToLatestChangeTypes)
+			bool inEnabled,
+			TimeSpan inTime,
+			bool anyOpenProject,
+			IEnumerable<UserSelectedProjectSettings> scheduledProjects,
+			IEnumerable<UserSelectedProjectSettings> openProjects,
+			Dictionary<UserSelectedProjectSettings, List<LatestChangeType>> inProjectToLatestChangeTypes)
 		{
 			InitializeComponent();
 
-			EnableCheckBox.Checked = InEnabled;
+			EnableCheckBox.Checked = inEnabled;
 
-			ProjectToLatestChangeTypes = InProjectToLatestChangeTypes;
+			_projectToLatestChangeTypes = inProjectToLatestChangeTypes;
 
-			DateTime CurrentTime = DateTime.Now;
+			DateTime currentTime = DateTime.Now;
 			TimePicker.CustomFormat = CultureInfo.CurrentCulture.DateTimeFormat.ShortTimePattern;
-			TimePicker.Value = new DateTime(CurrentTime.Year, CurrentTime.Month, CurrentTime.Day, InTime.Hours, InTime.Minutes, InTime.Seconds);
+			TimePicker.Value = new DateTime(currentTime.Year, currentTime.Month, currentTime.Day, inTime.Hours, inTime.Minutes, inTime.Seconds);
 
-			ProjectListBox.Items.Add("Any open projects", AnyOpenProject);
+			ProjectListBox.Items.Add("Any open projects", anyOpenProject);
 
-			Dictionary<string, UserSelectedProjectSettings> LocalFileToProject = new Dictionary<string, UserSelectedProjectSettings>(StringComparer.InvariantCultureIgnoreCase);
-			AddProjects(ScheduledProjects, LocalFileToProject);
-			AddProjects(OpenProjects, LocalFileToProject);
+			Dictionary<string, UserSelectedProjectSettings> localFileToProject = new Dictionary<string, UserSelectedProjectSettings>(StringComparer.InvariantCultureIgnoreCase);
+			AddProjects(scheduledProjects, localFileToProject);
+			AddProjects(openProjects, localFileToProject);
 
-			foreach(UserSelectedProjectSettings Project in LocalFileToProject.Values.OrderBy(x => x.ToString()))
+			foreach(UserSelectedProjectSettings project in localFileToProject.Values.OrderBy(x => x.ToString()))
 			{
-				bool Enabled = ScheduledProjects.Any(x => x.LocalPath == Project.LocalPath);
-				ProjectListBox.Items.Add(Project, Enabled);
+				bool enabled = scheduledProjects.Any(x => x.LocalPath == project.LocalPath);
+				ProjectListBox.Items.Add(project, enabled);
 			}
 
 			ProjectListBox.ItemCheck += ProjectListBox_ItemCheck;
@@ -56,30 +56,30 @@ namespace UnrealGameSync
 			UpdateEnabledControls();
 		}
 
-		private void UpdateSyncTypeDropDownWithProject(UserSelectedProjectSettings? SelectedProject)
+		private void UpdateSyncTypeDropDownWithProject(UserSelectedProjectSettings? selectedProject)
 		{
 			SyncTypeDropDown.Items.Clear();
 
-			if (SelectedProject != null
-				&& ProjectToLatestChangeTypes.ContainsKey(SelectedProject))
+			if (selectedProject != null
+				&& _projectToLatestChangeTypes.ContainsKey(selectedProject))
 			{
-				foreach (LatestChangeType ChangeType in ProjectToLatestChangeTypes[SelectedProject])
+				foreach (LatestChangeType changeType in _projectToLatestChangeTypes[selectedProject])
 				{
-					System.Windows.Forms.ToolStripMenuItem MenuItem = new System.Windows.Forms.ToolStripMenuItem();
-					MenuItem.Name = ChangeType.Name;
-					MenuItem.Text = ChangeType.Description;
-					MenuItem.Size = new System.Drawing.Size(189, 22);
-					MenuItem.Click += (sender, e) => SyncTypeDropDown_Click(sender, e, ChangeType.Name);
+					System.Windows.Forms.ToolStripMenuItem menuItem = new System.Windows.Forms.ToolStripMenuItem();
+					menuItem.Name = changeType.Name;
+					menuItem.Text = changeType.Description;
+					menuItem.Size = new System.Drawing.Size(189, 22);
+					menuItem.Click += (sender, e) => SyncTypeDropDown_Click(sender, e, changeType.Name);
 
-					SyncTypeDropDown.Items.Add(MenuItem);
+					SyncTypeDropDown.Items.Add(menuItem);
 				}
 			}
 		}
 
 		private void ProjectListBox_ItemCheck(object sender, ItemCheckEventArgs e)
 		{
-			bool bIsAnyOpenProjectIndex = e.Index == 0;
-			if (e.NewValue == CheckState.Checked && !bIsAnyOpenProjectIndex)
+			bool isAnyOpenProjectIndex = e.Index == 0;
+			if (e.NewValue == CheckState.Checked && !isAnyOpenProjectIndex)
 			{
 				UpdateSyncTypeDropDownWithProject(ProjectListBox.Items[e.Index] as UserSelectedProjectSettings);
 
@@ -98,47 +98,47 @@ namespace UnrealGameSync
 			}
 		}
 
-		private void SyncTypeDropDown_Click(object? sender, EventArgs e, string SyncTypeID)
+		private void SyncTypeDropDown_Click(object? sender, EventArgs e, string syncTypeId)
 		{
-			UserSelectedProjectSettings? ProjectSetting = ProjectListBox.Items[ProjectListBox.SelectedIndex] as UserSelectedProjectSettings;
-			if (ProjectSetting != null)
+			UserSelectedProjectSettings? projectSetting = ProjectListBox.Items[ProjectListBox.SelectedIndex] as UserSelectedProjectSettings;
+			if (projectSetting != null)
 			{
-				ProjectSetting.ScheduledSyncTypeID = SyncTypeID;
+				projectSetting.ScheduledSyncTypeId = syncTypeId;
 			}
 			SyncTypeDropDown.Close();
 		}
 
-		private void AddProjects(IEnumerable<UserSelectedProjectSettings> Projects, Dictionary<string, UserSelectedProjectSettings> LocalFileToProject)
+		private void AddProjects(IEnumerable<UserSelectedProjectSettings> projects, Dictionary<string, UserSelectedProjectSettings> localFileToProject)
 		{
-			foreach(UserSelectedProjectSettings Project in Projects)
+			foreach(UserSelectedProjectSettings project in projects)
 			{
-				if(Project.LocalPath != null)
+				if(project.LocalPath != null)
 				{
-					LocalFileToProject[Project.LocalPath] = Project;
+					localFileToProject[project.LocalPath] = project;
 				}
 			}
 		}
 
-		public void CopySettings(out bool OutEnabled, out TimeSpan OutTime, out bool OutAnyOpenProject, out List<UserSelectedProjectSettings> OutScheduledProjects)
+		public void CopySettings(out bool outEnabled, out TimeSpan outTime, out bool outAnyOpenProject, out List<UserSelectedProjectSettings> outScheduledProjects)
 		{
-			OutEnabled = EnableCheckBox.Checked;
-			OutTime = TimePicker.Value.TimeOfDay;
+			outEnabled = EnableCheckBox.Checked;
+			outTime = TimePicker.Value.TimeOfDay;
 
-			OutAnyOpenProject = false;
+			outAnyOpenProject = false;
 
-			List<UserSelectedProjectSettings> ScheduledProjects = new List<UserSelectedProjectSettings>();
-			foreach(int Index in ProjectListBox.CheckedIndices.OfType<int>())
+			List<UserSelectedProjectSettings> scheduledProjects = new List<UserSelectedProjectSettings>();
+			foreach(int index in ProjectListBox.CheckedIndices.OfType<int>())
 			{
-				if(Index == 0)
+				if(index == 0)
 				{
-					OutAnyOpenProject = true;
+					outAnyOpenProject = true;
 				}
 				else
 				{
-					ScheduledProjects.Add((UserSelectedProjectSettings)ProjectListBox.Items[Index]);
+					scheduledProjects.Add((UserSelectedProjectSettings)ProjectListBox.Items[index]);
 				}
 			}
-			OutScheduledProjects = ScheduledProjects;
+			outScheduledProjects = scheduledProjects;
 		}
 
 		private void EnableCheckBox_CheckedChanged(object sender, EventArgs e)
diff --git a/Engine/Source/Programs/UnrealGameSync/UnrealGameSync/Forms/SdkInfoWindow.cs b/Engine/Source/Programs/UnrealGameSync/UnrealGameSync/Forms/SdkInfoWindow.cs
--- a/Engine/Source/Programs/UnrealGameSync/UnrealGameSync/Forms/SdkInfoWindow.cs
+++ b/Engine/Source/Programs/UnrealGameSync/UnrealGameSync/Forms/SdkInfoWindow.cs
@@ -26,10 +26,10 @@ namespace UnrealGameSync
 			public string? Install;
 			public string? Browse;
 
-			public SdkItem(string Category, string Description)
+			public SdkItem(string category, string description)
 			{
-				this.Category = Category;
-				this.Description = Description;
+				this.Category = category;
+				this.Description = description;
 			}
 		}
 
@@ -40,95 +40,95 @@ namespace UnrealGameSync
 			public Rectangle Rectangle;
 			public Action? OnClick;
 
-			public BadgeInfo(string UniqueId, string Label)
+			public BadgeInfo(string uniqueId, string label)
 			{
-				this.UniqueId = UniqueId;
-				this.Label = Label;
+				this.UniqueId = uniqueId;
+				this.Label = label;
 			}
 		}
 
-		Font BadgeFont;
-		string? HoverBadgeUniqueId;
+		Font _badgeFont;
+		string? _hoverBadgeUniqueId;
 
-		public SdkInfoWindow(string[] SdkInfoEntries, Dictionary<string, string> Variables, Font BadgeFont)
+		public SdkInfoWindow(string[] sdkInfoEntries, Dictionary<string, string> variables, Font badgeFont)
 		{
 			InitializeComponent();
 
-			this.BadgeFont = BadgeFont;
+			this._badgeFont = badgeFont;
 
-			Dictionary<string, ConfigObject> UniqueIdToObject = new Dictionary<string, ConfigObject>(StringComparer.InvariantCultureIgnoreCase);
-			foreach(string SdkInfoEntry in SdkInfoEntries)
+			Dictionary<string, ConfigObject> uniqueIdToObject = new Dictionary<string, ConfigObject>(StringComparer.InvariantCultureIgnoreCase);
+			foreach(string sdkInfoEntry in sdkInfoEntries)
 			{
-				ConfigObject Object = new ConfigObject(SdkInfoEntry);
+				ConfigObject obj = new ConfigObject(sdkInfoEntry);
 
-				string UniqueId = Object.GetValue("UniqueId", Guid.NewGuid().ToString());
+				string uniqueId = obj.GetValue("UniqueId", Guid.NewGuid().ToString());
 
-				ConfigObject? ExistingObject;
-				if(UniqueIdToObject.TryGetValue(UniqueId, out ExistingObject))
+				ConfigObject? existingObject;
+				if(uniqueIdToObject.TryGetValue(uniqueId, out existingObject))
 				{
-					ExistingObject.AddOverrides(Object, null);
+					existingObject.AddOverrides(obj, null);
 				}
 				else
 				{
-					UniqueIdToObject.Add(UniqueId, Object);
+					uniqueIdToObject.Add(uniqueId, obj);
 				}
 			}
 
-			List<SdkItem> Items = new List<SdkItem>();
-			foreach(ConfigObject Object in UniqueIdToObject.Values)
+			List<SdkItem> items = new List<SdkItem>();
+			foreach(ConfigObject obj in uniqueIdToObject.Values)
 			{
-				string Category = Object.GetValue("Category", "Other");
-				string Description = Object.GetValue("Description", "");
-				SdkItem Item = new SdkItem(Category, Description);
+				string category = obj.GetValue("Category", "Other");
+				string description = obj.GetValue("Description", "");
+				SdkItem item = new SdkItem(category, description);
 
-				Item.Install = Utility.ExpandVariables(Object.GetValue("Install", ""), Variables);
-				if(Item.Install.Contains("$("))
+				item.Install = Utility.ExpandVariables(obj.GetValue("Install", ""), variables);
+				if(item.Install.Contains("$("))
 				{
-					Item.Install = null;
+					item.Install = null;
 				}
 
-				Item.Browse = Utility.ExpandVariables(Object.GetValue("Browse", ""), Variables);
-				if(Item.Browse.Contains("$("))
+				item.Browse = Utility.ExpandVariables(obj.GetValue("Browse", ""), variables);
+				if(item.Browse.Contains("$("))
 				{
-					Item.Browse = null;
+					item.Browse = null;
 				}
 
-				if(!String.IsNullOrEmpty(Item.Install) && String.IsNullOrEmpty(Item.Browse))
+				if(!String.IsNullOrEmpty(item.Install) && String.IsNullOrEmpty(item.Browse))
 				{
 					try
 					{
-						Item.Browse = Path.GetDirectoryName(Item.Install);
+						item.Browse = Path.GetDirectoryName(item.Install);
 					}
 					catch
 					{
-						Item.Browse = null;
+						item.Browse = null;
 					}
 				}
 
-				Items.Add(Item);
+				items.Add(item);
 			}
 
-			foreach(IGrouping<string, SdkItem> ItemGroup in Items.GroupBy(x => x.Category).OrderBy(x => x.Key))
+			foreach(IGrouping<string, SdkItem> itemGroup in items.GroupBy(x => x.Category).OrderBy(x => x.Key))
 			{
-				ListViewGroup Group = new ListViewGroup(ItemGroup.Key);
-				SdkListView.Groups.Add(Group);
+				ListViewGroup group = new ListViewGroup(itemGroup.Key);
+				SdkListView.Groups.Add(group);
 
-				foreach(SdkItem Item in ItemGroup)
+				foreach(SdkItem item in itemGroup)
 				{
-					ListViewItem NewItem = new ListViewItem(Group);
-					NewItem.SubItems.Add(Item.Description);
-					NewItem.SubItems.Add(new ListViewItem.ListViewSubItem(){ Tag = Item });
-					SdkListView.Items.Add(NewItem);
+					ListViewItem newItem = new ListViewItem(group);
+					newItem.SubItems.Add(item.Description);
+					newItem.SubItems.Add(new ListViewItem.ListViewSubItem(){ Tag = item });
+					SdkListView.Items.Add(newItem);
 				}
 			}
 
-			System.Reflection.PropertyInfo DoubleBufferedProperty = typeof(Control).GetProperty("DoubleBuffered", System.Reflection.BindingFlags.NonPublic | System.Reflection.BindingFlags.Instance)!;
-			DoubleBufferedProperty.SetValue(SdkListView, true, null);
+			System.Reflection.PropertyInfo doubleBufferedProperty = typeof(Control).GetProperty("DoubleBuffered", System.Reflection.BindingFlags.NonPublic | System.Reflection.BindingFlags.Instance)!;
+			doubleBufferedProperty.SetValue(SdkListView, true, null);
 
 			if(SdkListView.Items.Count > 0)
 			{
-				int ItemsHeight = SdkListView.Items[SdkListView.Items.Count - 1].Bounds.Bottom + 20;
-				Height = SdkListView.Top + ItemsHeight + (Height - SdkListView.Bottom);
+				int itemsHeight = SdkListView.Items[SdkListView.Items.Count - 1].Bounds.Bottom + 20;
+				Height = SdkListView.Top + itemsHeight + (Height - SdkListView.Bottom);
 			}
 		}
 
@@ -153,13 +153,13 @@ namespace UnrealGameSync
 			{
 				e.Graphics.SmoothingMode = SmoothingMode.AntiAlias;
 
-				List<BadgeInfo> Badges = GetBadges(e.Item, e.SubItem);
-				for(int Idx = 0; Idx < Badges.Count; Idx++)
+				List<BadgeInfo> badges = GetBadges(e.Item, e.SubItem);
+				for(int idx = 0; idx < badges.Count; idx++)
 				{
-					Color BadgeColor = (HoverBadgeUniqueId == Badges[Idx].UniqueId)? Color.FromArgb(140, 180, 230) : Color.FromArgb(112, 146, 190);
-					if(Badges[Idx].OnClick != null)
+					Color badgeColor = (_hoverBadgeUniqueId == badges[idx].UniqueId)? Color.FromArgb(140, 180, 230) : Color.FromArgb(112, 146, 190);
+					if(badges[idx].OnClick != null)
 					{
-						DrawBadge(e.Graphics, Badges[Idx].Label, Badges[Idx].Rectangle, (Idx > 0), (Idx < Badges.Count - 1), BadgeColor);
+						DrawBadge(e.Graphics, badges[idx].Label, badges[idx].Rectangle, (idx > 0), (idx < badges.Count - 1), badgeColor);
 					}
 				}
 			}
@@ -167,135 +167,135 @@ namespace UnrealGameSync
 
 		private void SdkListView_MouseMove(object sender, MouseEventArgs e)
 		{
-			string? NewHoverUniqueId = null;
+			string? newHoverUniqueId = null;
 
-			ListViewHitTestInfo HitTest = SdkListView.HitTest(e.Location);
-			if(HitTest.Item != null && HitTest.SubItem == HitTest.Item.SubItems[2])
+			ListViewHitTestInfo hitTest = SdkListView.HitTest(e.Location);
+			if(hitTest.Item != null && hitTest.SubItem == hitTest.Item.SubItems[2])
 			{
-				List<BadgeInfo> Badges = GetBadges(HitTest.Item, HitTest.SubItem);
-				foreach(BadgeInfo Badge in Badges)
+				List<BadgeInfo> badges = GetBadges(hitTest.Item, hitTest.SubItem);
+				foreach(BadgeInfo badge in badges)
 				{
-					if(Badge.Rectangle.Contains(e.Location))
+					if(badge.Rectangle.Contains(e.Location))
 					{
-						NewHoverUniqueId = Badge.UniqueId;
+						newHoverUniqueId = badge.UniqueId;
 					}
 				}
 			}
 
-			if(NewHoverUniqueId != HoverBadgeUniqueId)
+			if(newHoverUniqueId != _hoverBadgeUniqueId)
 			{
-				HoverBadgeUniqueId = NewHoverUniqueId;
+				_hoverBadgeUniqueId = newHoverUniqueId;
 				SdkListView.Invalidate();
 			}
 		}
 
 		private void SdkListView_MouseLeave(object sender, EventArgs e)
 		{
-			HoverBadgeUniqueId = null;
+			_hoverBadgeUniqueId = null;
 		}
 
 		private void SdkListView_MouseDown(object sender, MouseEventArgs e)
 		{
-			ListViewHitTestInfo HitTest = SdkListView.HitTest(e.Location);
-			if(HitTest.Item != null && HitTest.SubItem == HitTest.Item.SubItems[2])
+			ListViewHitTestInfo hitTest = SdkListView.HitTest(e.Location);
+			if(hitTest.Item != null && hitTest.SubItem == hitTest.Item.SubItems[2])
 			{
-				List<BadgeInfo> Badges = GetBadges(HitTest.Item, HitTest.SubItem);
-				foreach(BadgeInfo Badge in Badges)
+				List<BadgeInfo> badges = GetBadges(hitTest.Item, hitTest.SubItem);
+				foreach(BadgeInfo badge in badges)
 				{
-					if(Badge.Rectangle.Contains(e.Location) && Badge.OnClick != null)
+					if(badge.Rectangle.Contains(e.Location) && badge.OnClick != null)
 					{
-						Badge.OnClick();
+						badge.OnClick();
 					}
 				}
 			}
 		}
 
-		private List<BadgeInfo> GetBadges(ListViewItem Item, ListViewItem.ListViewSubItem SubItem)
+		private List<BadgeInfo> GetBadges(ListViewItem item, ListViewItem.ListViewSubItem subItem)
 		{
-			string UniqueIdPrefix = String.Format("{0}_", Item.Index);
+			string uniqueIdPrefix = String.Format("{0}_", item.Index);
 
-			List<BadgeInfo> Badges = new List<BadgeInfo>();
+			List<BadgeInfo> badges = new List<BadgeInfo>();
 
-			SdkItem Sdk = (SdkItem)SubItem.Tag;
+			SdkItem sdk = (SdkItem)subItem.Tag;
 
-			Action? InstallAction = null;
-			if(!String.IsNullOrEmpty(Sdk.Install))
+			Action? installAction = null;
+			if(!String.IsNullOrEmpty(sdk.Install))
 			{
-				InstallAction = () => { Install(Sdk.Install); };
+				installAction = () => { Install(sdk.Install); };
 			}
-			Badges.Add(new BadgeInfo(UniqueIdPrefix + "_Install", "Install") { OnClick = InstallAction });
+			badges.Add(new BadgeInfo(uniqueIdPrefix + "_Install", "Install") { OnClick = installAction });
 
-			Action? BrowseAction = null;
-			if(!String.IsNullOrEmpty(Sdk.Browse))
+			Action? browseAction = null;
+			if(!String.IsNullOrEmpty(sdk.Browse))
 			{
-				BrowseAction = () => { Browse(Sdk.Browse); };
+				browseAction = () => { Browse(sdk.Browse); };
 			}
-			Badges.Add(new BadgeInfo(UniqueIdPrefix + "_Browse", "Browse"){ OnClick = BrowseAction });
+			badges.Add(new BadgeInfo(uniqueIdPrefix + "_Browse", "Browse"){ OnClick = browseAction });
 
-			int Right = SubItem.Bounds.Right - 10;
-			for(int Idx = Badges.Count - 1; Idx >= 0; Idx--)
+			int right = subItem.Bounds.Right - 10;
+			for(int idx = badges.Count - 1; idx >= 0; idx--)
 			{
-				Size BadgeSize = GetBadgeSize(Badges[Idx].Label);
-				Right -= BadgeSize.Width;
-				Badges[Idx].Rectangle = new Rectangle(Right, SubItem.Bounds.Y + (SubItem.Bounds.Height - BadgeSize.Height) / 2, BadgeSize.Width, BadgeSize.Height);
+				Size badgeSize = GetBadgeSize(badges[idx].Label);
+				right -= badgeSize.Width;
+				badges[idx].Rectangle = new Rectangle(right, subItem.Bounds.Y + (subItem.Bounds.Height - badgeSize.Height) / 2, badgeSize.Width, badgeSize.Height);
 			}
 
-			return Badges;
+			return badges;
 		}
 
-		private void Browse(string DirectoryName)
+		private void Browse(string directoryName)
 		{
 			try
 			{
-				Process.Start("explorer.exe", String.Format("\"{0}\"", DirectoryName));
+				Process.Start("explorer.exe", String.Format("\"{0}\"", directoryName));
 			}
-			catch(Exception Ex)
+			catch(Exception ex)
 			{
-				MessageBox.Show(String.Format("Unable to open explorer to {0}: {1}", DirectoryName, Ex.Message));
+				MessageBox.Show(String.Format("Unable to open explorer to {0}: {1}", directoryName, ex.Message));
 			}
 		}
 
-		private void Install(string FileName)
+		private void Install(string fileName)
 		{
 			try
 			{
-				ProcessStartInfo StartInfo = new ProcessStartInfo();
-				StartInfo.FileName = FileName;
-				StartInfo.UseShellExecute = true;
-				Process.Start(StartInfo);
+				ProcessStartInfo startInfo = new ProcessStartInfo();
+				startInfo.FileName = fileName;
+				startInfo.UseShellExecute = true;
+				Process.Start(startInfo);
 			}
-			catch(Exception Ex)
+			catch(Exception ex)
 			{
-				MessageBox.Show(String.Format("Unable to run {0}: {1}", FileName, Ex.Message));
+				MessageBox.Show(String.Format("Unable to run {0}: {1}", fileName, ex.Message));
 			}
 		}
 
-		private Size GetBadgeSize(string BadgeText)
+		private Size GetBadgeSize(string badgeText)
 		{
-			Size LabelSize = TextRenderer.MeasureText(BadgeText, BadgeFont);
-			int BadgeHeight = BadgeFont.Height + 1;
+			Size labelSize = TextRenderer.MeasureText(badgeText, _badgeFont);
+			int badgeHeight = _badgeFont.Height + 1;
 
-			return new Size(LabelSize.Width + BadgeHeight - 4, BadgeHeight);
+			return new Size(labelSize.Width + badgeHeight - 4, badgeHeight);
 		}
 
-		private void DrawBadge(Graphics Graphics, string BadgeText, Rectangle BadgeRect, bool bMergeLeft, bool bMergeRight, Color BadgeColor)
+		private void DrawBadge(Graphics graphics, string badgeText, Rectangle badgeRect, bool mergeLeft, bool mergeRight, Color badgeColor)
 		{
-			using (GraphicsPath Path = new GraphicsPath())
+			using (GraphicsPath path = new GraphicsPath())
 			{
-				Path.StartFigure();
-				Path.AddLine(BadgeRect.Left + (bMergeLeft? 1 : 0), BadgeRect.Top, BadgeRect.Left - (bMergeLeft? 1 : 0), BadgeRect.Bottom);
-				Path.AddLine(BadgeRect.Left - (bMergeLeft? 1 : 0), BadgeRect.Bottom, BadgeRect.Right - 1 - (bMergeRight? 1 : 0), BadgeRect.Bottom);
-				Path.AddLine(BadgeRect.Right - 1 - (bMergeRight? 1 : 0), BadgeRect.Bottom, BadgeRect.Right - 1 + (bMergeRight? 1 : 0), BadgeRect.Top);
-				Path.AddLine(BadgeRect.Right - 1 + (bMergeRight? 1 : 0), BadgeRect.Top, BadgeRect.Left + (bMergeLeft? 1 : 0), BadgeRect.Top);
-				Path.CloseFigure();
-
-				using(SolidBrush Brush = new SolidBrush(BadgeColor))
+				path.StartFigure();
+				path.AddLine(badgeRect.Left + (mergeLeft? 1 : 0), badgeRect.Top, badgeRect.Left - (mergeLeft? 1 : 0), badgeRect.Bottom);
+				path.AddLine(badgeRect.Left - (mergeLeft? 1 : 0), badgeRect.Bottom, badgeRect.Right - 1 - (mergeRight? 1 : 0), badgeRect.Bottom);
+				path.AddLine(badgeRect.Right - 1 - (mergeRight? 1 : 0), badgeRect.Bottom, badgeRect.Right - 1 + (mergeRight? 1 : 0), badgeRect.Top);
+				path.AddLine(badgeRect.Right - 1 + (mergeRight? 1 : 0), badgeRect.Top, badgeRect.Left + (mergeLeft? 1 : 0), badgeRect.Top);
+				path.CloseFigure();
+
+				using(SolidBrush brush = new SolidBrush(badgeColor))
 				{
-					Graphics.FillPath(Brush, Path);
+					graphics.FillPath(brush, path);
 				}
 			}
 
-			TextRenderer.DrawText(Graphics, BadgeText, BadgeFont, BadgeRect, Color.White, TextFormatFlags.HorizontalCenter | TextFormatFlags.VerticalCenter | TextFormatFlags.SingleLine | TextFormatFlags.NoPrefix | TextFormatFlags.PreserveGraphicsClipping);
+			TextRenderer.DrawText(graphics, badgeText, _badgeFont, badgeRect, Color.White, TextFormatFlags.HorizontalCenter | TextFormatFlags.VerticalCenter | TextFormatFlags.SingleLine | TextFormatFlags.NoPrefix | TextFormatFlags.PreserveGraphicsClipping);
 		}
 	}
 }
