UGS: Reformat to standard C# style conventions.

diff --git a/Engine/Source/Programs/UnrealGameSync/UnrealGameSync/Forms/SelectProjectFromWorkspaceWindow.cs b/Engine/Source/Programs/UnrealGameSync/UnrealGameSync/Forms/SelectProjectFromWorkspaceWindow.cs
--- a/Engine/Source/Programs/UnrealGameSync/UnrealGameSync/Forms/SelectProjectFromWorkspaceWindow.cs
+++ b/Engine/Source/Programs/UnrealGameSync/UnrealGameSync/Forms/SelectProjectFromWorkspaceWindow.cs
@@ -31,51 +31,51 @@ namespace UnrealGameSync
 				"....uprojectdirs",
 			};
 
-			public static async Task<List<string>> RunAsync(IPerforceConnection Perforce, string ClientName, CancellationToken CancellationToken)
+			public static async Task<List<string>> RunAsync(IPerforceConnection perforce, string clientName, CancellationToken cancellationToken)
 			{
-				ClientRecord ClientSpec = await Perforce.GetClientAsync(Perforce.Settings.ClientName, CancellationToken);
+				ClientRecord clientSpec = await perforce.GetClientAsync(perforce.Settings.ClientName, cancellationToken);
 
-				string ClientRoot = ClientSpec.Root.TrimEnd(Path.DirectorySeparatorChar);
-				if (String.IsNullOrEmpty(ClientRoot))
+				string clientRoot = clientSpec.Root.TrimEnd(Path.DirectorySeparatorChar);
+				if (String.IsNullOrEmpty(clientRoot))
 				{
-					throw new UserErrorException($"Client '{ClientName}' does not have a valid root directory.");
+					throw new UserErrorException($"Client '{clientName}' does not have a valid root directory.");
 				}
 
-				List<FStatRecord> FileRecords = new List<FStatRecord>();
-				foreach(string Pattern in Patterns)
+				List<FStatRecord> fileRecords = new List<FStatRecord>();
+				foreach(string pattern in Patterns)
 				{
-					string Filter = String.Format("//{0}/{1}", ClientName, Pattern);
+					string filter = String.Format("//{0}/{1}", clientName, pattern);
 
-					List<FStatRecord> WildcardFileRecords = await Perforce.FStatAsync(Filter, CancellationToken).ToListAsync(CancellationToken);
-					WildcardFileRecords.RemoveAll(x => x.HeadAction == FileAction.Delete || x.HeadAction == FileAction.MoveDelete);
+					List<FStatRecord> wildcardFileRecords = await perforce.FStatAsync(filter, cancellationToken).ToListAsync(cancellationToken);
+					wildcardFileRecords.RemoveAll(x => x.HeadAction == FileAction.Delete || x.HeadAction == FileAction.MoveDelete);
 
-					FileRecords.AddRange(WildcardFileRecords);
+					fileRecords.AddRange(wildcardFileRecords);
 				}
 
-				string ClientPrefix = ClientRoot;
-				if (!ClientPrefix.EndsWith(Path.DirectorySeparatorChar))
+				string clientPrefix = clientRoot;
+				if (!clientPrefix.EndsWith(Path.DirectorySeparatorChar))
 				{
-					ClientPrefix += Path.DirectorySeparatorChar;
+					clientPrefix += Path.DirectorySeparatorChar;
 				}
 
-				List<string> Paths = new List<string>();
-				foreach(FStatRecord FileRecord in FileRecords)
+				List<string> paths = new List<string>();
+				foreach(FStatRecord fileRecord in fileRecords)
 				{
-					if(FileRecord.ClientFile != null && FileRecord.ClientFile.StartsWith(ClientPrefix, StringComparison.OrdinalIgnoreCase))
+					if(fileRecord.ClientFile != null && fileRecord.ClientFile.StartsWith(clientPrefix, StringComparison.OrdinalIgnoreCase))
 					{
-						Paths.Add(FileRecord.ClientFile.Substring(ClientRoot.Length).Replace(Path.DirectorySeparatorChar, '/'));
+						paths.Add(fileRecord.ClientFile.Substring(clientRoot.Length).Replace(Path.DirectorySeparatorChar, '/'));
 					}
 				}
 
-				return Paths;
+				return paths;
 			}
 		}
 
 		[DllImport("Shell32.dll", EntryPoint = "ExtractIconExW", CharSet = CharSet.Unicode, ExactSpelling = true, CallingConvention = CallingConvention.StdCall)]
 		private static extern int ExtractIconEx(string sFile, int iIndex, IntPtr piLargeVersion, out IntPtr piSmallVersion, int amountIcons);
 
-		List<string> ProjectFiles;
-		string SelectedProjectFile;
+		List<string> _projectFiles;
+		string _selectedProjectFile;
 
 		class ProjectNode
 		{
@@ -83,49 +83,49 @@ namespace UnrealGameSync
 			public string Folder;
 			public string Name;
 
-			public ProjectNode(string FullName)
+			public ProjectNode(string fullName)
 			{
-				this.FullName = FullName;
+				this.FullName = fullName;
 
-				int SlashIdx = FullName.LastIndexOf('/');
-				Folder = FullName.Substring(0, SlashIdx);
-				Name = FullName.Substring(SlashIdx + 1);
+				int slashIdx = fullName.LastIndexOf('/');
+				Folder = fullName.Substring(0, slashIdx);
+				Name = fullName.Substring(slashIdx + 1);
 			}
 		}
 		
-		public SelectProjectFromWorkspaceWindow(string WorkspaceName, List<string> ProjectFiles, string SelectedProjectFile)
+		public SelectProjectFromWorkspaceWindow(string workspaceName, List<string> projectFiles, string selectedProjectFile)
 		{
 			InitializeComponent();
 			
-			this.ProjectFiles = ProjectFiles;
-			this.SelectedProjectFile = SelectedProjectFile;
+			this._projectFiles = projectFiles;
+			this._selectedProjectFile = selectedProjectFile;
 
 			// Make the image strip containing icons for nodes in the tree
-			IntPtr FolderIconPtr;
-			ExtractIconEx("imageres.dll", 3, IntPtr.Zero, out FolderIconPtr, 1);
+			IntPtr folderIconPtr;
+			ExtractIconEx("imageres.dll", 3, IntPtr.Zero, out folderIconPtr, 1);
 
-			Icon[] Icons = new Icon[]{ Icon.FromHandle(FolderIconPtr), Properties.Resources.Icon };
+			Icon[] icons = new Icon[]{ Icon.FromHandle(folderIconPtr), Properties.Resources.Icon };
 
-			Bitmap TypeImageListBitmap = new Bitmap(Icons.Length * 16, 16, System.Drawing.Imaging.PixelFormat.Format32bppArgb);
-			using(Graphics Graphics = Graphics.FromImage(TypeImageListBitmap))
+			Bitmap typeImageListBitmap = new Bitmap(icons.Length * 16, 16, System.Drawing.Imaging.PixelFormat.Format32bppArgb);
+			using(Graphics graphics = Graphics.FromImage(typeImageListBitmap))
 			{
-				for(int IconIdx = 0; IconIdx < Icons.Length; IconIdx++)
+				for(int iconIdx = 0; iconIdx < icons.Length; iconIdx++)
 				{
-					Graphics.DrawIcon(Icons[IconIdx], new Rectangle(IconIdx * 16, 0, 16, 16));
+					graphics.DrawIcon(icons[iconIdx], new Rectangle(iconIdx * 16, 0, 16, 16));
 				}
 			}
 
-			ImageList TypeImageList = new ImageList();
-			TypeImageList.ImageSize = new Size(16, 16);
-			TypeImageList.ColorDepth = ColorDepth.Depth32Bit;
-			TypeImageList.Images.AddStrip(TypeImageListBitmap);
-			ProjectTreeView.ImageList = TypeImageList;
+			ImageList typeImageList = new ImageList();
+			typeImageList.ImageSize = new Size(16, 16);
+			typeImageList.ColorDepth = ColorDepth.Depth32Bit;
+			typeImageList.Images.AddStrip(typeImageListBitmap);
+			ProjectTreeView.ImageList = typeImageList;
 
 			// Create the root node
-			TreeNode RootNode = new TreeNode();
-			RootNode.Text = WorkspaceName;
-			RootNode.Expand();
-			ProjectTreeView.Nodes.Add(RootNode);
+			TreeNode rootNode = new TreeNode();
+			rootNode.Text = workspaceName;
+			rootNode.Expand();
+			ProjectTreeView.Nodes.Add(rootNode);
 
 			// Populate the tree
 			Populate();
@@ -134,91 +134,91 @@ namespace UnrealGameSync
 		private void Populate()
 		{
 			// Clear out the existing nodes
-			TreeNode RootNode = ProjectTreeView.Nodes[0];
-			RootNode.Nodes.Clear();
+			TreeNode rootNode = ProjectTreeView.Nodes[0];
+			rootNode.Nodes.Clear();
 
 			// Filter the project files
-			List<string> FilteredProjectFiles = new List<string>(ProjectFiles); 
+			List<string> filteredProjectFiles = new List<string>(_projectFiles); 
 			if(!ShowProjectDirsFiles.Checked)
 			{
-				FilteredProjectFiles.RemoveAll(x => x.EndsWith(".uprojectdirs", StringComparison.OrdinalIgnoreCase));
+				filteredProjectFiles.RemoveAll(x => x.EndsWith(".uprojectdirs", StringComparison.OrdinalIgnoreCase));
 			}
 
 			// Sort by paths, then files
-			List<ProjectNode> ProjectNodes = FilteredProjectFiles.Select(x => new ProjectNode(x)).OrderBy(x => x.Folder).ThenBy(x => x.Name).ToList();
+			List<ProjectNode> projectNodes = filteredProjectFiles.Select(x => new ProjectNode(x)).OrderBy(x => x.Folder).ThenBy(x => x.Name).ToList();
 
 			// Add the folders for each project
-			TreeNode[] ProjectParentNodes = new TreeNode[ProjectNodes.Count];
-			for(int Idx = 0; Idx < ProjectNodes.Count; Idx++)
+			TreeNode[] projectParentNodes = new TreeNode[projectNodes.Count];
+			for(int idx = 0; idx < projectNodes.Count; idx++)
 			{
-				TreeNode ParentNode = RootNode;
-				if(ProjectNodes[Idx].Folder.Length > 0)
+				TreeNode parentNode = rootNode;
+				if(projectNodes[idx].Folder.Length > 0)
 				{
-					string[] Fragments = ProjectNodes[Idx].Folder.Split(new char[]{ '/' }, StringSplitOptions.RemoveEmptyEntries);
-					foreach(string Fragment in Fragments)
+					string[] fragments = projectNodes[idx].Folder.Split(new char[]{ '/' }, StringSplitOptions.RemoveEmptyEntries);
+					foreach(string fragment in fragments)
 					{
-						ParentNode = FindOrAddChildNode(ParentNode, Fragment, 0);
+						parentNode = FindOrAddChildNode(parentNode, fragment, 0);
 					}
 				}
-				ProjectParentNodes[Idx] = ParentNode;
+				projectParentNodes[idx] = parentNode;
 			}
 
 			// Add the actual project nodes themselves
-			for(int Idx = 0; Idx < ProjectNodes.Count; Idx++)
+			for(int idx = 0; idx < projectNodes.Count; idx++)
 			{
-				TreeNode Node = FindOrAddChildNode(ProjectParentNodes[Idx], ProjectNodes[Idx].Name, 1);
-				Node.Tag = ProjectNodes[Idx].FullName;
+				TreeNode node = FindOrAddChildNode(projectParentNodes[idx], projectNodes[idx].Name, 1);
+				node.Tag = projectNodes[idx].FullName;
 
-				if(String.Compare(ProjectNodes[Idx].FullName, SelectedProjectFile, StringComparison.InvariantCultureIgnoreCase) == 0)
+				if(String.Compare(projectNodes[idx].FullName, _selectedProjectFile, StringComparison.InvariantCultureIgnoreCase) == 0)
 				{
-					ProjectTreeView.SelectedNode = Node;
-					for(TreeNode ParentNode = Node.Parent; ParentNode != RootNode; ParentNode = ParentNode.Parent)
+					ProjectTreeView.SelectedNode = node;
+					for(TreeNode parentNode = node.Parent; parentNode != rootNode; parentNode = parentNode.Parent)
 					{
-						ParentNode.Expand();
+						parentNode.Expand();
 					}
 				}
 			}
 		}
 
-		static TreeNode FindOrAddChildNode(TreeNode ParentNode, string Text, int ImageIndex)
+		static TreeNode FindOrAddChildNode(TreeNode parentNode, string text, int imageIndex)
 		{
-			foreach(TreeNode? ChildNode in ParentNode.Nodes)
+			foreach(TreeNode? childNode in parentNode.Nodes)
 			{
-				if(ChildNode != null && String.Compare(ChildNode.Text, Text, StringComparison.InvariantCultureIgnoreCase) == 0)
+				if(childNode != null && String.Compare(childNode.Text, text, StringComparison.InvariantCultureIgnoreCase) == 0)
 				{
-					return ChildNode;
+					return childNode;
 				}
 			}
 
-			TreeNode NextNode = new TreeNode(Text);
-			NextNode.ImageIndex = ImageIndex;
-			NextNode.SelectedImageIndex = ImageIndex;
-			ParentNode.Nodes.Add(NextNode);
-			return NextNode;
+			TreeNode nextNode = new TreeNode(text);
+			nextNode.ImageIndex = imageIndex;
+			nextNode.SelectedImageIndex = imageIndex;
+			parentNode.Nodes.Add(nextNode);
+			return nextNode;
 		}
 
-		public static bool ShowModal(IWin32Window Owner, IPerforceSettings Perforce, string WorkspaceName, string WorkspacePath, IServiceProvider ServiceProvider, [NotNullWhen(true)] out string? NewWorkspacePath)
+		public static bool ShowModal(IWin32Window owner, IPerforceSettings perforce, string workspaceName, string workspacePath, IServiceProvider serviceProvider, [NotNullWhen(true)] out string? newWorkspacePath)
 		{
-			Perforce = new PerforceSettings(Perforce) { ClientName = WorkspaceName };
+			perforce = new PerforceSettings(perforce) { ClientName = workspaceName };
 
-			ILogger Logger = ServiceProvider.GetRequiredService<ILogger<SelectProjectFromWorkspaceWindow>>();
+			ILogger logger = serviceProvider.GetRequiredService<ILogger<SelectProjectFromWorkspaceWindow>>();
 
-			ModalTask<List<string>>? PathsTask = PerforceModalTask.Execute(Owner, "Finding Projects", "Finding projects, please wait...", Perforce, (p, c) => EnumerateWorkspaceProjectsTask.RunAsync(p, WorkspaceName, c), Logger);
-			if(PathsTask == null || !PathsTask.Succeeded)
+			ModalTask<List<string>>? pathsTask = PerforceModalTask.Execute(owner, "Finding Projects", "Finding projects, please wait...", perforce, (p, c) => EnumerateWorkspaceProjectsTask.RunAsync(p, workspaceName, c), logger);
+			if(pathsTask == null || !pathsTask.Succeeded)
 			{
-				NewWorkspacePath = null;
+				newWorkspacePath = null;
 				return false;
 			}
 
-			SelectProjectFromWorkspaceWindow SelectProjectWindow = new SelectProjectFromWorkspaceWindow(WorkspaceName, PathsTask.Result, WorkspacePath);
-			if(SelectProjectWindow.ShowDialog() == DialogResult.OK && !String.IsNullOrEmpty(SelectProjectWindow.SelectedProjectFile))
+			SelectProjectFromWorkspaceWindow selectProjectWindow = new SelectProjectFromWorkspaceWindow(workspaceName, pathsTask.Result, workspacePath);
+			if(selectProjectWindow.ShowDialog() == DialogResult.OK && !String.IsNullOrEmpty(selectProjectWindow._selectedProjectFile))
 			{
-				NewWorkspacePath = SelectProjectWindow.SelectedProjectFile;
+				newWorkspacePath = selectProjectWindow._selectedProjectFile;
 				return true;
 			}
 			else
 			{
-				NewWorkspacePath = null;
+				newWorkspacePath = null;
 				return false;
 			}
 		}
@@ -232,7 +232,7 @@ namespace UnrealGameSync
 		{
 			if(ProjectTreeView.SelectedNode != null && ProjectTreeView.SelectedNode.Tag != null)
 			{
-				SelectedProjectFile = (string)ProjectTreeView.SelectedNode.Tag;
+				_selectedProjectFile = (string)ProjectTreeView.SelectedNode.Tag;
 				DialogResult = DialogResult.OK;
 				Close();
 			}
diff --git a/Engine/Source/Programs/UnrealGameSync/UnrealGameSync/Forms/SelectStreamWindow.cs b/Engine/Source/Programs/UnrealGameSync/UnrealGameSync/Forms/SelectStreamWindow.cs
--- a/Engine/Source/Programs/UnrealGameSync/UnrealGameSync/Forms/SelectStreamWindow.cs
+++ b/Engine/Source/Programs/UnrealGameSync/UnrealGameSync/Forms/SelectStreamWindow.cs
@@ -24,9 +24,9 @@ namespace UnrealGameSync
 	{
 		static class EnumerateStreamsTask
 		{
-			public static async Task<List<StreamsRecord>> RunAsync(IPerforceConnection Perforce, CancellationToken CancellationToken)
+			public static async Task<List<StreamsRecord>> RunAsync(IPerforceConnection perforce, CancellationToken cancellationToken)
 			{
-				return await Perforce.GetStreamsAsync(null, CancellationToken);
+				return await perforce.GetStreamsAsync(null, cancellationToken);
 			}
 		}
 
@@ -35,17 +35,17 @@ namespace UnrealGameSync
 			public StreamsRecord Record;
 			public List<StreamNode> ChildNodes = new List<StreamNode>();
 
-			public StreamNode(StreamsRecord Record)
+			public StreamNode(StreamsRecord record)
 			{
-				this.Record = Record;
+				this.Record = record;
 			}
 
 			public void Sort()
 			{
 				ChildNodes = ChildNodes.OrderBy(x => x.Record.Name).ToList();
-				foreach(StreamNode ChildNode in ChildNodes)
+				foreach(StreamNode childNode in ChildNodes)
 				{
-					ChildNode.Sort();
+					childNode.Sort();
 				}
 			}
 		}
@@ -55,80 +55,80 @@ namespace UnrealGameSync
 			public string Name;
 			public List<StreamNode> RootNodes = new List<StreamNode>();
 
-			public StreamDepot(string Name)
+			public StreamDepot(string name)
 			{
-				this.Name = Name;
+				this.Name = name;
 			}
 
 			public void Sort()
 			{
 				RootNodes = RootNodes.OrderBy(x => x.Record.Name).ToList();
-				foreach(StreamNode RootNode in RootNodes)
+				foreach(StreamNode rootNode in RootNodes)
 				{
-					RootNode.Sort();
+					rootNode.Sort();
 				}
 			}
 		}
 
-		private string? SelectedStream;
-		private List<StreamDepot> Depots = new List<StreamDepot>();
+		private string? _selectedStream;
+		private List<StreamDepot> _depots = new List<StreamDepot>();
 
-		private SelectStreamWindow(List<StreamsRecord> Streams, string? StreamName)
+		private SelectStreamWindow(List<StreamsRecord> streams, string? streamName)
 		{
 			InitializeComponent();
 
-			this.SelectedStream = StreamName;
+			this._selectedStream = streamName;
 
 			// Set up the image list
-			ImageList PerforceImageList = new ImageList();
-			PerforceImageList.ImageSize = new Size(16, 16);
-			PerforceImageList.ColorDepth = ColorDepth.Depth32Bit;
-			PerforceImageList.Images.AddStrip(Properties.Resources.Perforce);
-			StreamsTreeView.ImageList = PerforceImageList;
+			ImageList perforceImageList = new ImageList();
+			perforceImageList.ImageSize = new Size(16, 16);
+			perforceImageList.ColorDepth = ColorDepth.Depth32Bit;
+			perforceImageList.Images.AddStrip(Properties.Resources.Perforce);
+			StreamsTreeView.ImageList = perforceImageList;
 
 			// Build a map of stream names to their nodes
-			Dictionary<string, StreamNode> IdentifierToNode = new Dictionary<string, StreamNode>(StringComparer.InvariantCultureIgnoreCase);
-			foreach(StreamsRecord Stream in Streams)
+			Dictionary<string, StreamNode> identifierToNode = new Dictionary<string, StreamNode>(StringComparer.InvariantCultureIgnoreCase);
+			foreach(StreamsRecord stream in streams)
 			{
-				if(Stream.Stream != null && Stream.Name != null)
+				if(stream.Stream != null && stream.Name != null)
 				{
-					IdentifierToNode[Stream.Stream] = new StreamNode(Stream);
+					identifierToNode[stream.Stream] = new StreamNode(stream);
 				}
 			}
 
 			// Create all the depots
-			Dictionary<string, StreamDepot> NameToDepot = new Dictionary<string, StreamDepot>(StringComparer.InvariantCultureIgnoreCase);
-			foreach(StreamNode Node in IdentifierToNode.Values)
+			Dictionary<string, StreamDepot> nameToDepot = new Dictionary<string, StreamDepot>(StringComparer.InvariantCultureIgnoreCase);
+			foreach(StreamNode node in identifierToNode.Values)
 			{
-				if(Node.Record.Parent == null || Node.Record.Parent.Equals("none", StringComparison.OrdinalIgnoreCase))
+				if(node.Record.Parent == null || node.Record.Parent.Equals("none", StringComparison.OrdinalIgnoreCase))
 				{
-					string? DepotName;
-					if(PerforceUtils.TryGetDepotName(Node.Record.Stream, out DepotName))
+					string? depotName;
+					if(PerforceUtils.TryGetDepotName(node.Record.Stream, out depotName))
 					{
-						StreamDepot? Depot;
-						if(!NameToDepot.TryGetValue(DepotName, out Depot))
+						StreamDepot? depot;
+						if(!nameToDepot.TryGetValue(depotName, out depot))
 						{
-							Depot = new StreamDepot(DepotName);
-							NameToDepot.Add(DepotName, Depot);
+							depot = new StreamDepot(depotName);
+							nameToDepot.Add(depotName, depot);
 						}
-						Depot.RootNodes.Add(Node);
+						depot.RootNodes.Add(node);
 					}
 				}
 				else
 				{
-					StreamNode? ParentNode;
-					if(IdentifierToNode.TryGetValue(Node.Record.Parent, out ParentNode))
+					StreamNode? parentNode;
+					if(identifierToNode.TryGetValue(node.Record.Parent, out parentNode))
 					{
-						ParentNode.ChildNodes.Add(Node);
+						parentNode.ChildNodes.Add(node);
 					}
 				}
 			}
 
 			// Sort the tree
-			Depots = NameToDepot.Values.OrderBy(x => x.Name).ToList();
-			foreach(StreamDepot Depot in Depots)
+			_depots = nameToDepot.Values.OrderBy(x => x.Name).ToList();
+			foreach(StreamDepot depot in _depots)
 			{
-				Depot.Sort();
+				depot.Sort();
 			}
 
 			// Update the contents of the tree
@@ -136,85 +136,85 @@ namespace UnrealGameSync
 			UpdateOkButton();
 		}
 
-		private void GetExpandedNodes(TreeNodeCollection Nodes, List<TreeNode> ExpandedNodes)
+		private void GetExpandedNodes(TreeNodeCollection nodes, List<TreeNode> expandedNodes)
 		{
-			foreach(TreeNode? Node in Nodes)
+			foreach(TreeNode? node in nodes)
 			{
-				if (Node != null)
+				if (node != null)
 				{
-					ExpandedNodes.Add(Node);
-					if (Node.IsExpanded)
+					expandedNodes.Add(node);
+					if (node.IsExpanded)
 					{
-						GetExpandedNodes(Node.Nodes, ExpandedNodes);
+						GetExpandedNodes(node.Nodes, expandedNodes);
 					}
 				}
 			}
 		}
 
-		private void MoveSelection(int Delta)
+		private void MoveSelection(int delta)
 		{
 			if(StreamsTreeView.SelectedNode != null)
 			{
-				List<TreeNode> ExpandedNodes = new List<TreeNode>();
-				GetExpandedNodes(StreamsTreeView.Nodes, ExpandedNodes);
+				List<TreeNode> expandedNodes = new List<TreeNode>();
+				GetExpandedNodes(StreamsTreeView.Nodes, expandedNodes);
 
-				int Idx = ExpandedNodes.IndexOf(StreamsTreeView.SelectedNode);
-				if(Idx != -1)
+				int idx = expandedNodes.IndexOf(StreamsTreeView.SelectedNode);
+				if(idx != -1)
 				{
-					int NextIdx = Idx + Delta;
-					if(NextIdx < 0)
+					int nextIdx = idx + delta;
+					if(nextIdx < 0)
 					{
-						NextIdx = 0;
+						nextIdx = 0;
 					}
-					if(NextIdx >= ExpandedNodes.Count)
+					if(nextIdx >= expandedNodes.Count)
 					{
-						NextIdx = ExpandedNodes.Count - 1;
+						nextIdx = expandedNodes.Count - 1;
 					}
-					StreamsTreeView.SelectedNode = ExpandedNodes[NextIdx];
+					StreamsTreeView.SelectedNode = expandedNodes[nextIdx];
 				}
 			}
 		}
 
-		protected override bool ProcessCmdKey(ref Message Msg, Keys KeyData)
+		protected override bool ProcessCmdKey(ref Message msg, Keys keyData)
 		{
-			if(KeyData == Keys.Up)
+			if(keyData == Keys.Up)
 			{
 				MoveSelection(-1);
 				return true;
 			}
-			else if(KeyData == Keys.Down)
+			else if(keyData == Keys.Down)
 			{
 				MoveSelection(+1);
 				return true;
 			}
-			return base.ProcessCmdKey(ref Msg, KeyData);
+			return base.ProcessCmdKey(ref msg, keyData);
 		}
 
-		private bool IncludeNodeInFilter(StreamNode Node, string[] Filter)
+		private bool IncludeNodeInFilter(StreamNode node, string[] filter)
 		{
-			return Filter.All(x => Node.Record.Stream.IndexOf(x, StringComparison.InvariantCultureIgnoreCase) != -1 || Node.Record.Name.IndexOf(x, StringComparison.InvariantCultureIgnoreCase) != -1);
+			return filter.All(x => node.Record.Stream.IndexOf(x, StringComparison.InvariantCultureIgnoreCase) != -1 || node.Record.Name.IndexOf(x, StringComparison.InvariantCultureIgnoreCase) != -1);
 		}
 
-		private bool TryFilterTree(StreamNode Node, string[] Filter, [NotNullWhen(true)] out StreamNode? NewNode)
+		private bool TryFilterTree(StreamNode node, string[] filter, [NotNullWhen(true)] out StreamNode? newNode)
 		{
-			StreamNode FilteredNode = new StreamNode(Node.Record);
-			foreach(StreamNode ChildNode in Node.ChildNodes)
+			StreamNode filteredNode = new StreamNode(node.Record);
+			foreach(StreamNode childNode in node.ChildNodes)
 			{
-				StreamNode? FilteredChildNode;
-				if(TryFilterTree(ChildNode, Filter, out FilteredChildNode))
+				StreamNode? filteredChildNode;
+				if(TryFilterTree(childNode, filter, out filteredChildNode))
 				{
-					FilteredNode.ChildNodes.Add(FilteredChildNode);
+					filteredNode.ChildNodes.Add(filteredChildNode);
 				}
 			}
 
-			if(FilteredNode.ChildNodes.Count > 0 || IncludeNodeInFilter(FilteredNode, Filter))
+			if(filteredNode.ChildNodes.Count > 0 || IncludeNodeInFilter(filteredNode, filter))
 			{
-				NewNode = FilteredNode;
+				newNode = filteredNode;
 				return true;
 			}
 			else
 			{
-				NewNode = null;
+				newNode = null;
 				return false;
 			}
 		}
@@ -224,60 +224,60 @@ namespace UnrealGameSync
 			StreamsTreeView.BeginUpdate();
 			StreamsTreeView.Nodes.Clear();
 
-			string[] Filter = FilterTextBox.Text.Split(new char[]{' '}, StringSplitOptions.RemoveEmptyEntries);
+			string[] filter = FilterTextBox.Text.Split(new char[]{' '}, StringSplitOptions.RemoveEmptyEntries);
 
-			List<StreamDepot> FilteredDepots = Depots;
-			if(Filter.Length > 0)
+			List<StreamDepot> filteredDepots = _depots;
+			if(filter.Length > 0)
 			{
-				FilteredDepots = new List<StreamDepot>();
-				foreach(StreamDepot Depot in Depots)
+				filteredDepots = new List<StreamDepot>();
+				foreach(StreamDepot depot in _depots)
 				{
-					StreamDepot FilteredDepot = new StreamDepot(Depot.Name);
-					foreach(StreamNode RootNode in Depot.RootNodes)
+					StreamDepot filteredDepot = new StreamDepot(depot.Name);
+					foreach(StreamNode rootNode in depot.RootNodes)
 					{
-						StreamNode? FilteredRootNode;
-						if(TryFilterTree(RootNode, Filter, out FilteredRootNode))
+						StreamNode? filteredRootNode;
+						if(TryFilterTree(rootNode, filter, out filteredRootNode))
 						{
-							FilteredDepot.RootNodes.Add(FilteredRootNode);
+							filteredDepot.RootNodes.Add(filteredRootNode);
 						}
 					}
-					if(FilteredDepot.RootNodes.Count > 0)
+					if(filteredDepot.RootNodes.Count > 0)
 					{
-						FilteredDepots.Add(FilteredDepot);
+						filteredDepots.Add(filteredDepot);
 					}
 				}
 			}
 
-			bool bExpandAll = Filter.Length > 0;
-			foreach(StreamDepot Depot in FilteredDepots)
+			bool expandAll = filter.Length > 0;
+			foreach(StreamDepot depot in filteredDepots)
 			{
-				TreeNode DepotTreeNode = new TreeNode(Depot.Name);
-				DepotTreeNode.ImageIndex = 1;
-				DepotTreeNode.SelectedImageIndex = 1;
-				StreamsTreeView.Nodes.Add(DepotTreeNode);
+				TreeNode depotTreeNode = new TreeNode(depot.Name);
+				depotTreeNode.ImageIndex = 1;
+				depotTreeNode.SelectedImageIndex = 1;
+				StreamsTreeView.Nodes.Add(depotTreeNode);
 
-				bool bExpand = bExpandAll;
-				foreach(StreamNode RootNode in Depot.RootNodes)
+				bool expand = expandAll;
+				foreach(StreamNode rootNode in depot.RootNodes)
 				{
-					bExpand |= AddStreamNodeToTree(RootNode, Filter, DepotTreeNode, bExpandAll);
+					expand |= AddStreamNodeToTree(rootNode, filter, depotTreeNode, expandAll);
 				}
-				if(bExpand)
+				if(expand)
 				{
-					DepotTreeNode.Expand();
+					depotTreeNode.Expand();
 				}
 			}
 
-			if(StreamsTreeView.SelectedNode == null && Filter.Length > 0 && StreamsTreeView.Nodes.Count > 0)
+			if(StreamsTreeView.SelectedNode == null && filter.Length > 0 && StreamsTreeView.Nodes.Count > 0)
 			{
-				for(TreeNode Node = StreamsTreeView.Nodes[0];;Node = Node.Nodes[0])
+				for(TreeNode node = StreamsTreeView.Nodes[0];;node = node.Nodes[0])
 				{
-					StreamNode? Stream = Node.Tag as StreamNode;
-					if(Stream != null && IncludeNodeInFilter(Stream, Filter))
+					StreamNode? stream = node.Tag as StreamNode;
+					if(stream != null && IncludeNodeInFilter(stream, filter))
 					{
-						StreamsTreeView.SelectedNode = Node;
+						StreamsTreeView.SelectedNode = node;
 						break;
 					}
-					if(Node.Nodes.Count == 0)
+					if(node.Nodes.Count == 0)
 					{
 						break;
 					}
@@ -297,51 +297,51 @@ namespace UnrealGameSync
 			UpdateOkButton();
 		}
 
-		private bool AddStreamNodeToTree(StreamNode Stream, string[] Filter, TreeNode ParentTreeNode, bool bExpandAll)
+		private bool AddStreamNodeToTree(StreamNode stream, string[] filter, TreeNode parentTreeNode, bool expandAll)
 		{
-			TreeNode StreamTreeNode = new TreeNode(Stream.Record.Name);
-			StreamTreeNode.ImageIndex = 0;
-			StreamTreeNode.SelectedImageIndex = 0;
-			StreamTreeNode.Tag = Stream;
-			ParentTreeNode.Nodes.Add(StreamTreeNode);
+			TreeNode streamTreeNode = new TreeNode(stream.Record.Name);
+			streamTreeNode.ImageIndex = 0;
+			streamTreeNode.SelectedImageIndex = 0;
+			streamTreeNode.Tag = stream;
+			parentTreeNode.Nodes.Add(streamTreeNode);
 
-			if(Stream.Record.Name == SelectedStream && IncludeNodeInFilter(Stream, Filter))
+			if(stream.Record.Name == _selectedStream && IncludeNodeInFilter(stream, filter))
 			{
-				StreamsTreeView.SelectedNode = StreamTreeNode;
+				StreamsTreeView.SelectedNode = streamTreeNode;
 			}
 
-			bool bExpand = bExpandAll;
-			foreach(StreamNode ChildNode in Stream.ChildNodes)
+			bool expand = expandAll;
+			foreach(StreamNode childNode in stream.ChildNodes)
 			{
-				bExpand |= AddStreamNodeToTree(ChildNode, Filter, StreamTreeNode, bExpandAll);
+				expand |= AddStreamNodeToTree(childNode, filter, streamTreeNode, expandAll);
 			}
-			if(bExpand)
+			if(expand)
 			{
-				StreamTreeNode.Expand();
+				streamTreeNode.Expand();
 			}
-			return bExpand || (SelectedStream == Stream.Record.Stream);
+			return expand || (_selectedStream == stream.Record.Stream);
 		}
 
-		public static bool ShowModal(IWin32Window Owner, IPerforceSettings Perforce, string? StreamName, IServiceProvider ServiceProvider, [NotNullWhen(true)] out string? NewStreamName)
+		public static bool ShowModal(IWin32Window owner, IPerforceSettings perforce, string? streamName, IServiceProvider serviceProvider, [NotNullWhen(true)] out string? newStreamName)
 		{
-			ILogger Logger = ServiceProvider.GetRequiredService<ILogger<SelectStreamWindow>>();
+			ILogger logger = serviceProvider.GetRequiredService<ILogger<SelectStreamWindow>>();
 
-			ModalTask<List<StreamsRecord>>? StreamsTask = PerforceModalTask.Execute(Owner, "Finding streams", "Finding streams, please wait...", Perforce, EnumerateStreamsTask.RunAsync, Logger);
-			if(StreamsTask == null || !StreamsTask.Succeeded)
+			ModalTask<List<StreamsRecord>>? streamsTask = PerforceModalTask.Execute(owner, "Finding streams", "Finding streams, please wait...", perforce, EnumerateStreamsTask.RunAsync, logger);
+			if(streamsTask == null || !streamsTask.Succeeded)
 			{
-				NewStreamName = null;
+				newStreamName = null;
 				return false;
 			}
 
-			SelectStreamWindow SelectStream = new SelectStreamWindow(StreamsTask.Result, StreamName);
-			if(SelectStream.ShowDialog(Owner) == DialogResult.OK && SelectStream.SelectedStream != null)
+			SelectStreamWindow selectStream = new SelectStreamWindow(streamsTask.Result, streamName);
+			if(selectStream.ShowDialog(owner) == DialogResult.OK && selectStream._selectedStream != null)
 			{
-				NewStreamName = SelectStream.SelectedStream;
+				newStreamName = selectStream._selectedStream;
 				return true;
 			}
 			else
 			{
-				NewStreamName = null;
+				newStreamName = null;
 				return false;
 			}
 		}
@@ -358,21 +358,21 @@ namespace UnrealGameSync
 
 		private string? GetSelectedStream()
 		{
-			string? NewSelectedStream = null;
+			string? newSelectedStream = null;
 			if(StreamsTreeView.SelectedNode != null)
 			{
-				StreamNode StreamNode = (StreamNode)StreamsTreeView.SelectedNode.Tag;
-				if(StreamNode != null)
+				StreamNode streamNode = (StreamNode)StreamsTreeView.SelectedNode.Tag;
+				if(streamNode != null)
 				{
-					NewSelectedStream = StreamNode.Record.Stream;
+					newSelectedStream = streamNode.Record.Stream;
 				}
 			}
-			return NewSelectedStream;
+			return newSelectedStream;
 		}
 
 		private void UpdateSelectedStream()
 		{
-			SelectedStream = GetSelectedStream();
+			_selectedStream = GetSelectedStream();
 		}
 
 		private void StreamsTreeView_AfterSelect(object sender, TreeViewEventArgs e)
@@ -388,7 +388,7 @@ namespace UnrealGameSync
 		{
 			UpdateSelectedStream();
 
-			if(SelectedStream != null)
+			if(_selectedStream != null)
 			{
 				DialogResult = DialogResult.OK;
 				Close();
diff --git a/Engine/Source/Programs/UnrealGameSync/UnrealGameSync/Forms/SelectUserWindow.cs b/Engine/Source/Programs/UnrealGameSync/UnrealGameSync/Forms/SelectUserWindow.cs
--- a/Engine/Source/Programs/UnrealGameSync/UnrealGameSync/Forms/SelectUserWindow.cs
+++ b/Engine/Source/Programs/UnrealGameSync/UnrealGameSync/Forms/SelectUserWindow.cs
@@ -24,81 +24,81 @@ namespace UnrealGameSync
 	{
 		static class EnumerateUsersTask
 		{
-			public static async Task<List<UsersRecord>> RunAsync(IPerforceConnection Perforce, CancellationToken CancellationToken)
+			public static async Task<List<UsersRecord>> RunAsync(IPerforceConnection perforce, CancellationToken cancellationToken)
 			{
-				return await Perforce.GetUsersAsync(UsersOptions.None, -1, CancellationToken);
+				return await perforce.GetUsersAsync(UsersOptions.None, -1, cancellationToken);
 			}
 		}
 
-		private int SelectedUserIndex;
-		private List<UsersRecord> Users;
+		private int _selectedUserIndex;
+		private List<UsersRecord> _users;
 		
-		private SelectUserWindow(List<UsersRecord> Users, int SelectedUserIndex)
+		private SelectUserWindow(List<UsersRecord> users, int selectedUserIndex)
 		{
 			InitializeComponent();
 
-			this.SelectedUserIndex = SelectedUserIndex;
-			this.Users = Users;
+			this._selectedUserIndex = selectedUserIndex;
+			this._users = users;
 
 			PopulateList();
 			UpdateOkButton();
 		}
 
-		private void MoveSelection(int Delta)
+		private void MoveSelection(int delta)
 		{
 			if(UserListView.Items.Count > 0)
 			{
-				int CurrentIndex = -1;
+				int currentIndex = -1;
 				if(UserListView.SelectedIndices.Count > 0)
 				{
-					CurrentIndex = UserListView.SelectedIndices[0];
+					currentIndex = UserListView.SelectedIndices[0];
 				}
 
-				int NextIndex = CurrentIndex + Delta;
-				if(NextIndex < 0)
+				int nextIndex = currentIndex + delta;
+				if(nextIndex < 0)
 				{
-					NextIndex = 0;
+					nextIndex = 0;
 				}
-				else if(NextIndex >= UserListView.Items.Count)
+				else if(nextIndex >= UserListView.Items.Count)
 				{
-					NextIndex = UserListView.Items.Count - 1;
+					nextIndex = UserListView.Items.Count - 1;
 				}
 
-				if(CurrentIndex != NextIndex)
+				if(currentIndex != nextIndex)
 				{
-					if(CurrentIndex != -1)
+					if(currentIndex != -1)
 					{
-						UserListView.Items[CurrentIndex].Selected = false;
+						UserListView.Items[currentIndex].Selected = false;
 					}
 
-					UserListView.Items[NextIndex].Selected = true;
-					SelectedUserIndex = (int)UserListView.Items[NextIndex].Tag;
+					UserListView.Items[nextIndex].Selected = true;
+					_selectedUserIndex = (int)UserListView.Items[nextIndex].Tag;
 				}
 			}
 		}
 
-		protected override bool ProcessCmdKey(ref Message Msg, Keys KeyData)
+		protected override bool ProcessCmdKey(ref Message msg, Keys keyData)
 		{
-			if(KeyData == Keys.Up)
+			if(keyData == Keys.Up)
 			{
 				MoveSelection(-1);
 				return true;
 			}
-			else if(KeyData == Keys.Down)
+			else if(keyData == Keys.Down)
 			{
 				MoveSelection(+1);
 				return true;
 			}
-			return base.ProcessCmdKey(ref Msg, KeyData);
+			return base.ProcessCmdKey(ref msg, keyData);
 		}
 
-		private bool IncludeInFilter(UsersRecord User, string[] FilterWords)
+		private bool IncludeInFilter(UsersRecord user, string[] filterWords)
 		{
-			foreach(string FilterWord in FilterWords)
+			foreach(string filterWord in filterWords)
 			{
-				if(User.UserName.IndexOf(FilterWord, StringComparison.OrdinalIgnoreCase) == -1 
-					&& User.FullName.IndexOf(FilterWord, StringComparison.OrdinalIgnoreCase) == -1
-					&& User.Email.IndexOf(FilterWord, StringComparison.OrdinalIgnoreCase) == -1)
+				if(user.UserName.IndexOf(filterWord, StringComparison.OrdinalIgnoreCase) == -1 
+					&& user.FullName.IndexOf(filterWord, StringComparison.OrdinalIgnoreCase) == -1
+					&& user.Email.IndexOf(filterWord, StringComparison.OrdinalIgnoreCase) == -1)
 				{
 					return false;
 				}
@@ -111,37 +111,37 @@ namespace UnrealGameSync
 			UserListView.BeginUpdate();
 			UserListView.Items.Clear();
 
-			int SelectedItemIndex = -1;
+			int selectedItemIndex = -1;
 
-			string[] Filter = FilterTextBox.Text.Split(new char[]{' '}, StringSplitOptions.RemoveEmptyEntries);
-			for(int Idx = 0; Idx < Users.Count; Idx++)
+			string[] filter = FilterTextBox.Text.Split(new char[]{' '}, StringSplitOptions.RemoveEmptyEntries);
+			for(int idx = 0; idx < _users.Count; idx++)
 			{
-				UsersRecord User = Users[Idx];
-				if(IncludeInFilter(User, Filter))
+				UsersRecord user = _users[idx];
+				if(IncludeInFilter(user, filter))
 				{
-					ListViewItem Item = new ListViewItem(User.UserName);
-					Item.SubItems.Add(new ListViewItem.ListViewSubItem(Item, User.FullName));
-					Item.SubItems.Add(new ListViewItem.ListViewSubItem(Item, User.Email));
-					Item.Tag = (int)Idx;
-					UserListView.Items.Add(Item);
+					ListViewItem item = new ListViewItem(user.UserName);
+					item.SubItems.Add(new ListViewItem.ListViewSubItem(item, user.FullName));
+					item.SubItems.Add(new ListViewItem.ListViewSubItem(item, user.Email));
+					item.Tag = (int)idx;
+					UserListView.Items.Add(item);
 
-					if(SelectedItemIndex == -1 && Idx >= SelectedUserIndex)
+					if(selectedItemIndex == -1 && idx >= _selectedUserIndex)
 					{
-						SelectedItemIndex = UserListView.Items.Count - 1;
-						Item.Selected = true;
+						selectedItemIndex = UserListView.Items.Count - 1;
+						item.Selected = true;
 					}
 				}
 			}
 
-			if(SelectedItemIndex == -1 && UserListView.Items.Count > 0)
+			if(selectedItemIndex == -1 && UserListView.Items.Count > 0)
 			{
-				SelectedItemIndex = UserListView.Items.Count - 1;
-				UserListView.Items[SelectedItemIndex].Selected = true;
+				selectedItemIndex = UserListView.Items.Count - 1;
+				UserListView.Items[selectedItemIndex].Selected = true;
 			}
 
-			if(SelectedItemIndex != -1)
+			if(selectedItemIndex != -1)
 			{
-				UserListView.EnsureVisible(SelectedItemIndex);
+				UserListView.EnsureVisible(selectedItemIndex);
 			}
 
 			UserListView.EndUpdate();
@@ -149,26 +149,26 @@ namespace UnrealGameSync
 			UpdateOkButton();
 		}
 
-		public static bool ShowModal(IWin32Window Owner, IPerforceSettings PerforceSettings, IServiceProvider ServiceProvider, [NotNullWhen(true)] out string? SelectedUserName)
+		public static bool ShowModal(IWin32Window owner, IPerforceSettings perforceSettings, IServiceProvider serviceProvider, [NotNullWhen(true)] out string? selectedUserName)
 		{
-			ILogger<SelectUserWindow> Logger = ServiceProvider.GetRequiredService<ILogger<SelectUserWindow>>();
+			ILogger<SelectUserWindow> logger = serviceProvider.GetRequiredService<ILogger<SelectUserWindow>>();
 
-			ModalTask<List<UsersRecord>>? UsersTask = PerforceModalTask.Execute(Owner, "Finding users", "Finding users, please wait...", PerforceSettings, EnumerateUsersTask.RunAsync, Logger);
-			if(UsersTask == null || !UsersTask.Succeeded)
+			ModalTask<List<UsersRecord>>? usersTask = PerforceModalTask.Execute(owner, "Finding users", "Finding users, please wait...", perforceSettings, EnumerateUsersTask.RunAsync, logger);
+			if(usersTask == null || !usersTask.Succeeded)
 			{
-				SelectedUserName = null;
+				selectedUserName = null;
 				return false;
 			}
 
-			SelectUserWindow SelectUser = new SelectUserWindow(UsersTask.Result, 0);
-			if(SelectUser.ShowDialog(Owner) == DialogResult.OK)
+			SelectUserWindow selectUser = new SelectUserWindow(usersTask.Result, 0);
+			if(selectUser.ShowDialog(owner) == DialogResult.OK)
 			{
-				SelectedUserName = UsersTask.Result[SelectUser.SelectedUserIndex].UserName;
+				selectedUserName = usersTask.Result[selectUser._selectedUserIndex].UserName;
 				return true;
 			}
 			else
 			{
-				SelectedUserName = null;
+				selectedUserName = null;
 				return false;
 			}
 		}
@@ -187,7 +187,7 @@ namespace UnrealGameSync
 		{
 			if(UserListView.SelectedItems.Count > 0)
 			{
-				SelectedUserIndex = (int)UserListView.SelectedItems[0].Tag;
+				_selectedUserIndex = (int)UserListView.SelectedItems[0].Tag;
 				DialogResult = DialogResult.OK;
 				Close();
 			}
diff --git a/Engine/Source/Programs/UnrealGameSync/UnrealGameSync/Forms/SelectWorkspaceWindow.cs b/Engine/Source/Programs/UnrealGameSync/UnrealGameSync/Forms/SelectWorkspaceWindow.cs
--- a/Engine/Source/Programs/UnrealGameSync/UnrealGameSync/Forms/SelectWorkspaceWindow.cs
+++ b/Engine/Source/Programs/UnrealGameSync/UnrealGameSync/Forms/SelectWorkspaceWindow.cs
@@ -27,31 +27,31 @@ namespace UnrealGameSync
 			public InfoRecord Info { get; }
 			public List<ClientsRecord> Clients { get; }
 
-			public EnumerateWorkspaces(InfoRecord Info, List<ClientsRecord> Clients)
+			public EnumerateWorkspaces(InfoRecord info, List<ClientsRecord> clients)
 			{
-				this.Info = Info;
-				this.Clients = Clients;
+				this.Info = info;
+				this.Clients = clients;
 			}
 
-			public static async Task<EnumerateWorkspaces> RunAsync(IPerforceConnection Perforce, CancellationToken CancellationToken)
+			public static async Task<EnumerateWorkspaces> RunAsync(IPerforceConnection perforce, CancellationToken cancellationToken)
 			{
-				InfoRecord Info = await Perforce.GetInfoAsync(InfoOptions.ShortOutput, CancellationToken);
-				List<ClientsRecord> Clients = await Perforce.GetClientsAsync(ClientsOptions.None, Perforce.Settings.UserName, CancellationToken);
-				return new EnumerateWorkspaces(Info, Clients);
+				InfoRecord info = await perforce.GetInfoAsync(InfoOptions.ShortOutput, cancellationToken);
+				List<ClientsRecord> clients = await perforce.GetClientsAsync(ClientsOptions.None, perforce.Settings.UserName, cancellationToken);
+				return new EnumerateWorkspaces(info, clients);
 			}
 		}
 
-		InfoRecord Info;
-		List<ClientsRecord> Clients;
-		string? WorkspaceName;
+		InfoRecord _info;
+		List<ClientsRecord> _clients;
+		string? _workspaceName;
 
-		private SelectWorkspaceWindow(InfoRecord Info, List<ClientsRecord> Clients, string? WorkspaceName)
+		private SelectWorkspaceWindow(InfoRecord info, List<ClientsRecord> clients, string? workspaceName)
 		{
 			InitializeComponent();
 
-			this.Info = Info;
-			this.Clients = Clients;
-			this.WorkspaceName = WorkspaceName;
+			this._info = info;
+			this._clients = clients;
+			this._workspaceName = workspaceName;
 
 			UpdateListView();
 			UpdateOkButton();
@@ -61,25 +61,25 @@ namespace UnrealGameSync
 		{
 			if(WorkspaceListView.SelectedItems.Count > 0)
 			{
-				WorkspaceName = WorkspaceListView.SelectedItems[0].Text;
+				_workspaceName = WorkspaceListView.SelectedItems[0].Text;
 			}
 			else
 			{
-				WorkspaceName = null;
+				_workspaceName = null;
 			}
 
 			WorkspaceListView.Items.Clear();
 
-			foreach(ClientsRecord Client in Clients.OrderBy(x => x.Name))
+			foreach(ClientsRecord client in _clients.OrderBy(x => x.Name))
 			{
-				if(!OnlyForThisComputer.Checked || String.Compare(Client.Host, Info.ClientHost, StringComparison.InvariantCultureIgnoreCase) == 0)
+				if(!OnlyForThisComputer.Checked || String.Compare(client.Host, _info.ClientHost, StringComparison.InvariantCultureIgnoreCase) == 0)
 				{
-					ListViewItem Item = new ListViewItem(Client.Name);
-					Item.SubItems.Add(new ListViewItem.ListViewSubItem(Item, Client.Host));
-					Item.SubItems.Add(new ListViewItem.ListViewSubItem(Item, Client.Stream));
-					Item.SubItems.Add(new ListViewItem.ListViewSubItem(Item, Client.Root));
-					Item.Selected = (WorkspaceName == Client.Name);
-					WorkspaceListView.Items.Add(Item);
+					ListViewItem item = new ListViewItem(client.Name);
+					item.SubItems.Add(new ListViewItem.ListViewSubItem(item, client.Host));
+					item.SubItems.Add(new ListViewItem.ListViewSubItem(item, client.Stream));
+					item.SubItems.Add(new ListViewItem.ListViewSubItem(item, client.Root));
+					item.Selected = (_workspaceName == client.Name);
+					WorkspaceListView.Items.Add(item);
 				}
 			}
 		}
@@ -89,24 +89,24 @@ namespace UnrealGameSync
 			OkBtn.Enabled = (WorkspaceListView.SelectedItems.Count == 1);
 		}
 
-		public static bool ShowModal(IWin32Window Owner, IPerforceSettings Perforce, string WorkspaceName, IServiceProvider ServiceProvider, out string? NewWorkspaceName)
+		public static bool ShowModal(IWin32Window owner, IPerforceSettings perforce, string workspaceName, IServiceProvider serviceProvider, out string? newWorkspaceName)
 		{
-			ModalTask<EnumerateWorkspaces>? Task = PerforceModalTask.Execute(Owner, "Finding workspaces", "Finding workspaces, please wait...", Perforce, EnumerateWorkspaces.RunAsync, ServiceProvider.GetRequiredService<ILogger<EnumerateWorkspaces>>());
-			if (Task == null || !Task.Succeeded)
+			ModalTask<EnumerateWorkspaces>? task = PerforceModalTask.Execute(owner, "Finding workspaces", "Finding workspaces, please wait...", perforce, EnumerateWorkspaces.RunAsync, serviceProvider.GetRequiredService<ILogger<EnumerateWorkspaces>>());
+			if (task == null || !task.Succeeded)
 			{
-				NewWorkspaceName = null;
+				newWorkspaceName = null;
 				return false;
 			}
 
-			SelectWorkspaceWindow SelectWorkspace = new SelectWorkspaceWindow(Task.Result.Info, Task.Result.Clients, WorkspaceName);
-			if(SelectWorkspace.ShowDialog(Owner) == DialogResult.OK)
+			SelectWorkspaceWindow selectWorkspace = new SelectWorkspaceWindow(task.Result.Info, task.Result.Clients, workspaceName);
+			if(selectWorkspace.ShowDialog(owner) == DialogResult.OK)
 			{
-				NewWorkspaceName = SelectWorkspace.WorkspaceName;
+				newWorkspaceName = selectWorkspace._workspaceName;
 				return true;
 			}
 			else
 			{
-				NewWorkspaceName = null;
+				newWorkspaceName = null;
 				return false;
 			}
 		}
@@ -125,7 +125,7 @@ namespace UnrealGameSync
 		{
 			if(WorkspaceListView.SelectedItems.Count > 0)
 			{
-				WorkspaceName = WorkspaceListView.SelectedItems[0].Text;
+				_workspaceName = WorkspaceListView.SelectedItems[0].Text;
 				DialogResult = DialogResult.OK;
 				Close();
 			}
@@ -141,7 +141,7 @@ namespace UnrealGameSync
 		{
 			if(WorkspaceListView.SelectedItems.Count > 0)
 			{
-				WorkspaceName = WorkspaceListView.SelectedItems[0].Text;
+				_workspaceName = WorkspaceListView.SelectedItems[0].Text;
 				DialogResult = DialogResult.OK;
 				Close();
 			}
diff --git a/Engine/Source/Programs/UnrealGameSync/UnrealGameSync/Forms/SyncFilter.cs b/Engine/Source/Programs/UnrealGameSync/UnrealGameSync/Forms/SyncFilter.cs
--- a/Engine/Source/Programs/UnrealGameSync/UnrealGameSync/Forms/SyncFilter.cs
+++ b/Engine/Source/Programs/UnrealGameSync/UnrealGameSync/Forms/SyncFilter.cs
@@ -14,29 +14,29 @@ namespace UnrealGameSync
 {
 	partial class SyncFilter : Form
 	{
-		Dictionary<Guid, WorkspaceSyncCategory> UniqueIdToCategory;
+		Dictionary<Guid, WorkspaceSyncCategory> _uniqueIdToCategory;
 		public FilterSettings GlobalFilter;
 		public FilterSettings WorkspaceFilter;
 
-		public SyncFilter(Dictionary<Guid, WorkspaceSyncCategory> InUniqueIdToCategory, FilterSettings InGlobalFilter, FilterSettings InWorkspaceFilter)
+		public SyncFilter(Dictionary<Guid, WorkspaceSyncCategory> inUniqueIdToCategory, FilterSettings inGlobalFilter, FilterSettings inWorkspaceFilter)
 		{
 			InitializeComponent();
 
-			UniqueIdToCategory = InUniqueIdToCategory;
-			GlobalFilter = InGlobalFilter;
-			WorkspaceFilter = InWorkspaceFilter;
+			_uniqueIdToCategory = inUniqueIdToCategory;
+			GlobalFilter = inGlobalFilter;
+			WorkspaceFilter = inWorkspaceFilter;
 
-			Dictionary<Guid, bool> SyncCategories = WorkspaceSyncCategory.GetDefault(UniqueIdToCategory.Values);
+			Dictionary<Guid, bool> syncCategories = WorkspaceSyncCategory.GetDefault(_uniqueIdToCategory.Values);
 
-			WorkspaceSyncCategory.ApplyDelta(SyncCategories, GlobalFilter.GetCategories());
+			WorkspaceSyncCategory.ApplyDelta(syncCategories, GlobalFilter.GetCategories());
 			GlobalControl.SetView(GlobalFilter.View.ToArray());
-			SetExcludedCategories(GlobalControl.CategoriesCheckList, UniqueIdToCategory, SyncCategories);
+			SetExcludedCategories(GlobalControl.CategoriesCheckList, _uniqueIdToCategory, syncCategories);
 			GlobalControl.SyncAllProjects.Checked = GlobalFilter.AllProjects ?? false;
 			GlobalControl.IncludeAllProjectsInSolution.Checked = GlobalFilter.AllProjectsInSln ?? false;
 
-			WorkspaceSyncCategory.ApplyDelta(SyncCategories, WorkspaceFilter.GetCategories());
+			WorkspaceSyncCategory.ApplyDelta(syncCategories, WorkspaceFilter.GetCategories());
 			WorkspaceControl.SetView(WorkspaceFilter.View.ToArray());
-			SetExcludedCategories(WorkspaceControl.CategoriesCheckList, UniqueIdToCategory, SyncCategories);
+			SetExcludedCategories(WorkspaceControl.CategoriesCheckList, _uniqueIdToCategory, syncCategories);
 			WorkspaceControl.SyncAllProjects.Checked = WorkspaceFilter.AllProjects ?? GlobalFilter.AllProjects ?? false;
 			WorkspaceControl.IncludeAllProjectsInSolution.Checked = WorkspaceFilter.AllProjectsInSln ?? GlobalFilter.AllProjectsInSln ?? false;
 
@@ -60,80 +60,80 @@ namespace UnrealGameSync
 			WorkspaceControl.IncludeAllProjectsInSolution.Checked = GlobalControl.IncludeAllProjectsInSolution.Checked;
 		}
 
-		private static void SetExcludedCategories(CheckedListBox ListBox, Dictionary<Guid, WorkspaceSyncCategory> UniqueIdToFilter, Dictionary<Guid, bool> CategoryIdToSetting)
+		private static void SetExcludedCategories(CheckedListBox listBox, Dictionary<Guid, WorkspaceSyncCategory> uniqueIdToFilter, Dictionary<Guid, bool> categoryIdToSetting)
 		{
-			ListBox.Items.Clear();
-			foreach(WorkspaceSyncCategory Filter in UniqueIdToFilter.Values)
+			listBox.Items.Clear();
+			foreach(WorkspaceSyncCategory filter in uniqueIdToFilter.Values)
 			{
-				if (!Filter.bHidden)
+				if (!filter.Hidden)
 				{
-					CheckState State = CheckState.Checked;
-					if (!CategoryIdToSetting[Filter.UniqueId])
+					CheckState state = CheckState.Checked;
+					if (!categoryIdToSetting[filter.UniqueId])
 					{
-						State = CheckState.Unchecked;
+						state = CheckState.Unchecked;
 					}
-					ListBox.Items.Add(Filter, State);
+					listBox.Items.Add(filter, state);
 				}
 			}
 		}
 
-		private void GetSettings(out FilterSettings NewGlobalFilter, out FilterSettings NewWorkspaceFilter)
+		private void GetSettings(out FilterSettings newGlobalFilter, out FilterSettings newWorkspaceFilter)
 		{
-			Dictionary<Guid, bool> DefaultSyncCategories = WorkspaceSyncCategory.GetDefault(UniqueIdToCategory.Values);
+			Dictionary<Guid, bool> defaultSyncCategories = WorkspaceSyncCategory.GetDefault(_uniqueIdToCategory.Values);
 
-			NewGlobalFilter = new FilterSettings();
-			NewGlobalFilter.View = GlobalControl.GetView().ToList();
-			NewGlobalFilter.AllProjects = GlobalControl.SyncAllProjects.Checked;
-			NewGlobalFilter.AllProjectsInSln = GlobalControl.IncludeAllProjectsInSolution.Checked;
+			newGlobalFilter = new FilterSettings();
+			newGlobalFilter.View = GlobalControl.GetView().ToList();
+			newGlobalFilter.AllProjects = GlobalControl.SyncAllProjects.Checked;
+			newGlobalFilter.AllProjectsInSln = GlobalControl.IncludeAllProjectsInSolution.Checked;
 
-			Dictionary<Guid, bool> GlobalSyncCategories = GetCategorySettings(GlobalControl.CategoriesCheckList, GlobalFilter.GetCategories());
-			NewGlobalFilter.SetCategories(WorkspaceSyncCategory.GetDelta(DefaultSyncCategories, GlobalSyncCategories));
+			Dictionary<Guid, bool> globalSyncCategories = GetCategorySettings(GlobalControl.CategoriesCheckList, GlobalFilter.GetCategories());
+			newGlobalFilter.SetCategories(WorkspaceSyncCategory.GetDelta(defaultSyncCategories, globalSyncCategories));
 
-			NewWorkspaceFilter = new FilterSettings();
-			NewWorkspaceFilter.View = WorkspaceControl.GetView().ToList();
-			NewWorkspaceFilter.AllProjects = (WorkspaceControl.SyncAllProjects.Checked == NewGlobalFilter.AllProjects) ? (bool?)null : WorkspaceControl.SyncAllProjects.Checked;
-			NewWorkspaceFilter.AllProjectsInSln = (WorkspaceControl.IncludeAllProjectsInSolution.Checked == NewGlobalFilter.AllProjectsInSln) ? (bool?)null : WorkspaceControl.IncludeAllProjectsInSolution.Checked;
+			newWorkspaceFilter = new FilterSettings();
+			newWorkspaceFilter.View = WorkspaceControl.GetView().ToList();
+			newWorkspaceFilter.AllProjects = (WorkspaceControl.SyncAllProjects.Checked == newGlobalFilter.AllProjects) ? (bool?)null : WorkspaceControl.SyncAllProjects.Checked;
+			newWorkspaceFilter.AllProjectsInSln = (WorkspaceControl.IncludeAllProjectsInSolution.Checked == newGlobalFilter.AllProjectsInSln) ? (bool?)null : WorkspaceControl.IncludeAllProjectsInSolution.Checked;
 
-			Dictionary<Guid, bool> WorkspaceSyncCategories = GetCategorySettings(WorkspaceControl.CategoriesCheckList, WorkspaceFilter.GetCategories());
-			NewWorkspaceFilter.SetCategories(WorkspaceSyncCategory.GetDelta(GlobalSyncCategories, WorkspaceSyncCategories));
+			Dictionary<Guid, bool> workspaceSyncCategories = GetCategorySettings(WorkspaceControl.CategoriesCheckList, WorkspaceFilter.GetCategories());
+			newWorkspaceFilter.SetCategories(WorkspaceSyncCategory.GetDelta(globalSyncCategories, workspaceSyncCategories));
 		}
 
-		private Dictionary<Guid, bool> GetCategorySettings(CheckedListBox ListBox, IEnumerable<KeyValuePair<Guid, bool>> OriginalSettings)
+		private Dictionary<Guid, bool> GetCategorySettings(CheckedListBox listBox, IEnumerable<KeyValuePair<Guid, bool>> originalSettings)
 		{
-			Dictionary<Guid, bool> Result = new Dictionary<Guid, bool>();
-			for(int Idx = 0; Idx < ListBox.Items.Count; Idx++)
+			Dictionary<Guid, bool> result = new Dictionary<Guid, bool>();
+			for(int idx = 0; idx < listBox.Items.Count; idx++)
 			{
-				Guid UniqueId = ((WorkspaceSyncCategory)ListBox.Items[Idx]).UniqueId;
-				if (!Result.ContainsKey(UniqueId))
+				Guid uniqueId = ((WorkspaceSyncCategory)listBox.Items[idx]).UniqueId;
+				if (!result.ContainsKey(uniqueId))
 				{
-					Result[UniqueId] = ListBox.GetItemCheckState(Idx) == CheckState.Checked;
+					result[uniqueId] = listBox.GetItemCheckState(idx) == CheckState.Checked;
 				}
 			}
-			foreach (KeyValuePair<Guid, bool> OriginalSetting in OriginalSettings)
+			foreach (KeyValuePair<Guid, bool> originalSetting in originalSettings)
 			{
-				if (!UniqueIdToCategory.ContainsKey(OriginalSetting.Key))
+				if (!_uniqueIdToCategory.ContainsKey(originalSetting.Key))
 				{
-					Result[OriginalSetting.Key] = OriginalSetting.Value;
+					result[originalSetting.Key] = originalSetting.Value;
 				}
 			}
-			return Result;
+			return result;
 		}
 
-		private static string[] GetView(TextBox FilterText)
+		private static string[] GetView(TextBox filterText)
 		{
-			List<string> NewLines = new List<string>(FilterText.Lines);
-			while (NewLines.Count > 0 && NewLines.Last().Trim().Length == 0)
+			List<string> newLines = new List<string>(filterText.Lines);
+			while (newLines.Count > 0 && newLines.Last().Trim().Length == 0)
 			{
-				NewLines.RemoveAt(NewLines.Count - 1);
+				newLines.RemoveAt(newLines.Count - 1);
 			}
-			return NewLines.Count > 0 ? FilterText.Lines : new string[0];
+			return newLines.Count > 0 ? filterText.Lines : new string[0];
 		}
 
 		private void OkButton_Click(object sender, EventArgs e)
 		{
-			GetSettings(out FilterSettings NewGlobalFilter, out FilterSettings NewWorkspaceFilter);
+			GetSettings(out FilterSettings newGlobalFilter, out FilterSettings newWorkspaceFilter);
 
-			if(NewGlobalFilter.View.Any(x => x.Contains("//")) || NewWorkspaceFilter.View.Any(x => x.Contains("//")))
+			if(newGlobalFilter.View.Any(x => x.Contains("//")) || newWorkspaceFilter.View.Any(x => x.Contains("//")))
 			{
 				if(MessageBox.Show(this, "Custom views should be relative to the stream root (eg. -/Engine/...).\r\n\r\nFull depot paths (eg. //depot/...) will not match any files.\r\n\r\nAre you sure you want to continue?", "Invalid view", MessageBoxButtons.OKCancel) != System.Windows.Forms.DialogResult.OK)
 				{
@@ -141,8 +141,8 @@ namespace UnrealGameSync
 				}
 			}
 
-			GlobalFilter = NewGlobalFilter;
-			WorkspaceFilter = NewWorkspaceFilter;
+			GlobalFilter = newGlobalFilter;
+			WorkspaceFilter = newWorkspaceFilter;
 
 			DialogResult = DialogResult.OK;
 		}
@@ -154,14 +154,14 @@ namespace UnrealGameSync
 
 		private void ShowCombinedView_Click(object sender, EventArgs e)
 		{
-			GetSettings(out FilterSettings NewGlobalFilter, out FilterSettings NewWorkspaceFilter);
+			GetSettings(out FilterSettings newGlobalFilter, out FilterSettings newWorkspaceFilter);
 
-			string[] Filter = UserSettings.GetCombinedSyncFilter(UniqueIdToCategory, NewGlobalFilter, NewWorkspaceFilter);
-			if(Filter.Length == 0)
+			string[] filter = UserSettings.GetCombinedSyncFilter(_uniqueIdToCategory, newGlobalFilter, newWorkspaceFilter);
+			if(filter.Length == 0)
 			{
-				Filter = new string[]{ "All files will be synced." };
+				filter = new string[]{ "All files will be synced." };
 			}
-			MessageBox.Show(String.Join("\r\n", Filter), "Combined View");
+			MessageBox.Show(String.Join("\r\n", filter), "Combined View");
 		}
 	}
 }
diff --git a/Engine/Source/Programs/UnrealGameSync/UnrealGameSync/Forms/VariablesWindow.cs b/Engine/Source/Programs/UnrealGameSync/UnrealGameSync/Forms/VariablesWindow.cs
--- a/Engine/Source/Programs/UnrealGameSync/UnrealGameSync/Forms/VariablesWindow.cs
+++ b/Engine/Source/Programs/UnrealGameSync/UnrealGameSync/Forms/VariablesWindow.cs
@@ -17,7 +17,7 @@ namespace UnrealGameSync
 {
 	public partial class VariablesWindow : Form
 	{
-		static HashSet<string> LegacyVariables = new HashSet<string>()
+		static HashSet<string> _legacyVariables = new HashSet<string>()
 		{
 			"UE4EditorConfig",
 			"UE4EditorDebugArg",
@@ -26,40 +26,40 @@ namespace UnrealGameSync
 			"UseIncrementalBuilds",
 		};
 
-		public delegate void InsertVariableDelegate(string Name);
+		public delegate void InsertVariableDelegate(string name);
 
 		public event InsertVariableDelegate? OnInsertVariable;
 
-		public VariablesWindow(IReadOnlyDictionary<string, string> Variables)
+		public VariablesWindow(IReadOnlyDictionary<string, string> variables)
 		{
 			InitializeComponent();
 
-			ListViewGroup CurrentProjectGroup = new ListViewGroup("Current Project");
-			MacrosList.Groups.Add(CurrentProjectGroup);
+			ListViewGroup currentProjectGroup = new ListViewGroup("Current Project");
+			MacrosList.Groups.Add(currentProjectGroup);
 
-			ListViewGroup EnvironmentGroup = new ListViewGroup("Environment");
-			MacrosList.Groups.Add(EnvironmentGroup);
+			ListViewGroup environmentGroup = new ListViewGroup("Environment");
+			MacrosList.Groups.Add(environmentGroup);
 
-			foreach(KeyValuePair<string, string> Pair in Variables)
+			foreach(KeyValuePair<string, string> pair in variables)
 			{
-				if (!LegacyVariables.Contains(Pair.Key))
+				if (!_legacyVariables.Contains(pair.Key))
 				{
-					ListViewItem Item = new ListViewItem(String.Format("$({0})", Pair.Key));
-					Item.SubItems.Add(Pair.Value);
-					Item.Group = CurrentProjectGroup;
-					MacrosList.Items.Add(Item);
+					ListViewItem item = new ListViewItem(String.Format("$({0})", pair.Key));
+					item.SubItems.Add(pair.Value);
+					item.Group = currentProjectGroup;
+					MacrosList.Items.Add(item);
 				}
 			}
 
-			foreach(DictionaryEntry Entry in Environment.GetEnvironmentVariables().OfType<DictionaryEntry>())
+			foreach(DictionaryEntry entry in Environment.GetEnvironmentVariables().OfType<DictionaryEntry>())
 			{
-				string? Key = Entry.Key?.ToString();
-				if(Key != null && Entry.Value != null && !Variables.ContainsKey(Key))
+				string? key = entry.Key?.ToString();
+				if(key != null && entry.Value != null && !variables.ContainsKey(key))
 				{
-					ListViewItem Item = new ListViewItem(String.Format("$({0})", Key));
-					Item.SubItems.Add(Entry.Value.ToString());
-					Item.Group = EnvironmentGroup;
-					MacrosList.Items.Add(Item);
+					ListViewItem item = new ListViewItem(String.Format("$({0})", key));
+					item.SubItems.Add(entry.Value.ToString());
+					item.Group = environmentGroup;
+					MacrosList.Items.Add(item);
 				}
 			}
 		}
@@ -75,14 +75,14 @@ namespace UnrealGameSync
 			Close();
 		}
 
-		private void MacrosList_MouseDoubleClick(object Sender, MouseEventArgs Args)
+		private void MacrosList_MouseDoubleClick(object sender, MouseEventArgs args)
 		{
-			if(Args.Button == MouseButtons.Left)
+			if(args.Button == MouseButtons.Left)
 			{
-				ListViewHitTestInfo HitTest = MacrosList.HitTest(Args.Location);
-				if(HitTest.Item != null && OnInsertVariable != null)
+				ListViewHitTestInfo hitTest = MacrosList.HitTest(args.Location);
+				if(hitTest.Item != null && OnInsertVariable != null)
 				{
-					OnInsertVariable(HitTest.Item.Text);
+					OnInsertVariable(hitTest.Item.Text);
 				}
 			}
 		}
diff --git a/Engine/Source/Programs/UnrealGameSync/UnrealGameSync/GlobalSettings.cs b/Engine/Source/Programs/UnrealGameSync/UnrealGameSync/GlobalSettings.cs
--- a/Engine/Source/Programs/UnrealGameSync/UnrealGameSync/GlobalSettings.cs
+++ b/Engine/Source/Programs/UnrealGameSync/UnrealGameSync/GlobalSettings.cs
@@ -12,50 +12,50 @@ namespace UnrealGameSync
 {
 	internal class GlobalPerforceSettings
 	{
-		public static void ReadGlobalPerforceSettings(ref string? ServerAndPort, ref string? UserName, ref string? DepotPath, ref bool bPreview)
+		public static void ReadGlobalPerforceSettings(ref string? serverAndPort, ref string? userName, ref string? depotPath, ref bool preview)
 		{
-			using (RegistryKey? Key = Registry.CurrentUser.OpenSubKey("SOFTWARE\\Epic Games\\UnrealGameSync", false))
+			using (RegistryKey? key = Registry.CurrentUser.OpenSubKey("SOFTWARE\\Epic Games\\UnrealGameSync", false))
 			{
-				if (Key != null)
+				if (key != null)
 				{
-					ServerAndPort = Key.GetValue("ServerAndPort", ServerAndPort) as string;
-					UserName = Key.GetValue("UserName", UserName) as string;
-					DepotPath = Key.GetValue("DepotPath", DepotPath) as string;
-					bPreview = ((Key.GetValue("Preview", bPreview? 1 : 0) as int?) ?? 0) != 0;
+					serverAndPort = key.GetValue("ServerAndPort", serverAndPort) as string;
+					userName = key.GetValue("UserName", userName) as string;
+					depotPath = key.GetValue("DepotPath", depotPath) as string;
+					preview = ((key.GetValue("Preview", preview? 1 : 0) as int?) ?? 0) != 0;
 
 					// Fix corrupted depot path string
-					if (DepotPath != null)
+					if (depotPath != null)
 					{
-						Match Match = Regex.Match(DepotPath, "^(.*)/(Release|UnstableRelease)/\\.\\.\\.@.*$");
-						if (Match.Success)
+						Match match = Regex.Match(depotPath, "^(.*)/(Release|UnstableRelease)/\\.\\.\\.@.*$");
+						if (match.Success)
 						{
-							DepotPath = Match.Groups[1].Value;
-							SaveGlobalPerforceSettings(ServerAndPort, UserName, DepotPath, bPreview);
+							depotPath = match.Groups[1].Value;
+							SaveGlobalPerforceSettings(serverAndPort, userName, depotPath, preview);
 						}
 					}
 				}
 			}
 		}
 
-		public static void DeleteRegistryKey(RegistryKey RootKey, string KeyName, string ValueName)
+		public static void DeleteRegistryKey(RegistryKey rootKey, string keyName, string valueName)
 		{
-			using (RegistryKey? Key = RootKey.OpenSubKey(KeyName, true))
+			using (RegistryKey? key = rootKey.OpenSubKey(keyName, true))
 			{
-				if (Key != null)
+				if (key != null)
 				{
-					DeleteRegistryKey(Key, ValueName);
+					DeleteRegistryKey(key, valueName);
 				}
 			}
 		}
 
-		public static void DeleteRegistryKey(RegistryKey Key, string Name)
+		public static void DeleteRegistryKey(RegistryKey key, string name)
 		{
-			string[] ValueNames = Key.GetValueNames();
-			if (ValueNames.Any(x => String.Compare(x, Name, StringComparison.OrdinalIgnoreCase) == 0))
+			string[] valueNames = key.GetValueNames();
+			if (valueNames.Any(x => String.Compare(x, name, StringComparison.OrdinalIgnoreCase) == 0))
 			{
 				try
 				{
-					Key.DeleteValue(Name);
+					key.DeleteValue(name);
 				}
 				catch
 				{
@@ -63,55 +63,55 @@ namespace UnrealGameSync
 			}
 		}
 
-		public static void SaveGlobalPerforceSettings(string? ServerAndPort, string? UserName, string? DepotPath, bool bPreview)
+		public static void SaveGlobalPerforceSettings(string? serverAndPort, string? userName, string? depotPath, bool preview)
 		{
 			try
 			{
-				using (RegistryKey Key = Registry.CurrentUser.CreateSubKey("SOFTWARE\\Epic Games\\UnrealGameSync"))
+				using (RegistryKey key = Registry.CurrentUser.CreateSubKey("SOFTWARE\\Epic Games\\UnrealGameSync"))
 				{
 					// Delete this legacy setting
-					DeleteRegistryKey(Key, "Server");
+					DeleteRegistryKey(key, "Server");
 
-					if (String.IsNullOrEmpty(ServerAndPort))
+					if (String.IsNullOrEmpty(serverAndPort))
 					{
-						try { Key.DeleteValue("ServerAndPort"); } catch (Exception) { }
+						try { key.DeleteValue("ServerAndPort"); } catch (Exception) { }
 					}
 					else
 					{
-						Key.SetValue("ServerAndPort", ServerAndPort);
+						key.SetValue("ServerAndPort", serverAndPort);
 					}
 
-					if (String.IsNullOrEmpty(UserName))
+					if (String.IsNullOrEmpty(userName))
 					{
-						try { Key.DeleteValue("UserName"); } catch (Exception) { }
+						try { key.DeleteValue("UserName"); } catch (Exception) { }
 					}
 					else
 					{
-						Key.SetValue("UserName", UserName);
+						key.SetValue("UserName", userName);
 					}
 
-					if (String.IsNullOrEmpty(DepotPath) || (DeploymentSettings.DefaultDepotPath != null && String.Equals(DepotPath, DeploymentSettings.DefaultDepotPath, StringComparison.InvariantCultureIgnoreCase)))
+					if (String.IsNullOrEmpty(depotPath) || (DeploymentSettings.DefaultDepotPath != null && String.Equals(depotPath, DeploymentSettings.DefaultDepotPath, StringComparison.InvariantCultureIgnoreCase)))
 					{
-						DeleteRegistryKey(Key, "DepotPath");
+						DeleteRegistryKey(key, "DepotPath");
 					}
 					else
 					{
-						Key.SetValue("DepotPath", DepotPath);
+						key.SetValue("DepotPath", depotPath);
 					}
 
-					if (bPreview)
+					if (preview)
 					{
-						Key.SetValue("Preview", 1);
+						key.SetValue("Preview", 1);
 					}
 					else
 					{
-						try { Key.DeleteValue("Preview"); } catch (Exception) { }
+						try { key.DeleteValue("Preview"); } catch (Exception) { }
 					}
 				}
 			}
-			catch (Exception Ex)
+			catch (Exception ex)
 			{
-				MessageBox.Show("Unable to save settings.\n\n" + Ex.ToString());
+				MessageBox.Show("Unable to save settings.\n\n" + ex.ToString());
 			}
 		}
 	}
diff --git a/Engine/Source/Programs/UnrealGameSync/UnrealGameSync/IssueMonitor.cs b/Engine/Source/Programs/UnrealGameSync/UnrealGameSync/IssueMonitor.cs
--- a/Engine/Source/Programs/UnrealGameSync/UnrealGameSync/IssueMonitor.cs
+++ b/Engine/Source/Programs/UnrealGameSync/UnrealGameSync/IssueMonitor.cs
@@ -64,52 +64,52 @@ namespace UnrealGameSync
 		public DateTime? AcknowledgedAt { get; set; }
 		public int FixChange { get; set; }
 		public DateTime? ResolvedAt { get; set; }
-		public bool bNotify { get; set; }
-		public bool bIsWarning { get; set; }
+		public bool Notify { get; set; }
+		public bool IsWarning { get; set; }
 		public string BuildUrl { get; set; } = String.Empty;
 		public List<string> Streams { get; set; } = new List<string>();
 
-		HashSet<string>? CachedProjects;
+		HashSet<string>? _cachedProjects;
 
 		public HashSet<string> Projects
 		{
 			get
 			{
 				// HACK to infer project names from streams
-				if(CachedProjects == null)
+				if(_cachedProjects == null)
 				{
-					HashSet<string> NewProjects = new HashSet<string>(StringComparer.OrdinalIgnoreCase);
+					HashSet<string> newProjects = new HashSet<string>(StringComparer.OrdinalIgnoreCase);
 					if (!String.IsNullOrEmpty(Project))
 					{
-						NewProjects.Add(Project);
+						newProjects.Add(Project);
 					}
 					if (Streams != null)
 					{
-						foreach (string Stream in Streams)
+						foreach (string stream in Streams)
 						{
-							Match Match = Regex.Match(Stream, "^//([^/]+)/");
-							if (Match.Success)
+							Match match = Regex.Match(stream, "^//([^/]+)/");
+							if (match.Success)
 							{
-								string Project = Match.Groups[1].Value;
-								if (Project.StartsWith("UE", StringComparison.OrdinalIgnoreCase))
+								string project = match.Groups[1].Value;
+								if (project.StartsWith("UE", StringComparison.OrdinalIgnoreCase))
 								{
-									Project = "UE" + Project.Substring(2);
+									project = "UE" + project.Substring(2);
 								}
-								else if (Char.IsLower(Project[0]))
+								else if (Char.IsLower(project[0]))
 								{
-									Project = Char.ToUpper(Project[0], CultureInfo.InvariantCulture) + Project.Substring(1);
+									project = Char.ToUpper(project[0], CultureInfo.InvariantCulture) + project.Substring(1);
 								}
-								NewProjects.Add(Project);
+								newProjects.Add(project);
 							}
 						}
 					}
-					if (NewProjects.Count == 0)
+					if (newProjects.Count == 0)
 					{
-						NewProjects.Add("Default");
+						newProjects.Add("Default");
 					}
-					CachedProjects = NewProjects;
+					_cachedProjects = newProjects;
 				}
-				return CachedProjects;
+				return _cachedProjects;
 			}
 		}
 	}
@@ -138,42 +138,42 @@ namespace UnrealGameSync
 	{
 		public readonly string? ApiUrl;
 		public readonly string UserName;
-		int RefCount = 1;
-		Task? WorkerTask;
-		ILogger Logger;
-		CancellationTokenSource CancellationSource;
-		AsyncEvent RefreshEvent;
-		int UpdateIntervalMs;
-		List<long> TrackingIssueIds = new List<long>();
-		List<IssueData> Issues = new List<IssueData>();
-		object LockObject = new object();
-		List<IssueUpdateData> PendingUpdates = new List<IssueUpdateData>();
-		IAsyncDisposer AsyncDisposer;
+		int _refCount = 1;
+		Task? _workerTask;
+		ILogger _logger;
+		CancellationTokenSource _cancellationSource;
+		AsyncEvent _refreshEvent;
+		int _updateIntervalMs;
+		List<long> _trackingIssueIds = new List<long>();
+		List<IssueData> _issues = new List<IssueData>();
+		object _lockObject = new object();
+		List<IssueUpdateData> _pendingUpdates = new List<IssueUpdateData>();
+		IAsyncDisposer _asyncDisposer;
 
 		public Action? OnIssuesChanged { get; set; }
 
 		// Only used by MainWindow, but easier to just store here
 		public Dictionary<long, IssueAlertReason> IssueIdToAlertReason = new Dictionary<long, IssueAlertReason>();
 
-		public IssueMonitor(string? ApiUrl, string UserName, TimeSpan UpdateInterval, IServiceProvider ServiceProvider)
+		public IssueMonitor(string? apiUrl, string userName, TimeSpan updateInterval, IServiceProvider serviceProvider)
 		{
-			this.ApiUrl = ApiUrl;
-			this.UserName = UserName;
-			this.UpdateIntervalMs = (int)UpdateInterval.TotalMilliseconds;
-			this.Logger = ServiceProvider.GetRequiredService<ILogger<IssueMonitor>>();
-			CancellationSource = new CancellationTokenSource();
-			this.AsyncDisposer = ServiceProvider.GetRequiredService<IAsyncDisposer>();
-
-			if (ApiUrl == null)
+			this.ApiUrl = apiUrl;
+			this.UserName = userName;
+			this._updateIntervalMs = (int)updateInterval.TotalMilliseconds;
+			this._logger = serviceProvider.GetRequiredService<ILogger<IssueMonitor>>();
+			_cancellationSource = new CancellationTokenSource();
+			this._asyncDisposer = serviceProvider.GetRequiredService<IAsyncDisposer>();
+
+			if (apiUrl == null)
 			{
 				LastStatusMessage = "Database functionality disabled due to empty ApiUrl.";
 			}
 			else
 			{
-				Logger.LogInformation("Using connection string: {ApiUrl}", this.ApiUrl);
+				_logger.LogInformation("Using connection string: {ApiUrl}", this.ApiUrl);
 			}
 
-			RefreshEvent = new AsyncEvent();
+			_refreshEvent = new AsyncEvent();
 		}
 
 		public string LastStatusMessage
@@ -184,122 +184,122 @@ namespace UnrealGameSync
 
 		public List<IssueData> GetIssues()
 		{
-			return Issues;
+			return _issues;
 		}
 
 		public TimeSpan GetUpdateInterval()
 		{
-			return TimeSpan.FromMilliseconds(UpdateIntervalMs);
+			return TimeSpan.FromMilliseconds(_updateIntervalMs);
 		}
 
-		public void SetUpdateInterval(TimeSpan UpdateInterval)
+		public void SetUpdateInterval(TimeSpan updateInterval)
 		{
-			UpdateIntervalMs = (int)UpdateInterval.TotalMilliseconds;
-			RefreshEvent.Set();
+			_updateIntervalMs = (int)updateInterval.TotalMilliseconds;
+			_refreshEvent.Set();
 		}
 
-		public void StartTracking(long IssueId)
+		public void StartTracking(long issueId)
 		{
-			lock(LockObject)
+			lock(_lockObject)
 			{
-				TrackingIssueIds.Add(IssueId);
+				_trackingIssueIds.Add(issueId);
 			}
-			RefreshEvent.Set();
+			_refreshEvent.Set();
 		}
 
-		public void StopTracking(long IssueId)
+		public void StopTracking(long issueId)
 		{
-			lock(LockObject)
+			lock(_lockObject)
 			{
-				TrackingIssueIds.RemoveAt(TrackingIssueIds.IndexOf(IssueId));
+				_trackingIssueIds.RemoveAt(_trackingIssueIds.IndexOf(issueId));
 			}
 		}
 
 		public bool HasPendingUpdate()
 		{
-			return PendingUpdates.Count > 0;
+			return _pendingUpdates.Count > 0;
 		}
 
-		public void PostUpdate(IssueUpdateData Update)
+		public void PostUpdate(IssueUpdateData update)
 		{
-			bool bUpdatedIssues;
-			lock(LockObject)
+			bool updatedIssues;
+			lock(_lockObject)
 			{
-				PendingUpdates.Add(Update);
-				bUpdatedIssues = ApplyPendingUpdate(Issues, Update);
+				_pendingUpdates.Add(update);
+				updatedIssues = ApplyPendingUpdate(_issues, update);
 			}
 
-			RefreshEvent.Set();
+			_refreshEvent.Set();
 
-			if(bUpdatedIssues)
+			if(updatedIssues)
 			{
 				OnIssuesChanged?.Invoke();
 			}
 		}
 
-		static bool ApplyPendingUpdate(List<IssueData> Issues, IssueUpdateData Update)
+		static bool ApplyPendingUpdate(List<IssueData> issues, IssueUpdateData update)
 		{
-			bool bUpdated = false;
-			for(int Idx = 0; Idx < Issues.Count; Idx++)
+			bool updated = false;
+			for(int idx = 0; idx < issues.Count; idx++)
 			{
-				IssueData Issue = Issues[Idx];
-				if(Update.Id == Issue.Id)
+				IssueData issue = issues[idx];
+				if(update.Id == issue.Id)
 				{
-					if(Update.Owner != null && Update.Owner != Issue.Owner)
+					if(update.Owner != null && update.Owner != issue.Owner)
 					{
-						Issue.Owner = Update.Owner;
-						bUpdated = true;
+						issue.Owner = update.Owner;
+						updated = true;
 					}
-					if(Update.NominatedBy != null && Update.NominatedBy != Issue.NominatedBy)
+					if(update.NominatedBy != null && update.NominatedBy != issue.NominatedBy)
 					{
-						Issue.NominatedBy = Update.NominatedBy;
-						bUpdated = true;
+						issue.NominatedBy = update.NominatedBy;
+						updated = true;
 					}
-					if(Update.Acknowledged.HasValue && Update.Acknowledged.Value != Issue.AcknowledgedAt.HasValue)
+					if(update.Acknowledged.HasValue && update.Acknowledged.Value != issue.AcknowledgedAt.HasValue)
 					{
-						Issue.AcknowledgedAt = Update.Acknowledged.Value? (DateTime?)DateTime.UtcNow : null;
-						bUpdated = true;
+						issue.AcknowledgedAt = update.Acknowledged.Value? (DateTime?)DateTime.UtcNow : null;
+						updated = true;
 					}
-					if(Update.FixChange.HasValue)
+					if(update.FixChange.HasValue)
 					{
-						Issue.FixChange = Update.FixChange.Value;
-						if(Issue.FixChange != 0)
+						issue.FixChange = update.FixChange.Value;
+						if(issue.FixChange != 0)
 						{
-							Issues.RemoveAt(Idx);
+							issues.RemoveAt(idx);
 						}
-						bUpdated = true;
+						updated = true;
 					}
 					break;
 				}
 			}
-			return bUpdated;
+			return updated;
 		}
 
 		public void Start()
 		{
 			if(ApiUrl != null)
 			{
-				WorkerTask = Task.Run(() => PollForUpdatesAsync(CancellationSource.Token));
+				_workerTask = Task.Run(() => PollForUpdatesAsync(_cancellationSource.Token));
 			}
 		}
 
 		public void AddRef()
 		{
-			if(RefCount == 0)
+			if(_refCount == 0)
 			{
 				throw new Exception("Invalid reference count for IssueMonitor (zero)");
 			}
-			RefCount++;
+			_refCount++;
 		}
 
 		public void Release()
 		{
-			RefCount--;
-			if(RefCount < 0)
+			_refCount--;
+			if(_refCount < 0)
 			{
 				throw new Exception("Invalid reference count for IssueMonitor (ltz)");
 			}
-			if(RefCount == 0)
+			if(_refCount == 0)
 			{
 				DisposeInternal();
 			}
@@ -314,140 +314,140 @@ namespace UnrealGameSync
 		{
 			OnIssuesChanged = null;
 
-			if (WorkerTask != null)
+			if (_workerTask != null)
 			{
-				CancellationSource.Cancel();
-				AsyncDisposer.Add(WorkerTask.ContinueWith(_ => CancellationSource.Dispose()));
-				WorkerTask = null;
+				_cancellationSource.Cancel();
+				_asyncDisposer.Add(_workerTask.ContinueWith(_ => _cancellationSource.Dispose()));
+				_workerTask = null;
 			}
 		}
 
-		async Task PollForUpdatesAsync(CancellationToken CancellationToken)
+		async Task PollForUpdatesAsync(CancellationToken cancellationToken)
 		{
-			while (CancellationToken.IsCancellationRequested)
+			while (cancellationToken.IsCancellationRequested)
 			{
-				Task RefreshTask = RefreshEvent.Task;
+				Task refreshTask = _refreshEvent.Task;
 
 				// Check if there's any pending update
-				IssueUpdateData? PendingUpdate;
-				lock (LockObject)
+				IssueUpdateData? pendingUpdate;
+				lock (_lockObject)
 				{
-					if (PendingUpdates.Count > 0)
+					if (_pendingUpdates.Count > 0)
 					{
-						PendingUpdate = PendingUpdates[0];
+						pendingUpdate = _pendingUpdates[0];
 					}
 					else
 					{
-						PendingUpdate = null;
+						pendingUpdate = null;
 					}
 				}
 
 				// If we have an update, try to post it to the backend and check for another
-				if (PendingUpdate != null)
+				if (pendingUpdate != null)
 				{
-					if (await SendUpdateAsync(PendingUpdate, CancellationToken))
+					if (await SendUpdateAsync(pendingUpdate, cancellationToken))
 					{
-						lock (LockObject) { PendingUpdates.RemoveAt(0); }
+						lock (_lockObject) { _pendingUpdates.RemoveAt(0); }
 					}
 					else
 					{
-						await Task.WhenAny(RefreshTask, Task.Delay(TimeSpan.FromSeconds(5.0), CancellationToken));
+						await Task.WhenAny(refreshTask, Task.Delay(TimeSpan.FromSeconds(5.0), cancellationToken));
 					}
 					continue;
 				}
 
 				// Read all the current issues
-				await ReadCurrentIssuesAsync(CancellationToken);
+				await ReadCurrentIssuesAsync(cancellationToken);
 
 				// Wait for something else to do
-				await Task.WhenAny(RefreshTask, Task.Delay(UpdateIntervalMs, CancellationToken));
+				await Task.WhenAny(refreshTask, Task.Delay(_updateIntervalMs, cancellationToken));
 			}
 		}
 
-		async Task<bool> SendUpdateAsync(IssueUpdateData Update, CancellationToken CancellationToken)
+		async Task<bool> SendUpdateAsync(IssueUpdateData update, CancellationToken cancellationToken)
 		{
 			try
 			{
-				await RESTApi.PutAsync<IssueUpdateData>($"{ApiUrl}/api/issues/{Update.Id}", Update, CancellationToken);
+				await RestApi.PutAsync<IssueUpdateData>($"{ApiUrl}/api/issues/{update.Id}", update, cancellationToken);
 				return true;
 			}
-			catch(Exception Ex)
+			catch(Exception ex)
 			{
-				Logger.LogError(Ex, "Failed with exception.");
-				LastStatusMessage = String.Format("Failed to send update: ({0})", Ex.ToString());
+				_logger.LogError(ex, "Failed with exception.");
+				LastStatusMessage = String.Format("Failed to send update: ({0})", ex.ToString());
 				return false;
 			}
 		}
 
-		async Task<bool> ReadCurrentIssuesAsync(CancellationToken CancellationToken)
+		async Task<bool> ReadCurrentIssuesAsync(CancellationToken cancellationToken)
 		{
 			try
 			{
-				Stopwatch Timer = Stopwatch.StartNew();
-				Logger.LogInformation("Polling for issues...");
+				Stopwatch timer = Stopwatch.StartNew();
+				_logger.LogInformation("Polling for issues...");
 
 				// Get the initial number of issues. We won't post updates if this stays at zero.
-				int InitialNumIssues = Issues.Count;
+				int initialNumIssues = _issues.Count;
 
 				// Fetch the new issues
-				List<IssueData> NewIssues = await RESTApi.GetAsync<List<IssueData>>($"{ApiUrl}/api/issues?user={UserName}", CancellationToken);
+				List<IssueData> newIssues = await RestApi.GetAsync<List<IssueData>>($"{ApiUrl}/api/issues?user={UserName}", cancellationToken);
 
 				// Check if we're tracking a particular issue. If so, we want updates even when it's resolved.
-				long[] LocalTrackingIssueIds;
-				lock(LockObject)
+				long[] localTrackingIssueIds;
+				lock(_lockObject)
 				{
-					LocalTrackingIssueIds = TrackingIssueIds.Distinct().ToArray();
+					localTrackingIssueIds = _trackingIssueIds.Distinct().ToArray();
 				}
-				foreach(long LocalTrackingIssueId in LocalTrackingIssueIds)
+				foreach(long localTrackingIssueId in localTrackingIssueIds)
 				{
-					if(!NewIssues.Any(x => x.Id == LocalTrackingIssueId))
+					if(!newIssues.Any(x => x.Id == localTrackingIssueId))
 					{
 						try
 						{
-							IssueData Issue = await RESTApi.GetAsync<IssueData>($"{ApiUrl}/api/issues/{LocalTrackingIssueId}", CancellationToken);
-							if(Issue != null)
+							IssueData issue = await RestApi.GetAsync<IssueData>($"{ApiUrl}/api/issues/{localTrackingIssueId}", cancellationToken);
+							if(issue != null)
 							{
-								NewIssues.Add(Issue);
+								newIssues.Add(issue);
 							}
 						}
-						catch(Exception Ex)
+						catch(Exception ex)
 						{
-							Logger.LogError(Ex, "Exception while fetching tracked issue");
+							_logger.LogError(ex, "Exception while fetching tracked issue");
 						}
 					}
 				}
 
 				// Update all the builds for each issue
-				foreach (IssueData NewIssue in NewIssues)
+				foreach (IssueData newIssue in newIssues)
 				{
-					if (NewIssue.Version == 0)
+					if (newIssue.Version == 0)
 					{
-						List<IssueBuildData> Builds = await RESTApi.GetAsync<List<IssueBuildData>>($"{ApiUrl}/api/issues/{NewIssue.Id}/builds", CancellationToken);
-						if (Builds != null && Builds.Count > 0)
+						List<IssueBuildData> builds = await RestApi.GetAsync<List<IssueBuildData>>($"{ApiUrl}/api/issues/{newIssue.Id}/builds", cancellationToken);
+						if (builds != null && builds.Count > 0)
 						{
-							NewIssue.bIsWarning = !Builds.Any(x => x.Outcome != IssueBuildOutcome.Warning);
+							newIssue.IsWarning = !builds.Any(x => x.Outcome != IssueBuildOutcome.Warning);
 
-							IssueBuildData LastBuild = Builds.OrderByDescending(x => x.Change).FirstOrDefault();
-							if (LastBuild != null && !String.IsNullOrEmpty(LastBuild.ErrorUrl))
+							IssueBuildData lastBuild = builds.OrderByDescending(x => x.Change).FirstOrDefault();
+							if (lastBuild != null && !String.IsNullOrEmpty(lastBuild.ErrorUrl))
 							{
-								NewIssue.BuildUrl = LastBuild.ErrorUrl;
+								newIssue.BuildUrl = lastBuild.ErrorUrl;
 							}
 						}
 					}
 				}
 
 				// Apply any pending updates to this issue list, and update it
-				lock (LockObject)
+				lock (_lockObject)
 				{
-					foreach(IssueUpdateData PendingUpdate in PendingUpdates)
+					foreach(IssueUpdateData pendingUpdate in _pendingUpdates)
 					{
-						ApplyPendingUpdate(NewIssues, PendingUpdate);
+						ApplyPendingUpdate(newIssues, pendingUpdate);
 					}
-					Issues = NewIssues;
+					_issues = newIssues;
 				}
 
 				// Update the main thread
-				if(InitialNumIssues > 0 || Issues.Count > 0)
+				if(initialNumIssues > 0 || _issues.Count > 0)
 				{
 					if(OnIssuesChanged != null)
 					{
@@ -456,14 +456,14 @@ namespace UnrealGameSync
 				}
 
 				// Update the stats
-				LastStatusMessage = String.Format("Last update took {0}ms", Timer.ElapsedMilliseconds);
-				Logger.LogInformation("Done in {Time}ms.", Timer.ElapsedMilliseconds);
+				LastStatusMessage = String.Format("Last update took {0}ms", timer.ElapsedMilliseconds);
+				_logger.LogInformation("Done in {Time}ms.", timer.ElapsedMilliseconds);
 				return true;
 			}
-			catch(Exception Ex)
+			catch(Exception ex)
 			{
-				Logger.LogError(Ex, "Failed with exception.");
-				LastStatusMessage = String.Format("Last update failed: ({0})", Ex.ToString());
+				_logger.LogError(ex, "Failed with exception.");
+				LastStatusMessage = String.Format("Last update failed: ({0})", ex.ToString());
 				return false;
 			}
 		}
diff --git a/Engine/Source/Programs/UnrealGameSync/UnrealGameSync/Jupiter/JupiterMonitor.cs b/Engine/Source/Programs/UnrealGameSync/UnrealGameSync/Jupiter/JupiterMonitor.cs
--- a/Engine/Source/Programs/UnrealGameSync/UnrealGameSync/Jupiter/JupiterMonitor.cs
+++ b/Engine/Source/Programs/UnrealGameSync/UnrealGameSync/Jupiter/JupiterMonitor.cs
@@ -20,161 +20,161 @@ namespace UnrealGameSync
 {
 	class JupiterMonitor: IArchiveInfoSource, IDisposable
 	{
-		private readonly object ArchivesLock = new object();
-		private IReadOnlyList<IArchiveInfo> Archives;
+		private readonly object _archivesLock = new object();
+		private IReadOnlyList<IArchiveInfo> _archives;
 
-		private readonly OIDCTokenManager TokenManager;
-		private readonly string JupiterNamespace;
-		private readonly string ProviderIdentifier;
-		private readonly string ExpectedBranch;
-		private readonly Uri JupiterUrl;
+		private readonly OidcTokenManager _tokenManager;
+		private readonly string _jupiterNamespace;
+		private readonly string _providerIdentifier;
+		private readonly string _expectedBranch;
+		private readonly Uri _jupiterUrl;
 
-		private readonly Timer UpdateTimer;
-		private readonly ILogger Logger;
+		private readonly Timer _updateTimer;
+		private readonly ILogger _logger;
 
 		public IReadOnlyList<IArchiveInfo> AvailableArchives
 		{
 			get
 			{
-				lock (ArchivesLock)
+				lock (_archivesLock)
 				{
-					return Archives;
+					return _archives;
 				}
 			}
 		}
 
-		private JupiterMonitor(OIDCTokenManager InTokenManager, ILogger InLogger, string InNamespace, string InUrl,
-			string InProviderIdentifier, string InExpectedBranch)
+		private JupiterMonitor(OidcTokenManager inTokenManager, ILogger inLogger, string inNamespace, string inUrl,
+			string inProviderIdentifier, string inExpectedBranch)
 		{
-			TokenManager = InTokenManager;
-			JupiterNamespace = InNamespace;
-			ProviderIdentifier = InProviderIdentifier;
-			ExpectedBranch = InExpectedBranch;
-			JupiterUrl = new Uri(InUrl);
-
-			Logger = InLogger;
-			UpdateTimer = new Timer(DoUpdate, null, TimeSpan.Zero, TimeSpan.FromMinutes(5));
+			_tokenManager = inTokenManager;
+			_jupiterNamespace = inNamespace;
+			_providerIdentifier = inProviderIdentifier;
+			_expectedBranch = inExpectedBranch;
+			_jupiterUrl = new Uri(inUrl);
+
+			_logger = inLogger;
+			_updateTimer = new Timer(DoUpdate, null, TimeSpan.Zero, TimeSpan.FromMinutes(5));
 		}
 		public void Dispose()
 		{
-			UpdateTimer?.Dispose();
+			_updateTimer?.Dispose();
 		}
 
-		private async void DoUpdate(object State)
+		private async void DoUpdate(object state)
 		{
-			Logger.LogInformation("Starting poll of JupiterMonitor for namespace {JupiterNamespace}", JupiterNamespace);
+			_logger.LogInformation("Starting poll of JupiterMonitor for namespace {JupiterNamespace}", _jupiterNamespace);
 			try
 			{
-				IReadOnlyList<IArchiveInfo> NewArchives = await GetAvailableArchives();
-				lock (ArchivesLock)
+				IReadOnlyList<IArchiveInfo> newArchives = await GetAvailableArchives();
+				lock (_archivesLock)
 				{
-					Archives = NewArchives;
+					_archives = newArchives;
 				}
 
 			}
-			catch (Exception Exception)
+			catch (Exception exception)
 			{
-				Logger.LogError(Exception, "Exception occured during poll!");
+				_logger.LogError(exception, "Exception occured during poll!");
 			}
 		}
 
-		public static JupiterMonitor CreateFromConfigFile(OIDCTokenManager TokenManager, ILogger<JupiterMonitor> Logger, ConfigFile ConfigFile, string SelectedProjectIdentifier)
+		public static JupiterMonitor CreateFromConfigFile(OidcTokenManager tokenManager, ILogger<JupiterMonitor> logger, ConfigFile configFile, string selectedProjectIdentifier)
 		{
-			ConfigSection JupiterConfigSection = ConfigFile.FindSection("Jupiter");
-			if (JupiterConfigSection == null)
+			ConfigSection jupiterConfigSection = configFile.FindSection("Jupiter");
+			if (jupiterConfigSection == null)
 				return null;
 
-			string JupiterUrl = JupiterConfigSection.GetValue("JupiterUrl");
-			string OIDCProviderIdentifier = JupiterConfigSection.GetValue("OIDCProviderIdentifier");
+			string jupiterUrl = jupiterConfigSection.GetValue("JupiterUrl");
+			string oidcProviderIdentifier = jupiterConfigSection.GetValue("OIDCProviderIdentifier");
 
-			ConfigSection ProjectConfigSection = ConfigFile.FindSection(SelectedProjectIdentifier);
-			if (ProjectConfigSection == null)
+			ConfigSection projectConfigSection = configFile.FindSection(selectedProjectIdentifier);
+			if (projectConfigSection == null)
 				return null;
 
-			string JupiterNamespace = ProjectConfigSection.GetValue("JupiterNamespace");
+			string jupiterNamespace = projectConfigSection.GetValue("JupiterNamespace");
 			// Is no namespace has been specified we are unable to fetch builds
-			if (JupiterNamespace == null)
+			if (jupiterNamespace == null)
 				return null;
 
-			string ExpectedBranch = ProjectConfigSection.GetValue("ExpectedBranch");
+			string expectedBranch = projectConfigSection.GetValue("ExpectedBranch");
 			// If we do not know which branch to fetch we can not list builds, as it would risk getting binaries from other branches
-			if (ExpectedBranch == null)
+			if (expectedBranch == null)
 				return null;
 
-			OIDCProviderIdentifier = ProjectConfigSection.GetValue("OIDCProviderIdentifier") ?? OIDCProviderIdentifier;
+			oidcProviderIdentifier = projectConfigSection.GetValue("OIDCProviderIdentifier") ?? oidcProviderIdentifier;
 
 			// with no oidc provider we are unable to login, thus it is required
-			if (OIDCProviderIdentifier == null)
+			if (oidcProviderIdentifier == null)
 				return null;
 
 			// project specific overrides
-			JupiterUrl = ProjectConfigSection.GetValue("JupiterUrl") ?? JupiterUrl;
+			jupiterUrl = projectConfigSection.GetValue("JupiterUrl") ?? jupiterUrl;
 
-			return new JupiterMonitor(TokenManager, Logger, JupiterNamespace, JupiterUrl, OIDCProviderIdentifier, ExpectedBranch);
+			return new JupiterMonitor(tokenManager, logger, jupiterNamespace, jupiterUrl, oidcProviderIdentifier, expectedBranch);
 		}
 
 		private async Task<IReadOnlyList<IArchiveInfo>> GetAvailableArchives()
 		{
-			string Token = await TokenManager.GetAccessToken(ProviderIdentifier);
+			string token = await _tokenManager.GetAccessToken(_providerIdentifier);
 
-			List<JupiterArchiveInfo> NewArchives = new List<JupiterArchiveInfo>();
-			using (HttpClient Client = new HttpClient())
+			List<JupiterArchiveInfo> newArchives = new List<JupiterArchiveInfo>();
+			using (HttpClient client = new HttpClient())
 			{
-				Client.BaseAddress = JupiterUrl;
-				Client.SetBearerToken(Token);
+				client.BaseAddress = _jupiterUrl;
+				client.SetBearerToken(token);
 
-				string ResponseBody = await Client.GetStringAsync($"/api/v1/c/tree-root/{JupiterNamespace}");
-				TreeRootListResponse Response = JsonSerializer.Deserialize<TreeRootListResponse>(ResponseBody, new JsonSerializerOptions {PropertyNameCaseInsensitive = true});
+				string responseBody = await client.GetStringAsync($"/api/v1/c/tree-root/{_jupiterNamespace}");
+				TreeRootListResponse response = JsonSerializer.Deserialize<TreeRootListResponse>(responseBody, new JsonSerializerOptions {PropertyNameCaseInsensitive = true});
 
-				foreach (string TreeRoot in Response.TreeRoots)
+				foreach (string treeRoot in response.TreeRoots)
 				{
 					// fetch info on this tree root
-					string ResourceUrl = $"/api/v1/c/tree-root/{JupiterNamespace}/{TreeRoot}";
-					string ResponseBodyTreeReference = await Client.GetStringAsync(ResourceUrl);
-					TreeRootReferenceResponse ResponseTreeReference = JsonSerializer.Deserialize<TreeRootReferenceResponse>(ResponseBodyTreeReference, new JsonSerializerOptions { PropertyNameCaseInsensitive = true });
+					string resourceUrl = $"/api/v1/c/tree-root/{_jupiterNamespace}/{treeRoot}";
+					string responseBodyTreeReference = await client.GetStringAsync(resourceUrl);
+					TreeRootReferenceResponse responseTreeReference = JsonSerializer.Deserialize<TreeRootReferenceResponse>(responseBodyTreeReference, new JsonSerializerOptions { PropertyNameCaseInsensitive = true });
 
-					Dictionary<string, string> Metadata = ResponseTreeReference.Metadata.ToDictionary(Pair => Pair.Key, Pair => Pair.Value.ToString());
+					Dictionary<string, string> metadata = responseTreeReference.Metadata.ToDictionary(pair => pair.Key, pair => pair.Value.ToString());
 
-					string ArchiveType, Project, Branch, ChangelistString;
+					string archiveType, project, branch, changelistString;
 
-					if (!Metadata.TryGetValue("ArchiveType", out ArchiveType) ||
-					    !Metadata.TryGetValue("Project", out Project) ||
-					    !Metadata.TryGetValue("Branch", out Branch) ||
-						!Metadata.TryGetValue("Changelist", out ChangelistString))
+					if (!metadata.TryGetValue("ArchiveType", out archiveType) ||
+					    !metadata.TryGetValue("Project", out project) ||
+					    !metadata.TryGetValue("Branch", out branch) ||
+						!metadata.TryGetValue("Changelist", out changelistString))
 					{
 						continue;
 					}
 
 					// skip if we do not have metadata as we need that to be able to determine if this root is of the type we want
-					if (Branch == null || Project == null || ArchiveType == null || ChangelistString == null)
+					if (branch == null || project == null || archiveType == null || changelistString == null)
 						continue;
 
-					if (!string.Equals(ExpectedBranch, Branch, StringComparison.InvariantCultureIgnoreCase))
+					if (!string.Equals(_expectedBranch, branch, StringComparison.InvariantCultureIgnoreCase))
 					{
 						continue;
 					}
 
-					int Changelist;
-					if (!int.TryParse(ChangelistString, out Changelist))
+					int changelist;
+					if (!int.TryParse(changelistString, out changelist))
 					{
 						continue; // invalid changelist format
 					}
 
-					JupiterArchiveInfo ExistingArchive = NewArchives.FirstOrDefault(Info =>
-						string.Equals(Info.Name, Project, StringComparison.InvariantCultureIgnoreCase) &&
-						string.Equals(Info.Type, ArchiveType, StringComparison.InvariantCultureIgnoreCase));
-					if (ExistingArchive == null)
+					JupiterArchiveInfo existingArchive = newArchives.FirstOrDefault(info =>
+						string.Equals(info.Name, project, StringComparison.InvariantCultureIgnoreCase) &&
+						string.Equals(info.Type, archiveType, StringComparison.InvariantCultureIgnoreCase));
+					if (existingArchive == null)
 					{
-						JupiterArchiveInfo Archive = new JupiterArchiveInfo(Project, ArchiveType, JupiterUrl.ToString(), JupiterNamespace);
-						NewArchives.Add(Archive);
-						ExistingArchive = Archive;
+						JupiterArchiveInfo archive = new JupiterArchiveInfo(project, archiveType, _jupiterUrl.ToString(), _jupiterNamespace);
+						newArchives.Add(archive);
+						existingArchive = archive;
 					}
 
-					ExistingArchive.AddArchiveVersion(Changelist, ResponseTreeReference.TreeReferenceKey);
+					existingArchive.AddArchiveVersion(changelist, responseTreeReference.TreeReferenceKey);
 				}
 			}
 
-			return NewArchives.AsReadOnly();
+			return newArchives.AsReadOnly();
 		}
 
 		private class TreeRootListResponse
@@ -194,17 +194,17 @@ namespace UnrealGameSync
 
 		private class JupiterArchiveInfo : IArchiveInfo
 		{
-			private readonly string JupiterUrl;
-			private readonly string JupiterNamespace;
-			private readonly Dictionary<int, string> ChangeToKey = new Dictionary<int, string>();
+			private readonly string _jupiterUrl;
+			private readonly string _jupiterNamespace;
+			private readonly Dictionary<int, string> _changeToKey = new Dictionary<int, string>();
 
-			public JupiterArchiveInfo(string InName, string InType, string InJupiterUrl, string InJupiterNamespace)
+			public JupiterArchiveInfo(string inName, string inType, string inJupiterUrl, string inJupiterNamespace)
 			{
-				Name = InName;
-				Type = InType;
+				Name = inName;
+				Type = inType;
 
-				JupiterUrl = InJupiterUrl;
-				JupiterNamespace = InJupiterNamespace;
+				_jupiterUrl = inJupiterUrl;
+				_jupiterNamespace = inJupiterNamespace;
 			}
 
 			public string Name { get; }
@@ -219,61 +219,61 @@ namespace UnrealGameSync
 
 			public bool Exists()
 			{
-				return ChangeToKey.Count > 0;
+				return _changeToKey.Count > 0;
 			}
 
-			public bool TryGetArchiveKeyForChangeNumber(int ChangeNumber, out string ArchiveKey)
+			public bool TryGetArchiveKeyForChangeNumber(int changeNumber, out string archiveKey)
 			{
-				return ChangeToKey.TryGetValue(ChangeNumber, out ArchiveKey);
+				return _changeToKey.TryGetValue(changeNumber, out archiveKey);
 			}
 
-			public Task<bool> DownloadArchive(IPerforceConnection _, string ArchiveKey, DirectoryReference LocalRootPath, FileReference ManifestFileName, ILogger Logger, ProgressValue Progress, CancellationToken CancellationToken)
+			public Task<bool> DownloadArchive(IPerforceConnection _, string archiveKey, DirectoryReference localRootPath, FileReference manifestFileName, ILogger logger, ProgressValue progress, CancellationToken cancellationToken)
 			{
 				try
 				{
-					Progress<Tuple<float, FileReference>> ProgressCallback = new Progress<Tuple<float, FileReference>>(ProgressUpdate =>
+					Progress<Tuple<float, FileReference>> progressCallback = new Progress<Tuple<float, FileReference>>(progressUpdate =>
 					{
-						(float ProgressFraction, FileReference FileReference) = ProgressUpdate;
-						Progress.Set(ProgressFraction);
-						Logger.LogInformation("Writing {FileName}", FileReference.FullName);
+						(float progressFraction, FileReference fileReference) = progressUpdate;
+						progress.Set(progressFraction);
+						logger.LogInformation("Writing {FileName}", fileReference.FullName);
 					});
 
 					// place the manifest for the Jupiter download next to the UGS manifest
-					FileReference UGSManifestFileReference = ManifestFileName;
-					FileReference JupiterManifestFileReference = FileReference.Combine(UGSManifestFileReference.Directory, "Jupiter-Manifest.json");
+					FileReference ugsManifestFileReference = manifestFileName;
+					FileReference jupiterManifestFileReference = FileReference.Combine(ugsManifestFileReference.Directory, "Jupiter-Manifest.json");
 
-					DirectoryReference RootDirectory = LocalRootPath;
-					JupiterFileTree FileTree = new JupiterFileTree(RootDirectory, InDeferReadingFiles: true);
-					Task<List<FileReference>> DownloadTask = FileTree.DownloadFromJupiter(JupiterManifestFileReference, JupiterUrl, JupiterNamespace, ArchiveKey, ProgressCallback);
-					DownloadTask.Wait();
+					DirectoryReference rootDirectory = localRootPath;
+					JupiterFileTree fileTree = new JupiterFileTree(rootDirectory, InDeferReadingFiles: true);
+					Task<List<FileReference>> downloadTask = fileTree.DownloadFromJupiter(jupiterManifestFileReference, _jupiterUrl, _jupiterNamespace, archiveKey, progressCallback);
+					downloadTask.Wait();
 
-					List<FileReference> WrittenFiles = DownloadTask.Result;
-					ArchiveManifest ArchiveManifest = new ArchiveManifest();
-					foreach (FileReference File in WrittenFiles)
+					List<FileReference> writtenFiles = downloadTask.Result;
+					ArchiveManifest archiveManifest = new ArchiveManifest();
+					foreach (FileReference file in writtenFiles)
 					{
-						ArchiveManifest.Files.Add(new ArchiveManifestFile(File.FullName, File.ToFileInfo().Length, DateTime.Now));
+						archiveManifest.Files.Add(new ArchiveManifestFile(file.FullName, file.ToFileInfo().Length, DateTime.Now));
 					}
 
 					// Write it out to a temporary file, then move it into place
-					FileReference TempManifestFileName = ManifestFileName + ".tmp";
-					using (FileStream OutputStream = FileReference.Open(TempManifestFileName, FileMode.Create, FileAccess.Write))
+					FileReference tempManifestFileName = manifestFileName + ".tmp";
+					using (FileStream outputStream = FileReference.Open(tempManifestFileName, FileMode.Create, FileAccess.Write))
 					{
-						ArchiveManifest.Write(OutputStream);
+						archiveManifest.Write(outputStream);
 					}
-					FileReference.Move(TempManifestFileName, ManifestFileName);
+					FileReference.Move(tempManifestFileName, manifestFileName);
 
 					return Task.FromResult(true);
 				}
-				catch (Exception Exception)
+				catch (Exception exception)
 				{
-					Logger.LogError(Exception, "Exception occured when downloading build from Jupiter with key {Key}.", ArchiveKey);
+					logger.LogError(exception, "Exception occured when downloading build from Jupiter with key {Key}.", archiveKey);
 					return Task.FromResult(false);
 				}
 			}
 
-			public void AddArchiveVersion(int Changelist, string Key)
+			public void AddArchiveVersion(int changelist, string key)
 			{
-				ChangeToKey[Changelist] = Key;
+				_changeToKey[changelist] = key;
 			}
 		}
 	}
