UGS: Reformat to standard C# style conventions.

diff --git a/Engine/Source/Programs/UnrealGameSync/UnrealGameSync/Logging.cs b/Engine/Source/Programs/UnrealGameSync/UnrealGameSync/Logging.cs
--- a/Engine/Source/Programs/UnrealGameSync/UnrealGameSync/Logging.cs
+++ b/Engine/Source/Programs/UnrealGameSync/UnrealGameSync/Logging.cs
@@ -12,14 +12,14 @@ namespace UnrealGameSync
 {
 	public static class Logging
 	{
-		public static ILoggerProvider CreateLoggerProvider(FileReference File)
+		public static ILoggerProvider CreateLoggerProvider(FileReference file)
 		{
-			Serilog.ILogger Logger = new LoggerConfiguration()
+			Serilog.ILogger logger = new LoggerConfiguration()
 				.Enrich.FromLogContext()
-				.WriteTo.File(new ExpressionTemplate("[{@t:HH:mm:ss.fff} {@l:w3}] {Concat('[', Substring(SourceContext, LastIndexOf(SourceContext, '.') + 1, 14), ']'),16} {Indent}{@m}{@x}\n"), File.FullName, rollingInterval: RollingInterval.Day, rollOnFileSizeLimit: true, fileSizeLimitBytes: 20 * 1024 * 1024, retainedFileCountLimit: 10)
+				.WriteTo.File(new ExpressionTemplate("[{@t:HH:mm:ss.fff} {@l:w3}] {Concat('[', Substring(SourceContext, LastIndexOf(SourceContext, '.') + 1, 14), ']'),16} {Indent}{@m}{@x}\n"), file.FullName, rollingInterval: RollingInterval.Day, rollOnFileSizeLimit: true, fileSizeLimitBytes: 20 * 1024 * 1024, retainedFileCountLimit: 10)
 				.CreateLogger();
 
-			return new Serilog.Extensions.Logging.SerilogLoggerProvider(Logger, true);
+			return new Serilog.Extensions.Logging.SerilogLoggerProvider(logger, true);
 		}
 	}
 }
diff --git a/Engine/Source/Programs/UnrealGameSync/UnrealGameSync/ModalTask.cs b/Engine/Source/Programs/UnrealGameSync/UnrealGameSync/ModalTask.cs
--- a/Engine/Source/Programs/UnrealGameSync/UnrealGameSync/ModalTask.cs
+++ b/Engine/Source/Programs/UnrealGameSync/UnrealGameSync/ModalTask.cs
@@ -26,9 +26,9 @@ namespace UnrealGameSync
 
 		public Exception? Exception => Task.Exception?.InnerException;
 
-		public ModalTask(Task Task)
+		public ModalTask(Task task)
 		{
-			this.Task = Task;
+			this.Task = task;
 		}
 
 		public string Error
@@ -40,52 +40,52 @@ namespace UnrealGameSync
 					return "Succeeded.";
 				}
 
-				Exception? Ex = Exception;
-				if (Ex == null)
+				Exception? ex = Exception;
+				if (ex == null)
 				{
 					return "Failed.";
 				}
-				else if (Ex is UserErrorException UserEx)
+				else if (ex is UserErrorException userEx)
 				{
-					return UserEx.Message;
+					return userEx.Message;
 				}
 				else
 				{
-					return $"Unhandled exception ({Ex.Message})";
+					return $"Unhandled exception ({ex.Message})";
 				}
 			}
 		}
 
-		public static ModalTask? Execute(IWin32Window? Owner, string Title, string Message, Func<CancellationToken, Task> TaskFunc, ModalTaskFlags Flags = ModalTaskFlags.None)
+		public static ModalTask? Execute(IWin32Window? owner, string title, string message, Func<CancellationToken, Task> taskFunc, ModalTaskFlags flags = ModalTaskFlags.None)
 		{
-			Func<CancellationToken, Task<int>> TypedTaskFunc = async x => { await TaskFunc(x); return 0; };
-			return Execute(Owner, Title, Message, TypedTaskFunc, Flags);
+			Func<CancellationToken, Task<int>> typedTaskFunc = async x => { await taskFunc(x); return 0; };
+			return Execute(owner, title, message, typedTaskFunc, flags);
 		}
 
-		public static ModalTask<T>? Execute<T>(IWin32Window? Owner, string Title, string Message, Func<CancellationToken, Task<T>> TaskFunc, ModalTaskFlags Flags = ModalTaskFlags.None)
+		public static ModalTask<T>? Execute<T>(IWin32Window? owner, string title, string message, Func<CancellationToken, Task<T>> taskFunc, ModalTaskFlags flags = ModalTaskFlags.None)
 		{
-			using (CancellationTokenSource CancellationSource = new CancellationTokenSource())
+			using (CancellationTokenSource cancellationSource = new CancellationTokenSource())
 			{
-				Task<T> BackgroundTask = Task.Run(() => TaskFunc(CancellationSource.Token));
+				Task<T> backgroundTask = Task.Run(() => taskFunc(cancellationSource.Token));
 
-				ModalTaskWindow Window = new ModalTaskWindow(Title, Message, (Owner == null) ? FormStartPosition.CenterScreen : FormStartPosition.CenterParent, BackgroundTask, CancellationSource);
-				if(Owner == null)
+				ModalTaskWindow window = new ModalTaskWindow(title, message, (owner == null) ? FormStartPosition.CenterScreen : FormStartPosition.CenterParent, backgroundTask, cancellationSource);
+				if(owner == null)
 				{
-					Window.ShowInTaskbar = true;
+					window.ShowInTaskbar = true;
 				}
-				Window.ShowDialog(Owner);
+				window.ShowDialog(owner);
 
-				if (BackgroundTask.IsCanceled || (BackgroundTask.Exception != null && BackgroundTask.Exception.InnerException is OperationCanceledException))
+				if (backgroundTask.IsCanceled || (backgroundTask.Exception != null && backgroundTask.Exception.InnerException is OperationCanceledException))
 				{
 					return null;
 				}
 
-				ModalTask<T> Result = new ModalTask<T>(BackgroundTask);
-				if (Result.Failed && (Flags & ModalTaskFlags.Quiet) == 0)
+				ModalTask<T> result = new ModalTask<T>(backgroundTask);
+				if (result.Failed && (flags & ModalTaskFlags.Quiet) == 0)
 				{
-					MessageBox.Show(Owner, Result.Error, Title);
+					MessageBox.Show(owner, result.Error, title);
 				}
-				return Result;
+				return result;
 			}
 		}
 	}
@@ -96,7 +96,7 @@ namespace UnrealGameSync
 
 		public T Result => Task.Result;
 
-		public ModalTask(Task<T> Task) : base(Task)
+		public ModalTask(Task<T> task) : base(task)
 		{
 		}
 	}
diff --git a/Engine/Source/Programs/UnrealGameSync/UnrealGameSync/OIDCTokenManager.cs b/Engine/Source/Programs/UnrealGameSync/UnrealGameSync/OIDCTokenManager.cs
--- a/Engine/Source/Programs/UnrealGameSync/UnrealGameSync/OIDCTokenManager.cs
+++ b/Engine/Source/Programs/UnrealGameSync/UnrealGameSync/OIDCTokenManager.cs
@@ -19,27 +19,27 @@ using IdentityModel.OidcClient.Results;
 
 namespace UnrealGameSync
 {
-	public class OIDCTokenManager
+	public class OidcTokenManager
 	{
-		private UserSettings Settings;
-		private readonly Dictionary<string, OIDCTokenClient> TokenClients = new Dictionary<string, OIDCTokenClient>();
+		private UserSettings _settings;
+		private readonly Dictionary<string, OidcTokenClient> _tokenClients = new Dictionary<string, OidcTokenClient>();
 
 		public Dictionary<string, ProviderInfo> Providers { get; }
 
-		private OIDCTokenManager(UserSettings InSettings, Dictionary<string, ProviderInfo> InProviders)
+		private OidcTokenManager(UserSettings inSettings, Dictionary<string, ProviderInfo> inProviders)
 		{
-			Settings = InSettings;
-			Providers = InProviders;
+			_settings = inSettings;
+			Providers = inProviders;
 
-			Dictionary<string, string> RefreshTokens = new Dictionary<string, string>();
-			foreach (KeyValuePair<string, string> Pair in Settings.ProviderToRefreshTokens)
+			Dictionary<string, string> refreshTokens = new Dictionary<string, string>();
+			foreach (KeyValuePair<string, string> pair in _settings.ProviderToRefreshTokens)
 			{
-				string ProviderIdentifier = Pair.Key;
+				string providerIdentifier = pair.Key;
 				try
 				{
-					byte[] KeyBytes = Convert.FromBase64String(Pair.Value);
-					byte[] UnencryptedBytes = ProtectedData.Unprotect(KeyBytes, null, DataProtectionScope.CurrentUser);
-					RefreshTokens.TryAdd(ProviderIdentifier, Encoding.ASCII.GetString(UnencryptedBytes));
+					byte[] keyBytes = Convert.FromBase64String(pair.Value);
+					byte[] unencryptedBytes = ProtectedData.Unprotect(keyBytes, null, DataProtectionScope.CurrentUser);
+					refreshTokens.TryAdd(providerIdentifier, Encoding.ASCII.GetString(unencryptedBytes));
 				}
 				catch (Exception)
 				{
@@ -47,93 +47,93 @@ namespace UnrealGameSync
 				}
 			}
 
-			foreach (KeyValuePair<string, ProviderInfo> ProviderPair in InProviders)
+			foreach (KeyValuePair<string, ProviderInfo> providerPair in inProviders)
 			{
-				ProviderInfo Provider = ProviderPair.Value;
+				ProviderInfo provider = providerPair.Value;
 
-				OIDCTokenClient TokenClient = new OIDCTokenClient(new Uri(Provider.ServerUri), Provider.ClientId, Provider.RedirectUri);
+				OidcTokenClient tokenClient = new OidcTokenClient(new Uri(provider.ServerUri), provider.ClientId, provider.RedirectUri);
 
-				string RefreshToken;
-				if (RefreshTokens.TryGetValue(Provider.Identifier, out RefreshToken))
+				string refreshToken;
+				if (refreshTokens.TryGetValue(provider.Identifier, out refreshToken))
 				{
-					TokenClient.SetRefreshToken(RefreshToken);
+					tokenClient.SetRefreshToken(refreshToken);
 				}
-				TokenClients.Add(Provider.Identifier, TokenClient);
+				_tokenClients.Add(provider.Identifier, tokenClient);
 			}
 		}
 
 		public bool HasUnfinishedLogin()
 		{
-			return TokenClients.Any(Pair => Pair.Value.GetStatusForProvider() == OIDCStatus.NotLoggedIn);
+			return _tokenClients.Any(pair => pair.Value.GetStatusForProvider() == OidcStatus.NotLoggedIn);
 		}
 
-		internal static OIDCTokenManager CreateFromConfigFile(UserSettings Settings, List<OpenProjectInfo> ConfigFiles)
+		internal static OidcTokenManager CreateFromConfigFile(UserSettings settings, List<OpenProjectInfo> configFiles)
 		{
 			// join the provider configuration from all projects
-			Dictionary<string, ProviderInfo> Providers = new Dictionary<string, ProviderInfo>();
-			foreach (OpenProjectInfo DetectProjectSettingsTask in ConfigFiles)
+			Dictionary<string, ProviderInfo> providers = new Dictionary<string, ProviderInfo>();
+			foreach (OpenProjectInfo detectProjectSettingsTask in configFiles)
 			{
-				if(DetectProjectSettingsTask == null)
+				if(detectProjectSettingsTask == null)
 				{
 					continue;
 				}
 
-				ConfigFile ConfigFile = DetectProjectSettingsTask.LatestProjectConfigFile;
-				if(ConfigFile == null)
+				ConfigFile configFile = detectProjectSettingsTask.LatestProjectConfigFile;
+				if(configFile == null)
 				{
 					continue;
 				}
 
-				ConfigSection ProviderSection = ConfigFile.FindSection("OIDCProvider");
-				if (ProviderSection == null)
+				ConfigSection providerSection = configFile.FindSection("OIDCProvider");
+				if (providerSection == null)
 				{
 					continue;
 				}
 
-				string[] ProviderValues = ProviderSection.GetValues("Provider", (string[]) null);
-				foreach (ConfigObject Provider in ProviderValues.Select(s => new ConfigObject(s)).ToList())
+				string[] providerValues = providerSection.GetValues("Provider", (string[]) null);
+				foreach (ConfigObject provider in providerValues.Select(s => new ConfigObject(s)).ToList())
 				{
-					string Identifier = Provider.GetValue("Identifier");
-					string ServerUri = Provider.GetValue("ServerUri");
-					string ClientId = Provider.GetValue("ClientId");
-					string DisplayName = Provider.GetValue("DisplayName");
-					string RedirectUri = Provider.GetValue("RedirectUri");
+					string identifier = provider.GetValue("Identifier");
+					string serverUri = provider.GetValue("ServerUri");
+					string clientId = provider.GetValue("ClientId");
+					string displayName = provider.GetValue("DisplayName");
+					string redirectUri = provider.GetValue("RedirectUri");
 
 					// we might get a provider with the same identifier from another project, in which case we only keep the first one
-					Providers.TryAdd(Identifier, new ProviderInfo(Identifier, ServerUri, ClientId, DisplayName, RedirectUri));
+					providers.TryAdd(identifier, new ProviderInfo(identifier, serverUri, clientId, displayName, redirectUri));
 				}
 			}
 
-			if (Providers.Count == 0)
+			if (providers.Count == 0)
 				return null;
-			return new OIDCTokenManager(Settings, Providers);
+			return new OidcTokenManager(settings, providers);
 		}
 
-		public async Task Login(string ProviderIdentifier)
+		public async Task Login(string providerIdentifier)
 		{
-			OIDCTokenClient TokenClient = TokenClients[ProviderIdentifier];
+			OidcTokenClient tokenClient = _tokenClients[providerIdentifier];
 
-			string RefreshToken = await TokenClient.Login();
-			if (!string.IsNullOrEmpty(RefreshToken))
+			string refreshToken = await tokenClient.Login();
+			if (!string.IsNullOrEmpty(refreshToken))
 			{
 				// if we got a refresh token we store that for future use
-				byte[] RefreshTokenBytes = Encoding.ASCII.GetBytes(RefreshToken);
-				byte[] EncryptedBytes = ProtectedData.Protect(RefreshTokenBytes, null, DataProtectionScope.CurrentUser);
+				byte[] refreshTokenBytes = Encoding.ASCII.GetBytes(refreshToken);
+				byte[] encryptedBytes = ProtectedData.Protect(refreshTokenBytes, null, DataProtectionScope.CurrentUser);
 
-				Settings.ProviderToRefreshTokens[ProviderIdentifier] = Convert.ToBase64String(EncryptedBytes);
+				_settings.ProviderToRefreshTokens[providerIdentifier] = Convert.ToBase64String(encryptedBytes);
 			}
 
 
 		}
 
-		public Task<string> GetAccessToken(string ProviderIdentifier)
+		public Task<string> GetAccessToken(string providerIdentifier)
 		{
-			return TokenClients[ProviderIdentifier].GetAccessToken();
+			return _tokenClients[providerIdentifier].GetAccessToken();
 		}
 
-		public OIDCStatus GetStatusForProvider(string ProviderIdentifier)
+		public OidcStatus GetStatusForProvider(string providerIdentifier)
 		{
-			return TokenClients[ProviderIdentifier].GetStatusForProvider();
+			return _tokenClients[providerIdentifier].GetStatusForProvider();
 		}
 
 		public class ProviderInfo
@@ -144,83 +144,83 @@ namespace UnrealGameSync
 			public string DisplayName { get; }
 			public string RedirectUri { get; }
 
-			public ProviderInfo(string InIdentifier, string InServerUri, string InClientId, string InDisplayName, string InRedirectUri)
+			public ProviderInfo(string inIdentifier, string inServerUri, string inClientId, string inDisplayName, string inRedirectUri)
 			{
-				Identifier = InIdentifier; 
-				ServerUri = InServerUri;
-				ClientId = InClientId;
-				DisplayName = InDisplayName;
-				RedirectUri = InRedirectUri;
+				Identifier = inIdentifier; 
+				ServerUri = inServerUri;
+				ClientId = inClientId;
+				DisplayName = inDisplayName;
+				RedirectUri = inRedirectUri;
 			}
 		}
 	}
 
 
-	public enum OIDCStatus
+	public enum OidcStatus
 	{
 		Connected,
 		NotLoggedIn,
 		TokenRefreshRequired
 	}
 
-	class OIDCTokenClient
+	class OidcTokenClient
 	{
-		private readonly Uri AuthorityUri;
-		private readonly string ClientId;
-		private readonly string RedirectUri;
-		private readonly string Scopes;
+		private readonly Uri _authorityUri;
+		private readonly string _clientId;
+		private readonly string _redirectUri;
+		private readonly string _scopes;
 
-		private OidcClient OidcClient;
+		private OidcClient _oidcClient;
 
-		private string RefreshToken;
-		private string AccessToken;
-		private DateTime TokenExpiry;
+		private string _refreshToken;
+		private string _accessToken;
+		private DateTime _tokenExpiry;
 
-		private bool Initialized = false;
+		private bool _initialized = false;
 
 
-		public OIDCTokenClient(Uri InAuthorityUri, string InClientId, string InRedirectUri, string InScopes = "openid profile offline_access")
+		public OidcTokenClient(Uri inAuthorityUri, string inClientId, string inRedirectUri, string inScopes = "openid profile offline_access")
 		{
-			AuthorityUri = InAuthorityUri;
-			ClientId = InClientId;
-			RedirectUri = InRedirectUri;
-			Scopes = InScopes;
+			_authorityUri = inAuthorityUri;
+			_clientId = inClientId;
+			_redirectUri = inRedirectUri;
+			_scopes = inScopes;
 		}
 
 		private async Task Initialize()
 		{
-			if (Initialized)
+			if (_initialized)
 				return;
 
-			OidcClientOptions Options = new OidcClientOptions
+			OidcClientOptions options = new OidcClientOptions
 			{
-				Authority = AuthorityUri.ToString(),
-				Policy = new Policy { Discovery = new DiscoveryPolicy { Authority = AuthorityUri.ToString() } },
-				ClientId = ClientId,
-				Scope = Scopes,
+				Authority = _authorityUri.ToString(),
+				Policy = new Policy { Discovery = new DiscoveryPolicy { Authority = _authorityUri.ToString() } },
+				ClientId = _clientId,
+				Scope = _scopes,
 				FilterClaims = false,
-				RedirectUri = RedirectUri,
+				RedirectUri = _redirectUri,
 				Flow = OidcClientOptions.AuthenticationFlow.AuthorizationCode,
 				ResponseMode = OidcClientOptions.AuthorizeResponseMode.Redirect
 			};
 
 			// we need to fetch the discovery document ourselves to support OIDC Authorities which have a subresource for it
 			// with Okta has for authorization servers for instance.
-			DiscoveryDocumentResponse DiscoveryDocument = await GetDiscoveryDocument();
-			Options.ProviderInformation = new ProviderInformation
+			DiscoveryDocumentResponse discoveryDocument = await GetDiscoveryDocument();
+			options.ProviderInformation = new ProviderInformation
 			{
-				IssuerName = DiscoveryDocument.Issuer,
-				KeySet = DiscoveryDocument.KeySet,
-
-				AuthorizeEndpoint = DiscoveryDocument.AuthorizeEndpoint,
-				TokenEndpoint = DiscoveryDocument.TokenEndpoint,
-				EndSessionEndpoint = DiscoveryDocument.EndSessionEndpoint,
-				UserInfoEndpoint = DiscoveryDocument.UserInfoEndpoint,
-				TokenEndPointAuthenticationMethods = DiscoveryDocument.TokenEndpointAuthenticationMethodsSupported
+				IssuerName = discoveryDocument.Issuer,
+				KeySet = discoveryDocument.KeySet,
+
+				AuthorizeEndpoint = discoveryDocument.AuthorizeEndpoint,
+				TokenEndpoint = discoveryDocument.TokenEndpoint,
+				EndSessionEndpoint = discoveryDocument.EndSessionEndpoint,
+				UserInfoEndpoint = discoveryDocument.UserInfoEndpoint,
+				TokenEndPointAuthenticationMethods = discoveryDocument.TokenEndpointAuthenticationMethodsSupported
 			};
 
-			OidcClient = new OidcClient(Options);
-			Initialized = true;
+			_oidcClient = new OidcClient(options);
+			_initialized = true;
 		}
 
 		public async Task<string> Login()
@@ -228,107 +228,107 @@ namespace UnrealGameSync
 			await Initialize();
 
 			// setup a local http server to listen for the result of the login
-			using HttpListener Http = new HttpListener();
-			Uri Uri = new Uri(RedirectUri);
+			using HttpListener http = new HttpListener();
+			Uri uri = new Uri(_redirectUri);
 			// build the url the server should be hosted at
-			string Prefix = $"{Uri.Scheme}{Uri.SchemeDelimiter}{Uri.Authority}/";
-			Http.Prefixes.Add(Prefix);
-			Http.Start();
+			string prefix = $"{uri.Scheme}{Uri.SchemeDelimiter}{uri.Authority}/";
+			http.Prefixes.Add(prefix);
+			http.Start();
 
 			// generate the appropriate codes we need to login
-			AuthorizeState LoginState = await OidcClient.PrepareLoginAsync();
+			AuthorizeState loginState = await _oidcClient.PrepareLoginAsync();
 			// start the user browser
-			OpenBrowser(LoginState.StartUrl);
+			OpenBrowser(loginState.StartUrl);
 
-			string ResponseData = await ProcessHttpRequest(Http);
+			string responseData = await ProcessHttpRequest(http);
 
 			// parse the returned url for the tokens needed to complete the login
-			IdentityModel.OidcClient.LoginResult LoginResult = await OidcClient.ProcessResponseAsync(ResponseData, LoginState);
+			IdentityModel.OidcClient.LoginResult loginResult = await _oidcClient.ProcessResponseAsync(responseData, loginState);
 
-			Http.Stop();
+			http.Stop();
 
-			if (LoginResult.IsError)
-				throw new LoginFailedException("Failed to login due to error: " + LoginResult.Error);
+			if (loginResult.IsError)
+				throw new LoginFailedException("Failed to login due to error: " + loginResult.Error);
 
-			RefreshToken = LoginResult.RefreshToken;
-			AccessToken = LoginResult.AccessToken;
-			TokenExpiry = LoginResult.AccessTokenExpiration;
+			_refreshToken = loginResult.RefreshToken;
+			_accessToken = loginResult.AccessToken;
+			_tokenExpiry = loginResult.AccessTokenExpiration;
 
-			return RefreshToken;
+			return _refreshToken;
 		}
 
-		private async Task<string> ProcessHttpRequest(HttpListener Http)
+		private async Task<string> ProcessHttpRequest(HttpListener http)
 		{
-			HttpListenerContext Context = await Http.GetContextAsync();
-			string ResponseData;
-			if (Context.Request.HttpMethod == "GET")
+			HttpListenerContext context = await http.GetContextAsync();
+			string responseData;
+			if (context.Request.HttpMethod == "GET")
 			{
-				ResponseData = Context.Request.RawUrl;
+				responseData = context.Request.RawUrl;
 
 			}
-			else if (Context.Request.HttpMethod == "POST")
+			else if (context.Request.HttpMethod == "POST")
 			{
-				var Request = Context.Request;
-				if (Request.ContentType != null && !Request.ContentType.Equals("application/x-www-form-urlencoded", StringComparison.OrdinalIgnoreCase))
+				var request = context.Request;
+				if (request.ContentType != null && !request.ContentType.Equals("application/x-www-form-urlencoded", StringComparison.OrdinalIgnoreCase))
 				{
 					// we do not support url encoded return types
-					Context.Response.StatusCode = 415;
+					context.Response.StatusCode = 415;
 					return null;
 				}
 				
 				// attempt to parse the body
 
 				// if there is no body we can not handle the post
-				if (!Context.Request.HasEntityBody)
+				if (!context.Request.HasEntityBody)
 				{
-					Context.Response.StatusCode = 415;
+					context.Response.StatusCode = 415;
 					return null;
 				}
 
-				await using Stream Body = Request.InputStream;
-				using StreamReader Reader = new StreamReader(Body, Request.ContentEncoding);
-				ResponseData = Reader.ReadToEnd();
+				await using Stream body = request.InputStream;
+				using StreamReader reader = new StreamReader(body, request.ContentEncoding);
+				responseData = reader.ReadToEnd();
 			}
 			else
 			{
 				// if we receive any other http method something is very odd. Tell them to use a different method.
-				Context.Response.StatusCode = 415;
+				context.Response.StatusCode = 415;
 				return null;
 			}
 
 			// generate a simple http page to show the user
-			HttpListenerResponse Response = Context.Response;
-			string HttpPage = "<html><head><meta http-equiv='refresh' content='10></head><body>Please close this browser and return to UnrealGameSync.</body></html>";
-			byte[] Buffer = Encoding.UTF8.GetBytes(HttpPage);
-			Response.ContentLength64 = Buffer.Length;
-			Stream ResponseOutput = Response.OutputStream;
-			await ResponseOutput.WriteAsync(Buffer, 0, Buffer.Length);
-			ResponseOutput.Close();
-
-			return ResponseData;
+			HttpListenerResponse response = context.Response;
+			string httpPage = "<html><head><meta http-equiv='refresh' content='10></head><body>Please close this browser and return to UnrealGameSync.</body></html>";
+			byte[] buffer = Encoding.UTF8.GetBytes(httpPage);
+			response.ContentLength64 = buffer.Length;
+			Stream responseOutput = response.OutputStream;
+			await responseOutput.WriteAsync(buffer, 0, buffer.Length);
+			responseOutput.Close();
+
+			return responseData;
 		}
 
-		private void OpenBrowser(string Url)
+		private void OpenBrowser(string url)
 		{
 			try
 			{
-				Process.Start(Url);
+				Process.Start(url);
 			}
 			catch
 			{
 				// hack because of this: [REDACTED_PROFILE_URL]/corefx/issues/10361
 				if (RuntimeInformation.IsOSPlatform(OSPlatform.Windows))
 				{
-					Url = Url.Replace("&", "^&");
-					Process.Start(new ProcessStartInfo("cmd", $"/c start {Url}") { CreateNoWindow = true });
+					url = url.Replace("&", "^&");
+					Process.Start(new ProcessStartInfo("cmd", $"/c start {url}") { CreateNoWindow = true });
 				}
 				else if (RuntimeInformation.IsOSPlatform(OSPlatform.Linux))
 				{
-					Process.Start("xdg-open", Url);
+					Process.Start("xdg-open", url);
 				}
 				else if (RuntimeInformation.IsOSPlatform(OSPlatform.OSX))
 				{
-					Process.Start("open", Url);
+					Process.Start("open", url);
 				}
 				else
 				{
@@ -338,78 +338,78 @@ namespace UnrealGameSync
 		}
 
 
-		private async Task<string> DoRefreshToken(string InRefreshToken)
+		private async Task<string> DoRefreshToken(string inRefreshToken)
 		{
 			await Initialize();
 
 			// use the refresh token to acquire a new access token
-			RefreshTokenResult RefreshTokenResult = await OidcClient.RefreshTokenAsync(InRefreshToken);
+			RefreshTokenResult refreshTokenResult = await _oidcClient.RefreshTokenAsync(inRefreshToken);
 
-			RefreshToken = RefreshTokenResult.RefreshToken;
-			AccessToken = RefreshTokenResult.AccessToken;
-			TokenExpiry = RefreshTokenResult.AccessTokenExpiration;
+			_refreshToken = refreshTokenResult.RefreshToken;
+			_accessToken = refreshTokenResult.AccessToken;
+			_tokenExpiry = refreshTokenResult.AccessTokenExpiration;
 
-			return AccessToken;
+			return _accessToken;
 		}
 
 		private async Task<DiscoveryDocumentResponse> GetDiscoveryDocument()
 		{
-			string DiscoUrl = $"{AuthorityUri}/.well-known/openid-configuration";
+			string discoUrl = $"{_authorityUri}/.well-known/openid-configuration";
 
-			HttpClient Client = new HttpClient();
-			DiscoveryDocumentResponse Disco = await Client.GetDiscoveryDocumentAsync(new DiscoveryDocumentRequest
+			HttpClient client = new HttpClient();
+			DiscoveryDocumentResponse disco = await client.GetDiscoveryDocumentAsync(new DiscoveryDocumentRequest
 			{
-				Address = DiscoUrl,
+				Address = discoUrl,
 				Policy =
 				{
 					ValidateEndpoints = false
 				}
 			});
 
-			if (Disco.IsError) throw new Exception(Disco.Error);
+			if (disco.IsError) throw new Exception(disco.Error);
 
-			return Disco;
+			return disco;
 		}
 
 
 		public async Task<string> GetAccessToken()
 		{
-			if (string.IsNullOrEmpty(RefreshToken))
+			if (string.IsNullOrEmpty(_refreshToken))
 				throw new NotLoggedInException();
 
 			// if the token is valid for another few minutes we can use it
 			// we avoid using a token that is about to expire to make sure we can finish the call we expect to do with it before it expires
-			if (!string.IsNullOrEmpty(AccessToken) && TokenExpiry.AddMinutes(2) > DateTime.Now)
+			if (!string.IsNullOrEmpty(_accessToken) && _tokenExpiry.AddMinutes(2) > DateTime.Now)
 			{
-				return AccessToken;
+				return _accessToken;
 			}
 
-			return await DoRefreshToken(RefreshToken);
+			return await DoRefreshToken(_refreshToken);
 		}
 
-		public OIDCStatus GetStatusForProvider()
+		public OidcStatus GetStatusForProvider()
 		{
-			if (string.IsNullOrEmpty(RefreshToken))
-				return OIDCStatus.NotLoggedIn;
+			if (string.IsNullOrEmpty(_refreshToken))
+				return OidcStatus.NotLoggedIn;
 
-			if (string.IsNullOrEmpty(AccessToken))
-				return OIDCStatus.TokenRefreshRequired;
+			if (string.IsNullOrEmpty(_accessToken))
+				return OidcStatus.TokenRefreshRequired;
 
-			if (TokenExpiry < DateTime.Now)
-				return OIDCStatus.TokenRefreshRequired;
+			if (_tokenExpiry < DateTime.Now)
+				return OidcStatus.TokenRefreshRequired;
 
-			return OIDCStatus.Connected;
+			return OidcStatus.Connected;
 		}
 
-		public void SetRefreshToken(string InRefreshToken)
+		public void SetRefreshToken(string inRefreshToken)
 		{
-			RefreshToken = InRefreshToken;
+			_refreshToken = inRefreshToken;
 		}
 	}
 
 	internal class LoginFailedException : Exception
 	{
-		public LoginFailedException(string Message) : base(Message)
+		public LoginFailedException(string message) : base(message)
 		{
 		}
 	}
diff --git a/Engine/Source/Programs/UnrealGameSync/UnrealGameSync/PerforceModalTask.cs b/Engine/Source/Programs/UnrealGameSync/UnrealGameSync/PerforceModalTask.cs
--- a/Engine/Source/Programs/UnrealGameSync/UnrealGameSync/PerforceModalTask.cs
+++ b/Engine/Source/Programs/UnrealGameSync/UnrealGameSync/PerforceModalTask.cs
@@ -17,114 +17,114 @@ namespace UnrealGameSync
 	{
 		public PerforceResponse<LoginRecord> Response { get; }
 
-		public PerforceLoginException(PerforceResponse<LoginRecord> Response) : base($"Login failed: {Response}")
+		public PerforceLoginException(PerforceResponse<LoginRecord> response) : base($"Login failed: {response}")
 		{
-			this.Response = Response;
+			this.Response = response;
 		}
 	}
 
 	static class PerforceModalTask
 	{
-		public static ModalTask? Execute(IWin32Window? Owner, string Title, string Message, IPerforceSettings PerforceSettings, Func<IPerforceConnection, CancellationToken, Task> ExecuteAsync, ILogger Logger, ModalTaskFlags Flags = ModalTaskFlags.None)
+		public static ModalTask? Execute(IWin32Window? owner, string title, string message, IPerforceSettings perforceSettings, Func<IPerforceConnection, CancellationToken, Task> executeAsync, ILogger logger, ModalTaskFlags flags = ModalTaskFlags.None)
 		{
-			Func<IPerforceConnection, CancellationToken, Task<int>> ExecuteTypedAsync = async (p, c) => { await ExecuteAsync(p, c); return 0; };
-			return Execute(Owner, Title, Message, PerforceSettings, ExecuteTypedAsync, Logger, Flags);
+			Func<IPerforceConnection, CancellationToken, Task<int>> executeTypedAsync = async (p, c) => { await executeAsync(p, c); return 0; };
+			return Execute(owner, title, message, perforceSettings, executeTypedAsync, logger, flags);
 		}
 
-		public static ModalTask<T>? Execute<T>(IWin32Window? Owner, string Title, string Message, IPerforceSettings PerforceSettings, Func<IPerforceConnection, CancellationToken, Task<T>> ExecuteAsync, ILogger Logger, ModalTaskFlags Flags = ModalTaskFlags.None)
+		public static ModalTask<T>? Execute<T>(IWin32Window? owner, string title, string message, IPerforceSettings perforceSettings, Func<IPerforceConnection, CancellationToken, Task<T>> executeAsync, ILogger logger, ModalTaskFlags flags = ModalTaskFlags.None)
 		{
-			IPerforceConnection? Connection = null;
+			IPerforceConnection? connection = null;
 			try
 			{
-				Func<Task<IPerforceConnection>> ConnectAsync = async () =>
+				Func<Task<IPerforceConnection>> connectAsync = async () =>
 				{
-					Connection ??= await PerforceConnection.CreateAsync(PerforceSettings, Logger);
-					return Connection;
+					connection ??= await PerforceConnection.CreateAsync(perforceSettings, logger);
+					return connection;
 				};
-				return ExecuteInternal(Owner, Title, Message, PerforceSettings, ConnectAsync, ExecuteAsync, Flags);
+				return ExecuteInternal(owner, title, message, perforceSettings, connectAsync, executeAsync, flags);
 			}
 			finally
 			{
-				Connection?.Dispose();
+				connection?.Dispose();
 			}
 		}
 
-		public static ModalTask? Execute(IWin32Window? Owner, string Title, string Message, IPerforceConnection PerforceConnection, Func<IPerforceConnection, CancellationToken, Task> ExecuteAsync, ILogger Logger, ModalTaskFlags Flags = ModalTaskFlags.None)
+		public static ModalTask? Execute(IWin32Window? owner, string title, string message, IPerforceConnection perforceConnection, Func<IPerforceConnection, CancellationToken, Task> executeAsync, ILogger logger, ModalTaskFlags flags = ModalTaskFlags.None)
 		{
-			Func<IPerforceConnection, CancellationToken, Task<int>> ExecuteTypedAsync = async (p, c) => { await ExecuteAsync(p, c); return 0; };
-			return Execute(Owner, Title, Message, PerforceConnection, ExecuteTypedAsync, Logger, Flags);
+			Func<IPerforceConnection, CancellationToken, Task<int>> executeTypedAsync = async (p, c) => { await executeAsync(p, c); return 0; };
+			return Execute(owner, title, message, perforceConnection, executeTypedAsync, logger, flags);
 		}
 
-		public static ModalTask<T>? Execute<T>(IWin32Window? Owner, string Title, string Message, IPerforceConnection PerforceConnection, Func<IPerforceConnection, CancellationToken, Task<T>> ExecuteAsync, ILogger Logger, ModalTaskFlags Flags = ModalTaskFlags.None)
+		public static ModalTask<T>? Execute<T>(IWin32Window? owner, string title, string message, IPerforceConnection perforceConnection, Func<IPerforceConnection, CancellationToken, Task<T>> executeAsync, ILogger logger, ModalTaskFlags flags = ModalTaskFlags.None)
 		{
-			return ExecuteInternal(Owner, Title, Message, PerforceConnection.Settings, () => Task.FromResult(PerforceConnection), ExecuteAsync, Flags);
+			return ExecuteInternal(owner, title, message, perforceConnection.Settings, () => Task.FromResult(perforceConnection), executeAsync, flags);
 		}
 
-		private static ModalTask<T>? ExecuteInternal<T>(IWin32Window? Owner, string Title, string Message, IPerforceSettings PerforceSettings, Func<Task<IPerforceConnection>> ConnectAsync, Func<IPerforceConnection, CancellationToken, Task<T>> ExecuteAsync, ModalTaskFlags Flags = ModalTaskFlags.None)
+		private static ModalTask<T>? ExecuteInternal<T>(IWin32Window? owner, string title, string message, IPerforceSettings perforceSettings, Func<Task<IPerforceConnection>> connectAsync, Func<IPerforceConnection, CancellationToken, Task<T>> executeAsync, ModalTaskFlags flags = ModalTaskFlags.None)
 		{
-			string? Password = PerforceSettings.Password;
+			string? password = perforceSettings.Password;
 			for(;;)
 			{
-				Func<CancellationToken, Task<T>> RunAsync = CancellationToken => LoginAndExecuteAsync(Password, ConnectAsync, ExecuteAsync, CancellationToken);
+				Func<CancellationToken, Task<T>> runAsync = cancellationToken => LoginAndExecuteAsync(password, connectAsync, executeAsync, cancellationToken);
 
-				ModalTask<T>? Result = ModalTask.Execute(Owner, Title, Message, RunAsync, ModalTaskFlags.Quiet);
-				if (Result != null && Result.Failed && (Flags & ModalTaskFlags.Quiet) == 0)
+				ModalTask<T>? result = ModalTask.Execute(owner, title, message, runAsync, ModalTaskFlags.Quiet);
+				if (result != null && result.Failed && (flags & ModalTaskFlags.Quiet) == 0)
 				{
-					if (Result.Exception is PerforceLoginException)
+					if (result.Exception is PerforceLoginException)
 					{
-						string PasswordPrompt;
-						if (String.IsNullOrEmpty(Password))
+						string passwordPrompt;
+						if (String.IsNullOrEmpty(password))
 						{
-							PasswordPrompt = $"Enter the password for user '{PerforceSettings.UserName}' on server '{PerforceSettings.ServerAndPort}'.";
+							passwordPrompt = $"Enter the password for user '{perforceSettings.UserName}' on server '{perforceSettings.ServerAndPort}'.";
 						}
 						else
 						{
-							PasswordPrompt = $"Authentication failed. Enter the password for user '{PerforceSettings.UserName}' on server '{PerforceSettings.ServerAndPort}'.";
+							passwordPrompt = $"Authentication failed. Enter the password for user '{perforceSettings.UserName}' on server '{perforceSettings.ServerAndPort}'.";
 						}
 
-						PasswordWindow PasswordWindow = new PasswordWindow(PasswordPrompt, Password ?? String.Empty);
-						if (Owner == null)
+						PasswordWindow passwordWindow = new PasswordWindow(passwordPrompt, password ?? String.Empty);
+						if (owner == null)
 						{
-							PasswordWindow.ShowInTaskbar = true;
-							PasswordWindow.StartPosition = FormStartPosition.CenterScreen;
+							passwordWindow.ShowInTaskbar = true;
+							passwordWindow.StartPosition = FormStartPosition.CenterScreen;
 						}
-						if (PasswordWindow.ShowDialog(Owner) != DialogResult.OK)
+						if (passwordWindow.ShowDialog(owner) != DialogResult.OK)
 						{
 							return null;
 						}
 
-						Password = PasswordWindow.Password;
+						password = passwordWindow.Password;
 						continue;
 					}
-					MessageBox.Show(Owner, Result.Error, Title, MessageBoxButtons.OK);
+					MessageBox.Show(owner, result.Error, title, MessageBoxButtons.OK);
 				}
 
-				return Result;
+				return result;
 			}
 		}
 
-		private static async Task<T> LoginAndExecuteAsync<T>(string? Password, Func<Task<IPerforceConnection>> ConnectAsync, Func<IPerforceConnection, CancellationToken, Task<T>> ExecuteAsync, CancellationToken CancellationToken)
+		private static async Task<T> LoginAndExecuteAsync<T>(string? password, Func<Task<IPerforceConnection>> connectAsync, Func<IPerforceConnection, CancellationToken, Task<T>> executeAsync, CancellationToken cancellationToken)
 		{
-			IPerforceConnection Perforce = await ConnectAsync();
+			IPerforceConnection perforce = await connectAsync();
 
 			// If we've got a password, execute the login command
-			PerforceResponse<LoginRecord> Response;
-			if (String.IsNullOrEmpty(Password))
+			PerforceResponse<LoginRecord> response;
+			if (String.IsNullOrEmpty(password))
 			{
-				Response = await Perforce.TryGetLoginStateAsync(CancellationToken);
+				response = await perforce.TryGetLoginStateAsync(cancellationToken);
 			}
 			else
 			{
-				Response = await Perforce.TryLoginAsync(Password, CancellationToken);
+				response = await perforce.TryLoginAsync(password, cancellationToken);
 			}
 
-			if (!Response.Succeeded)
+			if (!response.Succeeded)
 			{
-				throw new PerforceLoginException(Response);
+				throw new PerforceLoginException(response);
 			}
 
 			// Execute the inner task
-			return await ExecuteAsync(Perforce, CancellationToken);
+			return await executeAsync(perforce, cancellationToken);
 		}
 	}
 }
diff --git a/Engine/Source/Programs/UnrealGameSync/UnrealGameSync/PerforceMonitor.cs b/Engine/Source/Programs/UnrealGameSync/UnrealGameSync/PerforceMonitor.cs
--- a/Engine/Source/Programs/UnrealGameSync/UnrealGameSync/PerforceMonitor.cs
+++ b/Engine/Source/Programs/UnrealGameSync/UnrealGameSync/PerforceMonitor.cs
@@ -41,19 +41,19 @@ namespace UnrealGameSync
 			// TODO: executable/configuration?
 			public SortedList<int, string> ChangeNumberToFileRevision = new SortedList<int, string>();
 
-			public PerforceArchiveInfo(string Name, string Type, string DepotPath, string? Target, PerforceMonitor Outer)
+			public PerforceArchiveInfo(string name, string type, string depotPath, string? target, PerforceMonitor outer)
 			{
-				this.Name = Name;
-				this.Type = Type;
-				this.DepotPath = DepotPath;
-				this.Target = Target;
-				this.Outer = Outer;
+				this.Name = name;
+				this.Type = type;
+				this.DepotPath = depotPath;
+				this.Target = target;
+				this.Outer = outer;
 			}
 
-			public override bool Equals(object? Other)
+			public override bool Equals(object? other)
 			{
-				PerforceArchiveInfo? OtherArchive = Other as PerforceArchiveInfo;
-				return OtherArchive != null && Name == OtherArchive.Name && Type == OtherArchive.Type && DepotPath == OtherArchive.DepotPath && Target == OtherArchive.Target && Enumerable.SequenceEqual(ChangeNumberToFileRevision, OtherArchive.ChangeNumberToFileRevision);
+				PerforceArchiveInfo? otherArchive = other as PerforceArchiveInfo;
+				return otherArchive != null && Name == otherArchive.Name && Type == otherArchive.Type && DepotPath == otherArchive.DepotPath && Target == otherArchive.Target && Enumerable.SequenceEqual(ChangeNumberToFileRevision, otherArchive.ChangeNumberToFileRevision);
 			}
 
 			public override int GetHashCode()
@@ -66,53 +66,53 @@ namespace UnrealGameSync
 				return ChangeNumberToFileRevision.Count > 0;
 			}
 
-			public static bool TryParseConfigEntry(string Text, PerforceMonitor Outer, [NotNullWhen(true)] out PerforceArchiveInfo? Info)
+			public static bool TryParseConfigEntry(string text, PerforceMonitor outer, [NotNullWhen(true)] out PerforceArchiveInfo? info)
 			{
-				ConfigObject Object = new ConfigObject(Text);
+				ConfigObject obj = new ConfigObject(text);
 
-				string? Name = Object.GetValue("Name", null);
-				if (Name == null)
+				string? name = obj.GetValue("Name", null);
+				if (name == null)
 				{
-					Info = null;
+					info = null;
 					return false;
 				}
 
-				string? DepotPath = Object.GetValue("DepotPath", null);
-				if (DepotPath == null)
+				string? depotPath = obj.GetValue("DepotPath", null);
+				if (depotPath == null)
 				{
-					Info = null;
+					info = null;
 					return false;
 				}
 
-				string? Target = Object.GetValue("Target", null);
+				string? target = obj.GetValue("Target", null);
 
-				string Type = Object.GetValue("Type", null) ?? Name;
+				string type = obj.GetValue("Type", null) ?? name;
 
-				Info = new PerforceArchiveInfo(Name, Type, DepotPath, Target, Outer);
+				info = new PerforceArchiveInfo(name, type, depotPath, target, outer);
 				return true;
 			}
 
-			public bool TryGetArchiveKeyForChangeNumber(int ChangeNumber, [NotNullWhen(true)] out string? ArchiveKey)
+			public bool TryGetArchiveKeyForChangeNumber(int changeNumber, [NotNullWhen(true)] out string? archiveKey)
 			{
-				return ChangeNumberToFileRevision.TryGetValue(ChangeNumber, out ArchiveKey);
+				return ChangeNumberToFileRevision.TryGetValue(changeNumber, out archiveKey);
 			}
 
-			public async Task<bool> DownloadArchive(IPerforceConnection Perforce, string ArchiveKey, DirectoryReference LocalRootPath, FileReference ManifestFileName, ILogger Logger, ProgressValue Progress, CancellationToken CancellationToken)
+			public async Task<bool> DownloadArchive(IPerforceConnection perforce, string archiveKey, DirectoryReference localRootPath, FileReference manifestFileName, ILogger logger, ProgressValue progress, CancellationToken cancellationToken)
 			{
-				FileReference TempZipFileName = new FileReference(Path.GetTempFileName());
+				FileReference tempZipFileName = new FileReference(Path.GetTempFileName());
 				try
 				{
-					PrintRecord Record = await Perforce.PrintAsync(TempZipFileName.FullName, ArchiveKey, CancellationToken);
-					if (TempZipFileName.ToFileInfo().Length == 0)
+					PrintRecord record = await perforce.PrintAsync(tempZipFileName.FullName, archiveKey, cancellationToken);
+					if (tempZipFileName.ToFileInfo().Length == 0)
 					{
 						return false;
 					}
-					ArchiveUtils.ExtractFiles(TempZipFileName, LocalRootPath, ManifestFileName, Progress, Logger);
+					ArchiveUtils.ExtractFiles(tempZipFileName, localRootPath, manifestFileName, progress, logger);
 				}
 				finally
 				{
-					FileReference.SetAttributes(TempZipFileName, FileAttributes.Normal);
-					FileReference.Delete(TempZipFileName);
+					FileReference.SetAttributes(tempZipFileName, FileAttributes.Normal);
+					FileReference.Delete(tempZipFileName);
 				}
 
 				return true;
@@ -126,33 +126,33 @@ namespace UnrealGameSync
 
 		class PerforceChangeSorter : IComparer<ChangesRecord>
 		{
-			public int Compare(ChangesRecord? SummaryA, ChangesRecord? SummaryB)
+			public int Compare(ChangesRecord? summaryA, ChangesRecord? summaryB)
 			{
-				return SummaryB!.Number - SummaryA!.Number;
+				return summaryB!.Number - summaryA!.Number;
 			}
 		}
 
 		public int InitialMaxChangesValue = 100;
 
-		IPerforceSettings PerforceSettings;
-		readonly string BranchClientPath;
-		readonly string SelectedClientFileName;
-		readonly string SelectedProjectIdentifier;
-		Task? WorkerTask;
-		CancellationTokenSource CancellationSource;
-		int PendingMaxChangesValue;
-		SortedSet<ChangesRecord> Changes = new SortedSet<ChangesRecord>(new PerforceChangeSorter());
-		SortedDictionary<int, PerforceChangeDetails> ChangeDetails = new SortedDictionary<int,PerforceChangeDetails>();
-		SortedSet<int> PromotedChangeNumbers = new SortedSet<int>();
-		List<PerforceArchiveInfo> Archives = new List<PerforceArchiveInfo>();
-		AsyncEvent RefreshEvent = new AsyncEvent();
-		ILogger Logger;
-		bool bIsEnterpriseProject;
-		DirectoryReference CacheFolder;
-		List<KeyValuePair<FileReference, DateTime>> LocalConfigFiles;
-		IAsyncDisposer AsyncDisposeTasks;
-
-		SynchronizationContext SynchronizationContext;
+		IPerforceSettings _perforceSettings;
+		readonly string _branchClientPath;
+		readonly string _selectedClientFileName;
+		readonly string _selectedProjectIdentifier;
+		Task? _workerTask;
+		CancellationTokenSource _cancellationSource;
+		int _pendingMaxChangesValue;
+		SortedSet<ChangesRecord> _changes = new SortedSet<ChangesRecord>(new PerforceChangeSorter());
+		SortedDictionary<int, PerforceChangeDetails> _changeDetails = new SortedDictionary<int,PerforceChangeDetails>();
+		SortedSet<int> _promotedChangeNumbers = new SortedSet<int>();
+		List<PerforceArchiveInfo> _archives = new List<PerforceArchiveInfo>();
+		AsyncEvent _refreshEvent = new AsyncEvent();
+		ILogger _logger;
+		bool _isEnterpriseProject;
+		DirectoryReference _cacheFolder;
+		List<KeyValuePair<FileReference, DateTime>> _localConfigFiles;
+		IAsyncDisposer _asyncDisposeTasks;
+
+		SynchronizationContext _synchronizationContext;
 		public event Action? OnUpdate;
 		public event Action? OnUpdateMetadata;
 		public event Action? OnStreamChange;
@@ -164,32 +164,32 @@ namespace UnrealGameSync
 			private set;
 		}
 
-		public PerforceMonitor(IPerforceSettings InPerforceSettings, ProjectInfo ProjectInfo, ConfigFile InProjectConfigFile, DirectoryReference InCacheFolder, List<KeyValuePair<FileReference, DateTime>> InLocalConfigFiles, IServiceProvider InServiceProvider)
+		public PerforceMonitor(IPerforceSettings inPerforceSettings, ProjectInfo projectInfo, ConfigFile inProjectConfigFile, DirectoryReference inCacheFolder, List<KeyValuePair<FileReference, DateTime>> inLocalConfigFiles, IServiceProvider inServiceProvider)
 		{
-			PerforceSettings = InPerforceSettings;
-			BranchClientPath = ProjectInfo.ClientRootPath;
-			SelectedClientFileName = ProjectInfo.ClientFileName;
-			SelectedProjectIdentifier = ProjectInfo.ProjectIdentifier;
-			PendingMaxChangesValue = InitialMaxChangesValue;
+			_perforceSettings = inPerforceSettings;
+			_branchClientPath = projectInfo.ClientRootPath;
+			_selectedClientFileName = projectInfo.ClientFileName;
+			_selectedProjectIdentifier = projectInfo.ProjectIdentifier;
+			_pendingMaxChangesValue = InitialMaxChangesValue;
 			LastChangeByCurrentUser = -1;
 			LastCodeChangeByCurrentUser = -1;
-			Logger = InServiceProvider.GetRequiredService<ILogger<PerforceMonitor>>();
-			bIsEnterpriseProject = ProjectInfo.bIsEnterpriseProject;
-			LatestProjectConfigFile = InProjectConfigFile;
-			CacheFolder = InCacheFolder;
-			LocalConfigFiles = InLocalConfigFiles;
-			AsyncDisposeTasks = InServiceProvider.GetRequiredService<IAsyncDisposer>();
-			SynchronizationContext = SynchronizationContext.Current!;
-			CancellationSource = new CancellationTokenSource();
+			_logger = inServiceProvider.GetRequiredService<ILogger<PerforceMonitor>>();
+			_isEnterpriseProject = projectInfo.IsEnterpriseProject;
+			LatestProjectConfigFile = inProjectConfigFile;
+			_cacheFolder = inCacheFolder;
+			_localConfigFiles = inLocalConfigFiles;
+			_asyncDisposeTasks = inServiceProvider.GetRequiredService<IAsyncDisposer>();
+			_synchronizationContext = SynchronizationContext.Current!;
+			_cancellationSource = new CancellationTokenSource();
 
 			AvailableArchives = (new List<IArchiveInfo>()).AsReadOnly();
 		}
 
 		public void Start()
 		{
-			if (WorkerTask == null)
+			if (_workerTask == null)
 			{
-				WorkerTask = Task.Run(() => PollForUpdates(CancellationSource.Token));
+				_workerTask = Task.Run(() => PollForUpdates(_cancellationSource.Token));
 			}
 		}
 
@@ -200,11 +200,11 @@ namespace UnrealGameSync
 			OnStreamChange = null;
 			OnLoginExpired = null;
 
-			if (WorkerTask != null)
+			if (_workerTask != null)
 			{
-				CancellationSource.Cancel();
-				AsyncDisposeTasks.Add(WorkerTask.ContinueWith(_ => CancellationSource.Dispose()));
-				WorkerTask = null;
+				_cancellationSource.Cancel();
+				_asyncDisposeTasks.Add(_workerTask.ContinueWith(_ => _cancellationSource.Dispose()));
+				_workerTask = null;
 			}
 		}
 
@@ -228,284 +228,284 @@ namespace UnrealGameSync
 
 		public int PendingMaxChanges
 		{
-			get { return PendingMaxChangesValue; }
-			set { lock(this){ if(value != PendingMaxChangesValue){ PendingMaxChangesValue = value; RefreshEvent.Set(); } } }
+			get { return _pendingMaxChangesValue; }
+			set { lock(this){ if(value != _pendingMaxChangesValue){ _pendingMaxChangesValue = value; _refreshEvent.Set(); } } }
 		}
 
-		async Task PollForUpdates(CancellationToken CancellationToken)
+		async Task PollForUpdates(CancellationToken cancellationToken)
 		{
-			while (!CancellationToken.IsCancellationRequested)
+			while (!cancellationToken.IsCancellationRequested)
 			{
 				try
 				{
-					await PollForUpdatesInner(CancellationToken);
+					await PollForUpdatesInner(cancellationToken);
 				}
-				catch (OperationCanceledException) when (CancellationToken.IsCancellationRequested)
+				catch (OperationCanceledException) when (cancellationToken.IsCancellationRequested)
 				{
 					break;
 				}
-				catch (Exception Ex)
+				catch (Exception ex)
 				{
-					Logger.LogError(Ex, "Unhandled exception in PollForUpdatesInner()");
-					if (!(Ex is PerforceException))
+					_logger.LogError(ex, "Unhandled exception in PollForUpdatesInner()");
+					if (!(ex is PerforceException))
 					{
-						Program.CaptureException(Ex);
+						Program.CaptureException(ex);
 					}
-					await Task.Delay(TimeSpan.FromSeconds(20.0), CancellationToken).ContinueWith(x => { });
+					await Task.Delay(TimeSpan.FromSeconds(20.0), cancellationToken).ContinueWith(x => { });
 				}
 			}
 		}
 
-		async Task PollForUpdatesInner(CancellationToken CancellationToken)
+		async Task PollForUpdatesInner(CancellationToken cancellationToken)
 		{
-			string? StreamName;
-			using (IPerforceConnection Perforce = await PerforceConnection.CreateAsync(PerforceSettings, Logger))
+			string? streamName;
+			using (IPerforceConnection perforce = await PerforceConnection.CreateAsync(_perforceSettings, _logger))
 			{
-				StreamName = await Perforce.GetCurrentStreamAsync(CancellationToken);
+				streamName = await perforce.GetCurrentStreamAsync(cancellationToken);
 
 				// Get the perforce server settings
-				PerforceResponse<InfoRecord> InfoResponse = await Perforce.TryGetInfoAsync(InfoOptions.ShortOutput, CancellationToken);
-				if (InfoResponse.Succeeded)
+				PerforceResponse<InfoRecord> infoResponse = await perforce.TryGetInfoAsync(InfoOptions.ShortOutput, cancellationToken);
+				if (infoResponse.Succeeded)
 				{
-					DateTimeOffset? ServerDate = InfoResponse.Data.ServerDate;
-					if (ServerDate.HasValue)
+					DateTimeOffset? serverDate = infoResponse.Data.ServerDate;
+					if (serverDate.HasValue)
 					{
-						ServerTimeZone = ServerDate.Value.Offset;
+						ServerTimeZone = serverDate.Value.Offset;
 					}
 				}
 
 				// Try to update the zipped binaries list before anything else, because it causes a state change in the UI
-				await UpdateArchivesAsync(Perforce, CancellationToken);
+				await UpdateArchivesAsync(perforce, cancellationToken);
 			}
 
-			while(!CancellationToken.IsCancellationRequested)
+			while(!cancellationToken.IsCancellationRequested)
 			{
-				Stopwatch Timer = Stopwatch.StartNew();
-				Task NextRefreshTask = RefreshEvent.Task;
+				Stopwatch timer = Stopwatch.StartNew();
+				Task nextRefreshTask = _refreshEvent.Task;
 
-				using (IPerforceConnection Perforce = await PerforceConnection.CreateAsync(PerforceSettings, Logger))
+				using (IPerforceConnection perforce = await PerforceConnection.CreateAsync(_perforceSettings, _logger))
 				{
 					// Check we still have a valid login ticket
-					PerforceResponse<LoginRecord> LoginState = await Perforce.TryGetLoginStateAsync(CancellationToken);
-					if (!LoginState.Succeeded)
+					PerforceResponse<LoginRecord> loginState = await perforce.TryGetLoginStateAsync(cancellationToken);
+					if (!loginState.Succeeded)
 					{
 						LastStatusMessage = "User is not logged in";
-						SynchronizationContext.Post(_ => OnLoginExpired?.Invoke(), null);
+						_synchronizationContext.Post(_ => OnLoginExpired?.Invoke(), null);
 					}
 					else
 					{
 						// Check we haven't switched streams
-						string? NewStreamName = await Perforce.GetCurrentStreamAsync(CancellationToken);
-						if (NewStreamName != StreamName)
+						string? newStreamName = await perforce.GetCurrentStreamAsync(cancellationToken);
+						if (newStreamName != streamName)
 						{
-							SynchronizationContext.Post(_ => OnStreamChange?.Invoke(), null);
+							_synchronizationContext.Post(_ => OnStreamChange?.Invoke(), null);
 						}
 
 						// Check for any p4 changes
-						if (!await UpdateChangesAsync(Perforce, CancellationToken))
+						if (!await UpdateChangesAsync(perforce, cancellationToken))
 						{
 							LastStatusMessage = "Failed to update changes";
 						}
-						else if (!await UpdateChangeTypesAsync(Perforce, CancellationToken))
+						else if (!await UpdateChangeTypesAsync(perforce, cancellationToken))
 						{
 							LastStatusMessage = "Failed to update change types";
 						}
-						else if (!await UpdateArchivesAsync(Perforce, CancellationToken))
+						else if (!await UpdateArchivesAsync(perforce, cancellationToken))
 						{
 							LastStatusMessage = "Failed to update zipped binaries list";
 						}
 						else
 						{
-							LastStatusMessage = String.Format("Last update took {0}ms", Timer.ElapsedMilliseconds);
+							LastStatusMessage = String.Format("Last update took {0}ms", timer.ElapsedMilliseconds);
 						}
 					}
 				}
 
 				// Wait for another request, or scan for new builds after a timeout
-				Task DelayTask = Task.Delay(TimeSpan.FromMinutes(IsActive ? 2 : 10), CancellationToken);
-				await Task.WhenAny(NextRefreshTask, DelayTask);
+				Task delayTask = Task.Delay(TimeSpan.FromMinutes(IsActive ? 2 : 10), cancellationToken);
+				await Task.WhenAny(nextRefreshTask, delayTask);
 			}
 		}
 
-		async Task<bool> UpdateChangesAsync(IPerforceConnection Perforce, CancellationToken CancellationToken)
+		async Task<bool> UpdateChangesAsync(IPerforceConnection perforce, CancellationToken cancellationToken)
 		{
 			// Get the current status of the build
-			int MaxChanges;
-			int OldestChangeNumber = -1;
-			int NewestChangeNumber = -1;
-			HashSet<int> CurrentChangelists;
-			SortedSet<int> PrevPromotedChangelists;
+			int maxChanges;
+			int oldestChangeNumber = -1;
+			int newestChangeNumber = -1;
+			HashSet<int> currentChangelists;
+			SortedSet<int> prevPromotedChangelists;
 			lock(this)
 			{
-				MaxChanges = PendingMaxChanges;
-				if(Changes.Count > 0)
+				maxChanges = PendingMaxChanges;
+				if(_changes.Count > 0)
 				{
-					NewestChangeNumber = Changes.First().Number;
-					OldestChangeNumber = Changes.Last().Number;
+					newestChangeNumber = _changes.First().Number;
+					oldestChangeNumber = _changes.Last().Number;
 				}
-				CurrentChangelists = new HashSet<int>(Changes.Select(x => x.Number));
-				PrevPromotedChangelists = new SortedSet<int>(PromotedChangeNumbers);
+				currentChangelists = new HashSet<int>(_changes.Select(x => x.Number));
+				prevPromotedChangelists = new SortedSet<int>(_promotedChangeNumbers);
 			}
 
 			// Build a full list of all the paths to sync
-			List<string> DepotPaths = new List<string>();
-			if (SelectedClientFileName.EndsWith(".uprojectdirs", StringComparison.InvariantCultureIgnoreCase))
+			List<string> depotPaths = new List<string>();
+			if (_selectedClientFileName.EndsWith(".uprojectdirs", StringComparison.InvariantCultureIgnoreCase))
 			{
-				DepotPaths.Add(String.Format("{0}/...", BranchClientPath));
+				depotPaths.Add(String.Format("{0}/...", _branchClientPath));
 			}
 			else
 			{
-				DepotPaths.Add(String.Format("{0}/*", BranchClientPath));
-				DepotPaths.Add(String.Format("{0}/Engine/...", BranchClientPath));
-				DepotPaths.Add(String.Format("{0}/...", PerforceUtils.GetClientOrDepotDirectoryName(SelectedClientFileName)));
-				if (bIsEnterpriseProject)
+				depotPaths.Add(String.Format("{0}/*", _branchClientPath));
+				depotPaths.Add(String.Format("{0}/Engine/...", _branchClientPath));
+				depotPaths.Add(String.Format("{0}/...", PerforceUtils.GetClientOrDepotDirectoryName(_selectedClientFileName)));
+				if (_isEnterpriseProject)
 				{
-					DepotPaths.Add(String.Format("{0}/Enterprise/...", BranchClientPath));
+					depotPaths.Add(String.Format("{0}/Enterprise/...", _branchClientPath));
 				}
 
 				// Add in additional paths property
-				ConfigSection ProjectConfigSection = LatestProjectConfigFile.FindSection("Perforce");
-				if (ProjectConfigSection != null)
+				ConfigSection projectConfigSection = LatestProjectConfigFile.FindSection("Perforce");
+				if (projectConfigSection != null)
 				{
-					IEnumerable<string> AdditionalPaths = ProjectConfigSection.GetValues("AdditionalPathsToSync", new string[0]);
+					IEnumerable<string> additionalPaths = projectConfigSection.GetValues("AdditionalPathsToSync", new string[0]);
 
 					// turn into //ws/path
-					DepotPaths.AddRange(AdditionalPaths.Select(P => string.Format("{0}/{1}", BranchClientPath, P.TrimStart('/'))));
+					depotPaths.AddRange(additionalPaths.Select(p => string.Format("{0}/{1}", _branchClientPath, p.TrimStart('/'))));
 				}
 			}
 
 			// Read any new changes
-			List<ChangesRecord> NewChanges;
-			if(MaxChanges > CurrentMaxChanges)
+			List<ChangesRecord> newChanges;
+			if(maxChanges > CurrentMaxChanges)
 			{
-				NewChanges = await Perforce.GetChangesAsync(ChangesOptions.IncludeTimes | ChangesOptions.LongOutput, MaxChanges, ChangeStatus.Submitted, DepotPaths, CancellationToken);
+				newChanges = await perforce.GetChangesAsync(ChangesOptions.IncludeTimes | ChangesOptions.LongOutput, maxChanges, ChangeStatus.Submitted, depotPaths, cancellationToken);
 			}
 			else
 			{
-				NewChanges = await Perforce.GetChangesAsync(ChangesOptions.IncludeTimes | ChangesOptions.LongOutput, -1, ChangeStatus.Submitted, DepotPaths.Select(x => $"{x}@>{NewestChangeNumber}").ToArray(), CancellationToken);
+				newChanges = await perforce.GetChangesAsync(ChangesOptions.IncludeTimes | ChangesOptions.LongOutput, -1, ChangeStatus.Submitted, depotPaths.Select(x => $"{x}@>{newestChangeNumber}").ToArray(), cancellationToken);
 			}
 
 			// Remove anything we already have
-			NewChanges.RemoveAll(x => CurrentChangelists.Contains(x.Number));
+			newChanges.RemoveAll(x => currentChangelists.Contains(x.Number));
 
 			// Update the change ranges
-			if(NewChanges.Count > 0)
+			if(newChanges.Count > 0)
 			{
-				OldestChangeNumber = Math.Max(OldestChangeNumber, NewChanges.Last().Number);
-				NewestChangeNumber = Math.Min(NewestChangeNumber, NewChanges.First().Number);
+				oldestChangeNumber = Math.Max(oldestChangeNumber, newChanges.Last().Number);
+				newestChangeNumber = Math.Min(newestChangeNumber, newChanges.First().Number);
 			}
 
 			// If we are using zipped binaries, make sure we have every change since the last zip containing them. This is necessary for ensuring that content changes show as
 			// syncable in the workspace view if there have been a large number of content changes since the last code change.
-			int MinZippedChangeNumber = -1;
-			foreach (PerforceArchiveInfo Archive in Archives)
+			int minZippedChangeNumber = -1;
+			foreach (PerforceArchiveInfo archive in _archives)
 			{
-				foreach (int ChangeNumber in Archive.ChangeNumberToFileRevision.Keys)
+				foreach (int changeNumber in archive.ChangeNumberToFileRevision.Keys)
 				{
-					if (ChangeNumber > MinZippedChangeNumber && ChangeNumber <= OldestChangeNumber)
+					if (changeNumber > minZippedChangeNumber && changeNumber <= oldestChangeNumber)
 					{
-						MinZippedChangeNumber = ChangeNumber;
+						minZippedChangeNumber = changeNumber;
 					}
 				}
 			}
-			if(MinZippedChangeNumber != -1 && MinZippedChangeNumber < OldestChangeNumber)
+			if(minZippedChangeNumber != -1 && minZippedChangeNumber < oldestChangeNumber)
 			{
-				string[] FilteredPaths = DepotPaths.Select(x => $"{x}@{MinZippedChangeNumber},{OldestChangeNumber - 1}").ToArray();
-				List<ChangesRecord> ZipChanges = await Perforce.GetChangesAsync(ChangesOptions.None, -1, ChangeStatus.Submitted, FilteredPaths, CancellationToken);
-				NewChanges.AddRange(ZipChanges);
+				string[] filteredPaths = depotPaths.Select(x => $"{x}@{minZippedChangeNumber},{oldestChangeNumber - 1}").ToArray();
+				List<ChangesRecord> zipChanges = await perforce.GetChangesAsync(ChangesOptions.None, -1, ChangeStatus.Submitted, filteredPaths, cancellationToken);
+				newChanges.AddRange(zipChanges);
 			}
 
 			// Fixup any ROBOMERGE authors
-			const string RoboMergePrefix = "#ROBOMERGE-AUTHOR:";
-			foreach (ChangesRecord Change in NewChanges)
+			const string roboMergePrefix = "#ROBOMERGE-AUTHOR:";
+			foreach (ChangesRecord change in newChanges)
 			{
-				if(Change.Description != null && Change.Description.StartsWith(RoboMergePrefix))
+				if(change.Description != null && change.Description.StartsWith(roboMergePrefix))
 				{
-					int StartIdx = RoboMergePrefix.Length;
-					while(StartIdx < Change.Description.Length && Change.Description[StartIdx] == ' ')
+					int startIdx = roboMergePrefix.Length;
+					while(startIdx < change.Description.Length && change.Description[startIdx] == ' ')
 					{
-						StartIdx++;
+						startIdx++;
 					}
 
-					int EndIdx = StartIdx;
-					while(EndIdx < Change.Description.Length && !Char.IsWhiteSpace(Change.Description[EndIdx]))
+					int endIdx = startIdx;
+					while(endIdx < change.Description.Length && !Char.IsWhiteSpace(change.Description[endIdx]))
 					{
-						EndIdx++;
+						endIdx++;
 					}
 
-					if(EndIdx > StartIdx)
+					if(endIdx > startIdx)
 					{
-						Change.User = Change.Description.Substring(StartIdx, EndIdx - StartIdx);
-						Change.Description = "ROBOMERGE: " + Change.Description.Substring(EndIdx).TrimStart();
+						change.User = change.Description.Substring(startIdx, endIdx - startIdx);
+						change.Description = "ROBOMERGE: " + change.Description.Substring(endIdx).TrimStart();
 					}
 				}
 			}
 
 			// Process the new changes received
-			if(NewChanges.Count > 0 || MaxChanges < CurrentMaxChanges)
+			if(newChanges.Count > 0 || maxChanges < CurrentMaxChanges)
 			{
 				// Insert them into the builds list
 				lock(this)
 				{
-					Changes.UnionWith(NewChanges);
-					if(Changes.Count > MaxChanges)
+					_changes.UnionWith(newChanges);
+					if(_changes.Count > maxChanges)
 					{
 						// Remove changes to shrink it to the max requested size, being careful to avoid removing changes that would affect our ability to correctly
 						// show the availability for content changes using zipped binaries.
-						SortedSet<ChangesRecord> TrimmedChanges = new SortedSet<ChangesRecord>(new PerforceChangeSorter());
-						foreach(ChangesRecord Change in Changes)
+						SortedSet<ChangesRecord> trimmedChanges = new SortedSet<ChangesRecord>(new PerforceChangeSorter());
+						foreach(ChangesRecord change in _changes)
 						{
-							TrimmedChanges.Add(Change);
-							if(TrimmedChanges.Count >= MaxChanges && Archives.Any(x => x.ChangeNumberToFileRevision.Count == 0 || x.ChangeNumberToFileRevision.ContainsKey(Change.Number) || x.ChangeNumberToFileRevision.First().Key > Change.Number))
+							trimmedChanges.Add(change);
+							if(trimmedChanges.Count >= maxChanges && _archives.Any(x => x.ChangeNumberToFileRevision.Count == 0 || x.ChangeNumberToFileRevision.ContainsKey(change.Number) || x.ChangeNumberToFileRevision.First().Key > change.Number))
 							{
 								break;
 							}
 						}
-						Changes = TrimmedChanges;
+						_changes = trimmedChanges;
 					}
-					CurrentMaxChanges = MaxChanges;
+					CurrentMaxChanges = maxChanges;
 				}
 
 				// Find the last submitted change by the current user
-				int NewLastChangeByCurrentUser = -1;
-				foreach(ChangesRecord Change in Changes)
+				int newLastChangeByCurrentUser = -1;
+				foreach(ChangesRecord change in _changes)
 				{
-					if(String.Compare(Change.User, Perforce.Settings.UserName, StringComparison.InvariantCultureIgnoreCase) == 0)
+					if(String.Compare(change.User, perforce.Settings.UserName, StringComparison.InvariantCultureIgnoreCase) == 0)
 					{
-						NewLastChangeByCurrentUser = Math.Max(NewLastChangeByCurrentUser, Change.Number);
+						newLastChangeByCurrentUser = Math.Max(newLastChangeByCurrentUser, change.Number);
 					}
 				}
-				LastChangeByCurrentUser = NewLastChangeByCurrentUser;
+				LastChangeByCurrentUser = newLastChangeByCurrentUser;
 
 				// Notify the main window that we've got more data
-				SynchronizationContext.Post(_ => OnUpdate?.Invoke(), null);
+				_synchronizationContext.Post(_ => OnUpdate?.Invoke(), null);
 			}
 			return true;
 		}
 
-		public async Task<bool> UpdateChangeTypesAsync(IPerforceConnection Perforce, CancellationToken CancellationToken)
+		public async Task<bool> UpdateChangeTypesAsync(IPerforceConnection perforce, CancellationToken cancellationToken)
 		{
 			// Find the changes we need to query
-			List<int> QueryChangeNumbers = new List<int>();
+			List<int> queryChangeNumbers = new List<int>();
 			lock(this)
 			{
-				foreach(ChangesRecord Change in Changes)
+				foreach(ChangesRecord change in _changes)
 				{
-					if(!ChangeDetails.ContainsKey(Change.Number))
+					if(!_changeDetails.ContainsKey(change.Number))
 					{
-						QueryChangeNumbers.Add(Change.Number);
+						queryChangeNumbers.Add(change.Number);
 					}
 				}
 			}
 
 			// Update them in batches
-			bool bUpdatedConfigFile = false;
+			bool updatedConfigFile = false;
 			using (CancellationTokenSource cancellationSource = new CancellationTokenSource())
 			{
 				Task notifyTask = Task.CompletedTask;
-				foreach (IReadOnlyList<int> QueryChangeNumberBatch in QueryChangeNumbers.OrderByDescending(x => x).Batch(10))
+				foreach (IReadOnlyList<int> queryChangeNumberBatch in queryChangeNumbers.OrderByDescending(x => x).Batch(10))
 				{
-					CancellationToken.ThrowIfCancellationRequested();
+					cancellationToken.ThrowIfCancellationRequested();
 
 					// Skip this stuff if the user wants us to query for more changes
 					if (PendingMaxChanges != CurrentMaxChanges)
@@ -514,52 +514,51 @@ namespace UnrealGameSync
 					}
 
 					// If there's something to check for, find all the content changes after this changelist
-					const int MaxFiles = 100;
+					int maxFiles = 100;
 
-					List<DescribeRecord> DescribeRecords = await Perforce.DescribeAsync(DescribeOptions.None, MaxFiles, QueryChangeNumberBatch.ToArray(), CancellationToken);
-					foreach (DescribeRecord DescribeRecordLoop in DescribeRecords.OrderByDescending(x => x.Number))
+					List<DescribeRecord> describeRecords = await perforce.DescribeAsync(DescribeOptions.None, maxFiles, queryChangeNumberBatch.ToArray(), cancellationToken);
+					foreach (DescribeRecord describeRecordLoop in describeRecords.OrderByDescending(x => x.Number))
 					{
-						DescribeRecord DescribeRecord = DescribeRecordLoop;
-						int QueryChangeNumber = DescribeRecord.Number;
+						DescribeRecord describeRecord = describeRecordLoop;
+						int queryChangeNumber = describeRecord.Number;
 
-						PerforceChangeDetails Details = new PerforceChangeDetails(DescribeRecord);
+						PerforceChangeDetails details = new PerforceChangeDetails(describeRecord);
 
 						// Content only changes must be flagged accurately, because code changes invalidate precompiled binaries. Increase the number of files fetched until we can classify it correctly.
-						int maxFiles = MaxFiles;
-						while (DescribeRecord.Files.Count >= maxFiles && !Details.bContainsCode)
+						while (describeRecord.Files.Count >= maxFiles && !details.ContainsCode)
 						{
-							CancellationToken.ThrowIfCancellationRequested();
+							cancellationToken.ThrowIfCancellationRequested();
 							maxFiles *= 10;
 
-							List<DescribeRecord> NewDescribeRecords = await Perforce.DescribeAsync(DescribeOptions.None, maxFiles, new int[] { QueryChangeNumber }, CancellationToken);
-							if (NewDescribeRecords.Count == 0)
+							List<DescribeRecord> newDescribeRecords = await perforce.DescribeAsync(DescribeOptions.None, maxFiles, new int[] { queryChangeNumber }, cancellationToken);
+							if (newDescribeRecords.Count == 0)
 							{
 								break;
 							}
 
-							DescribeRecord = NewDescribeRecords[0];
-							Details = new PerforceChangeDetails(DescribeRecord);
+							describeRecord = newDescribeRecords[0];
+							details = new PerforceChangeDetails(describeRecord);
 						}
 
 						lock (this)
 						{
-							if (!ChangeDetails.ContainsKey(QueryChangeNumber))
+							if (!_changeDetails.ContainsKey(queryChangeNumber))
 							{
-								ChangeDetails.Add(QueryChangeNumber, Details);
+								_changeDetails.Add(queryChangeNumber, details);
 							}
 						}
 
 						// Reload the config file if it changes
-						if (DescribeRecord.Files.Any(x => x.DepotFile.EndsWith("/UnrealGameSync.ini", StringComparison.OrdinalIgnoreCase)) && !bUpdatedConfigFile)
+						if (describeRecord.Files.Any(x => x.DepotFile.EndsWith("/UnrealGameSync.ini", StringComparison.OrdinalIgnoreCase)) && !updatedConfigFile)
 						{
-							await UpdateProjectConfigFileAsync(Perforce, CancellationToken);
-							bUpdatedConfigFile = true;
+							await UpdateProjectConfigFileAsync(perforce, cancellationToken);
+							updatedConfigFile = true;
 						}
 
 						// Notify the caller after a fixed period of time, in case further updates are slow to arrive
 						if (notifyTask.IsCompleted)
 						{
-							notifyTask = Task.Delay(TimeSpan.FromSeconds(5.0), cancellationSource.Token).ContinueWith(_ => SynchronizationContext.Post(_ => OnUpdateMetadata?.Invoke(), null));
+							notifyTask = Task.Delay(TimeSpan.FromSeconds(5.0), cancellationSource.Token).ContinueWith(_ => _synchronizationContext.Post(_ => OnUpdateMetadata?.Invoke(), null));
 						}
 					}
 				}
@@ -568,77 +567,77 @@ namespace UnrealGameSync
 			}
 
 			// Find the last submitted code change by the current user
-			int NewLastCodeChangeByCurrentUser = -1;
-			foreach(ChangesRecord Change in Changes)
+			int newLastCodeChangeByCurrentUser = -1;
+			foreach(ChangesRecord change in _changes)
 			{
-				if(String.Compare(Change.User, Perforce.Settings.UserName, StringComparison.InvariantCultureIgnoreCase) == 0)
+				if(String.Compare(change.User, perforce.Settings.UserName, StringComparison.InvariantCultureIgnoreCase) == 0)
 				{
-					PerforceChangeDetails? OtherDetails;
-					if(ChangeDetails.TryGetValue(Change.Number, out OtherDetails) && OtherDetails.bContainsCode)
+					PerforceChangeDetails? otherDetails;
+					if(_changeDetails.TryGetValue(change.Number, out otherDetails) && otherDetails.ContainsCode)
 					{
-						NewLastCodeChangeByCurrentUser = Math.Max(NewLastCodeChangeByCurrentUser, Change.Number);
+						newLastCodeChangeByCurrentUser = Math.Max(newLastCodeChangeByCurrentUser, change.Number);
 					}
 				}
 			}
-			LastCodeChangeByCurrentUser = NewLastCodeChangeByCurrentUser;
+			LastCodeChangeByCurrentUser = newLastCodeChangeByCurrentUser;
 
 			// Notify the main window that we've got an update
-			SynchronizationContext.Post(_ => OnUpdateMetadata?.Invoke(), null);
+			_synchronizationContext.Post(_ => OnUpdateMetadata?.Invoke(), null);
 
-			if(LocalConfigFiles.Any(x => FileReference.GetLastWriteTimeUtc(x.Key) != x.Value))
+			if(_localConfigFiles.Any(x => FileReference.GetLastWriteTimeUtc(x.Key) != x.Value))
 			{
-				await UpdateProjectConfigFileAsync(Perforce, CancellationToken);
-				SynchronizationContext.Post(_ => OnUpdateMetadata?.Invoke(), null);
+				await UpdateProjectConfigFileAsync(perforce, cancellationToken);
+				_synchronizationContext.Post(_ => OnUpdateMetadata?.Invoke(), null);
 			}
 
 			return true;
 		}
 
-		async Task<bool> UpdateArchivesAsync(IPerforceConnection Perforce, CancellationToken CancellationToken)
+		async Task<bool> UpdateArchivesAsync(IPerforceConnection perforce, CancellationToken cancellationToken)
 		{
-			List<PerforceArchiveInfo> NewArchives = new List<PerforceArchiveInfo>();
+			List<PerforceArchiveInfo> newArchives = new List<PerforceArchiveInfo>();
 
 			// Find all the zipped binaries under this stream
-			ConfigSection ProjectConfigSection = LatestProjectConfigFile.FindSection(SelectedProjectIdentifier);
-			if (ProjectConfigSection != null)
+			ConfigSection projectConfigSection = LatestProjectConfigFile.FindSection(_selectedProjectIdentifier);
+			if (projectConfigSection != null)
 			{
 				// Legacy
-				string? LegacyEditorArchivePath = ProjectConfigSection.GetValue("ZippedBinariesPath", null);
-				if (LegacyEditorArchivePath != null)
+				string? legacyEditorArchivePath = projectConfigSection.GetValue("ZippedBinariesPath", null);
+				if (legacyEditorArchivePath != null)
 				{
-					NewArchives.Add(new PerforceArchiveInfo("Editor", "Editor", LegacyEditorArchivePath, null, this));
+					newArchives.Add(new PerforceArchiveInfo("Editor", "Editor", legacyEditorArchivePath, null, this));
 				}
 
 				// New style
-				foreach (string ArchiveValue in ProjectConfigSection.GetValues("Archives", new string[0]))
+				foreach (string archiveValue in projectConfigSection.GetValues("Archives", new string[0]))
 				{
-					PerforceArchiveInfo? Archive;
-					if (PerforceArchiveInfo.TryParseConfigEntry(ArchiveValue, this, out Archive))
+					PerforceArchiveInfo? archive;
+					if (PerforceArchiveInfo.TryParseConfigEntry(archiveValue, this, out archive))
 					{
-						NewArchives.Add(Archive);
+						newArchives.Add(archive);
 					}
 				}
 
 				// Make sure the zipped binaries path exists
-				foreach (PerforceArchiveInfo NewArchive in NewArchives)
+				foreach (PerforceArchiveInfo newArchive in newArchives)
 				{
-					PerforceResponseList<FileLogRecord> Response = await Perforce.TryFileLogAsync(30, FileLogOptions.FullDescriptions, NewArchive.DepotPath, CancellationToken);
-					if (Response.Succeeded)
+					PerforceResponseList<FileLogRecord> response = await perforce.TryFileLogAsync(30, FileLogOptions.FullDescriptions, newArchive.DepotPath, cancellationToken);
+					if (response.Succeeded)
 					{
  						// Build a new list of zipped binaries
-						foreach (FileLogRecord File in Response.Data)
+						foreach (FileLogRecord file in response.Data)
 						{
-							foreach(RevisionRecord Revision in File.Revisions)
+							foreach(RevisionRecord revision in file.Revisions)
 							{
-								if (Revision.Action != FileAction.Purge)
+								if (revision.Action != FileAction.Purge)
 								{
-									string[] Tokens = Revision.Description.Split(' ');
-									if (Tokens[0].StartsWith("[CL") && Tokens[1].EndsWith("]"))
+									string[] tokens = revision.Description.Split(' ');
+									if (tokens[0].StartsWith("[CL") && tokens[1].EndsWith("]"))
 									{
-										int OriginalChangeNumber;
-										if (int.TryParse(Tokens[1].Substring(0, Tokens[1].Length - 1), out OriginalChangeNumber) && !NewArchive.ChangeNumberToFileRevision.ContainsKey(OriginalChangeNumber))
+										int originalChangeNumber;
+										if (int.TryParse(tokens[1].Substring(0, tokens[1].Length - 1), out originalChangeNumber) && !newArchive.ChangeNumberToFileRevision.ContainsKey(originalChangeNumber))
 										{
-											NewArchive.ChangeNumberToFileRevision[OriginalChangeNumber] = $"{NewArchive.DepotPath}#{Revision.RevisionNumber}";
+											newArchive.ChangeNumberToFileRevision[originalChangeNumber] = $"{newArchive.DepotPath}#{revision.RevisionNumber}";
 										}
 									}
 								}
@@ -649,39 +648,39 @@ namespace UnrealGameSync
 			}
 
 			// Check if the information has changed
-			if (!Enumerable.SequenceEqual(Archives, NewArchives))
+			if (!Enumerable.SequenceEqual(_archives, newArchives))
 			{
-				Archives = NewArchives;
-				AvailableArchives = Archives.Select(x => (IArchiveInfo)x).ToList();
+				_archives = newArchives;
+				AvailableArchives = _archives.Select(x => (IArchiveInfo)x).ToList();
 
-				if (Changes.Count > 0)
+				if (_changes.Count > 0)
 				{
-					SynchronizationContext.Post(_ => OnUpdateMetadata?.Invoke(), null);
+					_synchronizationContext.Post(_ => OnUpdateMetadata?.Invoke(), null);
 				}
 			}
 
 			return true;
 		}
 
-		async Task UpdateProjectConfigFileAsync(IPerforceConnection Perforce, CancellationToken CancellationToken)
+		async Task UpdateProjectConfigFileAsync(IPerforceConnection perforce, CancellationToken cancellationToken)
 		{
-			LocalConfigFiles.Clear();
-			LatestProjectConfigFile = await ConfigUtils.ReadProjectConfigFileAsync(Perforce, BranchClientPath, SelectedClientFileName, CacheFolder, LocalConfigFiles, Logger, CancellationToken);
+			_localConfigFiles.Clear();
+			LatestProjectConfigFile = await ConfigUtils.ReadProjectConfigFileAsync(perforce, _branchClientPath, _selectedClientFileName, _cacheFolder, _localConfigFiles, _logger, cancellationToken);
 		}
 
 		public List<ChangesRecord> GetChanges()
 		{
 			lock(this)
 			{
-				return new List<ChangesRecord>(Changes);
+				return new List<ChangesRecord>(_changes);
 			}
 		}
 
-		public bool TryGetChangeDetails(int Number, [NotNullWhen(true)] out PerforceChangeDetails? Details)
+		public bool TryGetChangeDetails(int number, [NotNullWhen(true)] out PerforceChangeDetails? details)
 		{
 			lock(this)
 			{
-				return ChangeDetails.TryGetValue(Number, out Details);
+				return _changeDetails.TryGetValue(number, out details);
 			}
 		}
 
@@ -689,7 +688,7 @@ namespace UnrealGameSync
 		{
 			lock(this)
 			{
-				return new HashSet<int>(PromotedChangeNumbers);
+				return new HashSet<int>(_promotedChangeNumbers);
 			}
 		}
 
@@ -719,7 +718,7 @@ namespace UnrealGameSync
 
 		public void Refresh()
 		{
-			RefreshEvent.Set();
+			_refreshEvent.Set();
 		}
 	}
 }
diff --git a/Engine/Source/Programs/UnrealGameSync/UnrealGameSync/Program.cs b/Engine/Source/Programs/UnrealGameSync/UnrealGameSync/Program.cs
--- a/Engine/Source/Programs/UnrealGameSync/UnrealGameSync/Program.cs
+++ b/Engine/Source/Programs/UnrealGameSync/UnrealGameSync/Program.cs
@@ -30,8 +30,8 @@ namespace UnrealGameSync
 	{
 		public static string GetVersionString()
 		{
-			AssemblyInformationalVersionAttribute? Version = Assembly.GetExecutingAssembly().GetCustomAttribute<AssemblyInformationalVersionAttribute>();
-			return Version?.InformationalVersion ?? "Unknown";
+			AssemblyInformationalVersionAttribute? version = Assembly.GetExecutingAssembly().GetCustomAttribute<AssemblyInformationalVersionAttribute>();
+			return version?.InformationalVersion ?? "Unknown";
 		}
 
 		public static string? SyncVersion = null;
@@ -45,7 +45,7 @@ namespace UnrealGameSync
 		}
 
 		[STAThread]
-		static void Main(string[] Args)
+		static void Main(string[] args)
 		{
 			if (DeploymentSettings.SentryDsn != null)
 			{
@@ -67,40 +67,40 @@ namespace UnrealGameSync
 				TaskScheduler.UnobservedTaskException += Application_UnobservedException_Sentry;
 			}
 
-			bool bFirstInstance;
-			using (Mutex InstanceMutex = new Mutex(true, "UnrealGameSyncRunning", out bFirstInstance))
+			bool firstInstance;
+			using (Mutex instanceMutex = new Mutex(true, "UnrealGameSyncRunning", out firstInstance))
 			{
-				if (bFirstInstance)
+				if (firstInstance)
 				{
 					Application.EnableVisualStyles();
 					Application.SetCompatibleTextRenderingDefault(false);
 				}
 
-				using (EventWaitHandle ActivateEvent = new EventWaitHandle(false, EventResetMode.AutoReset, "ActivateUnrealGameSync"))
+				using (EventWaitHandle activateEvent = new EventWaitHandle(false, EventResetMode.AutoReset, "ActivateUnrealGameSync"))
 				{
 					// handle any url passed in, possibly exiting
-					if (UriHandler.ProcessCommandLine(Args, bFirstInstance, ActivateEvent))
+					if (UriHandler.ProcessCommandLine(args, firstInstance, activateEvent))
 					{
 						return;
 					}
 
-					if (bFirstInstance)
+					if (firstInstance)
 					{
-						GuardedInnerMainAsync(InstanceMutex, ActivateEvent, Args);
+						GuardedInnerMainAsync(instanceMutex, activateEvent, args);
 					}
 					else
 					{
-						ActivateEvent.Set();
+						activateEvent.Set();
 					}
 				}
 			}
 		}
 
-		static void GuardedInnerMainAsync(Mutex InstanceMutex, EventWaitHandle ActivateEvent, string[] Args)
+		static void GuardedInnerMainAsync(Mutex instanceMutex, EventWaitHandle activateEvent, string[] args)
 		{
 			try
 			{
-				InnerMainAsync(InstanceMutex, ActivateEvent, Args).GetAwaiter().GetResult();
+				InnerMainAsync(instanceMutex, activateEvent, args).GetAwaiter().GetResult();
 			}
 			catch (Exception ex)
 			{
@@ -108,55 +108,55 @@ namespace UnrealGameSync
 			}
 		}
 
-		static async Task InnerMainAsync(Mutex InstanceMutex, EventWaitHandle ActivateEvent, string[] Args)
+		static async Task InnerMainAsync(Mutex instanceMutex, EventWaitHandle activateEvent, string[] args)
 		{
-			string? ServerAndPort = null;
-			string? UserName = null;
-			string? BaseUpdatePath = null;
-			bool bPreviewSetting = false;
-			GlobalPerforceSettings.ReadGlobalPerforceSettings(ref ServerAndPort, ref UserName, ref BaseUpdatePath, ref bPreviewSetting);
+			string? serverAndPort = null;
+			string? userName = null;
+			string? baseUpdatePath = null;
+			bool previewSetting = false;
+			GlobalPerforceSettings.ReadGlobalPerforceSettings(ref serverAndPort, ref userName, ref baseUpdatePath, ref previewSetting);
 
-			List<string> RemainingArgs = new List<string>(Args);
+			List<string> remainingArgs = new List<string>(args);
 
-			string? UpdateSpawn;
-			ParseArgument(RemainingArgs, "-updatespawn=", out UpdateSpawn);
+			string? updateSpawn;
+			ParseArgument(remainingArgs, "-updatespawn=", out updateSpawn);
 
-			string? UpdatePath;
-			ParseArgument(RemainingArgs, "-updatepath=", out UpdatePath);
+			string? updatePath;
+			ParseArgument(remainingArgs, "-updatepath=", out updatePath);
 
-			bool bRestoreState;
-			ParseOption(RemainingArgs, "-restorestate", out bRestoreState);
+			bool restoreState;
+			ParseOption(remainingArgs, "-restorestate", out restoreState);
 
-			bool bUnstable;
-			ParseOption(RemainingArgs, "-unstable", out bUnstable);
-			bool bPreview;
-			ParseOption(RemainingArgs, "-preview", out bPreview);
-			bPreview |= bUnstable;
+			bool unstable;
+			ParseOption(remainingArgs, "-unstable", out unstable);
+			bool preview;
+			ParseOption(remainingArgs, "-preview", out preview);
+			preview |= unstable;
 
-            string? ProjectFileName;
-            ParseArgument(RemainingArgs, "-project=", out ProjectFileName);
+            string? projectFileName;
+            ParseArgument(remainingArgs, "-project=", out projectFileName);
 
-			string? Uri;
-			ParseArgument(RemainingArgs, "-uri=", out Uri);
+			string? uri;
+			ParseArgument(remainingArgs, "-uri=", out uri);
 
-			FileReference UpdateConfigFile = FileReference.Combine(new FileReference(Assembly.GetExecutingAssembly().Location).Directory, "AutoUpdate.ini");
-			MergeUpdateSettings(UpdateConfigFile, ref UpdatePath, ref UpdateSpawn);
+			FileReference updateConfigFile = FileReference.Combine(new FileReference(Assembly.GetExecutingAssembly().Location).Directory, "AutoUpdate.ini");
+			MergeUpdateSettings(updateConfigFile, ref updatePath, ref updateSpawn);
 
 			// Set the current working directory to the update directory to prevent child-process file handles from disrupting auto-updates
-			if (UpdateSpawn != null)
+			if (updateSpawn != null)
 			{
-				if (File.Exists(UpdateSpawn))
+				if (File.Exists(updateSpawn))
 				{
-					Directory.SetCurrentDirectory(Path.GetDirectoryName(UpdateSpawn));
+					Directory.SetCurrentDirectory(Path.GetDirectoryName(updateSpawn));
 				}
 			}
 
-			string SyncVersionFile = Path.Combine(Path.GetDirectoryName(Assembly.GetExecutingAssembly().Location!)!, "SyncVersion.txt");
-			if(File.Exists(SyncVersionFile))
+			string syncVersionFile = Path.Combine(Path.GetDirectoryName(Assembly.GetExecutingAssembly().Location!)!, "SyncVersion.txt");
+			if(File.Exists(syncVersionFile))
 			{
 				try
 				{
-					SyncVersion = File.ReadAllText(SyncVersionFile).Trim();
+					SyncVersion = File.ReadAllText(syncVersionFile).Trim();
 				}
 				catch(Exception)
 				{
@@ -164,96 +164,96 @@ namespace UnrealGameSync
 				}
 			}
 
-			DirectoryReference DataFolder = DirectoryReference.Combine(DirectoryReference.GetSpecialFolder(Environment.SpecialFolder.LocalApplicationData)!, "UnrealGameSync");
-			DirectoryReference.CreateDirectory(DataFolder);
+			DirectoryReference dataFolder = DirectoryReference.Combine(DirectoryReference.GetSpecialFolder(Environment.SpecialFolder.LocalApplicationData)!, "UnrealGameSync");
+			DirectoryReference.CreateDirectory(dataFolder);
 
 			// Enable TLS 1.1 and 1.2. TLS 1.0 is now deprecated and not allowed by default in NET Core servers.
 			ServicePointManager.SecurityProtocol |= SecurityProtocolType.Tls11 | SecurityProtocolType.Tls12;
 
 			// Create a new logger
-			using (ILoggerProvider LoggerProvider = Logging.CreateLoggerProvider(FileReference.Combine(DataFolder, "UnrealGameSync.log")))
+			using (ILoggerProvider loggerProvider = Logging.CreateLoggerProvider(FileReference.Combine(dataFolder, "UnrealGameSync.log")))
 			{
-				ServiceCollection Services = new ServiceCollection();
-				Services.AddLogging(Builder => Builder.AddProvider(LoggerProvider));
-				Services.AddSingleton<IAsyncDisposer, AsyncDisposer>();
+				ServiceCollection services = new ServiceCollection();
+				services.AddLogging(builder => builder.AddProvider(loggerProvider));
+				services.AddSingleton<IAsyncDisposer, AsyncDisposer>();
 
-				await using (ServiceProvider ServiceProvider = Services.BuildServiceProvider())
+				await using (ServiceProvider serviceProvider = services.BuildServiceProvider())
 				{
-					ILoggerFactory LoggerFactory = ServiceProvider.GetRequiredService<ILoggerFactory>();
+					ILoggerFactory loggerFactory = serviceProvider.GetRequiredService<ILoggerFactory>();
 
-					ILogger Logger = LoggerFactory.CreateLogger("Startup");
-					Logger.LogInformation("Application version: {Version}", Assembly.GetExecutingAssembly().GetName().Version);
-					Logger.LogInformation("Started at {Time}", DateTime.Now.ToString());
+					ILogger logger = loggerFactory.CreateLogger("Startup");
+					logger.LogInformation("Application version: {Version}", Assembly.GetExecutingAssembly().GetName().Version);
+					logger.LogInformation("Started at {Time}", DateTime.Now.ToString());
 
-					string SessionId = Guid.NewGuid().ToString();
-					Logger.LogInformation("SessionId: {SessionId}", SessionId);
+					string sessionId = Guid.NewGuid().ToString();
+					logger.LogInformation("SessionId: {SessionId}", sessionId);
 
-					if (ServerAndPort == null || UserName == null)
+					if (serverAndPort == null || userName == null)
 					{
-						Logger.LogInformation("Missing server settings; finding defaults.");
-						ServerAndPort ??= PerforceSettings.Default.ServerAndPort;
-						UserName ??= PerforceSettings.Default.UserName;
-						GlobalPerforceSettings.SaveGlobalPerforceSettings(ServerAndPort, UserName, BaseUpdatePath, bPreviewSetting);
+						logger.LogInformation("Missing server settings; finding defaults.");
+						serverAndPort ??= PerforceSettings.Default.ServerAndPort;
+						userName ??= PerforceSettings.Default.UserName;
+						GlobalPerforceSettings.SaveGlobalPerforceSettings(serverAndPort, userName, baseUpdatePath, previewSetting);
 					}
 
-					ILogger TelemetryLogger = LoggerProvider.CreateLogger("Telemetry");
-					TelemetryLogger.LogInformation("Creating telemetry sink for session {SessionId}", SessionId);
+					ILogger telemetryLogger = loggerProvider.CreateLogger("Telemetry");
+					telemetryLogger.LogInformation("Creating telemetry sink for session {SessionId}", sessionId);
 
-					using (ITelemetrySink TelemetrySink = DeploymentSettings.CreateTelemetrySink(UserName, SessionId, TelemetryLogger))
+					using (ITelemetrySink telemetrySink = DeploymentSettings.CreateTelemetrySink(userName, sessionId, telemetryLogger))
 					{
-						ITelemetrySink? PrevTelemetrySink = Telemetry.ActiveSink;
+						ITelemetrySink? prevTelemetrySink = Telemetry.ActiveSink;
 						try
 						{
-							Telemetry.ActiveSink = TelemetrySink;
+							Telemetry.ActiveSink = telemetrySink;
 
 							Telemetry.SendEvent("Startup", new { User = Environment.UserName, Machine = Environment.MachineName });
 
 							AppDomain.CurrentDomain.UnhandledException += CurrentDomain_UnhandledException;
 
-							IPerforceSettings DefaultSettings = new PerforceSettings(ServerAndPort, UserName) { PreferNativeClient = true };
+							IPerforceSettings defaultSettings = new PerforceSettings(serverAndPort, userName) { PreferNativeClient = true };
 
-							using (UpdateMonitor UpdateMonitor = new UpdateMonitor(DefaultSettings, UpdatePath, ServiceProvider))
+							using (UpdateMonitor updateMonitor = new UpdateMonitor(defaultSettings, updatePath, serviceProvider))
 							{
-								using ProgramApplicationContext Context = new ProgramApplicationContext(DefaultSettings, UpdateMonitor, DeploymentSettings.ApiUrl, DataFolder, ActivateEvent, bRestoreState, UpdateSpawn, ProjectFileName, bPreview, ServiceProvider, Uri);
-								Application.Run(Context);
+								using ProgramApplicationContext context = new ProgramApplicationContext(defaultSettings, updateMonitor, DeploymentSettings.ApiUrl, dataFolder, activateEvent, restoreState, updateSpawn, projectFileName, preview, serviceProvider, uri);
+								Application.Run(context);
 
-								if (UpdateMonitor.IsUpdateAvailable && UpdateSpawn != null)
+								if (updateMonitor.IsUpdateAvailable && updateSpawn != null)
 								{
-									InstanceMutex.Close();
-									bool bLaunchPreview = UpdateMonitor.RelaunchPreview ?? bPreview;
-									Utility.SpawnProcess(UpdateSpawn, "-restorestate" + (bLaunchPreview ? " -unstable" : ""));
+									instanceMutex.Close();
+									bool launchPreview = updateMonitor.RelaunchPreview ?? preview;
+									Utility.SpawnProcess(updateSpawn, "-restorestate" + (launchPreview ? " -unstable" : ""));
 								}
 							}
 						}
-						catch (Exception Ex)
+						catch (Exception ex)
 						{
-							Telemetry.SendEvent("Crash", new { Exception = Ex.ToString() });
+							Telemetry.SendEvent("Crash", new { Exception = ex.ToString() });
 							throw;
 						}
 						finally
 						{
-							Telemetry.ActiveSink = PrevTelemetrySink;
+							Telemetry.ActiveSink = prevTelemetrySink;
 						}
 					}
 				}
 			}
 		}
 
-		private static void CurrentDomain_UnhandledException(object Sender, UnhandledExceptionEventArgs Args)
+		private static void CurrentDomain_UnhandledException(object sender, UnhandledExceptionEventArgs args)
 		{
-			Exception? Ex = Args.ExceptionObject as Exception;
-			if(Ex != null)
+			Exception? ex = args.ExceptionObject as Exception;
+			if(ex != null)
 			{
-				Telemetry.SendEvent("Crash", new {Exception = Ex.ToString()});
+				Telemetry.SendEvent("Crash", new {Exception = ex.ToString()});
 			}
 		}
 
-		private static void CurrentDomain_UnhandledException_Sentry(object Sender, UnhandledExceptionEventArgs Args)
+		private static void CurrentDomain_UnhandledException_Sentry(object sender, UnhandledExceptionEventArgs args)
 		{
-			Exception? Ex = Args.ExceptionObject as Exception;
-			if (Ex != null)
+			Exception? ex = args.ExceptionObject as Exception;
+			if (ex != null)
 			{
-				SentrySdk.CaptureException(Ex);
+				SentrySdk.CaptureException(ex);
 			}
 		}
 
@@ -274,117 +274,117 @@ namespace UnrealGameSync
 			}
 		}
 
-		static void MergeUpdateSettings(FileReference UpdateConfigFile, ref string? UpdatePath, ref string? UpdateSpawn)
+		static void MergeUpdateSettings(FileReference updateConfigFile, ref string? updatePath, ref string? updateSpawn)
 		{
 			try
 			{
-				ConfigFile UpdateConfig = new ConfigFile();
-				if(FileReference.Exists(UpdateConfigFile))
+				ConfigFile updateConfig = new ConfigFile();
+				if(FileReference.Exists(updateConfigFile))
 				{
-					UpdateConfig.Load(UpdateConfigFile);
+					updateConfig.Load(updateConfigFile);
 				}
 
-				if(UpdatePath == null)
+				if(updatePath == null)
 				{
-					UpdatePath = UpdateConfig.GetValue("Update.Path", null);
+					updatePath = updateConfig.GetValue("Update.Path", null);
 				}
 				else
 				{
-					UpdateConfig.SetValue("Update.Path", UpdatePath);
+					updateConfig.SetValue("Update.Path", updatePath);
 				}
 
-				if(UpdateSpawn == null)
+				if(updateSpawn == null)
 				{
-					UpdateSpawn = UpdateConfig.GetValue("Update.Spawn", null);
+					updateSpawn = updateConfig.GetValue("Update.Spawn", null);
 				}
 				else
 				{
-					UpdateConfig.SetValue("Update.Spawn", UpdateSpawn);
+					updateConfig.SetValue("Update.Spawn", updateSpawn);
 				}
 
-				UpdateConfig.Save(UpdateConfigFile);
+				updateConfig.Save(updateConfigFile);
 			}
 			catch(Exception)
 			{
 			}
 		}
 
-		static bool ParseOption(List<string> RemainingArgs, string Option, out bool Value)
+		static bool ParseOption(List<string> remainingArgs, string option, out bool value)
 		{
-			for(int Idx = 0; Idx < RemainingArgs.Count; Idx++)
+			for(int idx = 0; idx < remainingArgs.Count; idx++)
 			{
-				if(RemainingArgs[Idx].Equals(Option, StringComparison.InvariantCultureIgnoreCase))
+				if(remainingArgs[idx].Equals(option, StringComparison.InvariantCultureIgnoreCase))
 				{
-					Value = true;
-					RemainingArgs.RemoveAt(Idx);
+					value = true;
+					remainingArgs.RemoveAt(idx);
 					return true;
 				}
 			}
 
-			Value = false;
+			value = false;
 			return false;
 		}
 
-		static bool ParseArgument(List<string> RemainingArgs, string Prefix, [NotNullWhen(true)] out string? Value)
+		static bool ParseArgument(List<string> remainingArgs, string prefix, [NotNullWhen(true)] out string? value)
 		{
-			for(int Idx = 0; Idx < RemainingArgs.Count; Idx++)
+			for(int idx = 0; idx < remainingArgs.Count; idx++)
 			{
-				if(RemainingArgs[Idx].StartsWith(Prefix, StringComparison.InvariantCultureIgnoreCase))
+				if(remainingArgs[idx].StartsWith(prefix, StringComparison.InvariantCultureIgnoreCase))
 				{
-					Value = RemainingArgs[Idx].Substring(Prefix.Length);
-					RemainingArgs.RemoveAt(Idx);
+					value = remainingArgs[idx].Substring(prefix.Length);
+					remainingArgs.RemoveAt(idx);
 					return true;
 				}
 			}
 
-			Value = null;
+			value = null;
 			return false;
 		}
 
 		public static IEnumerable<string> GetPerforcePaths()
 		{
-			string? PathList = Environment.GetEnvironmentVariable("PATH");
-			if (!String.IsNullOrEmpty(PathList))
+			string? pathList = Environment.GetEnvironmentVariable("PATH");
+			if (!String.IsNullOrEmpty(pathList))
 			{
-				foreach (string PathEntry in PathList.Split(Path.PathSeparator))
+				foreach (string pathEntry in pathList.Split(Path.PathSeparator))
 				{
-					string? PerforcePath = null;
+					string? perforcePath = null;
 					try
 					{
-						string TestPerforcePath = Path.Combine(PathEntry, "p4.exe");
-						if (File.Exists(TestPerforcePath))
+						string testPerforcePath = Path.Combine(pathEntry, "p4.exe");
+						if (File.Exists(testPerforcePath))
 						{
-							PerforcePath = TestPerforcePath;
+							perforcePath = testPerforcePath;
 						}
 					}
 					catch
 					{
 					}
 
-					if (PerforcePath != null)
+					if (perforcePath != null)
 					{
-						yield return PerforcePath;
+						yield return perforcePath;
 					}
 				}
 			}
 		}
 
-		public static void SpawnP4VC(string Arguments)
+		public static void SpawnP4Vc(string arguments)
 		{
-			string Executable = "p4vc.exe";
+			string executable = "p4vc.exe";
 
-			foreach (string PerforcePath in GetPerforcePaths())
+			foreach (string perforcePath in GetPerforcePaths())
 			{
-				string? PerforceDir = Path.GetDirectoryName(PerforcePath);
-				if (PerforceDir != null && File.Exists(Path.Combine(PerforceDir, "p4vc.bat")) && !File.Exists(Path.Combine(PerforceDir, "p4vc.exe")))
+				string? perforceDir = Path.GetDirectoryName(perforcePath);
+				if (perforceDir != null && File.Exists(Path.Combine(perforceDir, "p4vc.bat")) && !File.Exists(Path.Combine(perforceDir, "p4vc.exe")))
 				{
-					Executable = Path.Combine(PerforceDir, "p4v.exe");
-					Arguments = "-p4vc " + Arguments;
+					executable = Path.Combine(perforceDir, "p4v.exe");
+					arguments = "-p4vc " + arguments;
 					break;
 				}
 			}
 
-			if (!Utility.SpawnHiddenProcess(Executable, Arguments))
+			if (!Utility.SpawnHiddenProcess(executable, arguments))
 			{
 				MessageBox.Show("Unable to spawn p4vc. Check you have P4V installed.");
 			}
