UGS: Reformat to standard C# style conventions.

diff --git a/Engine/Source/Programs/UnrealGameSync/UnrealGameSync/ProgramApplicationContext.cs b/Engine/Source/Programs/UnrealGameSync/UnrealGameSync/ProgramApplicationContext.cs
--- a/Engine/Source/Programs/UnrealGameSync/UnrealGameSync/ProgramApplicationContext.cs
+++ b/Engine/Source/Programs/UnrealGameSync/UnrealGameSync/ProgramApplicationContext.cs
@@ -23,58 +23,58 @@ namespace UnrealGameSync
 {
 	class ProgramApplicationContext : ApplicationContext
 	{
-		SynchronizationContext MainThreadSynchronizationContext;
-
-		IPerforceSettings DefaultPerforceSettings;
-		UpdateMonitor UpdateMonitor;
-		string? ApiUrl;
-		DirectoryReference DataFolder;
-		DirectoryReference CacheFolder;
-		bool bRestoreState;
-		string? UpdateSpawn;
-		bool bPreview;
-		bool bIsClosing;
-		string? Uri;
-
-		IServiceProvider ServiceProvider;
-		ILogger Logger;
-		UserSettings Settings;
-		ActivationListener ActivationListener;
-
-		Container Components = new Container();
-		NotifyIcon NotifyIcon;
-		ContextMenuStrip NotifyMenu;
-		ToolStripMenuItem NotifyMenu_OpenUnrealGameSync;
-		ToolStripSeparator NotifyMenu_OpenUnrealGameSync_Separator;
-		ToolStripMenuItem NotifyMenu_SyncNow;
-		ToolStripMenuItem NotifyMenu_LaunchEditor;
-		ToolStripSeparator NotifyMenu_ExitSeparator;
-		ToolStripMenuItem NotifyMenu_Exit;
-
-		CancellationTokenSource StartupCancellationSource = new CancellationTokenSource();
-		Task StartupTask;
-		ModalTaskWindow? StartupWindow;
-		MainWindow? MainWindowInstance;
-
-		OIDCTokenManager? OIDCTokenManager;
-
-		public ProgramApplicationContext(IPerforceSettings DefaultPerforceSettings, UpdateMonitor UpdateMonitor, string? ApiUrl, DirectoryReference DataFolder, EventWaitHandle ActivateEvent, bool bRestoreState, string? UpdateSpawn, string? ProjectFileName, bool bPreview, IServiceProvider ServiceProvider, string? Uri)
+		SynchronizationContext _mainThreadSynchronizationContext;
+
+		IPerforceSettings _defaultPerforceSettings;
+		UpdateMonitor _updateMonitor;
+		string? _apiUrl;
+		DirectoryReference _dataFolder;
+		DirectoryReference _cacheFolder;
+		bool _restoreState;
+		string? _updateSpawn;
+		bool _preview;
+		bool _isClosing;
+		string? _uri;
+
+		IServiceProvider _serviceProvider;
+		ILogger _logger;
+		UserSettings _settings;
+		ActivationListener _activationListener;
+
+		Container _components = new Container();
+		NotifyIcon _notifyIcon;
+		ContextMenuStrip _notifyMenu;
+		ToolStripMenuItem _notifyMenuOpenUnrealGameSync;
+		ToolStripSeparator _notifyMenuOpenUnrealGameSyncSeparator;
+		ToolStripMenuItem _notifyMenuSyncNow;
+		ToolStripMenuItem _notifyMenuLaunchEditor;
+		ToolStripSeparator _notifyMenuExitSeparator;
+		ToolStripMenuItem _notifyMenuExit;
+
+		CancellationTokenSource _startupCancellationSource = new CancellationTokenSource();
+		Task _startupTask;
+		ModalTaskWindow? _startupWindow;
+		MainWindow? _mainWindowInstance;
+
+		OidcTokenManager? _oidcTokenManager;
+
+		public ProgramApplicationContext(IPerforceSettings defaultPerforceSettings, UpdateMonitor updateMonitor, string? apiUrl, DirectoryReference dataFolder, EventWaitHandle activateEvent, bool restoreState, string? updateSpawn, string? projectFileName, bool preview, IServiceProvider serviceProvider, string? uri)
 		{
-			this.DefaultPerforceSettings = DefaultPerforceSettings;
-			this.UpdateMonitor = UpdateMonitor;
-			this.ApiUrl = ApiUrl;
-			this.DataFolder = DataFolder;
-			this.CacheFolder = DirectoryReference.Combine(DataFolder, "Cache");
-			this.bRestoreState = bRestoreState;
-			this.UpdateSpawn = UpdateSpawn;
-			this.bPreview = bPreview;
-			this.ServiceProvider = ServiceProvider;
-			this.Logger = ServiceProvider.GetRequiredService<ILogger<ProgramApplicationContext>>();
-			this.Uri = Uri;
+			this._defaultPerforceSettings = defaultPerforceSettings;
+			this._updateMonitor = updateMonitor;
+			this._apiUrl = apiUrl;
+			this._dataFolder = dataFolder;
+			this._cacheFolder = DirectoryReference.Combine(dataFolder, "Cache");
+			this._restoreState = restoreState;
+			this._updateSpawn = updateSpawn;
+			this._preview = preview;
+			this._serviceProvider = serviceProvider;
+			this._logger = serviceProvider.GetRequiredService<ILogger<ProgramApplicationContext>>();
+			this._uri = uri;
 
 			// Create the directories
-			DirectoryReference.CreateDirectory(DataFolder);
-			DirectoryReference.CreateDirectory(CacheFolder);
+			DirectoryReference.CreateDirectory(dataFolder);
+			DirectoryReference.CreateDirectory(_cacheFolder);
 
 			// Make sure a synchronization context is set. We spawn a bunch of threads (eg. UpdateMonitor) at startup, and need to make sure we can post messages 
 			// back to the main thread at any time.
@@ -84,140 +84,140 @@ namespace UnrealGameSync
 			}
 
 			// Capture the main thread's synchronization context for callbacks
-			MainThreadSynchronizationContext = SynchronizationContext.Current!;
+			_mainThreadSynchronizationContext = SynchronizationContext.Current!;
 
 			// Read the user's settings
-			Settings = UserSettings.Create(DataFolder, ServiceProvider.GetRequiredService<ILogger<UserSettings>>());
-			if(!String.IsNullOrEmpty(ProjectFileName))
+			_settings = UserSettings.Create(dataFolder, serviceProvider.GetRequiredService<ILogger<UserSettings>>());
+			if(!String.IsNullOrEmpty(projectFileName))
 			{
-				string FullProjectFileName = Path.GetFullPath(ProjectFileName);
-				if(!Settings.OpenProjects.Any(x => x.LocalPath != null && String.Compare(x.LocalPath, FullProjectFileName, StringComparison.InvariantCultureIgnoreCase) == 0))
+				string fullProjectFileName = Path.GetFullPath(projectFileName);
+				if(!_settings.OpenProjects.Any(x => x.LocalPath != null && String.Compare(x.LocalPath, fullProjectFileName, StringComparison.InvariantCultureIgnoreCase) == 0))
 				{
-					Settings.OpenProjects.Add(new UserSelectedProjectSettings(null, null, UserSelectedProjectType.Local, null, FullProjectFileName));
+					_settings.OpenProjects.Add(new UserSelectedProjectSettings(null, null, UserSelectedProjectType.Local, null, fullProjectFileName));
 				}
 			}
 
 			// Update the settings to the latest version
-			if(Settings.Version < UserSettingsVersion.Latest)
+			if(_settings.Version < UserSettingsVersion.Latest)
 			{
 				// Clear out the server settings for anything using the default server
-				if(Settings.Version < UserSettingsVersion.DefaultServerSettings)
+				if(_settings.Version < UserSettingsVersion.DefaultServerSettings)
 				{
-					Logger.LogInformation("Clearing project settings for default server");
-					for (int Idx = 0; Idx < Settings.OpenProjects.Count; Idx++)
+					_logger.LogInformation("Clearing project settings for default server");
+					for (int idx = 0; idx < _settings.OpenProjects.Count; idx++)
 					{
-						Settings.OpenProjects[Idx] = UpgradeSelectedProjectSettings(Settings.OpenProjects[Idx]);
+						_settings.OpenProjects[idx] = UpgradeSelectedProjectSettings(_settings.OpenProjects[idx]);
 					}
-					for (int Idx = 0; Idx < Settings.RecentProjects.Count; Idx++)
+					for (int idx = 0; idx < _settings.RecentProjects.Count; idx++)
 					{
-						Settings.RecentProjects[Idx] = UpgradeSelectedProjectSettings(Settings.RecentProjects[Idx]);
+						_settings.RecentProjects[idx] = UpgradeSelectedProjectSettings(_settings.RecentProjects[idx]);
 					}
 				}
 
 				// Save the new settings
-				Settings.Version = UserSettingsVersion.Latest;
-				Settings.Save(Logger);
+				_settings.Version = UserSettingsVersion.Latest;
+				_settings.Save(_logger);
 			}
 
 			// Register the update listener
-			UpdateMonitor.OnUpdateAvailable += OnUpdateAvailableCallback;
+			updateMonitor.OnUpdateAvailable += OnUpdateAvailableCallback;
 
 			// Create the activation listener
-			ActivationListener = new ActivationListener(ActivateEvent);
-			ActivationListener.Start();
-			ActivationListener.OnActivate += OnActivationListenerAsyncCallback;
+			_activationListener = new ActivationListener(activateEvent);
+			_activationListener.Start();
+			_activationListener.OnActivate += OnActivationListenerAsyncCallback;
 
 			// Create the notification menu items
-			NotifyMenu_OpenUnrealGameSync = new ToolStripMenuItem();
-			NotifyMenu_OpenUnrealGameSync.Name = nameof(NotifyMenu_OpenUnrealGameSync);
-			NotifyMenu_OpenUnrealGameSync.Size = new Size(196, 22);
-			NotifyMenu_OpenUnrealGameSync.Text = "Open UnrealGameSync";
-			NotifyMenu_OpenUnrealGameSync.Click += new EventHandler(NotifyMenu_OpenUnrealGameSync_Click);
-			NotifyMenu_OpenUnrealGameSync.Font = new Font(NotifyMenu_OpenUnrealGameSync.Font, FontStyle.Bold);
-
-			NotifyMenu_OpenUnrealGameSync_Separator = new ToolStripSeparator();
-			NotifyMenu_OpenUnrealGameSync_Separator.Name = nameof(NotifyMenu_OpenUnrealGameSync_Separator);
-			NotifyMenu_OpenUnrealGameSync_Separator.Size = new Size(193, 6);
-
-			NotifyMenu_SyncNow = new ToolStripMenuItem();
-			NotifyMenu_SyncNow.Name = nameof(NotifyMenu_SyncNow);
-			NotifyMenu_SyncNow.Size = new Size(196, 22);
-			NotifyMenu_SyncNow.Text = "Sync Now";
-			NotifyMenu_SyncNow.Click += new EventHandler(NotifyMenu_SyncNow_Click);
-
-			NotifyMenu_LaunchEditor = new ToolStripMenuItem();
-			NotifyMenu_LaunchEditor.Name = nameof(NotifyMenu_LaunchEditor);
-			NotifyMenu_LaunchEditor.Size = new Size(196, 22);
-			NotifyMenu_LaunchEditor.Text = "Launch Editor";
-			NotifyMenu_LaunchEditor.Click += new EventHandler(NotifyMenu_LaunchEditor_Click);
-
-			NotifyMenu_ExitSeparator = new ToolStripSeparator();
-			NotifyMenu_ExitSeparator.Name = nameof(NotifyMenu_ExitSeparator);
-			NotifyMenu_ExitSeparator.Size = new Size(193, 6);
-
-			NotifyMenu_Exit = new ToolStripMenuItem();
-			NotifyMenu_Exit.Name = nameof(NotifyMenu_Exit);
-			NotifyMenu_Exit.Size = new Size(196, 22);
-			NotifyMenu_Exit.Text = "Exit";
-			NotifyMenu_Exit.Click += new EventHandler(NotifyMenu_Exit_Click);
+			_notifyMenuOpenUnrealGameSync = new ToolStripMenuItem();
+			_notifyMenuOpenUnrealGameSync.Name = nameof(_notifyMenuOpenUnrealGameSync);
+			_notifyMenuOpenUnrealGameSync.Size = new Size(196, 22);
+			_notifyMenuOpenUnrealGameSync.Text = "Open UnrealGameSync";
+			_notifyMenuOpenUnrealGameSync.Click += new EventHandler(NotifyMenu_OpenUnrealGameSync_Click);
+			_notifyMenuOpenUnrealGameSync.Font = new Font(_notifyMenuOpenUnrealGameSync.Font, FontStyle.Bold);
+
+			_notifyMenuOpenUnrealGameSyncSeparator = new ToolStripSeparator();
+			_notifyMenuOpenUnrealGameSyncSeparator.Name = nameof(_notifyMenuOpenUnrealGameSyncSeparator);
+			_notifyMenuOpenUnrealGameSyncSeparator.Size = new Size(193, 6);
+
+			_notifyMenuSyncNow = new ToolStripMenuItem();
+			_notifyMenuSyncNow.Name = nameof(_notifyMenuSyncNow);
+			_notifyMenuSyncNow.Size = new Size(196, 22);
+			_notifyMenuSyncNow.Text = "Sync Now";
+			_notifyMenuSyncNow.Click += new EventHandler(NotifyMenu_SyncNow_Click);
+
+			_notifyMenuLaunchEditor = new ToolStripMenuItem();
+			_notifyMenuLaunchEditor.Name = nameof(_notifyMenuLaunchEditor);
+			_notifyMenuLaunchEditor.Size = new Size(196, 22);
+			_notifyMenuLaunchEditor.Text = "Launch Editor";
+			_notifyMenuLaunchEditor.Click += new EventHandler(NotifyMenu_LaunchEditor_Click);
+
+			_notifyMenuExitSeparator = new ToolStripSeparator();
+			_notifyMenuExitSeparator.Name = nameof(_notifyMenuExitSeparator);
+			_notifyMenuExitSeparator.Size = new Size(193, 6);
+
+			_notifyMenuExit = new ToolStripMenuItem();
+			_notifyMenuExit.Name = nameof(_notifyMenuExit);
+			_notifyMenuExit.Size = new Size(196, 22);
+			_notifyMenuExit.Text = "Exit";
+			_notifyMenuExit.Click += new EventHandler(NotifyMenu_Exit_Click);
 
 			// Create the notification menu
-			NotifyMenu = new ContextMenuStrip(Components);
-			NotifyMenu.Name = nameof(NotifyMenu);
-			NotifyMenu.Size = new System.Drawing.Size(197, 104);
-			NotifyMenu.SuspendLayout();
-			NotifyMenu.Items.Add(NotifyMenu_OpenUnrealGameSync);
-			NotifyMenu.Items.Add(NotifyMenu_OpenUnrealGameSync_Separator);
-			NotifyMenu.Items.Add(NotifyMenu_SyncNow);
-			NotifyMenu.Items.Add(NotifyMenu_LaunchEditor);
-			NotifyMenu.Items.Add(NotifyMenu_ExitSeparator);
-			NotifyMenu.Items.Add(NotifyMenu_Exit);
-			NotifyMenu.ResumeLayout(false);
+			_notifyMenu = new ContextMenuStrip(_components);
+			_notifyMenu.Name = nameof(_notifyMenu);
+			_notifyMenu.Size = new System.Drawing.Size(197, 104);
+			_notifyMenu.SuspendLayout();
+			_notifyMenu.Items.Add(_notifyMenuOpenUnrealGameSync);
+			_notifyMenu.Items.Add(_notifyMenuOpenUnrealGameSyncSeparator);
+			_notifyMenu.Items.Add(_notifyMenuSyncNow);
+			_notifyMenu.Items.Add(_notifyMenuLaunchEditor);
+			_notifyMenu.Items.Add(_notifyMenuExitSeparator);
+			_notifyMenu.Items.Add(_notifyMenuExit);
+			_notifyMenu.ResumeLayout(false);
 
 			// Create the notification icon
-			NotifyIcon = new NotifyIcon(Components);
-			NotifyIcon.ContextMenuStrip = NotifyMenu;
-			NotifyIcon.Icon = Properties.Resources.Icon;
-			NotifyIcon.Text = "UnrealGameSync";
-			NotifyIcon.Visible = true;
-			NotifyIcon.DoubleClick += new EventHandler(NotifyIcon_DoubleClick);
-			NotifyIcon.MouseDown += new MouseEventHandler(NotifyIcon_MouseDown);
+			_notifyIcon = new NotifyIcon(_components);
+			_notifyIcon.ContextMenuStrip = _notifyMenu;
+			_notifyIcon.Icon = Properties.Resources.Icon;
+			_notifyIcon.Text = "UnrealGameSync";
+			_notifyIcon.Visible = true;
+			_notifyIcon.DoubleClick += new EventHandler(NotifyIcon_DoubleClick);
+			_notifyIcon.MouseDown += new MouseEventHandler(NotifyIcon_MouseDown);
 
 			// Create the startup tasks
-			List<(UserSelectedProjectSettings, ModalTask<OpenProjectInfo>)> StartupTasks = new List<(UserSelectedProjectSettings, ModalTask<OpenProjectInfo>)>();
-			foreach (UserSelectedProjectSettings ProjectSettings in Settings.OpenProjects)
+			List<(UserSelectedProjectSettings, ModalTask<OpenProjectInfo>)> startupTasks = new List<(UserSelectedProjectSettings, ModalTask<OpenProjectInfo>)>();
+			foreach (UserSelectedProjectSettings projectSettings in _settings.OpenProjects)
 			{
-				ILogger<OpenProjectInfo> Logger = ServiceProvider.GetRequiredService<ILogger<OpenProjectInfo>>();
-				Task<OpenProjectInfo> StartupTask = Task.Run(() => OpenProjectInfo.CreateAsync(DefaultPerforceSettings, ProjectSettings, Settings, Logger, StartupCancellationSource.Token), StartupCancellationSource.Token);
-				StartupTasks.Add((ProjectSettings, new ModalTask<OpenProjectInfo>(StartupTask)));
+				ILogger<OpenProjectInfo> logger = serviceProvider.GetRequiredService<ILogger<OpenProjectInfo>>();
+				Task<OpenProjectInfo> startupTask = Task.Run(() => OpenProjectInfo.CreateAsync(defaultPerforceSettings, projectSettings, _settings, logger, _startupCancellationSource.Token), _startupCancellationSource.Token);
+				startupTasks.Add((projectSettings, new ModalTask<OpenProjectInfo>(startupTask)));
 			}
-			StartupTask = Task.Run(() => WaitForStartupTasks(StartupTasks));
+			_startupTask = Task.Run(() => WaitForStartupTasks(startupTasks));
 
-			StartupWindow = new ModalTaskWindow("Opening Projects", "Opening projects, please wait...", FormStartPosition.CenterScreen, StartupTask, StartupCancellationSource);
-			Components.Add(StartupWindow);
+			_startupWindow = new ModalTaskWindow("Opening Projects", "Opening projects, please wait...", FormStartPosition.CenterScreen, _startupTask, _startupCancellationSource);
+			_components.Add(_startupWindow);
 
-			if(bRestoreState)
+			if(restoreState)
 			{
-				if(Settings.bWindowVisible)
+				if(_settings.WindowVisible)
 				{
-					StartupWindow.Show();
+					_startupWindow.Show();
 				}
 			}
 			else
 			{
-				StartupWindow.Show();
-				StartupWindow.Activate();
+				_startupWindow.Show();
+				_startupWindow.Activate();
 			}
-			StartupWindow.FormClosed += (s, e) => OnStartupComplete(StartupTasks);
+			_startupWindow.FormClosed += (s, e) => OnStartupComplete(startupTasks);
 		}
 
-		static async Task WaitForStartupTasks(List<(UserSelectedProjectSettings, ModalTask<OpenProjectInfo>)> StartupTasks)
+		static async Task WaitForStartupTasks(List<(UserSelectedProjectSettings, ModalTask<OpenProjectInfo>)> startupTasks)
 		{
-			foreach ((_, ModalTask<OpenProjectInfo> ModalTask) in StartupTasks)
+			foreach ((_, ModalTask<OpenProjectInfo> modalTask) in startupTasks)
 			{
 				try
 				{
-					await ModalTask.Task;
+					await modalTask.Task;
 				}
 				catch (PerforceException)
 				{
@@ -235,57 +235,57 @@ namespace UnrealGameSync
 			}
 		}
 
-		private UserSelectedProjectSettings UpgradeSelectedProjectSettings(UserSelectedProjectSettings Project)
+		private UserSelectedProjectSettings UpgradeSelectedProjectSettings(UserSelectedProjectSettings project)
 		{
-			if (Project.ServerAndPort == null || String.Compare(Project.ServerAndPort, DefaultPerforceSettings.ServerAndPort, StringComparison.OrdinalIgnoreCase) == 0)
+			if (project.ServerAndPort == null || String.Compare(project.ServerAndPort, _defaultPerforceSettings.ServerAndPort, StringComparison.OrdinalIgnoreCase) == 0)
 			{
-				if (Project.UserName == null || String.Compare(Project.UserName, DefaultPerforceSettings.UserName, StringComparison.OrdinalIgnoreCase) == 0)
+				if (project.UserName == null || String.Compare(project.UserName, _defaultPerforceSettings.UserName, StringComparison.OrdinalIgnoreCase) == 0)
 				{
-					Project = new UserSelectedProjectSettings(null, null, Project.Type, Project.ClientPath, Project.LocalPath);
+					project = new UserSelectedProjectSettings(null, null, project.Type, project.ClientPath, project.LocalPath);
 				}
 			}
-			return Project;
+			return project;
 		}
 
-		private void OnStartupComplete(List<(UserSelectedProjectSettings, ModalTask<OpenProjectInfo>)> StartupTasks)
+		private void OnStartupComplete(List<(UserSelectedProjectSettings, ModalTask<OpenProjectInfo>)> startupTasks)
 		{
 			// Close the startup window
-			bool bVisible = StartupWindow!.Visible;
-			StartupWindow = null;
+			bool visible = _startupWindow!.Visible;
+			_startupWindow = null;
 
 			// Clear out the cache folder
-			Utility.ClearPrintCache(CacheFolder);
+			Utility.ClearPrintCache(_cacheFolder);
 
-			OIDCTokenManager = OIDCTokenManager.CreateFromConfigFile(Settings, StartupTasks.Where(x => x.Item2.Succeeded).Select(x => x.Item2.Result).ToList());
+			_oidcTokenManager = OidcTokenManager.CreateFromConfigFile(_settings, startupTasks.Where(x => x.Item2.Succeeded).Select(x => x.Item2.Result).ToList());
 			// Verify that none of the projects we are opening needs a OIDC login, if they do prompt for the login
-			if (OIDCTokenManager?.HasUnfinishedLogin() ?? false)
+			if (_oidcTokenManager?.HasUnfinishedLogin() ?? false)
 			{
-				OIDCLoginWindow LoginDialog = new OIDCLoginWindow(OIDCTokenManager);
-				LoginDialog.ShowDialog();
+				OidcLoginWindow loginDialog = new OidcLoginWindow(_oidcTokenManager);
+				loginDialog.ShowDialog();
 			}
 
 			// Get the application path
-			string OriginalExe = Assembly.GetExecutingAssembly().Location;
-			if (Path.GetExtension(OriginalExe).Equals(".dll", StringComparison.OrdinalIgnoreCase))
+			string originalExe = Assembly.GetExecutingAssembly().Location;
+			if (Path.GetExtension(originalExe).Equals(".dll", StringComparison.OrdinalIgnoreCase))
 			{
-				string NewExecutable = Path.ChangeExtension(OriginalExe, ".exe");
-				if (File.Exists(NewExecutable))
+				string newExecutable = Path.ChangeExtension(originalExe, ".exe");
+				if (File.Exists(newExecutable))
 				{
-					OriginalExe = NewExecutable;
+					originalExe = newExecutable;
 				}
 			}
 
 			// Create the main window 
-			MainWindowInstance = new MainWindow(UpdateMonitor, ApiUrl, DataFolder, CacheFolder, bRestoreState, UpdateSpawn ?? OriginalExe, bPreview, StartupTasks, DefaultPerforceSettings, ServiceProvider, Settings, Uri, OIDCTokenManager);
-			if(bVisible)
+			_mainWindowInstance = new MainWindow(_updateMonitor, _apiUrl, _dataFolder, _cacheFolder, _restoreState, _updateSpawn ?? originalExe, _preview, startupTasks, _defaultPerforceSettings, _serviceProvider, _settings, _uri, _oidcTokenManager);
+			if(visible)
 			{
-				MainWindowInstance.Show();
-				if(!bRestoreState)
+				_mainWindowInstance.Show();
+				if(!_restoreState)
 				{
-					MainWindowInstance.Activate();
+					_mainWindowInstance.Activate();
 				}
 			}
-			MainWindowInstance.FormClosed += MainWindowInstance_FormClosed;
+			_mainWindowInstance.FormClosed += MainWindowInstance_FormClosed;
 		}
 
 		private void MainWindowInstance_FormClosed(object sender, FormClosedEventArgs e)
@@ -295,153 +295,153 @@ namespace UnrealGameSync
 
 		private void OnActivationListenerCallback()
 		{
-			if(MainWindowInstance != null && !MainWindowInstance.IsDisposed)
+			if(_mainWindowInstance != null && !_mainWindowInstance.IsDisposed)
 			{
-				MainWindowInstance.ShowAndActivate();
+				_mainWindowInstance.ShowAndActivate();
 			}
 		}
 
 		private void OnActivationListenerAsyncCallback()
 		{
-			MainThreadSynchronizationContext.Post((o) => OnActivationListenerCallback(), null);
+			_mainThreadSynchronizationContext.Post((o) => OnActivationListenerCallback(), null);
 		}
 
-		private void OnUpdateAvailable(UpdateType Type)
+		private void OnUpdateAvailable(UpdateType type)
 		{
-			if(MainWindowInstance != null && !bIsClosing)
+			if(_mainWindowInstance != null && !_isClosing)
 			{
-				if(Type == UpdateType.UserInitiated || MainWindowInstance.CanPerformUpdate())
+				if(type == UpdateType.UserInitiated || _mainWindowInstance.CanPerformUpdate())
 				{
-					bIsClosing = true;
-					MainWindowInstance.ForceClose();
-					MainWindowInstance = null;
+					_isClosing = true;
+					_mainWindowInstance.ForceClose();
+					_mainWindowInstance = null;
 				}
 			}
 		}
 
-		private void OnUpdateAvailableCallback(UpdateType Type)
+		private void OnUpdateAvailableCallback(UpdateType type)
 		{ 
-			MainThreadSynchronizationContext.Post((o) => OnUpdateAvailable(Type), null);
+			_mainThreadSynchronizationContext.Post((o) => OnUpdateAvailable(type), null);
 		}
 
-		protected override void Dispose(bool bDisposing)
+		protected override void Dispose(bool disposing)
 		{
-			base.Dispose(bDisposing);
+			base.Dispose(disposing);
 
-			if(UpdateMonitor != null)
+			if(_updateMonitor != null)
 			{
-				UpdateMonitor.Dispose();
-				UpdateMonitor = null!;
+				_updateMonitor.Dispose();
+				_updateMonitor = null!;
 			}
 
-			if(ActivationListener != null)
+			if(_activationListener != null)
 			{
-				ActivationListener.OnActivate -= OnActivationListenerAsyncCallback;
-				ActivationListener.Stop();
-				ActivationListener.Dispose();
-				ActivationListener = null!;
+				_activationListener.OnActivate -= OnActivationListenerAsyncCallback;
+				_activationListener.Stop();
+				_activationListener.Dispose();
+				_activationListener = null!;
 			}
 
-			if (Components != null)
+			if (_components != null)
 			{
-				Components.Dispose();
-				Components = null!;
+				_components.Dispose();
+				_components = null!;
 			}
 
-			if (NotifyIcon != null)
+			if (_notifyIcon != null)
 			{
-				NotifyIcon.Dispose();
-				NotifyIcon = null!;
+				_notifyIcon.Dispose();
+				_notifyIcon = null!;
 			}
 
-			if(MainWindowInstance != null)
+			if(_mainWindowInstance != null)
 			{
-				MainWindowInstance.ForceClose();
-				MainWindowInstance = null!;
+				_mainWindowInstance.ForceClose();
+				_mainWindowInstance = null!;
 			}
 
-			if(StartupWindow != null)
+			if(_startupWindow != null)
 			{
-				StartupWindow.Close();
-				StartupWindow = null;
+				_startupWindow.Close();
+				_startupWindow = null;
 			}
 
-			if (StartupCancellationSource != null)
+			if (_startupCancellationSource != null)
 			{
-				StartupCancellationSource.Dispose();
-				StartupCancellationSource = null!;
+				_startupCancellationSource.Dispose();
+				_startupCancellationSource = null!;
 			}
 		}
 
 		private void NotifyIcon_MouseDown(object sender, MouseEventArgs e)
 		{
 			// Have to set up this stuff here, because the menu is laid out before Opening() is called on it after mouse-up.
-			bool bCanSyncNow = MainWindowInstance != null && MainWindowInstance.CanSyncNow();
-			bool bCanLaunchEditor = MainWindowInstance != null && MainWindowInstance.CanLaunchEditor();
-			NotifyMenu_SyncNow.Visible = bCanSyncNow;
-			NotifyMenu_LaunchEditor.Visible = bCanLaunchEditor;
-			NotifyMenu_ExitSeparator.Visible = bCanSyncNow || bCanLaunchEditor;
+			bool canSyncNow = _mainWindowInstance != null && _mainWindowInstance.CanSyncNow();
+			bool canLaunchEditor = _mainWindowInstance != null && _mainWindowInstance.CanLaunchEditor();
+			_notifyMenuSyncNow.Visible = canSyncNow;
+			_notifyMenuLaunchEditor.Visible = canLaunchEditor;
+			_notifyMenuExitSeparator.Visible = canSyncNow || canLaunchEditor;
 
 			// Show the startup window, if not already visible
-			if(StartupWindow != null)
+			if(_startupWindow != null)
 			{
-				StartupWindow.Show();
+				_startupWindow.Show();
 			}
 		}
 
 		private void NotifyIcon_DoubleClick(object? sender, EventArgs e)
 		{
-			if(MainWindowInstance != null)
+			if(_mainWindowInstance != null)
 			{
-				MainWindowInstance.ShowAndActivate();
+				_mainWindowInstance.ShowAndActivate();
 			}
 		}
 
 		private void NotifyMenu_OpenUnrealGameSync_Click(object? sender, EventArgs e)
 		{
-			if(StartupWindow != null)
+			if(_startupWindow != null)
 			{
-				StartupWindow.ShowAndActivate();
+				_startupWindow.ShowAndActivate();
 			}
-			if(MainWindowInstance != null)
+			if(_mainWindowInstance != null)
 			{
-				MainWindowInstance.ShowAndActivate();
+				_mainWindowInstance.ShowAndActivate();
 			}
 		}
 
 		private void NotifyMenu_SyncNow_Click(object? sender, EventArgs e)
 		{
-			if(MainWindowInstance != null)
+			if(_mainWindowInstance != null)
 			{
-				MainWindowInstance.SyncLatestChange();
+				_mainWindowInstance.SyncLatestChange();
 			}
 		}
 
 		private void NotifyMenu_LaunchEditor_Click(object? sender, EventArgs e)
 		{
-			if(MainWindowInstance != null)
+			if(_mainWindowInstance != null)
 			{
-				MainWindowInstance.LaunchEditor();
+				_mainWindowInstance.LaunchEditor();
 			}
 		}
 
 		private void NotifyMenu_Exit_Click(object? sender, EventArgs e)
 		{
-			if (MainWindowInstance != null && !MainWindowInstance.ConfirmClose())
+			if (_mainWindowInstance != null && !_mainWindowInstance.ConfirmClose())
 			{
 				return;
 			}
 
-			if (StartupWindow != null)
+			if (_startupWindow != null)
 			{
-				StartupWindow.Close();
-				StartupWindow = null;
+				_startupWindow.Close();
+				_startupWindow = null;
 			}
 
-			if(MainWindowInstance != null)
+			if(_mainWindowInstance != null)
 			{
-				MainWindowInstance.ForceClose();
-				MainWindowInstance = null;
+				_mainWindowInstance.ForceClose();
+				_mainWindowInstance = null;
 			}
 
 			ExitThread();
@@ -451,9 +451,9 @@ namespace UnrealGameSync
 		{
 			base.ExitThreadCore();
 
-			if(NotifyIcon != null)
+			if(_notifyIcon != null)
 			{
-				NotifyIcon.Visible = false;
+				_notifyIcon.Visible = false;
 			}
 		}
 	}
diff --git a/Engine/Source/Programs/UnrealGameSync/UnrealGameSync/RESTApi.cs b/Engine/Source/Programs/UnrealGameSync/UnrealGameSync/RESTApi.cs
--- a/Engine/Source/Programs/UnrealGameSync/UnrealGameSync/RESTApi.cs
+++ b/Engine/Source/Programs/UnrealGameSync/UnrealGameSync/RESTApi.cs
@@ -14,8 +14,8 @@ namespace UnrealGameSync
 {
 	class RestException : Exception
 	{
-		public RestException(string Method, string Uri, Exception InnerException)
-			: base(String.Format("Error executing {0} {1}", Method, Uri), InnerException)
+		public RestException(string method, string uri, Exception innerException)
+			: base(String.Format("Error executing {0} {1}", method, uri), innerException)
 		{
 		}
 
@@ -25,63 +25,63 @@ namespace UnrealGameSync
 		}
 	}
 
-	public static class RESTApi
+	public static class RestApi
 	{
-		private static async Task<string> SendRequestInternal(string Url, string Method, string? RequestBody, CancellationToken CancellationToken)
+		private static async Task<string> SendRequestInternal(string url, string method, string? requestBody, CancellationToken cancellationToken)
 		{
-			HttpWebRequest Request = (HttpWebRequest)WebRequest.Create(Url);
-			Request.ContentType = "application/json";
-			Request.Method = Method;
+			HttpWebRequest request = (HttpWebRequest)WebRequest.Create(url);
+			request.ContentType = "application/json";
+			request.Method = method;
 
 			// Add json to request body
-			if (!string.IsNullOrEmpty(RequestBody))
+			if (!string.IsNullOrEmpty(requestBody))
 			{
-				if (Method == "POST" || Method == "PUT")
+				if (method == "POST" || method == "PUT")
 				{
-					byte[] bytes = Encoding.UTF8.GetBytes(RequestBody);
-					using (Stream RequestStream = Request.GetRequestStream())
+					byte[] bytes = Encoding.UTF8.GetBytes(requestBody);
+					using (Stream requestStream = request.GetRequestStream())
 					{
-						await RequestStream.WriteAsync(bytes, 0, bytes.Length, CancellationToken);
+						await requestStream.WriteAsync(bytes, 0, bytes.Length, cancellationToken);
 					}
 				}
 			}
 			try
 			{
-				using (WebResponse Response = Request.GetResponse())
+				using (WebResponse response = request.GetResponse())
 				{
-					byte[] Data;
-					using (MemoryStream Buffer = new MemoryStream())
+					byte[] data;
+					using (MemoryStream buffer = new MemoryStream())
 					{
-						await Response.GetResponseStream().CopyToAsync(Buffer, CancellationToken);
-						Data = Buffer.ToArray();
+						await response.GetResponseStream().CopyToAsync(buffer, cancellationToken);
+						data = buffer.ToArray();
 					}
-					return Encoding.UTF8.GetString(Data);
+					return Encoding.UTF8.GetString(data);
 				}
 			}
-			catch (Exception Ex)
+			catch (Exception ex)
 			{
-				throw new RestException(Method, Request.RequestUri.ToString(), Ex);
+				throw new RestException(method, request.RequestUri.ToString(), ex);
 			}
 		}
 
-		public static Task<string> PostAsync(string Url, string RequestBody, CancellationToken CancellationToken)
+		public static Task<string> PostAsync(string url, string requestBody, CancellationToken cancellationToken)
 		{
-			return SendRequestInternal(Url, "POST", RequestBody, CancellationToken);
+			return SendRequestInternal(url, "POST", requestBody, cancellationToken);
 		}
 
-		public static Task<string> GetAsync(string Url, CancellationToken CancellationToken)
+		public static Task<string> GetAsync(string url, CancellationToken cancellationToken)
 		{
-			return SendRequestInternal(Url, "GET", null, CancellationToken);
+			return SendRequestInternal(url, "GET", null, cancellationToken);
 		}
 
-		public static async Task<T> GetAsync<T>(string Url, CancellationToken CancellationToken)
+		public static async Task<T> GetAsync<T>(string url, CancellationToken cancellationToken)
 		{
-			return JsonSerializer.Deserialize<T>(await GetAsync(Url, CancellationToken), Utility.DefaultJsonSerializerOptions);
+			return JsonSerializer.Deserialize<T>(await GetAsync(url, cancellationToken), Utility.DefaultJsonSerializerOptions);
 		}
 
-		public static Task<string> PutAsync<T>(string Url, T Object, CancellationToken CancellationToken)
+		public static Task<string> PutAsync<T>(string url, T obj, CancellationToken cancellationToken)
 		{
-			return SendRequestInternal(Url, "PUT", JsonSerializer.Serialize(Object), CancellationToken);
+			return SendRequestInternal(url, "PUT", JsonSerializer.Serialize(obj), cancellationToken);
 		}
 	}
 }
diff --git a/Engine/Source/Programs/UnrealGameSync/UnrealGameSync/Taskbar.cs b/Engine/Source/Programs/UnrealGameSync/UnrealGameSync/Taskbar.cs
--- a/Engine/Source/Programs/UnrealGameSync/UnrealGameSync/Taskbar.cs
+++ b/Engine/Source/Programs/UnrealGameSync/UnrealGameSync/Taskbar.cs
@@ -30,7 +30,7 @@ namespace UnrealGameSync
 			void SetActiveAlt(IntPtr hWnd);
 			void MarkFullscreenWindow(IntPtr hWnd, int fFullscreen);
 			void SetProgressValue(IntPtr hWnd, UInt64 ullCompleted, UInt64 ullTotal);
-			void SetProgressState(IntPtr hWnd, TaskbarState State);
+			void SetProgressState(IntPtr hWnd, TaskbarState state);
 		}
 
 		[ComImport, Guid("56FDF344-FD6D-11d0-958A-006097C9A090"), ClassInterfaceAttribute(ClassInterfaceType.None)]
@@ -38,23 +38,23 @@ namespace UnrealGameSync
 		{
 		}
 
-		static ITaskbarList3? Interface;
+		static ITaskbarList3? _interface;
 
 		static Taskbar()
 		{
 			if(Environment.OSVersion.Version >= new Version(6, 1))
 			{
-				Interface = new TaskbarList3() as ITaskbarList3;
+				_interface = new TaskbarList3() as ITaskbarList3;
 			}
 		}
 
-		public static void SetState(IntPtr WindowHandle, TaskbarState State)
+		public static void SetState(IntPtr windowHandle, TaskbarState state)
 		{
-			if(Interface != null)
+			if(_interface != null)
 			{
 				try
 				{
-					Interface.SetProgressState(WindowHandle, State);
+					_interface.SetProgressState(windowHandle, state);
 				}
 				catch
 				{
@@ -62,13 +62,13 @@ namespace UnrealGameSync
 			}
 		}
 
-		public static void SetProgress(IntPtr WindowHandle, ulong Completed, ulong Total)
+		public static void SetProgress(IntPtr windowHandle, ulong completed, ulong total)
 		{
-			if(Interface != null)
+			if(_interface != null)
 			{
 				try
 				{
-					Interface.SetProgressValue(WindowHandle, Completed, Total);
+					_interface.SetProgressValue(windowHandle, completed, total);
 				}
 				catch
 				{
diff --git a/Engine/Source/Programs/UnrealGameSync/UnrealGameSync/ToolUpdateMonitor.cs b/Engine/Source/Programs/UnrealGameSync/UnrealGameSync/ToolUpdateMonitor.cs
--- a/Engine/Source/Programs/UnrealGameSync/UnrealGameSync/ToolUpdateMonitor.cs
+++ b/Engine/Source/Programs/UnrealGameSync/UnrealGameSync/ToolUpdateMonitor.cs
@@ -25,10 +25,10 @@ namespace UnrealGameSync
 		public string? Arguments;
 		public string? WorkingDir;
 
-		public ToolLink(string Label, string FileName)
+		public ToolLink(string label, string fileName)
 		{
-			this.Label = Label;
-			this.FileName = FileName;
+			this.Label = label;
+			this.FileName = fileName;
 		}
 	}
 
@@ -46,28 +46,28 @@ namespace UnrealGameSync
 		public string ConfigPath;
 		public int ConfigChange;
 
-		public ToolDefinition(Guid Id, string Name, string Description, string ZipPath, int ZipChange, string ConfigPath, int ConfigChange)
+		public ToolDefinition(Guid id, string name, string description, string zipPath, int zipChange, string configPath, int configChange)
 		{
-			this.Id = Id;
-			this.Name = Name;
-			this.Description = Description;
-			this.ZipPath = ZipPath;
-			this.ZipChange = ZipChange;
-			this.ConfigPath = ConfigPath;
-			this.ConfigChange = ConfigChange;
+			this.Id = id;
+			this.Name = name;
+			this.Description = description;
+			this.ZipPath = zipPath;
+			this.ZipChange = zipChange;
+			this.ConfigPath = configPath;
+			this.ConfigChange = configChange;
 		}
 	}
 
 	class ToolUpdateMonitor : IDisposable
 	{
-		CancellationTokenSource CancellationSource;
-		SynchronizationContext SynchronizationContext;
-		Task? WorkerTask;
-		AsyncEvent WakeEvent;
-		ILogger Logger;
+		CancellationTokenSource _cancellationSource;
+		SynchronizationContext _synchronizationContext;
+		Task? _workerTask;
+		AsyncEvent _wakeEvent;
+		ILogger _logger;
 		public List<ToolDefinition> Tools { get; private set; } = new List<ToolDefinition>();
-		int LastChange = -1;
-		IAsyncDisposer AsyncDisposer;
+		int _lastChange = -1;
+		IAsyncDisposer _asyncDisposer;
 
 		IPerforceSettings PerforceSettings { get; }
 		DirectoryReference ToolsDir { get; }
@@ -75,26 +75,26 @@ namespace UnrealGameSync
 
 		public Action? OnChange;
 
-		public ToolUpdateMonitor(IPerforceSettings PerforceSettings, DirectoryReference DataDir, UserSettings Settings, IServiceProvider ServiceProvider)
+		public ToolUpdateMonitor(IPerforceSettings perforceSettings, DirectoryReference dataDir, UserSettings settings, IServiceProvider serviceProvider)
 		{
-			CancellationSource = new CancellationTokenSource();
-			SynchronizationContext = SynchronizationContext.Current!;
-			this.ToolsDir = DirectoryReference.Combine(DataDir, "Tools");
-			this.PerforceSettings = PerforceSettings;
-			this.Settings = Settings;
-			this.Logger = ServiceProvider.GetRequiredService<ILogger<ToolUpdateMonitor>>();
-			this.AsyncDisposer = ServiceProvider.GetRequiredService<IAsyncDisposer>();
+			_cancellationSource = new CancellationTokenSource();
+			_synchronizationContext = SynchronizationContext.Current!;
+			this.ToolsDir = DirectoryReference.Combine(dataDir, "Tools");
+			this.PerforceSettings = perforceSettings;
+			this.Settings = settings;
+			this._logger = serviceProvider.GetRequiredService<ILogger<ToolUpdateMonitor>>();
+			this._asyncDisposer = serviceProvider.GetRequiredService<IAsyncDisposer>();
 
 			DirectoryReference.CreateDirectory(ToolsDir);
 
-			WakeEvent = new AsyncEvent();
+			_wakeEvent = new AsyncEvent();
 		}
 
 		public void Start()
 		{
 			if (DeploymentSettings.ToolsDepotPath != null)
 			{
-				WorkerTask = Task.Run(() => PollForUpdatesAsync(CancellationSource.Token));
+				_workerTask = Task.Run(() => PollForUpdatesAsync(_cancellationSource.Token));
 			}
 		}
 
@@ -102,36 +102,36 @@ namespace UnrealGameSync
 		{
 			OnChange = null;
 
-			if (WorkerTask != null)
+			if (_workerTask != null)
 			{
-				CancellationSource.Cancel();
-				AsyncDisposer.Add(WorkerTask.ContinueWith(_ => CancellationSource.Dispose()));
-				WorkerTask = null;
+				_cancellationSource.Cancel();
+				_asyncDisposer.Add(_workerTask.ContinueWith(_ => _cancellationSource.Dispose()));
+				_workerTask = null;
 			}
 		}
 
-		DirectoryReference GetToolPathInternal(string ToolName)
+		DirectoryReference GetToolPathInternal(string toolName)
 		{
-			return DirectoryReference.Combine(ToolsDir, ToolName, "Current");
+			return DirectoryReference.Combine(ToolsDir, toolName, "Current");
 		}
 
-		public string? GetToolName(Guid ToolId)
+		public string? GetToolName(Guid toolId)
 		{
-			foreach (ToolDefinition Tool in Tools)
+			foreach (ToolDefinition tool in Tools)
 			{
-				if(Tool.Id == ToolId)
+				if(tool.Id == toolId)
 				{
-					return Tool.Name;
+					return tool.Name;
 				}
 			}
 			return null;
 		}
 
-		public DirectoryReference? GetToolPath(string ToolName)
+		public DirectoryReference? GetToolPath(string toolName)
 		{
-			if (GetToolChange(ToolName) != 0)
+			if (GetToolChange(toolName) != 0)
 			{
-				return GetToolPathInternal(ToolName);
+				return GetToolPathInternal(toolName);
 			}
 			else
 			{
@@ -141,315 +141,315 @@ namespace UnrealGameSync
 
 		public void UpdateNow()
 		{
-			WakeEvent.Set();
+			_wakeEvent.Set();
 		}
 
-		async Task PollForUpdatesAsync(CancellationToken CancellationToken)
+		async Task PollForUpdatesAsync(CancellationToken cancellationToken)
 		{
-			while (!CancellationToken.IsCancellationRequested)
+			while (!cancellationToken.IsCancellationRequested)
 			{
-				Task WakeTask = WakeEvent.Task;
+				Task wakeTask = _wakeEvent.Task;
 
 				try
 				{
-					await PollForUpdatesOnce(Logger, CancellationToken);
+					await PollForUpdatesOnce(_logger, cancellationToken);
 				}
-				catch(Exception Ex)
+				catch(Exception ex)
 				{
-					Logger.LogError(Ex, "Exception while checking for tool updates");
+					_logger.LogError(ex, "Exception while checking for tool updates");
 				}
 
-				Task DelayTask = Task.Delay(TimeSpan.FromMinutes(60.0), CancellationToken);
-				await Task.WhenAny(DelayTask, WakeTask);
+				Task delayTask = Task.Delay(TimeSpan.FromMinutes(60.0), cancellationToken);
+				await Task.WhenAny(delayTask, wakeTask);
 			}
 		}
 
-		async Task PollForUpdatesOnce(ILogger Logger, CancellationToken CancellationToken)
+		async Task PollForUpdatesOnce(ILogger logger, CancellationToken cancellationToken)
 		{
-			using IPerforceConnection Perforce = await PerforceConnection.CreateAsync(PerforceSettings, Logger);
+			using IPerforceConnection perforce = await PerforceConnection.CreateAsync(PerforceSettings, logger);
 
-			List<ChangesRecord> Changes = await Perforce.GetChangesAsync(ChangesOptions.None, 1, ChangeStatus.Submitted, $"{DeploymentSettings.ToolsDepotPath}/...", CancellationToken);
-			if (Changes.Count == 0 || Changes[0].Number == LastChange)
+			List<ChangesRecord> changes = await perforce.GetChangesAsync(ChangesOptions.None, 1, ChangeStatus.Submitted, $"{DeploymentSettings.ToolsDepotPath}/...", cancellationToken);
+			if (changes.Count == 0 || changes[0].Number == _lastChange)
 			{
 				return;
 			}
 
-			List<FStatRecord> FileRecords = await Perforce.FStatAsync($"{DeploymentSettings.ToolsDepotPath}/...", CancellationToken).ToListAsync(CancellationToken);
+			List<FStatRecord> fileRecords = await perforce.FStatAsync($"{DeploymentSettings.ToolsDepotPath}/...", cancellationToken).ToListAsync(cancellationToken);
 
 			// Update the tools list
-			List<ToolDefinition> NewTools = new List<ToolDefinition>();
-			foreach (FStatRecord FileRecord in FileRecords)
+			List<ToolDefinition> newTools = new List<ToolDefinition>();
+			foreach (FStatRecord fileRecord in fileRecords)
 			{
-				if (FileRecord.DepotFile != null && FileRecord.DepotFile.EndsWith(".ini"))
+				if (fileRecord.DepotFile != null && fileRecord.DepotFile.EndsWith(".ini"))
 				{
-					ToolDefinition? Tool = Tools.FirstOrDefault(x => x.ConfigPath.Equals(FileRecord.DepotFile, StringComparison.Ordinal));
-					if (Tool == null || Tool.ConfigChange != FileRecord.HeadChange)
+					ToolDefinition? tool = Tools.FirstOrDefault(x => x.ConfigPath.Equals(fileRecord.DepotFile, StringComparison.Ordinal));
+					if (tool == null || tool.ConfigChange != fileRecord.HeadChange)
 					{
-						Tool = await ReadToolDefinitionAsync(Perforce, FileRecord.DepotFile, FileRecord.HeadChange, Logger, CancellationToken);
+						tool = await ReadToolDefinitionAsync(perforce, fileRecord.DepotFile, fileRecord.HeadChange, logger, cancellationToken);
 					}
-					if (Tool != null)
+					if (tool != null)
 					{
-						NewTools.Add(Tool);
+						newTools.Add(tool);
 					}
 				}
 			}
-			Tools = NewTools;
+			Tools = newTools;
 
-			foreach (ToolDefinition Tool in Tools)
+			foreach (ToolDefinition tool in Tools)
 			{
-				Tool.Enabled = Settings.EnabledTools.Contains(Tool.Id);
+				tool.Enabled = Settings.EnabledTools.Contains(tool.Id);
 
-				if(!Tool.Enabled)
+				if(!tool.Enabled)
 				{
 					continue;
 				}
 
-				List<FStatRecord> ToolFileRecords = FileRecords.Where(x => String.Equals(x.DepotFile, Tool.ZipPath, StringComparison.OrdinalIgnoreCase)).ToList();
-				if (ToolFileRecords.Count == 0)
+				List<FStatRecord> toolFileRecords = fileRecords.Where(x => String.Equals(x.DepotFile, tool.ZipPath, StringComparison.OrdinalIgnoreCase)).ToList();
+				if (toolFileRecords.Count == 0)
 				{
 					continue;
 				}
 
-				int HeadChange = ToolFileRecords.Max(x => x.HeadChange);
-				if (HeadChange == GetToolChange(Tool.Name))
+				int headChange = toolFileRecords.Max(x => x.HeadChange);
+				if (headChange == GetToolChange(tool.Name))
 				{
 					continue;
 				}
 
-				List<FStatRecord> SyncFileRecords = ToolFileRecords.Where(x => x.Action != FileAction.Delete && x.Action != FileAction.MoveDelete).ToList();
+				List<FStatRecord> syncFileRecords = toolFileRecords.Where(x => x.Action != FileAction.Delete && x.Action != FileAction.MoveDelete).ToList();
 				try
 				{
-					await UpdateToolAsync(Perforce, Tool.Name, HeadChange, SyncFileRecords, Tool.InstallAction, Logger, CancellationToken);
+					await UpdateToolAsync(perforce, tool.Name, headChange, syncFileRecords, tool.InstallAction, logger, cancellationToken);
 				}
-				catch (Exception Ex)
+				catch (Exception ex)
 				{
-					Logger.LogError(Ex, "Exception while updating tool");
+					logger.LogError(ex, "Exception while updating tool");
 				}
 			}
 
-			foreach (ToolDefinition Tool in Tools)
+			foreach (ToolDefinition tool in Tools)
 			{
-				if (!Tool.Enabled && GetToolChange(Tool.Name) != 0)
+				if (!tool.Enabled && GetToolChange(tool.Name) != 0)
 				{
 					try
 					{
-						await RemoveToolAsync(Tool.Name, Tool.UninstallAction, Logger, CancellationToken);
+						await RemoveToolAsync(tool.Name, tool.UninstallAction, logger, cancellationToken);
 					}
-					catch (Exception Ex)
+					catch (Exception ex)
 					{
-						Logger.LogError(Ex, "Exception while removing tool");
+						logger.LogError(ex, "Exception while removing tool");
 					}
 				}
 			}
 
-			SynchronizationContext.Post(_ => OnChange?.Invoke(), null);
+			_synchronizationContext.Post(_ => OnChange?.Invoke(), null);
 		}
 
-		async Task<ToolDefinition?> ReadToolDefinitionAsync(IPerforceConnection Perforce, string DepotPath, int Change, ILogger Logger, CancellationToken CancellationToken)
+		async Task<ToolDefinition?> ReadToolDefinitionAsync(IPerforceConnection perforce, string depotPath, int change, ILogger logger, CancellationToken cancellationToken)
 		{
-			PerforceResponse<PrintRecord<string[]>> Response = await Perforce.TryPrintLinesAsync($"{DepotPath}@{Change}", CancellationToken);
-			if (!Response.Succeeded || Response.Data.Contents == null)
+			PerforceResponse<PrintRecord<string[]>> response = await perforce.TryPrintLinesAsync($"{depotPath}@{change}", cancellationToken);
+			if (!response.Succeeded || response.Data.Contents == null)
 			{
 				return null;
 			}
 
-			int NameIdx = DepotPath.LastIndexOf('/') + 1;
-			int ExtensionIdx = DepotPath.LastIndexOf('.');
+			int nameIdx = depotPath.LastIndexOf('/') + 1;
+			int extensionIdx = depotPath.LastIndexOf('.');
 
-			ConfigFile ConfigFile = new ConfigFile();
-			ConfigFile.Parse(Response.Data.Contents);
+			ConfigFile configFile = new ConfigFile();
+			configFile.Parse(response.Data.Contents);
 
-			string? Id = ConfigFile.GetValue("Settings.Id", null);
-			if (Id == null || !Guid.TryParse(Id, out Guid ToolId))
+			string? id = configFile.GetValue("Settings.Id", null);
+			if (id == null || !Guid.TryParse(id, out Guid toolId))
 			{
 				return null;
 			}
 
-			string ToolName = ConfigFile.GetValue("Settings.Name", DepotPath.Substring(NameIdx, ExtensionIdx - NameIdx));
-			string ToolDescription = ConfigFile.GetValue("Settings.Description", ToolName);
-			string ToolZipPath = DepotPath.Substring(0, ExtensionIdx) + ".zip";
-			int ToolZipChange = GetToolChange(ToolName);
-			string ToolConfigPath = DepotPath;
-			int ToolConfigChange = Change;
+			string toolName = configFile.GetValue("Settings.Name", depotPath.Substring(nameIdx, extensionIdx - nameIdx));
+			string toolDescription = configFile.GetValue("Settings.Description", toolName);
+			string toolZipPath = depotPath.Substring(0, extensionIdx) + ".zip";
+			int toolZipChange = GetToolChange(toolName);
+			string toolConfigPath = depotPath;
+			int toolConfigChange = change;
 
-			ToolDefinition Tool = new ToolDefinition(ToolId, ToolName, ToolDescription, ToolZipPath, ToolZipChange, ToolConfigPath, ToolConfigChange);
+			ToolDefinition tool = new ToolDefinition(toolId, toolName, toolDescription, toolZipPath, toolZipChange, toolConfigPath, toolConfigChange);
 
-			string? InstallCommand = ConfigFile.GetValue("Settings.InstallCommand", null);
-			if (!String.IsNullOrEmpty(InstallCommand))
+			string? installCommand = configFile.GetValue("Settings.InstallCommand", null);
+			if (!String.IsNullOrEmpty(installCommand))
 			{
-				Tool.InstallAction = async (Logger, CancellationToken) =>
+				tool.InstallAction = async (logger, cancellationToken) =>
 				{
-					Logger.LogInformation("Running install action: {Command}", InstallCommand);
-					await RunCommandAsync(Tool.Name, InstallCommand, Logger, CancellationToken);
+					logger.LogInformation("Running install action: {Command}", installCommand);
+					await RunCommandAsync(tool.Name, installCommand, logger, cancellationToken);
 				};
 			}
 
-			string? UninstallCommand = ConfigFile.GetValue("Settings.UninstallCommand", null);
-			if (!String.IsNullOrEmpty(UninstallCommand))
+			string? uninstallCommand = configFile.GetValue("Settings.UninstallCommand", null);
+			if (!String.IsNullOrEmpty(uninstallCommand))
 			{
-				Tool.UninstallAction = async (Logger, CancellationToken) =>
+				tool.UninstallAction = async (logger, cancellationToken) =>
 				{
-					Logger.LogInformation("Running unininstall action: {Command}", UninstallCommand);
-					await RunCommandAsync(Tool.Name, UninstallCommand, Logger, CancellationToken);
+					logger.LogInformation("Running unininstall action: {Command}", uninstallCommand);
+					await RunCommandAsync(tool.Name, uninstallCommand, logger, cancellationToken);
 				};
 			}
 
-			string[] StatusPanelLinks = ConfigFile.GetValues("Settings.StatusPanelLinks", new string[0]);
-			foreach (string StatusPanelLink in StatusPanelLinks)
+			string[] statusPanelLinks = configFile.GetValues("Settings.StatusPanelLinks", new string[0]);
+			foreach (string statusPanelLink in statusPanelLinks)
 			{
-				ConfigObject Object = new ConfigObject(StatusPanelLink);
+				ConfigObject obj = new ConfigObject(statusPanelLink);
 
-				string? Label = Object.GetValue("Label", null);
-				string? FileName = Object.GetValue("FileName", null);
+				string? label = obj.GetValue("Label", null);
+				string? fileName = obj.GetValue("FileName", null);
 
-				if (Label != null && FileName != null)
+				if (label != null && fileName != null)
 				{
-					ToolLink Link = new ToolLink(Label, FileName);
-					Link.Arguments = Object.GetValue("Arguments", null);
-					Link.WorkingDir = Object.GetValue("WorkingDir", null);
-					Tool.StatusPanelLinks.Add(Link);
+					ToolLink link = new ToolLink(label, fileName);
+					link.Arguments = obj.GetValue("Arguments", null);
+					link.WorkingDir = obj.GetValue("WorkingDir", null);
+					tool.StatusPanelLinks.Add(link);
 				}
 			}
 
-			return Tool;
+			return tool;
 		}
 
-		async Task RunCommandAsync(string ToolName, string Command, ILogger Logger, CancellationToken CancellationToken)
+		async Task RunCommandAsync(string toolName, string command, ILogger logger, CancellationToken cancellationToken)
 		{
-			DirectoryReference ToolPath = GetToolPathInternal(ToolName);
+			DirectoryReference toolPath = GetToolPathInternal(toolName);
 
-			string CommandExe = Command;
-			string CommandArgs = string.Empty;
+			string commandExe = command;
+			string commandArgs = string.Empty;
 
-			int SpaceIdx = Command.IndexOf(' ');
-			if (SpaceIdx != -1)
+			int spaceIdx = command.IndexOf(' ');
+			if (spaceIdx != -1)
 			{
-				CommandExe = Command.Substring(0, SpaceIdx);
-				CommandArgs = Command.Substring(SpaceIdx + 1);
+				commandExe = command.Substring(0, spaceIdx);
+				commandArgs = command.Substring(spaceIdx + 1);
 			}
 
-			int ExitCode = await Utility.ExecuteProcessAsync(FileReference.Combine(ToolPath, CommandExe).FullName, ToolPath.FullName, CommandArgs, Line => Logger.LogInformation("{ToolName}> {Line}", ToolName, Line), CancellationToken);
-			Logger.LogInformation("{ToolName}> Exit code {ExitCode})", ToolName, ExitCode);
+			int exitCode = await Utility.ExecuteProcessAsync(FileReference.Combine(toolPath, commandExe).FullName, toolPath.FullName, commandArgs, line => logger.LogInformation("{ToolName}> {Line}", toolName, line), cancellationToken);
+			logger.LogInformation("{ToolName}> Exit code {ExitCode})", toolName, exitCode);
 		}
 
-		async Task RemoveToolAsync(string ToolName, Func<ILogger, CancellationToken, Task>? UninstallAction, ILogger Logger, CancellationToken CancellationToken)
+		async Task RemoveToolAsync(string toolName, Func<ILogger, CancellationToken, Task>? uninstallAction, ILogger logger, CancellationToken cancellationToken)
 		{
-			Logger.LogInformation("Removing {0}", ToolName);
-			DirectoryReference? ToolPath = GetToolPath(ToolName);
+			logger.LogInformation("Removing {0}", toolName);
+			DirectoryReference? toolPath = GetToolPath(toolName);
 
-			if (UninstallAction != null)
+			if (uninstallAction != null)
 			{
-				Logger.LogInformation("Running uninstall...");
-				await UninstallAction(Logger, CancellationToken);
+				logger.LogInformation("Running uninstall...");
+				await uninstallAction(logger, cancellationToken);
 			}
 
-			SetToolChange(ToolName, null);
+			SetToolChange(toolName, null);
 
-			if (ToolPath != null)
+			if (toolPath != null)
 			{
-				Logger.LogInformation("Removing {ToolPath}", ToolPath);
-				TryDeleteDirectory(ToolPath, Logger);
+				logger.LogInformation("Removing {ToolPath}", toolPath);
+				TryDeleteDirectory(toolPath, logger);
 			}
 
-			Logger.LogInformation("{ToolName} has been removed successfully", ToolName);
+			logger.LogInformation("{ToolName} has been removed successfully", toolName);
 		}
 
-		static void ForceDeleteDirectory(DirectoryReference DirectoryName)
+		static void ForceDeleteDirectory(DirectoryReference directoryName)
 		{
-			DirectoryInfo BaseDir = DirectoryName.ToDirectoryInfo();
-			if (BaseDir.Exists)
+			DirectoryInfo baseDir = directoryName.ToDirectoryInfo();
+			if (baseDir.Exists)
 			{
-				foreach (FileInfo File in BaseDir.EnumerateFiles("*", SearchOption.AllDirectories))
+				foreach (FileInfo file in baseDir.EnumerateFiles("*", SearchOption.AllDirectories))
 				{
-					File.Attributes = FileAttributes.Normal;
+					file.Attributes = FileAttributes.Normal;
 				}
-				BaseDir.Delete(true);
+				baseDir.Delete(true);
 			}
 		}
 
-		static bool TryDeleteDirectory(DirectoryReference DirectoryName, ILogger Logger)
+		static bool TryDeleteDirectory(DirectoryReference directoryName, ILogger logger)
 		{
 			try
 			{
-				ForceDeleteDirectory(DirectoryName);
+				ForceDeleteDirectory(directoryName);
 				return true;
 			}
-			catch(Exception Ex)
+			catch(Exception ex)
 			{
-				Logger.LogWarning(Ex, "Unable to delete directory {DirectoryName}", DirectoryName);
+				logger.LogWarning(ex, "Unable to delete directory {DirectoryName}", directoryName);
 				return false;
 			}
 		}
 
-		async Task<bool> UpdateToolAsync(IPerforceConnection Perforce, string ToolName, int Change, List<FStatRecord> Records, Func<ILogger, CancellationToken, Task>? InstallAction, ILogger Logger, CancellationToken CancellationToken)
+		async Task<bool> UpdateToolAsync(IPerforceConnection perforce, string toolName, int change, List<FStatRecord> records, Func<ILogger, CancellationToken, Task>? installAction, ILogger logger, CancellationToken cancellationToken)
 		{
-			DirectoryReference ToolDir = DirectoryReference.Combine(ToolsDir, ToolName);
-			DirectoryReference.CreateDirectory(ToolDir);
+			DirectoryReference toolDir = DirectoryReference.Combine(ToolsDir, toolName);
+			DirectoryReference.CreateDirectory(toolDir);
 
-			foreach (DirectoryReference ExistingDir in DirectoryReference.EnumerateDirectories(ToolDir, "Prev-*"))
+			foreach (DirectoryReference existingDir in DirectoryReference.EnumerateDirectories(toolDir, "Prev-*"))
 			{
-				TryDeleteDirectory(ExistingDir, Logger);
+				TryDeleteDirectory(existingDir, logger);
 			}
 
-			DirectoryReference NextToolDir = DirectoryReference.Combine(ToolDir, "Next");
-			ForceDeleteDirectory(NextToolDir);
-			DirectoryReference.CreateDirectory(NextToolDir);
+			DirectoryReference nextToolDir = DirectoryReference.Combine(toolDir, "Next");
+			ForceDeleteDirectory(nextToolDir);
+			DirectoryReference.CreateDirectory(nextToolDir);
 
-			DirectoryReference NextToolZipsDir = DirectoryReference.Combine(NextToolDir, ".zips");
-			DirectoryReference.CreateDirectory(NextToolZipsDir);
+			DirectoryReference nextToolZipsDir = DirectoryReference.Combine(nextToolDir, ".zips");
+			DirectoryReference.CreateDirectory(nextToolZipsDir);
 
-			for (int Idx = 0; Idx < Records.Count; Idx++)
+			for (int idx = 0; idx < records.Count; idx++)
 			{
-				FileReference ZipFile = FileReference.Combine(NextToolZipsDir, String.Format("{0}.{1}.zip", ToolName, Idx));
+				FileReference zipFile = FileReference.Combine(nextToolZipsDir, String.Format("{0}.{1}.zip", toolName, idx));
 
-				PerforceResponse<PrintRecord> Response = await Perforce.TryPrintAsync(ZipFile.FullName, $"{Records[Idx].DepotFile}#{Records[Idx].HeadRevision}", CancellationToken);
-				if(!Response.Succeeded || !FileReference.Exists(ZipFile))
+				PerforceResponse<PrintRecord> response = await perforce.TryPrintAsync(zipFile.FullName, $"{records[idx].DepotFile}#{records[idx].HeadRevision}", cancellationToken);
+				if(!response.Succeeded || !FileReference.Exists(zipFile))
 				{
-					Logger.LogError("Unable to print {0}", Records[Idx].DepotFile);
+					logger.LogError("Unable to print {0}", records[idx].DepotFile);
 					return false;
 				}
 
-				ArchiveUtils.ExtractFiles(ZipFile, NextToolDir, null, new ProgressValue(), Logger);
+				ArchiveUtils.ExtractFiles(zipFile, nextToolDir, null, new ProgressValue(), logger);
 			}
 
-			SetToolChange(ToolName, null);
+			SetToolChange(toolName, null);
 
-			DirectoryReference CurrentToolDir = DirectoryReference.Combine(ToolDir, "Current");
-			if (DirectoryReference.Exists(CurrentToolDir))
+			DirectoryReference currentToolDir = DirectoryReference.Combine(toolDir, "Current");
+			if (DirectoryReference.Exists(currentToolDir))
 			{
-				DirectoryReference PrevDirectoryName = DirectoryReference.Combine(ToolDir, String.Format("Prev-{0:X16}", Stopwatch.GetTimestamp()));
-				Directory.Move(CurrentToolDir.FullName, PrevDirectoryName.FullName);
-				TryDeleteDirectory(PrevDirectoryName, Logger);
+				DirectoryReference prevDirectoryName = DirectoryReference.Combine(toolDir, String.Format("Prev-{0:X16}", Stopwatch.GetTimestamp()));
+				Directory.Move(currentToolDir.FullName, prevDirectoryName.FullName);
+				TryDeleteDirectory(prevDirectoryName, logger);
 			}
 
-			Directory.Move(NextToolDir.FullName, CurrentToolDir.FullName);
+			Directory.Move(nextToolDir.FullName, currentToolDir.FullName);
 
-			if (InstallAction != null)
+			if (installAction != null)
 			{
-				Logger.LogInformation("Running installer...");
-				await InstallAction.Invoke(Logger, CancellationToken);
+				logger.LogInformation("Running installer...");
+				await installAction.Invoke(logger, cancellationToken);
 			}
 
-			SetToolChange(ToolName, Change);
-			Logger.LogInformation("Updated {ToolName} to change {Change}", ToolName, Change);
+			SetToolChange(toolName, change);
+			logger.LogInformation("Updated {ToolName} to change {Change}", toolName, change);
 			return true;
 		}
 
-		FileReference GetConfigFilePath(string ToolName)
+		FileReference GetConfigFilePath(string toolName)
 		{
-			return FileReference.Combine(ToolsDir, ToolName, ToolName + ".ini");
+			return FileReference.Combine(ToolsDir, toolName, toolName + ".ini");
 		}
 
-		int GetToolChange(string ToolName)
+		int GetToolChange(string toolName)
 		{
 			try
 			{
-				FileReference ConfigFilePath = GetConfigFilePath(ToolName);
-				if (FileReference.Exists(ConfigFilePath))
+				FileReference configFilePath = GetConfigFilePath(toolName);
+				if (FileReference.Exists(configFilePath))
 				{
-					ConfigFile ConfigFile = new ConfigFile();
-					ConfigFile.Load(ConfigFilePath);
-					return ConfigFile.GetValue("Settings.Change", 0);
+					ConfigFile configFile = new ConfigFile();
+					configFile.Load(configFilePath);
+					return configFile.GetValue("Settings.Change", 0);
 				}
 			}
 			catch
@@ -458,18 +458,18 @@ namespace UnrealGameSync
 			return 0;
 		}
 
-		void SetToolChange(string ToolName, int? Change)
+		void SetToolChange(string toolName, int? change)
 		{
-			FileReference ConfigFilePath = GetConfigFilePath(ToolName);
-			if (Change.HasValue)
+			FileReference configFilePath = GetConfigFilePath(toolName);
+			if (change.HasValue)
 			{
-				ConfigFile ConfigFile = new ConfigFile();
-				ConfigFile.SetValue("Settings.Change", Change.Value);
-				ConfigFile.Save(ConfigFilePath);
+				ConfigFile configFile = new ConfigFile();
+				configFile.SetValue("Settings.Change", change.Value);
+				configFile.Save(configFilePath);
 			}
 			else
 			{
-				FileReference.Delete(ConfigFilePath);
+				FileReference.Delete(configFilePath);
 			}
 		}
 	}
diff --git a/Engine/Source/Programs/UnrealGameSync/UnrealGameSync/UpdateMonitor.cs b/Engine/Source/Programs/UnrealGameSync/UnrealGameSync/UpdateMonitor.cs
--- a/Engine/Source/Programs/UnrealGameSync/UnrealGameSync/UpdateMonitor.cs
+++ b/Engine/Source/Programs/UnrealGameSync/UnrealGameSync/UpdateMonitor.cs
@@ -22,10 +22,10 @@ namespace UnrealGameSync
 
 	class UpdateMonitor : IDisposable
 	{
-		Task? WorkerTask;
-		CancellationTokenSource CancellationSource = new CancellationTokenSource();
-		ILogger Logger;
-		IAsyncDisposer AsyncDisposer;
+		Task? _workerTask;
+		CancellationTokenSource _cancellationSource = new CancellationTokenSource();
+		ILogger _logger;
+		IAsyncDisposer _asyncDisposer;
 
 		public Action<UpdateType>? OnUpdateAvailable;
 
@@ -35,15 +35,15 @@ namespace UnrealGameSync
 			private set;
 		}
 
-		public UpdateMonitor(IPerforceSettings PerforceSettings, string? WatchPath, IServiceProvider ServiceProvider)
+		public UpdateMonitor(IPerforceSettings perforceSettings, string? watchPath, IServiceProvider serviceProvider)
 		{
-			this.Logger = ServiceProvider.GetRequiredService<ILogger<UpdateMonitor>>();
-			this.AsyncDisposer = ServiceProvider.GetRequiredService<IAsyncDisposer>();
+			this._logger = serviceProvider.GetRequiredService<ILogger<UpdateMonitor>>();
+			this._asyncDisposer = serviceProvider.GetRequiredService<IAsyncDisposer>();
 
-			if(WatchPath != null)
+			if(watchPath != null)
 			{
-				Logger.LogInformation("Watching for updates on {WatchPath}", WatchPath);
-				WorkerTask = Task.Run(() => PollForUpdatesAsync(PerforceSettings, WatchPath, CancellationSource.Token));
+				_logger.LogInformation("Watching for updates on {WatchPath}", watchPath);
+				_workerTask = Task.Run(() => PollForUpdatesAsync(perforceSettings, watchPath, _cancellationSource.Token));
 			}
 		}
 
@@ -51,11 +51,11 @@ namespace UnrealGameSync
 		{
 			OnUpdateAvailable = null;
 
-			if (WorkerTask != null)
+			if (_workerTask != null)
 			{
-				CancellationSource.Cancel();
-				AsyncDisposer.Add(WorkerTask.ContinueWith(_ => CancellationSource.Dispose()));
-				WorkerTask = null;
+				_cancellationSource.Cancel();
+				_asyncDisposer.Add(_workerTask.ContinueWith(_ => _cancellationSource.Dispose()));
+				_workerTask = null;
 			}
 		}
 
@@ -65,46 +65,46 @@ namespace UnrealGameSync
 			private set;
 		}
 
-		async Task PollForUpdatesAsync(IPerforceSettings PerforceSettings, string WatchPath, CancellationToken CancellationToken)
+		async Task PollForUpdatesAsync(IPerforceSettings perforceSettings, string watchPath, CancellationToken cancellationToken)
 		{
 			for (; ; )
 			{
-				await Task.Delay(TimeSpan.FromMinutes(5.0), CancellationToken);
+				await Task.Delay(TimeSpan.FromMinutes(5.0), cancellationToken);
 
-				IPerforceConnection? Perforce = null;
+				IPerforceConnection? perforce = null;
 				try
 				{
-					Perforce = await PerforceConnection.CreateAsync(PerforceSettings, Logger);
+					perforce = await PerforceConnection.CreateAsync(perforceSettings, _logger);
 
-					PerforceResponseList<ChangesRecord> Changes = await Perforce.TryGetChangesAsync(ChangesOptions.None, -1, ChangeStatus.Submitted, WatchPath, CancellationToken);
-					if (Changes.Succeeded && Changes.Data.Count > 0)
+					PerforceResponseList<ChangesRecord> changes = await perforce.TryGetChangesAsync(ChangesOptions.None, -1, ChangeStatus.Submitted, watchPath, cancellationToken);
+					if (changes.Succeeded && changes.Data.Count > 0)
 					{
 						TriggerUpdate(UpdateType.Background, null);
 					}
 				}
 				catch (PerforceException ex)
 				{
-					Logger.LogInformation(ex, "Perforce exception while attempting to poll for updates.");
+					_logger.LogInformation(ex, "Perforce exception while attempting to poll for updates.");
 				}
 				catch (Exception ex)
 				{
-					Logger.LogWarning(ex, "Exception while attempting to poll for updates.");
+					_logger.LogWarning(ex, "Exception while attempting to poll for updates.");
 					Program.CaptureException(ex);
 				}
 				finally
 				{
-					Perforce?.Dispose();
+					perforce?.Dispose();
 				}
 			}
 		}
 
-		public void TriggerUpdate(UpdateType UpdateType, bool? RelaunchPreview)
+		public void TriggerUpdate(UpdateType updateType, bool? relaunchPreview)
 		{
-			this.RelaunchPreview = RelaunchPreview;
+			this.RelaunchPreview = relaunchPreview;
 			IsUpdateAvailable = true;
 			if(OnUpdateAvailable != null)
 			{
-				OnUpdateAvailable(UpdateType);
+				OnUpdateAvailable(updateType);
 			}
 		}
 	}
diff --git a/Engine/Source/Programs/UnrealGameSync/UnrealGameSync/UriHandlers/P4Handler.cs b/Engine/Source/Programs/UnrealGameSync/UnrealGameSync/UriHandlers/P4Handler.cs
--- a/Engine/Source/Programs/UnrealGameSync/UnrealGameSync/UriHandlers/P4Handler.cs
+++ b/Engine/Source/Programs/UnrealGameSync/UnrealGameSync/UriHandlers/P4Handler.cs
@@ -9,24 +9,24 @@ namespace UnrealGameSync
 	static class P4Handler
 	{
 		[UriHandler(true)]
-		public static UriResult P4V(string DepotPath)
+		public static UriResult P4V(string depotPath)
 		{
-			string CommandLine = string.Format("-s \"{0}\"", DepotPath);
+			string commandLine = string.Format("-s \"{0}\"", depotPath);
 
-			if (!Utility.SpawnHiddenProcess("p4v.exe", CommandLine))
+			if (!Utility.SpawnHiddenProcess("p4v.exe", commandLine))
 			{
-				return new UriResult() { Error = string.Format("Error spawning p4v.exe with command line: {0}", CommandLine) };
+				return new UriResult() { Error = string.Format("Error spawning p4v.exe with command line: {0}", commandLine) };
 			}
 
 			return new UriResult() { Success = true };
 		}
 
 		[UriHandler(true)]
-		public static UriResult Timelapse(string DepotPath, int Line = -1)
+		public static UriResult Timelapse(string depotPath, int line = -1)
 		{
-			string CommandLine = string.Format("timelapse {0}{1}", Line == -1 ? "" : string.Format(" -l {0} ", Line), DepotPath);
+			string commandLine = string.Format("timelapse {0}{1}", line == -1 ? "" : string.Format(" -l {0} ", line), depotPath);
 
-			Program.SpawnP4VC(CommandLine);
+			Program.SpawnP4Vc(commandLine);
 
 			return new UriResult() { Success = true };
 		}
@@ -34,7 +34,7 @@ namespace UnrealGameSync
 		[UriHandler(true)]
 		public static UriResult Change(int number)
 		{
-			Program.SpawnP4VC($"change {number}");
+			Program.SpawnP4Vc($"change {number}");
 			return new UriResult() { Success = true };
 		}
 	}
diff --git a/Engine/Source/Programs/UnrealGameSync/UnrealGameSync/UriHandlers/RdpHandler.cs b/Engine/Source/Programs/UnrealGameSync/UnrealGameSync/UriHandlers/RdpHandler.cs
--- a/Engine/Source/Programs/UnrealGameSync/UnrealGameSync/UriHandlers/RdpHandler.cs
+++ b/Engine/Source/Programs/UnrealGameSync/UnrealGameSync/UriHandlers/RdpHandler.cs
@@ -22,12 +22,12 @@ namespace UnrealGameSync
 {
 	static class RdpHandler
 	{
-		const uint CRED_TYPE_GENERIC = 1;
-		const uint CRED_TYPE_DOMAIN_PASSWORD = 2;
-		const uint CRED_PERSIST_LOCAL_MACHINE = 2;
+		const uint CredTypeGeneric = 1;
+		const uint CredTypeDomainPassword = 2;
+		const uint CredPersistLocalMachine = 2;
 
 		[StructLayout(LayoutKind.Sequential, CharSet = CharSet.Auto)]
-		struct CREDENTIAL_ATTRIBUTE
+		struct CredentialAttribute
 		{
 			string Keyword;
 			uint Flags;
@@ -36,7 +36,7 @@ namespace UnrealGameSync
 		}
 
 		[StructLayout(LayoutKind.Sequential, CharSet = CharSet.Unicode)]
-		class CREDENTIAL
+		class Credential
 		{
 			public uint Flags;
 			public uint Type;
@@ -53,87 +53,87 @@ namespace UnrealGameSync
 		}
 
 		[DllImport("Advapi32.dll", EntryPoint = "CredReadW", CharSet = CharSet.Unicode, SetLastError = true)]
-		static extern bool CredRead(string Target, uint Type, int ReservedFlag, out IntPtr Buffer);
+		static extern bool CredRead(string target, uint type, int reservedFlag, out IntPtr buffer);
 
 		[DllImport("Advapi32.dll", EntryPoint = "CredWriteW", CharSet = CharSet.Unicode, SetLastError = true)]
-		static extern bool CredWrite(CREDENTIAL userCredential, uint Flags);
+		static extern bool CredWrite(Credential userCredential, uint flags);
 
 		[DllImport("advapi32.dll", SetLastError = true)]
-		static extern bool CredFree(IntPtr Buffer);
+		static extern bool CredFree(IntPtr buffer);
 
-		public static bool ReadCredential(string Name, uint Type, [NotNullWhen(true)] out string? UserName, [NotNullWhen(true)] out string? Password)
+		public static bool ReadCredential(string name, uint type, [NotNullWhen(true)] out string? userName, [NotNullWhen(true)] out string? password)
 		{
-			IntPtr Buffer = IntPtr.Zero;
+			IntPtr buffer = IntPtr.Zero;
 			try
 			{
-				if (!CredRead(Name, Type, 0, out Buffer))
+				if (!CredRead(name, type, 0, out buffer))
 				{
-					UserName = null;
-					Password = null;
+					userName = null;
+					password = null;
 					return false;
 				}
 
-				CREDENTIAL? Credential = Marshal.PtrToStructure<CREDENTIAL>(Buffer);
-				if (Credential == null || Credential.UserName == null || Credential.CredentialBlob == IntPtr.Zero)
+				Credential? credential = Marshal.PtrToStructure<Credential>(buffer);
+				if (credential == null || credential.UserName == null || credential.CredentialBlob == IntPtr.Zero)
 				{
-					UserName = null;
-					Password = null;
+					userName = null;
+					password = null;
 					return false;
 				}
 
-				UserName = Credential.UserName;
-				Password = Marshal.PtrToStringUni(Credential.CredentialBlob, Credential.CredentialBlobSize / sizeof(char));
+				userName = credential.UserName;
+				password = Marshal.PtrToStringUni(credential.CredentialBlob, credential.CredentialBlobSize / sizeof(char));
 				return true;
 			}
 			finally
 			{
-				if (Buffer != IntPtr.Zero)
+				if (buffer != IntPtr.Zero)
 				{
-					CredFree(Buffer);
+					CredFree(buffer);
 				}
 			}
 		}
 
-		public static void WriteCredential(string Name, uint Type, string UserName, string Password)
+		public static void WriteCredential(string name, uint type, string userName, string password)
 		{
-			CREDENTIAL Credential = new CREDENTIAL();
+			Credential credential = new Credential();
 			try
 			{
-				Credential.Type = Type;
-				Credential.Persist = CRED_PERSIST_LOCAL_MACHINE;
-				Credential.CredentialBlobSize = Password.Length * sizeof(char);
-				Credential.TargetName = Name;
-				Credential.CredentialBlob = Marshal.StringToCoTaskMemUni(Password);
-				Credential.UserName = UserName;
-				CredWrite(Credential, 0);
+				credential.Type = type;
+				credential.Persist = CredPersistLocalMachine;
+				credential.CredentialBlobSize = password.Length * sizeof(char);
+				credential.TargetName = name;
+				credential.CredentialBlob = Marshal.StringToCoTaskMemUni(password);
+				credential.UserName = userName;
+				CredWrite(credential, 0);
 			}
 			finally
 			{
-				if (Credential.CredentialBlob != IntPtr.Zero)
+				if (credential.CredentialBlob != IntPtr.Zero)
 				{
-					Marshal.FreeCoTaskMem(Credential.CredentialBlob);
+					Marshal.FreeCoTaskMem(credential.CredentialBlob);
 				}
 			}
 		}
 
 		[UriHandler(Terminate = true)]
-		public static UriResult Rdp(string Host)
+		public static UriResult Rdp(string host)
 		{
 			// Copy the credentials from a generic Windows credential to 
-			if (!ReadCredential(Host, CRED_TYPE_DOMAIN_PASSWORD, out _, out _))
+			if (!ReadCredential(host, CredTypeDomainPassword, out _, out _))
 			{
-				if (ReadCredential("UnrealGameSync:RDP", CRED_TYPE_GENERIC, out string? UserName, out string? Password))
+				if (ReadCredential("UnrealGameSync:RDP", CredTypeGeneric, out string? userName, out string? password))
 				{
-					WriteCredential(Host, CRED_TYPE_DOMAIN_PASSWORD, UserName, Password);
+					WriteCredential(host, CredTypeDomainPassword, userName, password);
 				}
 			}
 
-			using (Process Process = new Process())
+			using (Process process = new Process())
 			{
-				Process.StartInfo.FileName = Path.Combine(Environment.GetFolderPath(Environment.SpecialFolder.System), "mstsc.exe");
-				Process.StartInfo.ArgumentList.Add($"/v:{Host}");
-				Process.StartInfo.ArgumentList.Add($"/f");
-				Process.Start();
+				process.StartInfo.FileName = Path.Combine(Environment.GetFolderPath(Environment.SpecialFolder.System), "mstsc.exe");
+				process.StartInfo.ArgumentList.Add($"/v:{host}");
+				process.StartInfo.ArgumentList.Add($"/f");
+				process.Start();
 			}
 
 			return new UriResult() { Success = true };
diff --git a/Engine/Source/Programs/UnrealGameSync/UnrealGameSync/UriHandlers/UGSHandler.cs b/Engine/Source/Programs/UnrealGameSync/UnrealGameSync/UriHandlers/UGSHandler.cs
--- a/Engine/Source/Programs/UnrealGameSync/UnrealGameSync/UriHandlers/UGSHandler.cs
+++ b/Engine/Source/Programs/UnrealGameSync/UnrealGameSync/UriHandlers/UGSHandler.cs
@@ -11,64 +11,64 @@ namespace UnrealGameSync
 	/// <summary>
 	/// UGS handler 
 	/// </summary>
-	static class UGSHandler
+	static class UgsHandler
 	{
 
 		[UriHandler]
-		public static UriResult OpenProject(string Stream, string Project, bool Sync = false)
+		public static UriResult OpenProject(string stream, string project, bool sync = false)
 		{
 			// Create the request
-			using (MemoryStream InputDataStream = new MemoryStream())
+			using (MemoryStream inputDataStream = new MemoryStream())
 			{
-				using (BinaryWriter Writer = new BinaryWriter(InputDataStream))
+				using (BinaryWriter writer = new BinaryWriter(inputDataStream))
 				{
-					Writer.Write(Stream);
-					Writer.Write(Project);
+					writer.Write(stream);
+					writer.Write(project);
 				}
 
-				AutomationRequestInput Input = new AutomationRequestInput(Sync ? AutomationRequestType.SyncProject : AutomationRequestType.OpenProject, InputDataStream.GetBuffer());
-				return new UriResult() { Success = true, Request = new AutomationRequest(Input) };
+				AutomationRequestInput input = new AutomationRequestInput(sync ? AutomationRequestType.SyncProject : AutomationRequestType.OpenProject, inputDataStream.GetBuffer());
+				return new UriResult() { Success = true, Request = new AutomationRequest(input) };
 			}
 		}
 
 		[UriHandler]
-		public static UriResult BuildStep(string Project, string Stream, string Step, string Changelist, string Arguments)
+		public static UriResult BuildStep(string project, string stream, string step, string changelist, string arguments)
 		{
-			MessageBox.Show(string.Format("Project: {0}\nStream: {1}\nStep: {2}\nChange: {3}\nArguments: {4}", Project, Stream, Step, Changelist, Arguments), "UGS Build Step Handler");
+			MessageBox.Show(string.Format("Project: {0}\nStream: {1}\nStep: {2}\nChange: {3}\nArguments: {4}", project, stream, step, changelist, arguments), "UGS Build Step Handler");
 
 			return new UriResult() { Success = true };
 		}
 
 		[UriHandler]
-		public static UriResult Execute(string Stream, int Changelist, string Command, string Project = "")
+		public static UriResult Execute(string stream, int changelist, string command, string project = "")
 		{
-			using (MemoryStream InputDataStream = new MemoryStream())
+			using (MemoryStream inputDataStream = new MemoryStream())
 			{
-				using (BinaryWriter Writer = new BinaryWriter(InputDataStream))
+				using (BinaryWriter writer = new BinaryWriter(inputDataStream))
 				{
-					Writer.Write(Stream);
-					Writer.Write(Changelist);
-					Writer.Write(Command);
-					Writer.Write(Project);
+					writer.Write(stream);
+					writer.Write(changelist);
+					writer.Write(command);
+					writer.Write(project);
 				}
 
-				AutomationRequestInput Input = new AutomationRequestInput(AutomationRequestType.ExecCommand, InputDataStream.GetBuffer());
-				return new UriResult() { Success = true, Request = new AutomationRequest(Input) };
+				AutomationRequestInput input = new AutomationRequestInput(AutomationRequestType.ExecCommand, inputDataStream.GetBuffer());
+				return new UriResult() { Success = true, Request = new AutomationRequest(input) };
 			}
 		}
 
 		[UriHandler]
-		public static UriResult OpenIssue(int Id)
+		public static UriResult OpenIssue(int id)
 		{
-			using (MemoryStream InputDataStream = new MemoryStream())
+			using (MemoryStream inputDataStream = new MemoryStream())
 			{
-				using (BinaryWriter Writer = new BinaryWriter(InputDataStream))
+				using (BinaryWriter writer = new BinaryWriter(inputDataStream))
 				{
-					Writer.Write(Id);
+					writer.Write(id);
 				}
 
-				AutomationRequestInput Input = new AutomationRequestInput(AutomationRequestType.OpenIssue, InputDataStream.GetBuffer());
-				return new UriResult() { Success = true, Request = new AutomationRequest(Input) };
+				AutomationRequestInput input = new AutomationRequestInput(AutomationRequestType.OpenIssue, inputDataStream.GetBuffer());
+				return new UriResult() { Success = true, Request = new AutomationRequest(input) };
 			}
 		}
 	}
diff --git a/Engine/Source/Programs/UnrealGameSync/UnrealGameSync/UriHandlers/UriHandler.cs b/Engine/Source/Programs/UnrealGameSync/UnrealGameSync/UriHandlers/UriHandler.cs
--- a/Engine/Source/Programs/UnrealGameSync/UnrealGameSync/UriHandlers/UriHandler.cs
+++ b/Engine/Source/Programs/UnrealGameSync/UnrealGameSync/UriHandlers/UriHandler.cs
@@ -31,81 +31,81 @@ namespace UnrealGameSync
 	/// </summary>
 	static class UriHandler
 	{
-		public static UriResult HandleUri(string UriIn)
+		public static UriResult HandleUri(string uriIn)
 		{			
 			try
 			{
-				Uri Uri = new Uri(UriIn);
+				Uri uri = new Uri(uriIn);
 
 				// Check if this is a registered uri request
-				if (!Handlers.TryGetValue(Uri.Host, out MethodInfo? Info))
+				if (!Handlers.TryGetValue(uri.Host, out MethodInfo? info))
 				{
-					return new UriResult() { Error = string.Format("Unknown Uri {0}", Uri.Host) };
+					return new UriResult() { Error = string.Format("Unknown Uri {0}", uri.Host) };
 				}
 
-				NameValueCollection Query = HttpUtility.ParseQueryString(Uri.Query);
+				NameValueCollection query = HttpUtility.ParseQueryString(uri.Query);
 
-				List<object> Parameters = new List<object>();
+				List<object> parameters = new List<object>();
 
-				foreach (ParameterInfo Param in Info.GetParameters())
+				foreach (ParameterInfo param in info.GetParameters())
 				{					
-					string Value = Query.Get(Param.Name);
+					string value = query.Get(param.Name);
 
-					if (Value == null)
+					if (value == null)
 					{
-						if (!Param.HasDefaultValue)
+						if (!param.HasDefaultValue)
 						{
-							return new UriResult() { Error = string.Format("Uri {0} is missing required parameter {1}", Uri.Host, Param.Name) };
+							return new UriResult() { Error = string.Format("Uri {0} is missing required parameter {1}", uri.Host, param.Name) };
 						}
 
-						Parameters.Add(Param.DefaultValue!);
+						parameters.Add(param.DefaultValue!);
 						continue;
 					}
 
-					if (Param.ParameterType == typeof(string))
+					if (param.ParameterType == typeof(string))
 					{
-						Parameters.Add(Value);
+						parameters.Add(value);
 					}
-					else if (Param.ParameterType == typeof(bool))
+					else if (param.ParameterType == typeof(bool))
 					{
-						if (Value.Equals("true", StringComparison.OrdinalIgnoreCase))
+						if (value.Equals("true", StringComparison.OrdinalIgnoreCase))
 						{
-							Parameters.Add(true);
+							parameters.Add(true);
 						}
-						else if (Value.Equals("false", StringComparison.OrdinalIgnoreCase))
+						else if (value.Equals("false", StringComparison.OrdinalIgnoreCase))
 						{
-							Parameters.Add(false);
+							parameters.Add(false);
 						}
 						else
 						{
-							return new UriResult() { Error = string.Format("Uri {0} bool parameter {1} must be true or false", Uri.Host, Param.Name) };
+							return new UriResult() { Error = string.Format("Uri {0} bool parameter {1} must be true or false", uri.Host, param.Name) };
 						}
 					}
-					else if (Param.ParameterType == typeof(int))
+					else if (param.ParameterType == typeof(int))
 					{
-						int NumberValue;
-						if (!int.TryParse(Value, out NumberValue))
+						int numberValue;
+						if (!int.TryParse(value, out numberValue))
 						{
-							return new UriResult() { Error = string.Format("Uri {0} invalid integer parameter {1} : {2}", Uri.Host, Param.Name, Value) };
+							return new UriResult() { Error = string.Format("Uri {0} invalid integer parameter {1} : {2}", uri.Host, param.Name, value) };
 						}
-						Parameters.Add(NumberValue);
+						parameters.Add(numberValue);
 					}
-					else if (Param.ParameterType == typeof(float))
+					else if (param.ParameterType == typeof(float))
 					{
-						float NumberValue;
-						if (!float.TryParse(Value, out NumberValue))
+						float numberValue;
+						if (!float.TryParse(value, out numberValue))
 						{
-							return new UriResult() { Error = string.Format("Uri {0} invalid number parameter {1} : {2}", Uri.Host, Param.Name, Value) };
+							return new UriResult() { Error = string.Format("Uri {0} invalid number parameter {1} : {2}", uri.Host, param.Name, value) };
 						}
-						Parameters.Add(NumberValue);
+						parameters.Add(numberValue);
 					}
 				}
 
-				return (UriResult)Info.Invoke(null, Parameters.ToArray())!;					
+				return (UriResult)info.Invoke(null, parameters.ToArray())!;					
 			}
-			catch (Exception Ex)
+			catch (Exception ex)
 			{
-				return new UriResult() { Error = Ex.Message };
+				return new UriResult() { Error = ex.Message };
 			}
 		}
 
@@ -116,11 +116,11 @@ namespace UnrealGameSync
 		/// <summary>
 		/// Handle URI passed in via command lines
 		/// </summary>
-		public static bool ProcessCommandLine(string[] Args, bool FirstInstance, EventWaitHandle? ActivateEvent = null)
+		public static bool ProcessCommandLine(string[] args, bool firstInstance, EventWaitHandle? activateEvent = null)
 		{
-			if (Args.Any(x => x.Equals(InstallHandlerArg, StringComparison.OrdinalIgnoreCase)))
+			if (args.Any(x => x.Equals(InstallHandlerArg, StringComparison.OrdinalIgnoreCase)))
 			{
-				if (Args.Any(x => x.Equals(ElevatedArg, StringComparison.OrdinalIgnoreCase)))
+				if (args.Any(x => x.Equals(ElevatedArg, StringComparison.OrdinalIgnoreCase)))
 				{
 					ProtocolHandlerUtils.InstallElevated();
 				}
@@ -130,9 +130,9 @@ namespace UnrealGameSync
 				}
 				return true;
 			}
-			else if (Args.Any(x => x.Equals(UninstallHandlerArg, StringComparison.OrdinalIgnoreCase)))
+			else if (args.Any(x => x.Equals(UninstallHandlerArg, StringComparison.OrdinalIgnoreCase)))
 			{
-				if (Args.Any(x => x.Equals(ElevatedArg, StringComparison.OrdinalIgnoreCase)))
+				if (args.Any(x => x.Equals(ElevatedArg, StringComparison.OrdinalIgnoreCase)))
 				{
 					ProtocolHandlerUtils.UninstallElevated();
 				}
@@ -144,61 +144,61 @@ namespace UnrealGameSync
 			}
 			else
 			{
-				string UriIn = string.Empty;
-				for (int Idx = 0; Idx < Args.Length; Idx++)
+				string uriIn = string.Empty;
+				for (int idx = 0; idx < args.Length; idx++)
 				{
-					const string Prefix = "-uri=";
-					if (Args[Idx].StartsWith(Prefix, StringComparison.OrdinalIgnoreCase))
+					const string prefix = "-uri=";
+					if (args[idx].StartsWith(prefix, StringComparison.OrdinalIgnoreCase))
 					{
-						UriIn = Args[Idx].Substring(Prefix.Length);
+						uriIn = args[idx].Substring(prefix.Length);
 					}
 				}
 
-				if (UriIn == string.Empty)
+				if (uriIn == string.Empty)
 				{
 					return false;
 				}
 
-				Uri Uri;
+				Uri uri;
 				try
 				{
-					Uri = new Uri(UriIn);
+					uri = new Uri(uriIn);
 				}
 				catch
 				{
-					MessageBox.Show(String.Format("Invalid URI: {0}", UriIn));
+					MessageBox.Show(String.Format("Invalid URI: {0}", uriIn));
 					return true;
 				}
 
-				MethodInfo? Handler;
-				if (!Handlers.TryGetValue(Uri.Host, out Handler))
+				MethodInfo? handler;
+				if (!Handlers.TryGetValue(uri.Host, out handler))
 				{
-					MessageBox.Show(String.Format("Unknown action from URI request ('{0}')", Uri.Host));
+					MessageBox.Show(String.Format("Unknown action from URI request ('{0}')", uri.Host));
 					return true;
 				}
 
-				UriHandlerAttribute Attribute = Handler.GetCustomAttribute<UriHandlerAttribute>()!;
+				UriHandlerAttribute attribute = handler.GetCustomAttribute<UriHandlerAttribute>()!;
 
 				// handle case where we terminate after invoking handler
-				if (Attribute.Terminate)
+				if (attribute.Terminate)
 				{
-					UriResult Result = HandleUri(UriIn);
-					if (!Result.Success)
+					UriResult result = HandleUri(uriIn);
+					if (!result.Success)
 					{
-						MessageBox.Show(Result.Error);
+						MessageBox.Show(result.Error);
 					}
 					return true;
 				}
 
-				if (!FirstInstance)
+				if (!firstInstance)
 				{
-					if (ActivateEvent != null)
+					if (activateEvent != null)
 					{
-						ActivateEvent.Set();
+						activateEvent.Set();
 					}
 
 					// send to main UGS process using IPC
-					AutomationServer.SendUri(UriIn);
+					AutomationServer.SendUri(uriIn);
 					return true;
 				}
 
@@ -210,36 +210,36 @@ namespace UnrealGameSync
 
 		static UriHandler()
 		{
-			foreach (Type Type in Assembly.GetExecutingAssembly().GetTypes())
+			foreach (Type type in Assembly.GetExecutingAssembly().GetTypes())
 			{
-				List<MethodInfo> HandlerMethods = new List<MethodInfo> (Type.GetMethods().Where(MethodInfo => MethodInfo.GetCustomAttribute<UriHandlerAttribute>() != null));
-				foreach (MethodInfo MethodInfo in HandlerMethods)
+				List<MethodInfo> handlerMethods = new List<MethodInfo> (type.GetMethods().Where(methodInfo => methodInfo.GetCustomAttribute<UriHandlerAttribute>() != null));
+				foreach (MethodInfo methodInfo in handlerMethods)
 				{
-					if (!MethodInfo.IsStatic)
+					if (!methodInfo.IsStatic)
 					{
-						throw new Exception(string.Format("UriHandler method {0} must be static", MethodInfo.Name));
+						throw new Exception(string.Format("UriHandler method {0} must be static", methodInfo.Name));
 					}
 
-					if (MethodInfo.ReturnType != typeof(UriResult))
+					if (methodInfo.ReturnType != typeof(UriResult))
 					{
-						throw new Exception(string.Format("UriHandler method {0} must return UriResult type", MethodInfo.Name));
+						throw new Exception(string.Format("UriHandler method {0} must return UriResult type", methodInfo.Name));
 					}
 
-					if (Handlers.ContainsKey(MethodInfo.Name))
+					if (Handlers.ContainsKey(methodInfo.Name))
 					{
-						throw new Exception(string.Format("UriHandler method {0} clashes with another handler", MethodInfo.Name));
+						throw new Exception(string.Format("UriHandler method {0} clashes with another handler", methodInfo.Name));
 					}
 
-					foreach (ParameterInfo ParameterInfo in MethodInfo.GetParameters())
+					foreach (ParameterInfo parameterInfo in methodInfo.GetParameters())
 					{
-						Type ParameterType = ParameterInfo.ParameterType;
-						if (ParameterType != typeof(bool) && ParameterType != typeof(string) && ParameterType != typeof(float) && ParameterType != typeof(int))
+						Type parameterType = parameterInfo.ParameterType;
+						if (parameterType != typeof(bool) && parameterType != typeof(string) && parameterType != typeof(float) && parameterType != typeof(int))
 						{
-							throw new Exception(string.Format("UriHandler method parameter {0} must be bool, string, int, or float", ParameterInfo.Name));
+							throw new Exception(string.Format("UriHandler method parameter {0} must be bool, string, int, or float", parameterInfo.Name));
 						}
 					}
 
-					Handlers[MethodInfo.Name] = MethodInfo;
+					Handlers[methodInfo.Name] = methodInfo;
 				}
 			}
 		}
@@ -259,9 +259,9 @@ namespace UnrealGameSync
 	{
 		public bool Terminate;
 
-		public UriHandlerAttribute(bool Terminate = false)
+		public UriHandlerAttribute(bool terminate = false)
 		{
-			this.Terminate = Terminate;			
+			this.Terminate = terminate;			
 		}
 	}
 
@@ -285,12 +285,12 @@ namespace UnrealGameSync
 			public string? ValueName;
 			public string Value;
 
-			public RegistrySetting(RegistryKey RootKey, string KeyName, string? ValueName, string Value)
+			public RegistrySetting(RegistryKey rootKey, string keyName, string? valueName, string value)
 			{
-				this.RootKey = RootKey;
-				this.KeyName = KeyName;
-				this.ValueName = ValueName;
-				this.Value = Value;
+				this.RootKey = rootKey;
+				this.KeyName = keyName;
+				this.ValueName = valueName;
+				this.Value = value;
 			}
 		}
 
@@ -301,44 +301,44 @@ namespace UnrealGameSync
 
 		static List<RegistrySetting> GetRegistrySettings()
 		{
-			string ApplicationPath = CurrentProcessFilePath;
-
-			List<RegistrySetting> Keys = new List<RegistrySetting>();
-			Keys.Add(new RegistrySetting(Registry.ClassesRoot, "UGS", null, "URL:UGS Protocol"));
-			Keys.Add(new RegistrySetting(Registry.ClassesRoot, "UGS", "URL Protocol", ""));
-			Keys.Add(new RegistrySetting(Registry.ClassesRoot, "UGS\\DefaultIcon", null, String.Format("\"{0}\",0", ApplicationPath)));
-			Keys.Add(new RegistrySetting(Registry.ClassesRoot, "UGS\\shell\\open\\command", null, String.Format("\"{0}\" -uri=\"%1\"", ApplicationPath)));
-			return Keys;
+			string applicationPath = CurrentProcessFilePath;
+
+			List<RegistrySetting> keys = new List<RegistrySetting>();
+			keys.Add(new RegistrySetting(Registry.ClassesRoot, "UGS", null, "URL:UGS Protocol"));
+			keys.Add(new RegistrySetting(Registry.ClassesRoot, "UGS", "URL Protocol", ""));
+			keys.Add(new RegistrySetting(Registry.ClassesRoot, "UGS\\DefaultIcon", null, String.Format("\"{0}\",0", applicationPath)));
+			keys.Add(new RegistrySetting(Registry.ClassesRoot, "UGS\\shell\\open\\command", null, String.Format("\"{0}\" -uri=\"%1\"", applicationPath)));
+			return keys;
 		}
 
 		public static ProtocolHandlerState GetState()
 		{
 			try
 			{
-				bool bHasAny = false;
-				bool bHasAll = true;
+				bool hasAny = false;
+				bool hasAll = true;
 
-				List<RegistrySetting> Keys = GetRegistrySettings();
-				foreach (IGrouping<RegistryKey, RegistrySetting> RootKeyGroup in Keys.GroupBy(x => x.RootKey))
+				List<RegistrySetting> keys = GetRegistrySettings();
+				foreach (IGrouping<RegistryKey, RegistrySetting> rootKeyGroup in keys.GroupBy(x => x.RootKey))
 				{
-					foreach (IGrouping<string, RegistrySetting> KeyNameGroup in RootKeyGroup.GroupBy(x => x.KeyName))
+					foreach (IGrouping<string, RegistrySetting> keyNameGroup in rootKeyGroup.GroupBy(x => x.KeyName))
 					{
-						using (RegistryKey? RegistryKey = RootKeyGroup.Key.OpenSubKey(KeyNameGroup.Key))
+						using (RegistryKey? registryKey = rootKeyGroup.Key.OpenSubKey(keyNameGroup.Key))
 						{
-							if (RegistryKey == null)
+							if (registryKey == null)
 							{
-								bHasAll = false;
+								hasAll = false;
 							}
 							else
 							{
-								bHasAll &= KeyNameGroup.All(x => (RegistryKey.GetValue(x.ValueName) as string) == x.Value);
-								bHasAny = true;
+								hasAll &= keyNameGroup.All(x => (registryKey.GetValue(x.ValueName) as string) == x.Value);
+								hasAny = true;
 							}
 						}
 					}
 				}
 
-				return bHasAll? ProtocolHandlerState.Installed : bHasAny? ProtocolHandlerState.Unknown : ProtocolHandlerState.NotInstalled;
+				return hasAll? ProtocolHandlerState.Installed : hasAny? ProtocolHandlerState.Unknown : ProtocolHandlerState.NotInstalled;
 			}
 			catch
 			{
@@ -355,24 +355,24 @@ namespace UnrealGameSync
 		{
 			try
 			{
-				List<RegistrySetting> Keys = GetRegistrySettings();
-				foreach (IGrouping<RegistryKey, RegistrySetting> RootKeyGroup in Keys.GroupBy(x => x.RootKey))
+				List<RegistrySetting> keys = GetRegistrySettings();
+				foreach (IGrouping<RegistryKey, RegistrySetting> rootKeyGroup in keys.GroupBy(x => x.RootKey))
 				{
-					foreach (IGrouping<string, RegistrySetting> KeyNameGroup in RootKeyGroup.GroupBy(x => x.KeyName))
+					foreach (IGrouping<string, RegistrySetting> keyNameGroup in rootKeyGroup.GroupBy(x => x.KeyName))
 					{
-						using (RegistryKey RegistryKey = RootKeyGroup.Key.CreateSubKey(KeyNameGroup.Key))
+						using (RegistryKey registryKey = rootKeyGroup.Key.CreateSubKey(keyNameGroup.Key))
 						{
-							foreach (RegistrySetting Setting in KeyNameGroup)
+							foreach (RegistrySetting setting in keyNameGroup)
 							{
-								RegistryKey.SetValue(Setting.ValueName, Setting.Value);
+								registryKey.SetValue(setting.ValueName, setting.Value);
 							}
 						}
 					}
 				}
 			}
-			catch (Exception Ex)
+			catch (Exception ex)
 			{
-				MessageBox.Show(String.Format("Unable to register protocol handler: {0}", Ex));
+				MessageBox.Show(String.Format("Unable to register protocol handler: {0}", ex));
 			}
 		}
 
@@ -387,22 +387,22 @@ namespace UnrealGameSync
 			{
 				Registry.ClassesRoot.DeleteSubKeyTree("UGS", false);
 			}
-			catch (Exception Ex)
+			catch (Exception ex)
 			{
-				MessageBox.Show(String.Format("Unable to register protocol handler: {0}", Ex));
+				MessageBox.Show(String.Format("Unable to register protocol handler: {0}", ex));
 			}
 		}
 
-		private static void RunElevated(string Arguments)
+		private static void RunElevated(string arguments)
 		{
-			using (Process Process = new Process())
+			using (Process process = new Process())
 			{
-				Process.StartInfo.FileName = Path.ChangeExtension(Assembly.GetExecutingAssembly().Location, ".exe");
-				Process.StartInfo.Arguments = Arguments;
-				Process.StartInfo.Verb = "runas";
-				Process.StartInfo.UseShellExecute = true;
-				Process.Start();
-				Process.WaitForExit();
+				process.StartInfo.FileName = Path.ChangeExtension(Assembly.GetExecutingAssembly().Location, ".exe");
+				process.StartInfo.Arguments = arguments;
+				process.StartInfo.Verb = "runas";
+				process.StartInfo.UseShellExecute = true;
+				process.Start();
+				process.WaitForExit();
 			}
 		}
 	}
diff --git a/Engine/Source/Programs/UnrealGameSync/UnrealGameSync/UriHandlers/VisualStudioHandler.cs b/Engine/Source/Programs/UnrealGameSync/UnrealGameSync/UriHandlers/VisualStudioHandler.cs
--- a/Engine/Source/Programs/UnrealGameSync/UnrealGameSync/UriHandlers/VisualStudioHandler.cs
+++ b/Engine/Source/Programs/UnrealGameSync/UnrealGameSync/UriHandlers/VisualStudioHandler.cs
@@ -15,14 +15,14 @@ namespace UnrealGameSync
 	static class VisualStudioHandler
 	{
 		[UriHandler(true)]
-		public static UriResult VSOpen(string DepotPath, int Line = -1)
+		public static UriResult VsOpen(string depotPath, int line = -1)
 		{
-			string TempFileName = P4Automation.PrintToTempFile(null, DepotPath, NullLogger.Instance).GetAwaiter().GetResult();
+			string tempFileName = P4Automation.PrintToTempFile(null, depotPath, NullLogger.Instance).GetAwaiter().GetResult();
 
-			string? ErrorMessage;
-			if (!VisualStudioAutomation.OpenFile(TempFileName, out ErrorMessage, Line))
+			string? errorMessage;
+			if (!VisualStudioAutomation.OpenFile(tempFileName, out errorMessage, line))
 			{
-				return new UriResult() { Error = ErrorMessage ?? "Unknown Visual Studio Error" };
+				return new UriResult() { Error = errorMessage ?? "Unknown Visual Studio Error" };
 			}
 
 			return new UriResult() { Success = true };
