UGS: Reformat to standard C# style conventions.

diff --git a/Engine/Source/Programs/UnrealGameSync/UnrealGameSync/Workspace.cs b/Engine/Source/Programs/UnrealGameSync/UnrealGameSync/Workspace.cs
--- a/Engine/Source/Programs/UnrealGameSync/UnrealGameSync/Workspace.cs
+++ b/Engine/Source/Programs/UnrealGameSync/UnrealGameSync/Workspace.cs
@@ -28,35 +28,35 @@ namespace UnrealGameSync
 		public UserWorkspaceState State { get; }
 		public ProjectInfo Project { get; }
 		public SynchronizationContext SynchronizationContext { get; }
-		ILogger Logger;
+		ILogger _logger;
 
-		bool bSyncing => CurrentUpdate != null;
+		bool Syncing => _currentUpdate != null;
 
-		WorkspaceUpdate? CurrentUpdate;
+		WorkspaceUpdate? _currentUpdate;
 
 		public event Action<WorkspaceUpdateContext, WorkspaceUpdateResult, string>? OnUpdateComplete;
 
-		IAsyncDisposer AsyncDisposer;
+		IAsyncDisposer _asyncDisposer;
 
-		public Workspace(IPerforceSettings InPerfoceSettings, ProjectInfo InProject, UserWorkspaceState InState, ConfigFile ProjectConfigFile, IReadOnlyList<string>? ProjectStreamFilter, ILogger Logger, IServiceProvider ServiceProvider)
+		public Workspace(IPerforceSettings inPerfoceSettings, ProjectInfo inProject, UserWorkspaceState inState, ConfigFile projectConfigFile, IReadOnlyList<string>? projectStreamFilter, ILogger logger, IServiceProvider serviceProvider)
 		{
-			PerforceSettings = InPerfoceSettings;
-			Project = InProject;
-			State = InState;
+			PerforceSettings = inPerfoceSettings;
+			Project = inProject;
+			State = inState;
 			this.SynchronizationContext = SynchronizationContext.Current!;
-			this.Logger = Logger;
-			this.AsyncDisposer = ServiceProvider.GetRequiredService<IAsyncDisposer>();
+			this._logger = logger;
+			this._asyncDisposer = serviceProvider.GetRequiredService<IAsyncDisposer>();
 
-			this.ProjectConfigFile = ProjectConfigFile;
-			this.ProjectStreamFilter = ProjectStreamFilter;
+			this.ProjectConfigFile = projectConfigFile;
+			this.ProjectStreamFilter = projectStreamFilter;
 		}
 
 		public void Dispose()
 		{
 			CancelUpdate();
-			if (PrevUpdateTask != null)
+			if (_prevUpdateTask != null)
 			{
-				AsyncDisposer.Add(PrevUpdateTask);
+				_asyncDisposer.Add(_prevUpdateTask);
 			}
 		}
 
@@ -70,121 +70,121 @@ namespace UnrealGameSync
 			get; private set;
 		}
 
-		CancellationTokenSource? PrevCancellationSource;
-		Task PrevUpdateTask = Task.CompletedTask;
+		CancellationTokenSource? _prevCancellationSource;
+		Task _prevUpdateTask = Task.CompletedTask;
 
-		public void Update(WorkspaceUpdateContext Context)
+		public void Update(WorkspaceUpdateContext context)
 		{
 			CancelUpdate();
 
-			Task PrevUpdateTaskCopy = PrevUpdateTask;
+			Task prevUpdateTaskCopy = _prevUpdateTask;
 
-			WorkspaceUpdate Update = new WorkspaceUpdate(Context);
-			CurrentUpdate = Update;
+			WorkspaceUpdate update = new WorkspaceUpdate(context);
+			_currentUpdate = update;
 
-			CancellationTokenSource CancellationSource = new CancellationTokenSource();
-			PrevCancellationSource = CancellationSource;
-			PrevUpdateTask = Task.Run(() => UpdateWorkspaceMini(Update, PrevUpdateTaskCopy, CancellationSource.Token));
+			CancellationTokenSource cancellationSource = new CancellationTokenSource();
+			_prevCancellationSource = cancellationSource;
+			_prevUpdateTask = Task.Run(() => UpdateWorkspaceMini(update, prevUpdateTaskCopy, cancellationSource.Token));
 		}
 
 		public void CancelUpdate()
 		{
 			// Cancel the current task. We actually terminate the operation asynchronously, but we can signal the cancellation and 
 			// send a cancelled event, then wait for the heavy lifting to finish in the new update task.
-			if (PrevCancellationSource != null)
+			if (_prevCancellationSource != null)
 			{
-				CancellationTokenSource PrevCancellationSourceCopy = PrevCancellationSource;
-				PrevCancellationSourceCopy.Cancel();
-				PrevUpdateTask = PrevUpdateTask.ContinueWith(Task => PrevCancellationSourceCopy.Dispose());
-				PrevCancellationSource = null;
+				CancellationTokenSource prevCancellationSourceCopy = _prevCancellationSource;
+				prevCancellationSourceCopy.Cancel();
+				_prevUpdateTask = _prevUpdateTask.ContinueWith(task => prevCancellationSourceCopy.Dispose());
+				_prevCancellationSource = null;
 			}
-			if(CurrentUpdate != null)
+			if(_currentUpdate != null)
 			{
-				CompleteUpdate(CurrentUpdate, WorkspaceUpdateResult.Canceled, "Cancelled");
+				CompleteUpdate(_currentUpdate, WorkspaceUpdateResult.Canceled, "Cancelled");
 			}
 		}
 
-		async Task UpdateWorkspaceMini(WorkspaceUpdate Update, Task PrevUpdateTask, CancellationToken CancellationToken)
+		async Task UpdateWorkspaceMini(WorkspaceUpdate update, Task prevUpdateTask, CancellationToken cancellationToken)
 		{
-			if (PrevUpdateTask != null)
+			if (prevUpdateTask != null)
 			{
-				await PrevUpdateTask;
+				await prevUpdateTask;
 			}
 
-			WorkspaceUpdateContext Context = Update.Context;
-			Context.ProjectConfigFile = ProjectConfigFile;
-			Context.ProjectStreamFilter = ProjectStreamFilter;
+			WorkspaceUpdateContext context = update.Context;
+			context.ProjectConfigFile = ProjectConfigFile;
+			context.ProjectStreamFilter = ProjectStreamFilter;
 
-			string StatusMessage;
-			WorkspaceUpdateResult Result = WorkspaceUpdateResult.FailedToSync;
+			string statusMessage;
+			WorkspaceUpdateResult result = WorkspaceUpdateResult.FailedToSync;
 
 			try
 			{
-				(Result, StatusMessage) = await Update.ExecuteAsync(PerforceSettings, Project, State, Logger, CancellationToken);
-				if (Result != WorkspaceUpdateResult.Success)
+				(result, statusMessage) = await update.ExecuteAsync(PerforceSettings, Project, State, _logger, cancellationToken);
+				if (result != WorkspaceUpdateResult.Success)
 				{
-					Logger.LogError("{Message}", StatusMessage);
+					_logger.LogError("{Message}", statusMessage);
 				}
 			}
 			catch (OperationCanceledException)
 			{
-				StatusMessage = "Canceled.";
-				Logger.LogError("Canceled.");
+				statusMessage = "Canceled.";
+				_logger.LogError("Canceled.");
 			}
-			catch (Exception Ex)
+			catch (Exception ex)
 			{
-				StatusMessage = "Failed with exception - " + Ex.ToString();
-				Logger.LogError(Ex, "Failed with exception");
+				statusMessage = "Failed with exception - " + ex.ToString();
+				_logger.LogError(ex, "Failed with exception");
 			}
 
-			ProjectConfigFile = Context.ProjectConfigFile;
-			ProjectStreamFilter = Context.ProjectStreamFilter;
+			ProjectConfigFile = context.ProjectConfigFile;
+			ProjectStreamFilter = context.ProjectStreamFilter;
 
-			SynchronizationContext.Post(x => CompleteUpdate(Update, Result, StatusMessage), null);
+			SynchronizationContext.Post(x => CompleteUpdate(update, result, statusMessage), null);
 		}
 
-		void CompleteUpdate(WorkspaceUpdate Update, WorkspaceUpdateResult Result, string StatusMessage)
+		void CompleteUpdate(WorkspaceUpdate update, WorkspaceUpdateResult result, string statusMessage)
 		{
-			if (CurrentUpdate == Update)
+			if (_currentUpdate == update)
 			{
-				WorkspaceUpdateContext Context = Update.Context;
+				WorkspaceUpdateContext context = update.Context;
 
-				State.SetLastSyncState(Result, Context, StatusMessage);
-				State.Save(Logger);
+				State.SetLastSyncState(result, context, statusMessage);
+				State.Save(_logger);
 
-				OnUpdateComplete?.Invoke(Context, Result, StatusMessage);
-				CurrentUpdate = null;
+				OnUpdateComplete?.Invoke(context, result, statusMessage);
+				_currentUpdate = null;
 			}
 		}
 
-		public Dictionary<string, string> GetVariables(BuildConfig EditorConfig, int? OverrideChange = null, int? OverrideCodeChange = null)
+		public Dictionary<string, string> GetVariables(BuildConfig editorConfig, int? overrideChange = null, int? overrideCodeChange = null)
 		{
-			FileReference EditorReceiptFile = ConfigUtils.GetEditorReceiptFile(Project, ProjectConfigFile, EditorConfig);
+			FileReference editorReceiptFile = ConfigUtils.GetEditorReceiptFile(Project, ProjectConfigFile, editorConfig);
 
-			TargetReceipt? EditorReceipt;
-			if (!ConfigUtils.TryReadEditorReceipt(Project, EditorReceiptFile, out EditorReceipt))
+			TargetReceipt? editorReceipt;
+			if (!ConfigUtils.TryReadEditorReceipt(Project, editorReceiptFile, out editorReceipt))
 			{
-				EditorReceipt = ConfigUtils.CreateDefaultEditorReceipt(Project, ProjectConfigFile, EditorConfig);
+				editorReceipt = ConfigUtils.CreateDefaultEditorReceipt(Project, ProjectConfigFile, editorConfig);
 			}
 
-			Dictionary<string, string> Variables = ConfigUtils.GetWorkspaceVariables(Project, OverrideChange ?? State.CurrentChangeNumber, OverrideCodeChange ?? State.CurrentCodeChangeNumber, EditorReceipt, ProjectConfigFile);
-			return Variables;
+			Dictionary<string, string> variables = ConfigUtils.GetWorkspaceVariables(Project, overrideChange ?? State.CurrentChangeNumber, overrideCodeChange ?? State.CurrentCodeChangeNumber, editorReceipt, ProjectConfigFile);
+			return variables;
 		}
 
 		public bool IsBusy()
 		{
-			return bSyncing;
+			return Syncing;
 		}
 
 		public int CurrentChangeNumber => State.CurrentChangeNumber;
 
-		public int PendingChangeNumber => CurrentUpdate?.Context?.ChangeNumber ?? CurrentChangeNumber;
+		public int PendingChangeNumber => _currentUpdate?.Context?.ChangeNumber ?? CurrentChangeNumber;
 
 		public string ClientName
 		{
 			get { return PerforceSettings.ClientName!; }
 		}
 
-		public Tuple<string, float> CurrentProgress => CurrentUpdate?.CurrentProgress ?? new Tuple<string, float>("", 0.0f);
+		public Tuple<string, float> CurrentProgress => _currentUpdate?.CurrentProgress ?? new Tuple<string, float>("", 0.0f);
 	}
 }
diff --git a/Engine/Source/Programs/UnrealGameSync/UnrealGameSyncCmd/Program.cs b/Engine/Source/Programs/UnrealGameSync/UnrealGameSyncCmd/Program.cs
--- a/Engine/Source/Programs/UnrealGameSync/UnrealGameSyncCmd/Program.cs
+++ b/Engine/Source/Programs/UnrealGameSync/UnrealGameSyncCmd/Program.cs
@@ -28,15 +28,15 @@ namespace UnrealGameSyncCmd
 		public LogEvent Event { get; }
 		public int Code { get; }
 
-		public UserErrorException(LogEvent Event)
-			: base(Event.ToString())
+		public UserErrorException(LogEvent evt)
+			: base(evt.ToString())
 		{
-			this.Event = Event;
+			this.Event = evt;
 			this.Code = 1;
 		}
 
-		public UserErrorException(string Message, params object[] Args)
-			: this(LogEvent.Create(LogLevel.Error, Message, Args))
+		public UserErrorException(string message, params object[] args)
+			: this(LogEvent.Create(LogLevel.Error, message, args))
 		{
 		}
 	}
@@ -52,16 +52,16 @@ namespace UnrealGameSyncCmd
 			public string Usage { get; }
 			public string Brief { get; }
 
-			public CommandInfo(string Name, Type Type, string Usage, string Brief)
+			public CommandInfo(string name, Type type, string usage, string brief)
 			{
-				this.Name = Name;
-				this.Type = Type;
-				this.Usage = Usage;
-				this.Brief = Brief;
+				this.Name = name;
+				this.Type = type;
+				this.Usage = usage;
+				this.Brief = brief;
 			}
 		}
 
-		static CommandInfo[] Commands =
+		static CommandInfo[] _commands =
 		{
 			new CommandInfo("init", typeof(InitCommand),
 				"ugs init [stream-path] [-client=..] [-server=..] [-user=..] [-branch=..] [-project=..]",
@@ -116,12 +116,12 @@ namespace UnrealGameSyncCmd
 			public ILoggerFactory LoggerFactory { get; }
 			public GlobalSettingsFile UserSettings { get; }
 
-			public CommandContext(CommandLineArguments Arguments, ILogger Logger, ILoggerFactory LoggerFactory, GlobalSettingsFile UserSettings)
+			public CommandContext(CommandLineArguments arguments, ILogger logger, ILoggerFactory loggerFactory, GlobalSettingsFile userSettings)
 			{
-				this.Arguments = Arguments;
-				this.Logger = Logger;
-				this.LoggerFactory = LoggerFactory;
-				this.UserSettings = UserSettings;
+				this.Arguments = arguments;
+				this.Logger = logger;
+				this.LoggerFactory = loggerFactory;
+				this.UserSettings = userSettings;
 			}
 		}
 
@@ -136,15 +136,15 @@ namespace UnrealGameSyncCmd
 
 		class ProjectConfigOptions : ServerOptions
 		{
-			public void ApplyTo(UserWorkspaceSettings Settings)
+			public void ApplyTo(UserWorkspaceSettings settings)
 			{
 				if (ServerAndPort != null)
 				{
-					Settings.ServerAndPort = (ServerAndPort.Length == 0) ? null : ServerAndPort;
+					settings.ServerAndPort = (ServerAndPort.Length == 0) ? null : ServerAndPort;
 				}
 				if (UserName != null)
 				{
-					Settings.UserName = (UserName.Length == 0) ? null : UserName;
+					settings.UserName = (UserName.Length == 0) ? null : UserName;
 				}
 			}
 		}
@@ -161,92 +161,92 @@ namespace UnrealGameSyncCmd
 			public string? ProjectName { get; set; }
 		}
 
-		public static async Task<int> Main(string[] RawArgs)
+		public static async Task<int> Main(string[] rawArgs)
 		{
-			DirectoryReference GlobalConfigFolder;
+			DirectoryReference globalConfigFolder;
 			if (RuntimeInformation.IsOSPlatform(OSPlatform.Windows))
 			{
-				GlobalConfigFolder = DirectoryReference.Combine(DirectoryReference.GetSpecialFolder(Environment.SpecialFolder.LocalApplicationData)!, "UnrealGameSync");
+				globalConfigFolder = DirectoryReference.Combine(DirectoryReference.GetSpecialFolder(Environment.SpecialFolder.LocalApplicationData)!, "UnrealGameSync");
 			}
 			else
 			{
-				GlobalConfigFolder = DirectoryReference.Combine(DirectoryReference.GetSpecialFolder(Environment.SpecialFolder.UserProfile)!, ".config", "UnrealGameSync");
+				globalConfigFolder = DirectoryReference.Combine(DirectoryReference.GetSpecialFolder(Environment.SpecialFolder.UserProfile)!, ".config", "UnrealGameSync");
 			}
-			DirectoryReference.CreateDirectory(GlobalConfigFolder);
+			DirectoryReference.CreateDirectory(globalConfigFolder);
 
-			string LogName;
-			DirectoryReference LogFolder;
+			string logName;
+			DirectoryReference logFolder;
 			if (RuntimeInformation.IsOSPlatform(OSPlatform.OSX))
 			{
-				LogFolder = DirectoryReference.Combine(DirectoryReference.GetSpecialFolder(Environment.SpecialFolder.UserProfile)!, "Library", "Logs", "Unreal Engine", "UnrealGameSync");
-				LogName = "UnrealGameSync-.log";
+				logFolder = DirectoryReference.Combine(DirectoryReference.GetSpecialFolder(Environment.SpecialFolder.UserProfile)!, "Library", "Logs", "Unreal Engine", "UnrealGameSync");
+				logName = "UnrealGameSync-.log";
 			}
 			else if (RuntimeInformation.IsOSPlatform(OSPlatform.Linux))
 			{
-				LogFolder = DirectoryReference.GetSpecialFolder(Environment.SpecialFolder.UserProfile)!;
-				LogName = ".ugs-.log";
+				logFolder = DirectoryReference.GetSpecialFolder(Environment.SpecialFolder.UserProfile)!;
+				logName = ".ugs-.log";
 			}
 			else
 			{
-				LogFolder = GlobalConfigFolder;
-				LogName = "UnrealGameSyncCmd-.log";
+				logFolder = globalConfigFolder;
+				logName = "UnrealGameSyncCmd-.log";
 			}
 
-			Serilog.ILogger SerilogLogger = new LoggerConfiguration()
+			Serilog.ILogger serilogLogger = new LoggerConfiguration()
 				.Enrich.FromLogContext()
 				.WriteTo.Console(Serilog.Events.LogEventLevel.Information, outputTemplate: "{Message:lj}{NewLine}")
-				.WriteTo.File(FileReference.Combine(LogFolder, LogName).FullName, Serilog.Events.LogEventLevel.Debug, rollingInterval: RollingInterval.Day, rollOnFileSizeLimit: true, fileSizeLimitBytes: 20 * 1024 * 1024, retainedFileCountLimit: 10)
+				.WriteTo.File(FileReference.Combine(logFolder, logName).FullName, Serilog.Events.LogEventLevel.Debug, rollingInterval: RollingInterval.Day, rollOnFileSizeLimit: true, fileSizeLimitBytes: 20 * 1024 * 1024, retainedFileCountLimit: 10)
 				.CreateLogger();
 
-			using ILoggerFactory LoggerFactory = new Serilog.Extensions.Logging.SerilogLoggerFactory(SerilogLogger, true);
-			ILogger Logger = LoggerFactory.CreateLogger("Main");
+			using ILoggerFactory loggerFactory = new Serilog.Extensions.Logging.SerilogLoggerFactory(serilogLogger, true);
+			ILogger logger = loggerFactory.CreateLogger("Main");
 			try
 			{
-				GlobalSettingsFile Settings;
+				GlobalSettingsFile settings;
 				if (RuntimeInformation.IsOSPlatform(OSPlatform.Windows))
 				{
-					Settings = UserSettings.Create(GlobalConfigFolder, Logger);
+					settings = UserSettings.Create(globalConfigFolder, logger);
 				}
 				else
 				{
-					Settings = GlobalSettingsFile.Create(FileReference.Combine(GlobalConfigFolder, "Global.json"));
+					settings = GlobalSettingsFile.Create(FileReference.Combine(globalConfigFolder, "Global.json"));
 				}
 
-				CommandLineArguments Args = new CommandLineArguments(RawArgs);
+				CommandLineArguments args = new CommandLineArguments(rawArgs);
 
-				string? CommandName;
-				if (!Args.TryGetPositionalArgument(out CommandName))
+				string? commandName;
+				if (!args.TryGetPositionalArgument(out commandName))
 				{
 					PrintHelp();
 					return 0;
 				}
 
-				CommandInfo? Command = Commands.FirstOrDefault(x => x.Name.Equals(CommandName, StringComparison.OrdinalIgnoreCase));
-				if (Command == null)
+				CommandInfo? command = _commands.FirstOrDefault(x => x.Name.Equals(commandName, StringComparison.OrdinalIgnoreCase));
+				if (command == null)
 				{
-					Logger.LogError($"unknown command '{CommandName}'");
+					logger.LogError($"unknown command '{commandName}'");
 					Console.WriteLine();
 					PrintHelp();
 					return 1;
 				}
 
-				Command Instance = (Command)Activator.CreateInstance(Command.Type)!;
-				await Instance.ExecuteAsync(new CommandContext(Args, Logger, LoggerFactory, Settings));
+				Command instance = (Command)Activator.CreateInstance(command.Type)!;
+				await instance.ExecuteAsync(new CommandContext(args, logger, loggerFactory, settings));
 				return 0;
 			}
-			catch (UserErrorException Ex)
+			catch (UserErrorException ex)
 			{
-				Logger.Log(Ex.Event.Level, "{Message}", Ex.Event.ToString());
-				return Ex.Code;
+				logger.Log(ex.Event.Level, "{Message}", ex.Event.ToString());
+				return ex.Code;
 			}
-			catch (PerforceException Ex)
+			catch (PerforceException ex)
 			{
-				Logger.LogError(Ex, "{Message}", Ex.Message);
+				logger.LogError(ex, "{Message}", ex.Message);
 				return 1;
 			}
-			catch (Exception Ex)
+			catch (Exception ex)
 			{
-				Logger.LogError(Ex, "Unhandled exception.\n{Str}", Ex);
+				logger.LogError(ex, "Unhandled exception.\n{Str}", ex);
 				return 1;
 			}
 		}
@@ -254,25 +254,25 @@ namespace UnrealGameSyncCmd
 		static void PrintHelp()
 		{
 			Console.WriteLine("Usage:");
-			foreach (CommandInfo Command in Commands)
+			foreach (CommandInfo command in _commands)
 			{
 				Console.WriteLine();
-				ConsoleUtils.WriteLineWithWordWrap(Command.Usage, 2, 8);
-				ConsoleUtils.WriteLineWithWordWrap(Command.Brief, 4, 4);
+				ConsoleUtils.WriteLineWithWordWrap(command.Usage, 2, 8);
+				ConsoleUtils.WriteLineWithWordWrap(command.Brief, 4, 4);
 			}
 		}
 
 		public static UserWorkspaceSettings? ReadOptionalUserWorkspaceSettings()
 		{
-			DirectoryReference? Dir = DirectoryReference.GetCurrentDirectory();
-			for (; Dir != null; Dir = Dir.ParentDirectory)
+			DirectoryReference? dir = DirectoryReference.GetCurrentDirectory();
+			for (; dir != null; dir = dir.ParentDirectory)
 			{
 				try
 				{
-					UserWorkspaceSettings? Settings;
-					if (UserWorkspaceSettings.TryLoad(Dir, out Settings))
+					UserWorkspaceSettings? settings;
+					if (UserWorkspaceSettings.TryLoad(dir, out settings))
 					{
-						return Settings;
+						return settings;
 					}
 				}
 				catch
@@ -285,195 +285,195 @@ namespace UnrealGameSyncCmd
 
 		public static UserWorkspaceSettings ReadRequiredUserWorkspaceSettings()
 		{
-			UserWorkspaceSettings? Settings = ReadOptionalUserWorkspaceSettings();
-			if (Settings == null)
+			UserWorkspaceSettings? settings = ReadOptionalUserWorkspaceSettings();
+			if (settings == null)
 			{
 				throw new UserErrorException("Unable to find UGS workspace in current directory.");
 			}
-			return Settings;
+			return settings;
 		}
 
-		public static async Task<UserWorkspaceState> ReadWorkspaceState(IPerforceConnection PerforceClient, UserWorkspaceSettings Settings, GlobalSettingsFile UserSettings, ILogger Logger)
+		public static async Task<UserWorkspaceState> ReadWorkspaceState(IPerforceConnection perforceClient, UserWorkspaceSettings settings, GlobalSettingsFile userSettings, ILogger logger)
 		{
-			UserWorkspaceState State = UserSettings.FindOrAddWorkspaceState(Settings, Logger);
-			if (State.SettingsTimeUtc != Settings.LastModifiedTimeUtc)
+			UserWorkspaceState state = userSettings.FindOrAddWorkspaceState(settings, logger);
+			if (state.SettingsTimeUtc != settings.LastModifiedTimeUtc)
 			{
-				Logger.LogDebug("Updating state due to modified settings timestamp");
-				ProjectInfo Info = await ProjectInfo.CreateAsync(PerforceClient, Settings, CancellationToken.None);
-				State.UpdateCachedProjectInfo(Info, Settings.LastModifiedTimeUtc);
-				State.Save(Logger);
+				logger.LogDebug("Updating state due to modified settings timestamp");
+				ProjectInfo info = await ProjectInfo.CreateAsync(perforceClient, settings, CancellationToken.None);
+				state.UpdateCachedProjectInfo(info, settings.LastModifiedTimeUtc);
+				state.Save(logger);
 			}
-			return State;
+			return state;
 		}
 
-		public static Task<IPerforceConnection> ConnectAsync(string? ServerAndPort, string? UserName, string? ClientName, ILoggerFactory LoggerFactory)
+		public static Task<IPerforceConnection> ConnectAsync(string? serverAndPort, string? userName, string? clientName, ILoggerFactory loggerFactory)
 		{
-			PerforceSettings Settings = new PerforceSettings(PerforceSettings.Default);
-			Settings.ClientName = ClientName;
-			Settings.PreferNativeClient = true;
-			if (!String.IsNullOrEmpty(ServerAndPort))
+			PerforceSettings settings = new PerforceSettings(PerforceSettings.Default);
+			settings.ClientName = clientName;
+			settings.PreferNativeClient = true;
+			if (!String.IsNullOrEmpty(serverAndPort))
 			{
-				Settings.ServerAndPort = ServerAndPort;
+				settings.ServerAndPort = serverAndPort;
 			}
-			if (!String.IsNullOrEmpty(UserName))
+			if (!String.IsNullOrEmpty(userName))
 			{
-				Settings.UserName = UserName;
+				settings.UserName = userName;
 			}
 
-			return PerforceConnection.CreateAsync(Settings, LoggerFactory.CreateLogger("Perforce"));
+			return PerforceConnection.CreateAsync(settings, loggerFactory.CreateLogger("Perforce"));
 		}
 
-		public static Task<IPerforceConnection> ConnectAsync(UserWorkspaceSettings Settings, ILoggerFactory LoggerFactory)
+		public static Task<IPerforceConnection> ConnectAsync(UserWorkspaceSettings settings, ILoggerFactory loggerFactory)
 		{
-			return ConnectAsync(Settings.ServerAndPort, Settings.UserName, Settings.ClientName, LoggerFactory);
+			return ConnectAsync(settings.ServerAndPort, settings.UserName, settings.ClientName, loggerFactory);
 		}
 
-		static string[] ReadSyncFilter(UserWorkspaceSettings WorkspaceSettings, GlobalSettingsFile UserSettings, ConfigFile ProjectConfig)
+		static string[] ReadSyncFilter(UserWorkspaceSettings workspaceSettings, GlobalSettingsFile userSettings, ConfigFile projectConfig)
 		{
-			Dictionary<Guid, WorkspaceSyncCategory> SyncCategories = ConfigUtils.GetSyncCategories(ProjectConfig);
-			string[] CombinedSyncFilter = GlobalSettingsFile.GetCombinedSyncFilter(SyncCategories, UserSettings.Global.Filter, WorkspaceSettings.Filter);
+			Dictionary<Guid, WorkspaceSyncCategory> syncCategories = ConfigUtils.GetSyncCategories(projectConfig);
+			string[] combinedSyncFilter = GlobalSettingsFile.GetCombinedSyncFilter(syncCategories, userSettings.Global.Filter, workspaceSettings.Filter);
 
-			ConfigSection PerforceSection = ProjectConfig.FindSection("Perforce");
-			if (PerforceSection != null)
+			ConfigSection perforceSection = projectConfig.FindSection("Perforce");
+			if (perforceSection != null)
 			{
-				IEnumerable<string> AdditionalPaths = PerforceSection.GetValues("AdditionalPathsToSync", new string[0]);
-				CombinedSyncFilter = AdditionalPaths.Union(CombinedSyncFilter).ToArray();
+				IEnumerable<string> additionalPaths = perforceSection.GetValues("AdditionalPathsToSync", new string[0]);
+				combinedSyncFilter = additionalPaths.Union(combinedSyncFilter).ToArray();
 			}
 
-			return CombinedSyncFilter;
+			return combinedSyncFilter;
 		}
 
-		static async Task<string> FindProjectPathAsync(IPerforceConnection Perforce, string ClientName, string BranchPath, string? ProjectName)
+		static async Task<string> FindProjectPathAsync(IPerforceConnection perforce, string clientName, string branchPath, string? projectName)
 		{
-			using IPerforceConnection PerforceClient = await PerforceConnection.CreateAsync(new PerforceSettings(Perforce.Settings) { ClientName = ClientName }, Perforce.Logger);
+			using IPerforceConnection perforceClient = await PerforceConnection.CreateAsync(new PerforceSettings(perforce.Settings) { ClientName = clientName }, perforce.Logger);
 
 			// Find or validate the selected project
-			string SearchPath;
-			if (ProjectName == null)
+			string searchPath;
+			if (projectName == null)
 			{
-				SearchPath = $"//{ClientName}{BranchPath}/*.uprojectdirs";
+				searchPath = $"//{clientName}{branchPath}/*.uprojectdirs";
 			}
-			else if (ProjectName.Contains('.'))
+			else if (projectName.Contains('.'))
 			{
-				SearchPath = $"//{ClientName}{BranchPath}/{ProjectName.TrimStart('/')}";
+				searchPath = $"//{clientName}{branchPath}/{projectName.TrimStart('/')}";
 			}
 			else
 			{
-				SearchPath = $"//{ClientName}{BranchPath}/.../{ProjectName}.uproject";
+				searchPath = $"//{clientName}{branchPath}/.../{projectName}.uproject";
 			}
 
-			List<FStatRecord> ProjectFileRecords = await PerforceClient.FStatAsync(FStatOptions.ClientFileInPerforceSyntax, SearchPath).ToListAsync();
-			ProjectFileRecords.RemoveAll(x => x.HeadAction == FileAction.Delete || x.HeadAction == FileAction.MoveDelete);
-			ProjectFileRecords.RemoveAll(x => !x.IsMapped);
+			List<FStatRecord> projectFileRecords = await perforceClient.FStatAsync(FStatOptions.ClientFileInPerforceSyntax, searchPath).ToListAsync();
+			projectFileRecords.RemoveAll(x => x.HeadAction == FileAction.Delete || x.HeadAction == FileAction.MoveDelete);
+			projectFileRecords.RemoveAll(x => !x.IsMapped);
 
-			List<string> Paths = ProjectFileRecords.Select(x => PerforceUtils.GetClientRelativePath(x.ClientFile!)).Distinct(StringComparer.Ordinal).ToList();
-			if (Paths.Count == 0)
+			List<string> paths = projectFileRecords.Select(x => PerforceUtils.GetClientRelativePath(x.ClientFile!)).Distinct(StringComparer.Ordinal).ToList();
+			if (paths.Count == 0)
 			{
-				throw new UserErrorException("No project file found matching {SearchPath}", SearchPath);
+				throw new UserErrorException("No project file found matching {SearchPath}", searchPath);
 			}
-			if (Paths.Count > 1)
+			if (paths.Count > 1)
 			{
-				throw new UserErrorException("Multiple projects found matching {SearchPath}: {Paths}", SearchPath, String.Join(", ", Paths));
+				throw new UserErrorException("Multiple projects found matching {SearchPath}: {Paths}", searchPath, String.Join(", ", paths));
 			}
 
-			return "/" + Paths[0];
+			return "/" + paths[0];
 		}
 
 		abstract class Command
 		{
-			public abstract Task ExecuteAsync(CommandContext Context);
+			public abstract Task ExecuteAsync(CommandContext context);
 		}
 
 		class InitCommand : Command
 		{
-			public override async Task ExecuteAsync(CommandContext Context)
+			public override async Task ExecuteAsync(CommandContext context)
 			{
-				ILogger Logger = Context.Logger;
+				ILogger logger = context.Logger;
 
 				// Get the positional argument indicating the file to look for
-				string? InitName;
-				Context.Arguments.TryGetPositionalArgument(out InitName);
+				string? initName;
+				context.Arguments.TryGetPositionalArgument(out initName);
 
 				// Get the config settings from the command line
-				ProjectInitOptions Options = new ProjectInitOptions();
-				Context.Arguments.ApplyTo(Options);
-				Context.Arguments.CheckAllArgumentsUsed();
+				ProjectInitOptions options = new ProjectInitOptions();
+				context.Arguments.ApplyTo(options);
+				context.Arguments.CheckAllArgumentsUsed();
 
 				// Get the host name
-				using IPerforceConnection Perforce = await ConnectAsync(Options.ServerAndPort, Options.UserName, null, Context.LoggerFactory);
-				InfoRecord PerforceInfo = await Perforce.GetInfoAsync(InfoOptions.ShortOutput);
-				string HostName = PerforceInfo.ClientHost ?? Dns.GetHostName();
+				using IPerforceConnection perforce = await ConnectAsync(options.ServerAndPort, options.UserName, null, context.LoggerFactory);
+				InfoRecord perforceInfo = await perforce.GetInfoAsync(InfoOptions.ShortOutput);
+				string hostName = perforceInfo.ClientHost ?? Dns.GetHostName();
 
 				// Create the perforce connection
-				if (InitName != null)
+				if (initName != null)
 				{
-					await InitNewClientAsync(Perforce, Context, InitName, HostName, Options, Logger);
+					await InitNewClientAsync(perforce, context, initName, hostName, options, logger);
 				}
 				else
 				{
-					await InitExistingClientAsync(Perforce, Context, HostName, Options, Logger);
+					await InitExistingClientAsync(perforce, context, hostName, options, logger);
 				}
 			}
 
-			async Task InitNewClientAsync(IPerforceConnection Perforce, CommandContext Context, string StreamName, string HostName, ProjectInitOptions Options, ILogger Logger)
+			async Task InitNewClientAsync(IPerforceConnection perforce, CommandContext context, string streamName, string hostName, ProjectInitOptions options, ILogger logger)
 			{
-				Logger.LogInformation("Checking stream...");
+				logger.LogInformation("Checking stream...");
 
 				// Get the given stream
-				PerforceResponse<StreamRecord> StreamResponse = await Perforce.TryGetStreamAsync(StreamName, true);
-				if (!StreamResponse.Succeeded)
+				PerforceResponse<StreamRecord> streamResponse = await perforce.TryGetStreamAsync(streamName, true);
+				if (!streamResponse.Succeeded)
 				{
-					throw new UserErrorException($"Unable to find stream '{StreamName}'");
+					throw new UserErrorException($"Unable to find stream '{streamName}'");
 				}
-				StreamRecord Stream = StreamResponse.Data;
+				StreamRecord stream = streamResponse.Data;
 
 				// Get the new directory for the client
-				DirectoryReference ClientDir = DirectoryReference.Combine(DirectoryReference.GetCurrentDirectory(), Stream.Stream.Replace('/', '+'));
-				DirectoryReference.CreateDirectory(ClientDir);
+				DirectoryReference clientDir = DirectoryReference.Combine(DirectoryReference.GetCurrentDirectory(), stream.Stream.Replace('/', '+'));
+				DirectoryReference.CreateDirectory(clientDir);
 
 				// Make up a new client name 
-				string ClientName = Options.ClientName ?? Regex.Replace($"{Perforce.Settings.UserName}_{HostName}_{Stream.Stream.Trim('/')}", "[^0-9a-zA-Z_.-]", "+");
+				string clientName = options.ClientName ?? Regex.Replace($"{perforce.Settings.UserName}_{hostName}_{stream.Stream.Trim('/')}", "[^0-9a-zA-Z_.-]", "+");
 
 				// Check there are no existing clients under the current path
-				List<ClientsRecord> Clients = await FindExistingClients(Perforce, HostName, ClientDir);
-				if (Clients.Count > 0)
+				List<ClientsRecord> clients = await FindExistingClients(perforce, hostName, clientDir);
+				if (clients.Count > 0)
 				{
-					if (Clients.Count == 1 && ClientName.Equals(Clients[0].Name, StringComparison.OrdinalIgnoreCase) && ClientDir == TryParseRoot(Clients[0].Root))
+					if (clients.Count == 1 && clientName.Equals(clients[0].Name, StringComparison.OrdinalIgnoreCase) && clientDir == TryParseRoot(clients[0].Root))
 					{
-						Logger.LogInformation("Reusing existing client for {ClientDir} ({ClientName})", ClientDir, Options.ClientName);
+						logger.LogInformation("Reusing existing client for {ClientDir} ({ClientName})", clientDir, options.ClientName);
 					}
 					else
 					{
-						throw new UserErrorException("Current directory is already within a Perforce workspace ({ClientName})", Clients[0].Name);
+						throw new UserErrorException("Current directory is already within a Perforce workspace ({ClientName})", clients[0].Name);
 					}
 				}
 
 				// Create the new client
-				ClientRecord Client = new ClientRecord(ClientName, Perforce.Settings.UserName, ClientDir.FullName);
-				Client.Host = HostName;
-				Client.Stream = Stream.Stream;
-				Client.Options = ClientOptions.Rmdir;
-				await Perforce.CreateClientAsync(Client);
+				ClientRecord client = new ClientRecord(clientName, perforce.Settings.UserName, clientDir.FullName);
+				client.Host = hostName;
+				client.Stream = stream.Stream;
+				client.Options = ClientOptions.Rmdir;
+				await perforce.CreateClientAsync(client);
 
 				// Branch root is currently hard-coded at the root
-				string BranchPath = Options.BranchPath ?? String.Empty;
-				string ProjectPath = await FindProjectPathAsync(Perforce, ClientName, BranchPath, Options.ProjectName);
+				string branchPath = options.BranchPath ?? String.Empty;
+				string projectPath = await FindProjectPathAsync(perforce, clientName, branchPath, options.ProjectName);
 
 				// Create the settings object
-				UserWorkspaceSettings Settings = new UserWorkspaceSettings();
-				Settings.RootDir = ClientDir;
-				Settings.Init(Perforce.Settings.ServerAndPort, Perforce.Settings.UserName, ClientName, BranchPath, ProjectPath);
-				Options.ApplyTo(Settings);
-				Settings.Save(Logger);
+				UserWorkspaceSettings settings = new UserWorkspaceSettings();
+				settings.RootDir = clientDir;
+				settings.Init(perforce.Settings.ServerAndPort, perforce.Settings.UserName, clientName, branchPath, projectPath);
+				options.ApplyTo(settings);
+				settings.Save(logger);
 
-				Logger.LogInformation("Initialized {ClientName} with root at {RootDir}", ClientName, ClientDir);
+				logger.LogInformation("Initialized {ClientName} with root at {RootDir}", clientName, clientDir);
 			}
 
-			static DirectoryReference? TryParseRoot(string Root)
+			static DirectoryReference? TryParseRoot(string root)
 			{
 				try
 				{
-					return new DirectoryReference(Root);
+					return new DirectoryReference(root);
 				}
 				catch
 				{
@@ -481,80 +481,80 @@ namespace UnrealGameSyncCmd
 				}
 			}
 
-			async Task InitExistingClientAsync(IPerforceConnection Perforce, CommandContext Context, string HostName, ProjectInitOptions Options, ILogger Logger)
+			async Task InitExistingClientAsync(IPerforceConnection perforce, CommandContext context, string hostName, ProjectInitOptions options, ILogger logger)
 			{
-				DirectoryReference CurrentDir = DirectoryReference.GetCurrentDirectory();
+				DirectoryReference currentDir = DirectoryReference.GetCurrentDirectory();
 
 				// Make sure the client name is set
-				string? ClientName = Options.ClientName;
-				if (ClientName == null)
+				string? clientName = options.ClientName;
+				if (clientName == null)
 				{
-					List<ClientsRecord> Clients = await FindExistingClients(Perforce, HostName, CurrentDir);
-					if (Clients.Count == 0)
+					List<ClientsRecord> clients = await FindExistingClients(perforce, hostName, currentDir);
+					if (clients.Count == 0)
 					{
-						throw new UserErrorException("Unable to find client for {HostName} under {ClientDir}", HostName, CurrentDir);
+						throw new UserErrorException("Unable to find client for {HostName} under {ClientDir}", hostName, currentDir);
 					}
-					if (Clients.Count > 1)
+					if (clients.Count > 1)
 					{
-						throw new UserErrorException("Multiple clients found for {HostName} under {ClientDir}: {ClientList}", HostName, CurrentDir, String.Join(", ", Clients.Select(x => x.Name)));
+						throw new UserErrorException("Multiple clients found for {HostName} under {ClientDir}: {ClientList}", hostName, currentDir, String.Join(", ", clients.Select(x => x.Name)));
 					}
 
-					ClientName = Clients[0].Name;
-					Logger.LogInformation("Found client {ClientName}", ClientName);
+					clientName = clients[0].Name;
+					logger.LogInformation("Found client {ClientName}", clientName);
 				}
 
 				// Get the client info
-				ClientRecord Client = await Perforce.GetClientAsync(ClientName);
-				DirectoryReference ClientDir = new DirectoryReference(Client.Root);
+				ClientRecord client = await perforce.GetClientAsync(clientName);
+				DirectoryReference clientDir = new DirectoryReference(client.Root);
 
 				// If a project path was specified in local syntax, try to convert it to client-relative syntax
-				string? ProjectName = Options.ProjectName;
-				if (Options.ProjectName != null && Options.ProjectName.Contains('.'))
+				string? projectName = options.ProjectName;
+				if (options.ProjectName != null && options.ProjectName.Contains('.'))
 				{
-					Options.ProjectName = FileReference.Combine(CurrentDir, Options.ProjectName).MakeRelativeTo(ClientDir).Replace('\\', '/');
+					options.ProjectName = FileReference.Combine(currentDir, options.ProjectName).MakeRelativeTo(clientDir).Replace('\\', '/');
 				}
 
 				// Branch root is currently hard-coded at the root
-				string BranchPath = Options.BranchPath ?? String.Empty;
-				string ProjectPath = await FindProjectPathAsync(Perforce, ClientName, BranchPath, ProjectName);
+				string branchPath = options.BranchPath ?? String.Empty;
+				string projectPath = await FindProjectPathAsync(perforce, clientName, branchPath, projectName);
 
 				// Create the settings object
-				UserWorkspaceSettings Settings = new UserWorkspaceSettings();
-				Settings.RootDir = ClientDir;
-				Settings.Init(Perforce.Settings.ServerAndPort, Perforce.Settings.UserName, ClientName, BranchPath, ProjectPath);
-				Options.ApplyTo(Settings);
-				Settings.Save(Logger);
+				UserWorkspaceSettings settings = new UserWorkspaceSettings();
+				settings.RootDir = clientDir;
+				settings.Init(perforce.Settings.ServerAndPort, perforce.Settings.UserName, clientName, branchPath, projectPath);
+				options.ApplyTo(settings);
+				settings.Save(logger);
 
-				Logger.LogInformation("Initialized workspace at {RootDir} for {ClientProject}", ClientDir, Settings.ClientProjectPath);
+				logger.LogInformation("Initialized workspace at {RootDir} for {ClientProject}", clientDir, settings.ClientProjectPath);
 			}
 
-			static async Task<List<ClientsRecord>> FindExistingClients(IPerforceConnection Perforce, string HostName, DirectoryReference ClientDir)
+			static async Task<List<ClientsRecord>> FindExistingClients(IPerforceConnection perforce, string hostName, DirectoryReference clientDir)
 			{
-				List<ClientsRecord> MatchingClients = new List<ClientsRecord>();
+				List<ClientsRecord> matchingClients = new List<ClientsRecord>();
 
-				List<ClientsRecord> Clients = await Perforce.GetClientsAsync(ClientsOptions.None, Perforce.Settings.UserName);
-				foreach (ClientsRecord Client in Clients)
+				List<ClientsRecord> clients = await perforce.GetClientsAsync(ClientsOptions.None, perforce.Settings.UserName);
+				foreach (ClientsRecord client in clients)
 				{
-					if (!String.IsNullOrEmpty(Client.Root) && !String.IsNullOrEmpty(Client.Host) && String.Compare(HostName, Client.Host, StringComparison.OrdinalIgnoreCase) == 0)
+					if (!String.IsNullOrEmpty(client.Root) && !String.IsNullOrEmpty(client.Host) && String.Compare(hostName, client.Host, StringComparison.OrdinalIgnoreCase) == 0)
 					{
-						DirectoryReference? RootDir;
+						DirectoryReference? rootDir;
 						try
 						{
-							RootDir = new DirectoryReference(Client.Root);
+							rootDir = new DirectoryReference(client.Root);
 						}
 						catch
 						{
-							RootDir = null;
+							rootDir = null;
 						}
 
-						if (RootDir != null && ClientDir.IsUnderDirectory(RootDir))
+						if (rootDir != null && clientDir.IsUnderDirectory(rootDir))
 						{
-							MatchingClients.Add(Client);
+							matchingClients.Add(client);
 						}
 					}
 				}
 
-				return MatchingClients;
+				return matchingClients;
 			}
 		}
 
@@ -579,81 +579,81 @@ namespace UnrealGameSyncCmd
 				public bool Refilter { get; set; }
 			}
 
-			public override async Task ExecuteAsync(CommandContext Context)
+			public override async Task ExecuteAsync(CommandContext context)
 			{
-				ILogger Logger = Context.Logger;
-				Context.Arguments.TryGetPositionalArgument(out string? ChangeString);
+				ILogger logger = context.Logger;
+				context.Arguments.TryGetPositionalArgument(out string? changeString);
 
-				SyncOptions SyncOptions = new SyncOptions();
-				Context.Arguments.ApplyTo(SyncOptions);
+				SyncOptions syncOptions = new SyncOptions();
+				context.Arguments.ApplyTo(syncOptions);
 
-				Context.Arguments.CheckAllArgumentsUsed();
+				context.Arguments.CheckAllArgumentsUsed();
 
-				UserWorkspaceSettings Settings = ReadRequiredUserWorkspaceSettings();
-				using IPerforceConnection PerforceClient = await ConnectAsync(Settings, Context.LoggerFactory);
-				UserWorkspaceState State = await ReadWorkspaceState(PerforceClient, Settings, Context.UserSettings, Logger);
+				UserWorkspaceSettings settings = ReadRequiredUserWorkspaceSettings();
+				using IPerforceConnection perforceClient = await ConnectAsync(settings, context.LoggerFactory);
+				UserWorkspaceState state = await ReadWorkspaceState(perforceClient, settings, context.UserSettings, logger);
 
-				ChangeString ??= "latest";
+				changeString ??= "latest";
 
-				int Change;
-				if (!int.TryParse(ChangeString, out Change))
+				int change;
+				if (!int.TryParse(changeString, out change))
 				{
-					if (String.Equals(ChangeString, "latest", StringComparison.OrdinalIgnoreCase))
+					if (String.Equals(changeString, "latest", StringComparison.OrdinalIgnoreCase))
 					{
-						List<ChangesRecord> Changes = await PerforceClient.GetChangesAsync(ChangesOptions.None, 1, ChangeStatus.Submitted, $"//{Settings.ClientName}/...");
-						Change = Changes[0].Number;
+						List<ChangesRecord> changes = await perforceClient.GetChangesAsync(ChangesOptions.None, 1, ChangeStatus.Submitted, $"//{settings.ClientName}/...");
+						change = changes[0].Number;
 					}
 					else
 					{
-						throw new UserErrorException("Unknown change type for sync '{Change}'", ChangeString);
+						throw new UserErrorException("Unknown change type for sync '{Change}'", changeString);
 					}
 				}
 
-				WorkspaceUpdateOptions Options = SyncOptions.SingleChange? WorkspaceUpdateOptions.SyncSingleChange : WorkspaceUpdateOptions.Sync;
-				if (SyncOptions.Build)
+				WorkspaceUpdateOptions options = syncOptions.SingleChange? WorkspaceUpdateOptions.SyncSingleChange : WorkspaceUpdateOptions.Sync;
+				if (syncOptions.Build)
 				{
-					Options |= WorkspaceUpdateOptions.Build;
+					options |= WorkspaceUpdateOptions.Build;
 				}
-				if (SyncOptions.ProjectFiles)
+				if (syncOptions.ProjectFiles)
 				{
-					Options |= WorkspaceUpdateOptions.GenerateProjectFiles;
+					options |= WorkspaceUpdateOptions.GenerateProjectFiles;
 				}
-				if (SyncOptions.Clobber)
+				if (syncOptions.Clobber)
 				{
-					Options |= WorkspaceUpdateOptions.Clobber;
+					options |= WorkspaceUpdateOptions.Clobber;
 				}
-				if (SyncOptions.Refilter)
+				if (syncOptions.Refilter)
 				{
-					Options |= WorkspaceUpdateOptions.Refilter;
+					options |= WorkspaceUpdateOptions.Refilter;
 				}
-				Options |= WorkspaceUpdateOptions.RemoveFilteredFiles;
+				options |= WorkspaceUpdateOptions.RemoveFilteredFiles;
 
-				ProjectInfo ProjectInfo = State.CreateProjectInfo();
-				UserProjectSettings ProjectSettings = Context.UserSettings.FindOrAddProjectSettings(ProjectInfo, Settings, Logger);
+				ProjectInfo projectInfo = state.CreateProjectInfo();
+				UserProjectSettings projectSettings = context.UserSettings.FindOrAddProjectSettings(projectInfo, settings, logger);
 
-				ConfigFile ProjectConfig = await ConfigUtils.ReadProjectConfigFileAsync(PerforceClient, ProjectInfo, Logger, CancellationToken.None);
-				string[] SyncFilter = ReadSyncFilter(Settings, Context.UserSettings, ProjectConfig);
+				ConfigFile projectConfig = await ConfigUtils.ReadProjectConfigFileAsync(perforceClient, projectInfo, logger, CancellationToken.None);
+				string[] syncFilter = ReadSyncFilter(settings, context.UserSettings, projectConfig);
 
-				WorkspaceUpdateContext UpdateContext = new WorkspaceUpdateContext(Change, Options, BuildConfig.Development, SyncFilter, ProjectSettings.BuildSteps, null);
+				WorkspaceUpdateContext updateContext = new WorkspaceUpdateContext(change, options, BuildConfig.Development, syncFilter, projectSettings.BuildSteps, null);
 
-				WorkspaceUpdate Update = new WorkspaceUpdate(UpdateContext);
-				(WorkspaceUpdateResult Result, string Message) = await Update.ExecuteAsync(PerforceClient.Settings, ProjectInfo, State, Context.Logger, CancellationToken.None);
-				if (Result == WorkspaceUpdateResult.FilesToClobber)
+				WorkspaceUpdate update = new WorkspaceUpdate(updateContext);
+				(WorkspaceUpdateResult result, string message) = await update.ExecuteAsync(perforceClient.Settings, projectInfo, state, context.Logger, CancellationToken.None);
+				if (result == WorkspaceUpdateResult.FilesToClobber)
 				{
-					Logger.LogWarning("The following files are modified in your workspace:");
-					foreach (string File in UpdateContext.ClobberFiles.Keys.OrderBy(x => x))
+					logger.LogWarning("The following files are modified in your workspace:");
+					foreach (string file in updateContext.ClobberFiles.Keys.OrderBy(x => x))
 					{
-						Logger.LogWarning("  {File}", File);
+						logger.LogWarning("  {File}", file);
 					}
-					Logger.LogWarning("Use -Clobber to overwrite");
+					logger.LogWarning("Use -Clobber to overwrite");
 				}
-				else if (Result != WorkspaceUpdateResult.Success)
+				else if (result != WorkspaceUpdateResult.Success)
 				{
-					Logger.LogError("{Message} (Result: {Result})", Message, Result);
+					logger.LogError("{Message} (Result: {Result})", message, result);
 				}
 
-				State.SetLastSyncState(Result, UpdateContext, Message);
-				State.Save(Logger);
+				state.SetLastSyncState(result, updateContext, message);
+				state.Save(logger);
 			}
 		}
 
@@ -663,22 +663,22 @@ namespace UnrealGameSyncCmd
 			{
 			}
 
-			public override async Task ExecuteAsync(CommandContext Context)
+			public override async Task ExecuteAsync(CommandContext context)
 			{
-				ILogger Logger = Context.Logger;
+				ILogger logger = context.Logger;
 
-				ClientsOptions Options = Context.Arguments.ApplyTo<ClientsOptions>(Logger);
-				Context.Arguments.CheckAllArgumentsUsed();
+				ClientsOptions options = context.Arguments.ApplyTo<ClientsOptions>(logger);
+				context.Arguments.CheckAllArgumentsUsed();
 
-				using IPerforceConnection PerforceClient = await ConnectAsync(Options.ServerAndPort, Options.UserName, null, Context.LoggerFactory);
-				InfoRecord Info = await PerforceClient.GetInfoAsync(InfoOptions.ShortOutput);
+				using IPerforceConnection perforceClient = await ConnectAsync(options.ServerAndPort, options.UserName, null, context.LoggerFactory);
+				InfoRecord info = await perforceClient.GetInfoAsync(InfoOptions.ShortOutput);
 
-				List<ClientsRecord> Clients = await PerforceClient.GetClientsAsync(EpicGames.Perforce.ClientsOptions.None, PerforceClient.Settings.UserName);
-				foreach (ClientsRecord Client in Clients)
+				List<ClientsRecord> clients = await perforceClient.GetClientsAsync(EpicGames.Perforce.ClientsOptions.None, perforceClient.Settings.UserName);
+				foreach (ClientsRecord client in clients)
 				{
-					if (String.Equals(Info.ClientHost, Client.Host, StringComparison.OrdinalIgnoreCase))
+					if (String.Equals(info.ClientHost, client.Host, StringComparison.OrdinalIgnoreCase))
 					{
-						Logger.LogInformation("{Client,-50} {Root}", Client.Name, Client.Root);
+						logger.LogInformation("{Client,-50} {Root}", client.Name, client.Root);
 					}
 				}
 			}
@@ -686,115 +686,115 @@ namespace UnrealGameSyncCmd
 
 		class RunCommand : Command
 		{
-			public override async Task ExecuteAsync(CommandContext Context)
+			public override async Task ExecuteAsync(CommandContext context)
 			{
-				ILogger Logger = Context.Logger;
+				ILogger logger = context.Logger;
 
-				UserWorkspaceSettings Settings = ReadRequiredUserWorkspaceSettings();
-				using IPerforceConnection PerforceClient = await ConnectAsync(Settings, Context.LoggerFactory);
-				UserWorkspaceState State = await ReadWorkspaceState(PerforceClient, Settings, Context.UserSettings, Logger);
+				UserWorkspaceSettings settings = ReadRequiredUserWorkspaceSettings();
+				using IPerforceConnection perforceClient = await ConnectAsync(settings, context.LoggerFactory);
+				UserWorkspaceState state = await ReadWorkspaceState(perforceClient, settings, context.UserSettings, logger);
 
-				ProjectInfo ProjectInfo = State.CreateProjectInfo();
-				ConfigFile ProjectConfig = await ConfigUtils.ReadProjectConfigFileAsync(PerforceClient, ProjectInfo, Logger, CancellationToken.None);
+				ProjectInfo projectInfo = state.CreateProjectInfo();
+				ConfigFile projectConfig = await ConfigUtils.ReadProjectConfigFileAsync(perforceClient, projectInfo, logger, CancellationToken.None);
 
-				FileReference ReceiptFile = ConfigUtils.GetEditorReceiptFile(ProjectInfo, ProjectConfig, EditorConfig);
-				Logger.LogDebug("Receipt file: {Receipt}", ReceiptFile);
+				FileReference receiptFile = ConfigUtils.GetEditorReceiptFile(projectInfo, projectConfig, EditorConfig);
+				logger.LogDebug("Receipt file: {Receipt}", receiptFile);
 
-				if (!ConfigUtils.TryReadEditorReceipt(ProjectInfo, ReceiptFile, out TargetReceipt? Receipt) || String.IsNullOrEmpty(Receipt.Launch))
+				if (!ConfigUtils.TryReadEditorReceipt(projectInfo, receiptFile, out TargetReceipt? receipt) || String.IsNullOrEmpty(receipt.Launch))
 				{
-					throw new UserErrorException("The editor needs to be built before you can run it. (Missing {ReceiptFile}).", ReceiptFile);
+					throw new UserErrorException("The editor needs to be built before you can run it. (Missing {ReceiptFile}).", receiptFile);
 				}
-				if (!File.Exists(Receipt.Launch))
+				if (!File.Exists(receipt.Launch))
 				{
-					throw new UserErrorException("The editor needs to be built before you can run it. (Missing {LaunchFile}).", Receipt.Launch);
+					throw new UserErrorException("The editor needs to be built before you can run it. (Missing {LaunchFile}).", receipt.Launch);
 				}
 
-				List<string> LaunchArguments = new List<string>();
-				if (Settings.LocalProjectPath.HasExtension(".uproject"))
+				List<string> launchArguments = new List<string>();
+				if (settings.LocalProjectPath.HasExtension(".uproject"))
 				{
-					LaunchArguments.Add($"\"{Settings.LocalProjectPath}\"");
+					launchArguments.Add($"\"{settings.LocalProjectPath}\"");
 				}
 				if (EditorConfig == BuildConfig.Debug || EditorConfig == BuildConfig.DebugGame)
 				{
-					LaunchArguments.Append(" -debug");
+					launchArguments.Append(" -debug");
 				}
-				for (int Idx = 0; Idx < Context.Arguments.Count; Idx++)
+				for (int idx = 0; idx < context.Arguments.Count; idx++)
 				{
-					if (!Context.Arguments.HasBeenUsed(Idx))
+					if (!context.Arguments.HasBeenUsed(idx))
 					{
-						LaunchArguments.Add(Context.Arguments[Idx]);
+						launchArguments.Add(context.Arguments[idx]);
 					}
 				}
 
-				string CommandLine = CommandLineArguments.Join(LaunchArguments);
-				Logger.LogInformation("Spawning: {LaunchFile} {CommandLine}", CommandLineArguments.Quote(Receipt.Launch), CommandLine);
+				string commandLine = CommandLineArguments.Join(launchArguments);
+				logger.LogInformation("Spawning: {LaunchFile} {CommandLine}", CommandLineArguments.Quote(receipt.Launch), commandLine);
 
-				if (!Utility.SpawnProcess(Receipt.Launch, CommandLine))
+				if (!Utility.SpawnProcess(receipt.Launch, commandLine))
 				{
-					Logger.LogError("Unable to spawn {0} {1}", Receipt.Launch, LaunchArguments.ToString());
+					logger.LogError("Unable to spawn {0} {1}", receipt.Launch, launchArguments.ToString());
 				}
 			}
 		}
 
 		class ChangesCommand : Command
 		{
-			public override async Task ExecuteAsync(CommandContext Context)
+			public override async Task ExecuteAsync(CommandContext context)
 			{
-				ILogger Logger = Context.Logger;
+				ILogger logger = context.Logger;
 
-				int Count = Context.Arguments.GetIntegerOrDefault("-Count=", 10);
-				int LineCount = Context.Arguments.GetIntegerOrDefault("-Lines=", 3);
-				Context.Arguments.CheckAllArgumentsUsed(Context.Logger);
+				int count = context.Arguments.GetIntegerOrDefault("-Count=", 10);
+				int lineCount = context.Arguments.GetIntegerOrDefault("-Lines=", 3);
+				context.Arguments.CheckAllArgumentsUsed(context.Logger);
 
-				UserWorkspaceSettings Settings = ReadRequiredUserWorkspaceSettings();
-				using IPerforceConnection PerforceClient = await ConnectAsync(Settings, Context.LoggerFactory);
+				UserWorkspaceSettings settings = ReadRequiredUserWorkspaceSettings();
+				using IPerforceConnection perforceClient = await ConnectAsync(settings, context.LoggerFactory);
 
-				List<ChangesRecord> Changes = await PerforceClient.GetChangesAsync(ChangesOptions.None, Count, ChangeStatus.Submitted, $"//{Settings.ClientName}/...");
-				foreach(IEnumerable<ChangesRecord> ChangesBatch in Changes.Batch(10))
+				List<ChangesRecord> changes = await perforceClient.GetChangesAsync(ChangesOptions.None, count, ChangeStatus.Submitted, $"//{settings.ClientName}/...");
+				foreach(IEnumerable<ChangesRecord> changesBatch in changes.Batch(10))
 				{
-					List<DescribeRecord> DescribeRecords = await PerforceClient.DescribeAsync(ChangesBatch.Select(x => x.Number).ToArray());
+					List<DescribeRecord> describeRecords = await perforceClient.DescribeAsync(changesBatch.Select(x => x.Number).ToArray());
 
-					Logger.LogInformation("  Change    Type     Author          Description");
-					foreach (DescribeRecord DescribeRecord in DescribeRecords)
+					logger.LogInformation("  Change    Type     Author          Description");
+					foreach (DescribeRecord describeRecord in describeRecords)
 					{
-						PerforceChangeDetails Details = new PerforceChangeDetails(DescribeRecord);
+						PerforceChangeDetails details = new PerforceChangeDetails(describeRecord);
 
-						string Type;
-						if (Details.bContainsCode)
+						string type;
+						if (details.ContainsCode)
 						{
-							if (Details.bContainsContent)
+							if (details.ContainsContent)
 							{
-								Type = "Both";
+								type = "Both";
 							}
 							else
 							{
-								Type = "Code";
+								type = "Code";
 							}
 						}
 						else
 						{
-							if (Details.bContainsContent)
+							if (details.ContainsContent)
 							{
-								Type = "Content";
+								type = "Content";
 							}
 							else
 							{
-								Type = "None";
+								type = "None";
 							}
 						}
 
-						string Author = StringUtils.Truncate(DescribeRecord.User, 15);
+						string author = StringUtils.Truncate(describeRecord.User, 15);
 
-						List<string> Lines = StringUtils.WordWrap(Details.Description, Math.Max(ConsoleUtils.WindowWidth - 40, 10)).ToList();
-						if (Lines.Count == 0)
+						List<string> lines = StringUtils.WordWrap(details.Description, Math.Max(ConsoleUtils.WindowWidth - 40, 10)).ToList();
+						if (lines.Count == 0)
 						{
-							Lines.Add(String.Empty);
+							lines.Add(String.Empty);
 						}
 
-						Logger.LogInformation("  {Change,-9} {Type,-8} {Author,-15} {Description}", DescribeRecord.Number, Type, Author, Lines[0]);
-						for (int LineIndex = 1; LineIndex < LineCount; LineIndex++)
+						logger.LogInformation("  {Change,-9} {Type,-8} {Author,-15} {Description}", describeRecord.Number, type, author, lines[0]);
+						for (int lineIndex = 1; lineIndex < lineCount; lineIndex++)
 						{
-							Logger.LogInformation("                                     {Description}", Lines[LineIndex]);
+							logger.LogInformation("                                     {Description}", lines[lineIndex]);
 						}
 					}
 				}
@@ -803,34 +803,34 @@ namespace UnrealGameSyncCmd
 
 		class ConfigCommand : Command
 		{
-			public override Task ExecuteAsync(CommandContext Context)
+			public override Task ExecuteAsync(CommandContext context)
 			{
-				ILogger Logger = Context.Logger;
+				ILogger logger = context.Logger;
 
-				UserWorkspaceSettings Settings = ReadRequiredUserWorkspaceSettings();
-				if (!Context.Arguments.GetUnusedArguments().Any())
+				UserWorkspaceSettings settings = ReadRequiredUserWorkspaceSettings();
+				if (!context.Arguments.GetUnusedArguments().Any())
 				{
-					ProcessStartInfo StartInfo = new ProcessStartInfo();
-					StartInfo.FileName = Settings.ConfigFile.FullName;
-					StartInfo.UseShellExecute = true;
-					using (Process? Editor = Process.Start(StartInfo))
+					ProcessStartInfo startInfo = new ProcessStartInfo();
+					startInfo.FileName = settings.ConfigFile.FullName;
+					startInfo.UseShellExecute = true;
+					using (Process? editor = Process.Start(startInfo))
 					{
-						if (Editor != null)
+						if (editor != null)
 						{
-							Editor.WaitForExit();
+							editor.WaitForExit();
 						}
 					}
 				}
 				else
 				{
-					ProjectConfigOptions Options = new ProjectConfigOptions();
-					Context.Arguments.ApplyTo(Options);
-					Context.Arguments.CheckAllArgumentsUsed(Context.Logger);
+					ProjectConfigOptions options = new ProjectConfigOptions();
+					context.Arguments.ApplyTo(options);
+					context.Arguments.CheckAllArgumentsUsed(context.Logger);
 
-					Options.ApplyTo(Settings);
-					Settings.Save(Logger);
+					options.ApplyTo(settings);
+					settings.Save(logger);
 
-					Logger.LogInformation("Updated {ConfigFile}", Settings.ConfigFile);
+					logger.LogInformation("Updated {ConfigFile}", settings.ConfigFile);
 				}
 				
 				return Task.CompletedTask;
@@ -871,298 +871,298 @@ namespace UnrealGameSyncCmd
 				public bool Global { get; set; }
 			}
 
-			public override async Task ExecuteAsync(CommandContext Context)
+			public override async Task ExecuteAsync(CommandContext context)
 			{
-				ILogger Logger = Context.Logger;
+				ILogger logger = context.Logger;
 
-				UserWorkspaceSettings WorkspaceSettings = ReadRequiredUserWorkspaceSettings();
-				using IPerforceConnection PerforceClient = await ConnectAsync(WorkspaceSettings, Context.LoggerFactory);
-				UserWorkspaceState WorkspaceState = await ReadWorkspaceState(PerforceClient, WorkspaceSettings, Context.UserSettings, Logger);
-				ProjectInfo ProjectInfo = WorkspaceState.CreateProjectInfo();
+				UserWorkspaceSettings workspaceSettings = ReadRequiredUserWorkspaceSettings();
+				using IPerforceConnection perforceClient = await ConnectAsync(workspaceSettings, context.LoggerFactory);
+				UserWorkspaceState workspaceState = await ReadWorkspaceState(perforceClient, workspaceSettings, context.UserSettings, logger);
+				ProjectInfo projectInfo = workspaceState.CreateProjectInfo();
 
-				ConfigFile ProjectConfig = await ConfigUtils.ReadProjectConfigFileAsync(PerforceClient, ProjectInfo, Logger, CancellationToken.None);
-				Dictionary<Guid, WorkspaceSyncCategory> SyncCategories = ConfigUtils.GetSyncCategories(ProjectConfig);
+				ConfigFile projectConfig = await ConfigUtils.ReadProjectConfigFileAsync(perforceClient, projectInfo, logger, CancellationToken.None);
+				Dictionary<Guid, WorkspaceSyncCategory> syncCategories = ConfigUtils.GetSyncCategories(projectConfig);
 
-				FilterSettings GlobalFilter = Context.UserSettings.Global.Filter;
-				FilterSettings WorkspaceFilter = WorkspaceSettings.Filter;
+				FilterSettings globalFilter = context.UserSettings.Global.Filter;
+				FilterSettings workspaceFilter = workspaceSettings.Filter;
 
-				FilterCommandOptions Options = Context.Arguments.ApplyTo<FilterCommandOptions>(Logger);
-				Context.Arguments.CheckAllArgumentsUsed(Context.Logger);
+				FilterCommandOptions options = context.Arguments.ApplyTo<FilterCommandOptions>(logger);
+				context.Arguments.CheckAllArgumentsUsed(context.Logger);
 
-				if (Options.Global)
+				if (options.Global)
 				{
-					ApplyCommandOptions(Context.UserSettings.Global.Filter, Options, SyncCategories.Values, Logger);
-					Context.UserSettings.Save(Logger);
+					ApplyCommandOptions(context.UserSettings.Global.Filter, options, syncCategories.Values, logger);
+					context.UserSettings.Save(logger);
 				}
 				else
 				{
-					ApplyCommandOptions(WorkspaceSettings.Filter, Options, SyncCategories.Values, Logger);
-					WorkspaceSettings.Save(Logger);
+					ApplyCommandOptions(workspaceSettings.Filter, options, syncCategories.Values, logger);
+					workspaceSettings.Save(logger);
 				}
 
-				Dictionary<Guid, bool> GlobalCategories = GlobalFilter.GetCategories();
-				Dictionary<Guid, bool> WorkspaceCategories = WorkspaceFilter.GetCategories();
+				Dictionary<Guid, bool> globalCategories = globalFilter.GetCategories();
+				Dictionary<Guid, bool> workspaceCategories = workspaceFilter.GetCategories();
 
-				Logger.LogInformation("Categories:");
-				foreach (WorkspaceSyncCategory SyncCategory in SyncCategories.Values)
+				logger.LogInformation("Categories:");
+				foreach (WorkspaceSyncCategory syncCategory in syncCategories.Values)
 				{
-					bool bEnabled;
+					bool enabled;
 
-					string Scope = "(Default)";
-					if (GlobalCategories.TryGetValue(SyncCategory.UniqueId, out bEnabled))
+					string scope = "(Default)";
+					if (globalCategories.TryGetValue(syncCategory.UniqueId, out enabled))
 					{
-						Scope = "(Global)";
+						scope = "(Global)";
 					}
-					else if (WorkspaceCategories.TryGetValue(SyncCategory.UniqueId, out bEnabled))
+					else if (workspaceCategories.TryGetValue(syncCategory.UniqueId, out enabled))
 					{
-						Scope = "(Workspace)";
+						scope = "(Workspace)";
 					}
 					else
 					{
-						bEnabled = SyncCategory.bEnable;
+						enabled = syncCategory.Enable;
 					}
 
-					Logger.LogInformation("  {Id,30} {Enabled,3} {Scope,-9} {Name}", SyncCategory.UniqueId, bEnabled? "Yes" : "No", Scope, SyncCategory.Name);
+					logger.LogInformation("  {Id,30} {Enabled,3} {Scope,-9} {Name}", syncCategory.UniqueId, enabled? "Yes" : "No", scope, syncCategory.Name);
 				}
 
-				if (GlobalFilter.View.Count > 0)
+				if (globalFilter.View.Count > 0)
 				{
-					Logger.LogInformation("");
-					Logger.LogInformation("Global View:");
-					foreach (string Line in GlobalFilter.View)
+					logger.LogInformation("");
+					logger.LogInformation("Global View:");
+					foreach (string line in globalFilter.View)
 					{
-						Logger.LogInformation("  {Line}", Line);
+						logger.LogInformation("  {Line}", line);
 					}
 				}
-				if (WorkspaceFilter.View.Count > 0)
+				if (workspaceFilter.View.Count > 0)
 				{
-					Logger.LogInformation("");
-					Logger.LogInformation("Workspace View:");
-					foreach (string Line in WorkspaceFilter.View)
+					logger.LogInformation("");
+					logger.LogInformation("Workspace View:");
+					foreach (string line in workspaceFilter.View)
 					{
-						Logger.LogInformation("  {Line}", Line);
+						logger.LogInformation("  {Line}", line);
 					}
 				}
 
-				string[] Filter = ReadSyncFilter(WorkspaceSettings, Context.UserSettings, ProjectConfig);
+				string[] filter = ReadSyncFilter(workspaceSettings, context.UserSettings, projectConfig);
 
-				Logger.LogInformation("");
-				Logger.LogInformation("Combined view:");
-				foreach (string FilterLine in Filter)
+				logger.LogInformation("");
+				logger.LogInformation("Combined view:");
+				foreach (string filterLine in filter)
 				{
-					Logger.LogInformation("  {FilterLine}", FilterLine);
+					logger.LogInformation("  {FilterLine}", filterLine);
 				}
 			}
 
-			static void ApplyCommandOptions(FilterSettings Settings, FilterCommandOptions CommandOptions, IEnumerable<WorkspaceSyncCategory> SyncCategories, ILogger Logger)
+			static void ApplyCommandOptions(FilterSettings settings, FilterCommandOptions commandOptions, IEnumerable<WorkspaceSyncCategory> syncCategories, ILogger logger)
 			{
-				if (CommandOptions.Reset)
+				if (commandOptions.Reset)
 				{
-					Logger.LogInformation("Resetting settings...");
-					Settings.Reset();
+					logger.LogInformation("Resetting settings...");
+					settings.Reset();
 				}
 
-				HashSet<Guid> IncludeCategories = new HashSet<Guid>(CommandOptions.Include.Select(x => GetCategoryId(x, SyncCategories)));
-				HashSet<Guid> ExcludeCategories = new HashSet<Guid>(CommandOptions.Exclude.Select(x => GetCategoryId(x, SyncCategories)));
+				HashSet<Guid> includeCategories = new HashSet<Guid>(commandOptions.Include.Select(x => GetCategoryId(x, syncCategories)));
+				HashSet<Guid> excludeCategories = new HashSet<Guid>(commandOptions.Exclude.Select(x => GetCategoryId(x, syncCategories)));
 
-				Guid Id = IncludeCategories.FirstOrDefault(x => ExcludeCategories.Contains(x));
-				if (Id != Guid.Empty)
+				Guid id = includeCategories.FirstOrDefault(x => excludeCategories.Contains(x));
+				if (id != Guid.Empty)
 				{
-					throw new UserErrorException("Category {Id} cannot be both included and excluded", Id);
+					throw new UserErrorException("Category {Id} cannot be both included and excluded", id);
 				}
 
-				IncludeCategories.ExceptWith(Settings.IncludeCategories);
-				Settings.IncludeCategories.AddRange(IncludeCategories);
+				includeCategories.ExceptWith(settings.IncludeCategories);
+				settings.IncludeCategories.AddRange(includeCategories);
 
-				ExcludeCategories.ExceptWith(Settings.ExcludeCategories);
-				Settings.ExcludeCategories.AddRange(ExcludeCategories);
+				excludeCategories.ExceptWith(settings.ExcludeCategories);
+				settings.ExcludeCategories.AddRange(excludeCategories);
 
-				if (CommandOptions.View != null)
+				if (commandOptions.View != null)
 				{
-					Settings.View = CommandOptions.View;
+					settings.View = commandOptions.View;
 				}
-				if (CommandOptions.RemoveView.Count > 0)
+				if (commandOptions.RemoveView.Count > 0)
 				{
-					HashSet<string> ViewRemove = new HashSet<string>(CommandOptions.RemoveView, StringComparer.OrdinalIgnoreCase);
-					Settings.View.RemoveAll(x => ViewRemove.Contains(x));
+					HashSet<string> viewRemove = new HashSet<string>(commandOptions.RemoveView, StringComparer.OrdinalIgnoreCase);
+					settings.View.RemoveAll(x => viewRemove.Contains(x));
 				}
-				if (CommandOptions.AddView.Count > 0)
+				if (commandOptions.AddView.Count > 0)
 				{
-					HashSet<string> ViewLines = new HashSet<string>(Settings.View, StringComparer.OrdinalIgnoreCase);
-					Settings.View.AddRange(CommandOptions.AddView.Where(x => !ViewLines.Contains(x)));
+					HashSet<string> viewLines = new HashSet<string>(settings.View, StringComparer.OrdinalIgnoreCase);
+					settings.View.AddRange(commandOptions.AddView.Where(x => !viewLines.Contains(x)));
 				}
 
-				Settings.AllProjects = CommandOptions.AllProjects ?? Settings.AllProjects;
-				Settings.AllProjectsInSln = CommandOptions.AllProjectsInSln ?? Settings.AllProjectsInSln;
+				settings.AllProjects = commandOptions.AllProjects ?? settings.AllProjects;
+				settings.AllProjectsInSln = commandOptions.AllProjectsInSln ?? settings.AllProjectsInSln;
 			}
 
-			static Guid GetCategoryId(string Text, IEnumerable<WorkspaceSyncCategory> SyncCategories)
+			static Guid GetCategoryId(string text, IEnumerable<WorkspaceSyncCategory> syncCategories)
 			{
-				Guid Id;
-				if (Guid.TryParse(Text, out Id))
+				Guid id;
+				if (Guid.TryParse(text, out id))
 				{
-					return Id;
+					return id;
 				}
 
-				WorkspaceSyncCategory? Category = SyncCategories.FirstOrDefault(x => x.Name.Equals(Text, StringComparison.OrdinalIgnoreCase));
-				if (Category != null)
+				WorkspaceSyncCategory? category = syncCategories.FirstOrDefault(x => x.Name.Equals(text, StringComparison.OrdinalIgnoreCase));
+				if (category != null)
 				{
-					return Category.UniqueId;
+					return category.UniqueId;
 				}
 
-				throw new UserErrorException("Unable to find category '{Category}'", Text);
+				throw new UserErrorException("Unable to find category '{Category}'", text);
 			}
 		}
 
 		class BuildCommand : Command
 		{
-			public override async Task ExecuteAsync(CommandContext Context)
+			public override async Task ExecuteAsync(CommandContext context)
 			{
-				ILogger Logger = Context.Logger;
-				Context.Arguments.TryGetPositionalArgument(out string? Target);
-				bool ListOnly = Context.Arguments.HasOption("-List");
-				Context.Arguments.CheckAllArgumentsUsed();
+				ILogger logger = context.Logger;
+				context.Arguments.TryGetPositionalArgument(out string? target);
+				bool listOnly = context.Arguments.HasOption("-List");
+				context.Arguments.CheckAllArgumentsUsed();
 
-				UserWorkspaceSettings Settings = ReadRequiredUserWorkspaceSettings();
-				using IPerforceConnection PerforceClient = await ConnectAsync(Settings, Context.LoggerFactory);
-				UserWorkspaceState State = await ReadWorkspaceState(PerforceClient, Settings, Context.UserSettings, Logger);
+				UserWorkspaceSettings settings = ReadRequiredUserWorkspaceSettings();
+				using IPerforceConnection perforceClient = await ConnectAsync(settings, context.LoggerFactory);
+				UserWorkspaceState state = await ReadWorkspaceState(perforceClient, settings, context.UserSettings, logger);
 
-				ProjectInfo ProjectInfo = State.CreateProjectInfo();
+				ProjectInfo projectInfo = state.CreateProjectInfo();
 
-				if (ListOnly)
+				if (listOnly)
 				{
-					ConfigFile ProjectConfig = await ConfigUtils.ReadProjectConfigFileAsync(PerforceClient, ProjectInfo, Logger, CancellationToken.None);
+					ConfigFile projectConfig = await ConfigUtils.ReadProjectConfigFileAsync(perforceClient, projectInfo, logger, CancellationToken.None);
 
-					FileReference EditorTarget = ConfigUtils.GetEditorTargetFile(ProjectInfo, ProjectConfig);
+					FileReference editorTarget = ConfigUtils.GetEditorTargetFile(projectInfo, projectConfig);
 
-					Dictionary<Guid, ConfigObject> BuildStepObjects = ConfigUtils.GetDefaultBuildStepObjects(ProjectInfo, EditorTarget.GetFileNameWithoutAnyExtensions(), EditorConfig, ProjectConfig, false);
+					Dictionary<Guid, ConfigObject> buildStepObjects = ConfigUtils.GetDefaultBuildStepObjects(projectInfo, editorTarget.GetFileNameWithoutAnyExtensions(), EditorConfig, projectConfig, false);
 
-					Logger.LogInformation("Available build steps:");
-					Logger.LogInformation("");
-					Logger.LogInformation("  Id                                   | Description                              | Type       | Enabled");
-					Logger.LogInformation("  -------------------------------------|------------------------------------------|------------|-----------------");
-					foreach (BuildStep BuildStep in BuildStepObjects.Values.Select(x => new BuildStep(x)).OrderBy(x => x.OrderIndex))
+					logger.LogInformation("Available build steps:");
+					logger.LogInformation("");
+					logger.LogInformation("  Id                                   | Description                              | Type       | Enabled");
+					logger.LogInformation("  -------------------------------------|------------------------------------------|------------|-----------------");
+					foreach (BuildStep buildStep in buildStepObjects.Values.Select(x => new BuildStep(x)).OrderBy(x => x.OrderIndex))
 					{
-						Logger.LogInformation("  {Id,-36} | {Name,-40} | {Type,-10} | {Enabled,-8}", BuildStep.UniqueId, BuildStep.Description, BuildStep.Type, BuildStep.bNormalSync);
+						logger.LogInformation("  {Id,-36} | {Name,-40} | {Type,-10} | {Enabled,-8}", buildStep.UniqueId, buildStep.Description, buildStep.Type, buildStep.NormalSync);
 					}
 					return;
 				}
 
-				HashSet<Guid>? Steps = null;
-				if (Target != null)
+				HashSet<Guid>? steps = null;
+				if (target != null)
 				{
-					Guid Id;
-					if (!Guid.TryParse(Target, out Id))
+					Guid id;
+					if (!Guid.TryParse(target, out id))
 					{
-						Logger.LogError("Unable to parse '{Target}' as a GUID. Pass -List to show all available build steps and their identifiers.", Target);
+						logger.LogError("Unable to parse '{Target}' as a GUID. Pass -List to show all available build steps and their identifiers.", target);
 					}
-					Steps = new HashSet<Guid> { Id };
+					steps = new HashSet<Guid> { id };
 				}
 
-				WorkspaceUpdateContext UpdateContext = new WorkspaceUpdateContext(State.CurrentChangeNumber, WorkspaceUpdateOptions.Build, BuildConfig.Development, null, new List<ConfigObject>(), Steps);
+				WorkspaceUpdateContext updateContext = new WorkspaceUpdateContext(state.CurrentChangeNumber, WorkspaceUpdateOptions.Build, BuildConfig.Development, null, new List<ConfigObject>(), steps);
 
-				WorkspaceUpdate Update = new WorkspaceUpdate(UpdateContext);
-				await Update.ExecuteAsync(PerforceClient.Settings, ProjectInfo, State, Context.Logger, CancellationToken.None);
+				WorkspaceUpdate update = new WorkspaceUpdate(updateContext);
+				await update.ExecuteAsync(perforceClient.Settings, projectInfo, state, context.Logger, CancellationToken.None);
 			}
 		}
 
 		class StatusCommand : Command
 		{
-			public override async Task ExecuteAsync(CommandContext Context)
+			public override async Task ExecuteAsync(CommandContext context)
 			{
-				ILogger Logger = Context.Logger;
-				bool bUpdate = Context.Arguments.HasOption("-Update");
-				Context.Arguments.CheckAllArgumentsUsed();
+				ILogger logger = context.Logger;
+				bool update = context.Arguments.HasOption("-Update");
+				context.Arguments.CheckAllArgumentsUsed();
 
-				UserWorkspaceSettings Settings = ReadRequiredUserWorkspaceSettings();
-				Logger.LogInformation("User: {UserName}", Settings.UserName);
-				Logger.LogInformation("Server: {ServerAndPort}", Settings.ServerAndPort);
-				Logger.LogInformation("Project: {ClientProjectPath}", Settings.ClientProjectPath);
+				UserWorkspaceSettings settings = ReadRequiredUserWorkspaceSettings();
+				logger.LogInformation("User: {UserName}", settings.UserName);
+				logger.LogInformation("Server: {ServerAndPort}", settings.ServerAndPort);
+				logger.LogInformation("Project: {ClientProjectPath}", settings.ClientProjectPath);
 
-				using IPerforceConnection PerforceClient = await ConnectAsync(Settings, Context.LoggerFactory);
+				using IPerforceConnection perforceClient = await ConnectAsync(settings, context.LoggerFactory);
 
-				UserWorkspaceState State = await ReadWorkspaceState(PerforceClient, Settings, Context.UserSettings, Logger);
-				if (bUpdate)
+				UserWorkspaceState state = await ReadWorkspaceState(perforceClient, settings, context.UserSettings, logger);
+				if (update)
 				{
-					ProjectInfo NewProjectInfo = await ProjectInfo.CreateAsync(PerforceClient, Settings, CancellationToken.None);
-					State.UpdateCachedProjectInfo(NewProjectInfo, Settings.LastModifiedTimeUtc);
+					ProjectInfo newProjectInfo = await ProjectInfo.CreateAsync(perforceClient, settings, CancellationToken.None);
+					state.UpdateCachedProjectInfo(newProjectInfo, settings.LastModifiedTimeUtc);
 				}
 
-				string StreamOrBranchName = State.StreamName ?? Settings.BranchPath.TrimStart('/');
-				if (State.LastSyncResultMessage == null)
+				string streamOrBranchName = state.StreamName ?? settings.BranchPath.TrimStart('/');
+				if (state.LastSyncResultMessage == null)
 				{
-					Logger.LogInformation("Not currently synced to {Stream}", StreamOrBranchName);
+					logger.LogInformation("Not currently synced to {Stream}", streamOrBranchName);
 				}
-				else if (State.LastSyncResult == WorkspaceUpdateResult.Success)
+				else if (state.LastSyncResult == WorkspaceUpdateResult.Success)
 				{
-					Logger.LogInformation("Synced to {Stream} CL {Change}", StreamOrBranchName, State.LastSyncChangeNumber);
+					logger.LogInformation("Synced to {Stream} CL {Change}", streamOrBranchName, state.LastSyncChangeNumber);
 				}
 				else
 				{
-					Logger.LogWarning("Last sync to {Stream} CL {Change} failed: {Result}", StreamOrBranchName, State.LastSyncChangeNumber, State.LastSyncResultMessage);
+					logger.LogWarning("Last sync to {Stream} CL {Change} failed: {Result}", streamOrBranchName, state.LastSyncChangeNumber, state.LastSyncResultMessage);
 				}
 			}
 		}
 
 		class SwitchCommand : Command
 		{
-			public override async Task ExecuteAsync(CommandContext Context)
+			public override async Task ExecuteAsync(CommandContext context)
 			{
 				// Get the positional argument indicating the file to look for
-				string? TargetName;
-				if (!Context.Arguments.TryGetPositionalArgument(out TargetName))
+				string? targetName;
+				if (!context.Arguments.TryGetPositionalArgument(out targetName))
 				{
 					throw new UserErrorException("Missing stream or project name to switch to.");
 				}
 
-				bool Force = TargetName.StartsWith("//", StringComparison.Ordinal) && Context.Arguments.HasOption("-Force");
+				bool force = targetName.StartsWith("//", StringComparison.Ordinal) && context.Arguments.HasOption("-Force");
 
 				// Finish argument parsing
-				Context.Arguments.CheckAllArgumentsUsed();
+				context.Arguments.CheckAllArgumentsUsed();
 
 				// Get a connection to the client for this workspace
-				UserWorkspaceSettings Settings = ReadRequiredUserWorkspaceSettings();
-				using IPerforceConnection PerforceClient = await ConnectAsync(Settings, Context.LoggerFactory);
+				UserWorkspaceSettings settings = ReadRequiredUserWorkspaceSettings();
+				using IPerforceConnection perforceClient = await ConnectAsync(settings, context.LoggerFactory);
 
 				// Check whether we're switching stream or project
-				if (TargetName.StartsWith("//", StringComparison.Ordinal))
+				if (targetName.StartsWith("//", StringComparison.Ordinal))
 				{
-					await SwitchStreamAsync(PerforceClient, TargetName, Force, Context.Logger);
+					await SwitchStreamAsync(perforceClient, targetName, force, context.Logger);
 				}
 				else
 				{
-					await SwitchProjectAsync(PerforceClient, Settings, TargetName, Context.Logger);
+					await SwitchProjectAsync(perforceClient, settings, targetName, context.Logger);
 				}
 			}
 
-			public async Task SwitchStreamAsync(IPerforceConnection PerforceClient, string StreamName, bool Force, ILogger Logger)
+			public async Task SwitchStreamAsync(IPerforceConnection perforceClient, string streamName, bool force, ILogger logger)
 			{
-				if (!Force && await PerforceClient.OpenedAsync(OpenedOptions.None, FileSpecList.Any).AnyAsync())
+				if (!force && await perforceClient.OpenedAsync(OpenedOptions.None, FileSpecList.Any).AnyAsync())
 				{
-					throw new UserErrorException("Client {ClientName} has files opened. Use -Force to switch anyway.", PerforceClient.Settings.ClientName!);
+					throw new UserErrorException("Client {ClientName} has files opened. Use -Force to switch anyway.", perforceClient.Settings.ClientName!);
 				}
 
-				await PerforceClient.SwitchClientToStreamAsync(StreamName, SwitchClientOptions.IgnoreOpenFiles);
+				await perforceClient.SwitchClientToStreamAsync(streamName, SwitchClientOptions.IgnoreOpenFiles);
 
-				Logger.LogInformation("Switched to stream {StreamName}", StreamName);
+				logger.LogInformation("Switched to stream {StreamName}", streamName);
 			}
 
-			public async Task SwitchProjectAsync(IPerforceConnection PerforceClient, UserWorkspaceSettings Settings, string ProjectName, ILogger Logger)
+			public async Task SwitchProjectAsync(IPerforceConnection perforceClient, UserWorkspaceSettings settings, string projectName, ILogger logger)
 			{
-				Settings.ProjectPath = await FindProjectPathAsync(PerforceClient, Settings.ClientName, Settings.BranchPath, ProjectName);
-				Settings.Save(Logger);
-				Logger.LogInformation("Switched to project {ProjectPath}", Settings.ClientProjectPath);
+				settings.ProjectPath = await FindProjectPathAsync(perforceClient, settings.ClientName, settings.BranchPath, projectName);
+				settings.Save(logger);
+				logger.LogInformation("Switched to project {ProjectPath}", settings.ClientProjectPath);
 			}
 		}
 
 		class VersionCommand : Command
 		{
-			public override Task ExecuteAsync(CommandContext Context)
+			public override Task ExecuteAsync(CommandContext context)
 			{
-				ILogger Logger = Context.Logger;
+				ILogger logger = context.Logger;
  
-				AssemblyInformationalVersionAttribute? Version = Assembly.GetExecutingAssembly().GetCustomAttribute<AssemblyInformationalVersionAttribute>();
-				Logger.LogInformation("UnrealGameSync {Version}", Version?.InformationalVersion ?? "Unknown");
+				AssemblyInformationalVersionAttribute? version = Assembly.GetExecutingAssembly().GetCustomAttribute<AssemblyInformationalVersionAttribute>();
+				logger.LogInformation("UnrealGameSync {Version}", version?.InformationalVersion ?? "Unknown");
 
 				return Task.CompletedTask;
 			}
diff --git a/Engine/Source/Programs/UnrealGameSync/UnrealGameSyncLauncher/Forms/LogWindow.cs b/Engine/Source/Programs/UnrealGameSync/UnrealGameSyncLauncher/Forms/LogWindow.cs
--- a/Engine/Source/Programs/UnrealGameSync/UnrealGameSyncLauncher/Forms/LogWindow.cs
+++ b/Engine/Source/Programs/UnrealGameSync/UnrealGameSyncLauncher/Forms/LogWindow.cs
@@ -14,12 +14,12 @@ namespace UnrealGameSyncLauncher
 {
 	public partial class LogWindow : Form
 	{
-		public LogWindow(string Text)
+		public LogWindow(string text)
 		{
 			InitializeComponent();
 
-			LogTextBox.Text = Text;
-			LogTextBox.Select(Text.Length, 0);
+			LogTextBox.Text = text;
+			LogTextBox.Select(text.Length, 0);
 		}
 	}
 }
diff --git a/Engine/Source/Programs/UnrealGameSync/UnrealGameSyncLauncher/Forms/SettingsWindow.cs b/Engine/Source/Programs/UnrealGameSync/UnrealGameSyncLauncher/Forms/SettingsWindow.cs
--- a/Engine/Source/Programs/UnrealGameSync/UnrealGameSyncLauncher/Forms/SettingsWindow.cs
+++ b/Engine/Source/Programs/UnrealGameSync/UnrealGameSyncLauncher/Forms/SettingsWindow.cs
@@ -22,101 +22,101 @@ namespace UnrealGameSyncLauncher
 	partial class SettingsWindow : Form
 	{
 		[DllImport("user32.dll")]
-		private static extern IntPtr SendMessage(IntPtr hWnd, int Msg, int wParam, [MarshalAs(UnmanagedType.LPWStr)] string lParam);
+		private static extern IntPtr SendMessage(IntPtr hWnd, int msg, int wParam, [MarshalAs(UnmanagedType.LPWStr)] string lParam);
 
-		public delegate Task SyncAndRunDelegate(IPerforceConnection Perforce, string? DepotPath, bool bPreview, ILogger LogWriter, CancellationToken CancellationToken);
+		public delegate Task SyncAndRunDelegate(IPerforceConnection perforce, string? depotPath, bool preview, ILogger logWriter, CancellationToken cancellationToken);
 
-		const int EM_SETCUEBANNER = 0x1501;
+		const int EmSetcuebanner = 0x1501;
 
-		string? LogText;
-		SyncAndRunDelegate SyncAndRun;
+		string? _logText;
+		SyncAndRunDelegate _syncAndRun;
 
-		public SettingsWindow(string? Prompt, string? LogText, string? ServerAndPort, string? UserName, string? DepotPath, bool bPreview, SyncAndRunDelegate SyncAndRun)
+		public SettingsWindow(string? prompt, string? logText, string? serverAndPort, string? userName, string? depotPath, bool preview, SyncAndRunDelegate syncAndRun)
 		{
 			InitializeComponent();
 
-			if(Prompt != null)
+			if(prompt != null)
 			{
-				this.PromptLabel.Text = Prompt;
+				this.PromptLabel.Text = prompt;
 			}
 
-			this.LogText = LogText;
-			this.ServerTextBox.Text = ServerAndPort ?? String.Empty;
-			this.UserNameTextBox.Text = UserName ?? String.Empty;
-			this.DepotPathTextBox.Text = DepotPath ?? String.Empty;
-			this.UsePreviewBuildCheckBox.Checked = bPreview;
-			this.SyncAndRun = SyncAndRun;
+			this._logText = logText;
+			this.ServerTextBox.Text = serverAndPort ?? String.Empty;
+			this.UserNameTextBox.Text = userName ?? String.Empty;
+			this.DepotPathTextBox.Text = depotPath ?? String.Empty;
+			this.UsePreviewBuildCheckBox.Checked = preview;
+			this._syncAndRun = syncAndRun;
 
-			ViewLogBtn.Visible = LogText != null;
+			ViewLogBtn.Visible = logText != null;
 		}
 
 		protected override void OnLoad(EventArgs e)
 		{
 			base.OnLoad(e);
 
-			SendMessage(ServerTextBox.Handle, EM_SETCUEBANNER, 1, "Default Server");
-			SendMessage(UserNameTextBox.Handle, EM_SETCUEBANNER, 1, "Default User");
+			SendMessage(ServerTextBox.Handle, EmSetcuebanner, 1, "Default Server");
+			SendMessage(UserNameTextBox.Handle, EmSetcuebanner, 1, "Default User");
 		}
 
 		private void ViewLogBtn_Click(object sender, EventArgs e)
 		{
-			LogWindow Log = new LogWindow(LogText ?? String.Empty);
-			Log.ShowDialog(this);
+			LogWindow log = new LogWindow(_logText ?? String.Empty);
+			log.ShowDialog(this);
 		}
 
 		private void ConnectBtn_Click(object sender, EventArgs e)
 		{
 			// Update the settings
-			string? ServerAndPort = ServerTextBox.Text.Trim();
-			if(ServerAndPort.Length == 0)
+			string? serverAndPort = ServerTextBox.Text.Trim();
+			if(serverAndPort.Length == 0)
 			{
-				ServerAndPort = null;
+				serverAndPort = null;
 			}
 
-			string? UserName = UserNameTextBox.Text.Trim();
-			if(UserName.Length == 0)
+			string? userName = UserNameTextBox.Text.Trim();
+			if(userName.Length == 0)
 			{
-				UserName = null;
+				userName = null;
 			}
 
-			string? DepotPath = DepotPathTextBox.Text.Trim();
-			if(DepotPath.Length == 0)
+			string? depotPath = DepotPathTextBox.Text.Trim();
+			if(depotPath.Length == 0)
 			{
-				DepotPath = null;
+				depotPath = null;
 			}
 
-			bool bPreview = UsePreviewBuildCheckBox.Checked;
+			bool preview = UsePreviewBuildCheckBox.Checked;
 
-			GlobalPerforceSettings.SaveGlobalPerforceSettings(ServerAndPort, UserName, DepotPath, bPreview);
+			GlobalPerforceSettings.SaveGlobalPerforceSettings(serverAndPort, userName, depotPath, preview);
 
-			PerforceSettings PerforceSettings = new PerforceSettings(PerforceSettings.Default);
-			if (!String.IsNullOrEmpty(ServerAndPort))
+			PerforceSettings perforceSettings = new PerforceSettings(PerforceSettings.Default);
+			if (!String.IsNullOrEmpty(serverAndPort))
 			{
-				PerforceSettings.ServerAndPort = ServerAndPort;
+				perforceSettings.ServerAndPort = serverAndPort;
 			}
-			if (!String.IsNullOrEmpty(UserName))
+			if (!String.IsNullOrEmpty(userName))
 			{
-				PerforceSettings.UserName = UserName;
+				perforceSettings.UserName = userName;
 			}
-			PerforceSettings.PreferNativeClient = true;
+			perforceSettings.PreferNativeClient = true;
 
 			// Create the P4 connection
-			CaptureLogger Logger = new CaptureLogger();
+			CaptureLogger logger = new CaptureLogger();
 
 			// Create the task for connecting to this server
-			ModalTask? Task = PerforceModalTask.Execute(this, "Updating", "Checking for updates, please wait...", PerforceSettings, (p, c) => SyncAndRun(p, DepotPath, bPreview, Logger, c), Logger);
-			if (Task != null)
+			ModalTask? task = PerforceModalTask.Execute(this, "Updating", "Checking for updates, please wait...", perforceSettings, (p, c) => _syncAndRun(p, depotPath, preview, logger, c), logger);
+			if (task != null)
 			{
-				if(Task.Succeeded)
+				if(task.Succeeded)
 				{
-					GlobalPerforceSettings.SaveGlobalPerforceSettings(ServerAndPort, UserName, DepotPath, bPreview);
+					GlobalPerforceSettings.SaveGlobalPerforceSettings(serverAndPort, userName, depotPath, preview);
 					DialogResult = DialogResult.OK;
 					Close();
 				}
-				PromptLabel.Text = Task.Error;
+				PromptLabel.Text = task.Error;
 			}
 
-			LogText = Logger.Render(Environment.NewLine);
+			_logText = logger.Render(Environment.NewLine);
 			ViewLogBtn.Visible = true;
 		}
 
