UGS: Reformat to standard C# style conventions.

diff --git a/Engine/Source/Programs/UnrealGameSync/UnrealGameSyncLauncher/Program.cs b/Engine/Source/Programs/UnrealGameSync/UnrealGameSyncLauncher/Program.cs
--- a/Engine/Source/Programs/UnrealGameSync/UnrealGameSyncLauncher/Program.cs
+++ b/Engine/Source/Programs/UnrealGameSync/UnrealGameSyncLauncher/Program.cs
@@ -23,39 +23,39 @@ namespace UnrealGameSyncLauncher
 	static class Program
 	{
 		[STAThread]
-		static int Main(string[] Args)
+		static int Main(string[] args)
 		{
 			Application.EnableVisualStyles();
 			Application.SetCompatibleTextRenderingDefault(false);
 
-			bool bFirstInstance;
-			using(Mutex InstanceMutex = new Mutex(true, "UnrealGameSyncRunning", out bFirstInstance))
+			bool firstInstance;
+			using(Mutex instanceMutex = new Mutex(true, "UnrealGameSyncRunning", out firstInstance))
 			{
-				if(!bFirstInstance)
+				if(!firstInstance)
 				{
-					using(EventWaitHandle ActivateEvent = new EventWaitHandle(false, EventResetMode.AutoReset, "ActivateUnrealGameSync"))
+					using(EventWaitHandle activateEvent = new EventWaitHandle(false, EventResetMode.AutoReset, "ActivateUnrealGameSync"))
 					{
-						ActivateEvent.Set();
+						activateEvent.Set();
 					}
 					return 0;
 				}
 
 				// Figure out if we should sync the unstable build by default
-				bool bPreview = Args.Contains("-unstable", StringComparer.InvariantCultureIgnoreCase) || Args.Contains("-preview", StringComparer.InvariantCultureIgnoreCase);
+				bool preview = args.Contains("-unstable", StringComparer.InvariantCultureIgnoreCase) || args.Contains("-preview", StringComparer.InvariantCultureIgnoreCase);
 
 				// Read the settings
-				string? ServerAndPort = null;
-				string? UserName = null;
-				string? DepotPath = DeploymentSettings.DefaultDepotPath;
-				GlobalPerforceSettings.ReadGlobalPerforceSettings(ref ServerAndPort, ref UserName, ref DepotPath, ref bPreview);
+				string? serverAndPort = null;
+				string? userName = null;
+				string? depotPath = DeploymentSettings.DefaultDepotPath;
+				GlobalPerforceSettings.ReadGlobalPerforceSettings(ref serverAndPort, ref userName, ref depotPath, ref preview);
 
 				// If the shift key is held down, immediately show the settings window
-				SettingsWindow.SyncAndRunDelegate SyncAndRunWrapper = (Perforce, DepotParam, bPreviewParam, LogWriter, CancellationToken) => SyncAndRun(Perforce, DepotParam, bPreviewParam, Args, InstanceMutex, LogWriter, CancellationToken);
+				SettingsWindow.SyncAndRunDelegate syncAndRunWrapper = (perforce, depotParam, previewParam, logWriter, cancellationToken) => SyncAndRun(perforce, depotParam, previewParam, args, instanceMutex, logWriter, cancellationToken);
 				if ((Control.ModifierKeys & Keys.Shift) != 0)
 				{
 					// Show the settings window immediately
-					SettingsWindow UpdateError = new SettingsWindow(null, null, ServerAndPort, UserName, DepotPath, bPreview, SyncAndRunWrapper);
-					if(UpdateError.ShowDialog() == DialogResult.OK)
+					SettingsWindow updateError = new SettingsWindow(null, null, serverAndPort, userName, depotPath, preview, syncAndRunWrapper);
+					if(updateError.ShowDialog() == DialogResult.OK)
 					{
 						return 0;
 					}
@@ -63,22 +63,22 @@ namespace UnrealGameSyncLauncher
 				else
 				{
 					// Try to do a sync with the current settings first
-					CaptureLogger Logger = new CaptureLogger();
+					CaptureLogger logger = new CaptureLogger();
 
-					IPerforceSettings Settings = new PerforceSettings(PerforceSettings.Default) { PreferNativeClient = true }.MergeWith(newServerAndPort: ServerAndPort, newUserName: UserName);
+					IPerforceSettings settings = new PerforceSettings(PerforceSettings.Default) { PreferNativeClient = true }.MergeWith(newServerAndPort: serverAndPort, newUserName: userName);
 
-					ModalTask? Task = PerforceModalTask.Execute(null, "Updating", "Checking for updates, please wait...", Settings, (p, c) => SyncAndRun(p, DepotPath, bPreview, Args, InstanceMutex, Logger, c), Logger);
-					if (Task == null)
+					ModalTask? task = PerforceModalTask.Execute(null, "Updating", "Checking for updates, please wait...", settings, (p, c) => SyncAndRun(p, depotPath, preview, args, instanceMutex, logger, c), logger);
+					if (task == null)
 					{
-						Logger.LogInformation("Canceled by user");
+						logger.LogInformation("Canceled by user");
 					}
-					else if (Task.Succeeded)
+					else if (task.Succeeded)
 					{
 						return 0;
 					}
 
-					SettingsWindow UpdateError = new SettingsWindow("Unable to update UnrealGameSync from Perforce. Verify that your connection settings are correct.", Logger.Render(Environment.NewLine), ServerAndPort, UserName, DepotPath, bPreview, SyncAndRunWrapper);
-					if(UpdateError.ShowDialog() == DialogResult.OK)
+					SettingsWindow updateError = new SettingsWindow("Unable to update UnrealGameSync from Perforce. Verify that your connection settings are correct.", logger.Render(Environment.NewLine), serverAndPort, userName, depotPath, preview, syncAndRunWrapper);
+					if(updateError.ShowDialog() == DialogResult.OK)
 					{
 						return 0;
 					}
@@ -87,188 +87,188 @@ namespace UnrealGameSyncLauncher
 			return 1;
 		}
 
-		public static async Task SyncAndRun(IPerforceConnection Perforce, string? BaseDepotPath, bool bPreview, string[] Args, Mutex InstanceMutex, ILogger Logger, CancellationToken CancellationToken)
+		public static async Task SyncAndRun(IPerforceConnection perforce, string? baseDepotPath, bool preview, string[] args, Mutex instanceMutex, ILogger logger, CancellationToken cancellationToken)
 		{
 			try
 			{
-				if (String.IsNullOrEmpty(BaseDepotPath))
+				if (String.IsNullOrEmpty(baseDepotPath))
 				{
 					throw new UserErrorException($"Invalid setting for sync path");
 				}
 
-				string BaseDepotPathPrefix = BaseDepotPath.TrimEnd('/');
+				string baseDepotPathPrefix = baseDepotPath.TrimEnd('/');
 
 				// Find the most recent changelist
-				string SyncPath = BaseDepotPathPrefix + (bPreview ? "/UnstableRelease.zip" : "/Release.zip");
-				List<ChangesRecord> Changes = await Perforce.GetChangesAsync(ChangesOptions.None, 1, ChangeStatus.Submitted, SyncPath, CancellationToken);
-				if (Changes.Count == 0)
+				string syncPath = baseDepotPathPrefix + (preview ? "/UnstableRelease.zip" : "/Release.zip");
+				List<ChangesRecord> changes = await perforce.GetChangesAsync(ChangesOptions.None, 1, ChangeStatus.Submitted, syncPath, cancellationToken);
+				if (changes.Count == 0)
 				{
-					SyncPath = BaseDepotPathPrefix + (bPreview ? "/UnstableRelease/..." : "/Release/...");
-					Changes = await Perforce.GetChangesAsync(ChangesOptions.None, 1, ChangeStatus.Submitted, SyncPath, CancellationToken);
-					if (Changes.Count == 0)
+					syncPath = baseDepotPathPrefix + (preview ? "/UnstableRelease/..." : "/Release/...");
+					changes = await perforce.GetChangesAsync(ChangesOptions.None, 1, ChangeStatus.Submitted, syncPath, cancellationToken);
+					if (changes.Count == 0)
 					{
-						throw new UserErrorException($"Unable to find any UGS binaries under {SyncPath}");
+						throw new UserErrorException($"Unable to find any UGS binaries under {syncPath}");
 					}
 				}
 
-				int RequiredChangeNumber = Changes[0].Number;
-				Logger.LogInformation("Syncing from {SyncPath}", SyncPath);
+				int requiredChangeNumber = changes[0].Number;
+				logger.LogInformation("Syncing from {SyncPath}", syncPath);
 
 				// Create the target folder
-				string ApplicationFolder = Path.Combine(Environment.GetFolderPath(Environment.SpecialFolder.LocalApplicationData), "UnrealGameSync", "Latest");
-				if (!SafeCreateDirectory(ApplicationFolder))
+				string applicationFolder = Path.Combine(Environment.GetFolderPath(Environment.SpecialFolder.LocalApplicationData), "UnrealGameSync", "Latest");
+				if (!SafeCreateDirectory(applicationFolder))
 				{
-					throw new UserErrorException($"Couldn't create directory: {ApplicationFolder}");
+					throw new UserErrorException($"Couldn't create directory: {applicationFolder}");
 				}
 
 				// Read the current version
-				string SyncVersionFile = Path.Combine(ApplicationFolder, "SyncVersion.txt");
-				string RequiredSyncText = String.Format("{0}\n{1}@{2}", Perforce.Settings.ServerAndPort ?? "", SyncPath, RequiredChangeNumber);
+				string syncVersionFile = Path.Combine(applicationFolder, "SyncVersion.txt");
+				string requiredSyncText = String.Format("{0}\n{1}@{2}", perforce.Settings.ServerAndPort ?? "", syncPath, requiredChangeNumber);
 
 				// Check the application exists
-				string ApplicationExe = Path.Combine(ApplicationFolder, "UnrealGameSync.exe");
+				string applicationExe = Path.Combine(applicationFolder, "UnrealGameSync.exe");
 
 				// Check if the version has changed
-				string? SyncText;
-				if (!File.Exists(SyncVersionFile) || !File.Exists(ApplicationExe) || !TryReadAllText(SyncVersionFile, out SyncText) || SyncText != RequiredSyncText)
+				string? syncText;
+				if (!File.Exists(syncVersionFile) || !File.Exists(applicationExe) || !TryReadAllText(syncVersionFile, out syncText) || syncText != requiredSyncText)
 				{
 					// Try to delete the directory contents. Retry for a while, in case we've been spawned by an application in this folder to do an update.
-					for (int NumRetries = 0; !SafeDeleteDirectoryContents(ApplicationFolder); NumRetries++)
+					for (int numRetries = 0; !SafeDeleteDirectoryContents(applicationFolder); numRetries++)
 					{
-						if (NumRetries > 20)
+						if (numRetries > 20)
 						{
-							throw new UserErrorException($"Couldn't delete contents of {ApplicationFolder} (retried {NumRetries} times).");
+							throw new UserErrorException($"Couldn't delete contents of {applicationFolder} (retried {numRetries} times).");
 						}
 						Thread.Sleep(500);
 					}
 
 					// Find all the files in the sync path at this changelist
-					List<FStatRecord> FileRecords = await Perforce.FStatAsync(FStatOptions.None, $"{SyncPath}@{RequiredChangeNumber}", CancellationToken).ToListAsync(CancellationToken);
-					if (FileRecords.Count == 0)
+					List<FStatRecord> fileRecords = await perforce.FStatAsync(FStatOptions.None, $"{syncPath}@{requiredChangeNumber}", cancellationToken).ToListAsync(cancellationToken);
+					if (fileRecords.Count == 0)
 					{
-						throw new UserErrorException($"Couldn't find any matching files for {SyncPath}@{RequiredChangeNumber}");
+						throw new UserErrorException($"Couldn't find any matching files for {syncPath}@{requiredChangeNumber}");
 					}
 
 					// Sync all the files in this list to the same directory structure under the application folder
-					string DepotPathPrefix = SyncPath.Substring(0, SyncPath.LastIndexOf('/') + 1);
-					foreach (FStatRecord FileRecord in FileRecords)
+					string depotPathPrefix = syncPath.Substring(0, syncPath.LastIndexOf('/') + 1);
+					foreach (FStatRecord fileRecord in fileRecords)
 					{
-						if (FileRecord.DepotFile == null)
+						if (fileRecord.DepotFile == null)
 						{
 							throw new UserErrorException("Missing depot path for returned file");
 						}
 
-						string LocalPath = Path.Combine(ApplicationFolder, FileRecord.DepotFile.Substring(DepotPathPrefix.Length).Replace('/', Path.DirectorySeparatorChar));
-						if (!SafeCreateDirectory(Path.GetDirectoryName(LocalPath)!))
+						string localPath = Path.Combine(applicationFolder, fileRecord.DepotFile.Substring(depotPathPrefix.Length).Replace('/', Path.DirectorySeparatorChar));
+						if (!SafeCreateDirectory(Path.GetDirectoryName(localPath)!))
 						{
-							throw new UserErrorException($"Couldn't create folder {Path.GetDirectoryName(LocalPath)}");
+							throw new UserErrorException($"Couldn't create folder {Path.GetDirectoryName(localPath)}");
 						}
 
-						await Perforce.PrintAsync(LocalPath, FileRecord.DepotFile, CancellationToken);
+						await perforce.PrintAsync(localPath, fileRecord.DepotFile, cancellationToken);
 					}
 
 					// If it was a zip file, extract it
-					if (SyncPath.EndsWith(".zip", StringComparison.OrdinalIgnoreCase))
+					if (syncPath.EndsWith(".zip", StringComparison.OrdinalIgnoreCase))
 					{
-						string LocalPath = Path.Combine(ApplicationFolder, SyncPath.Substring(DepotPathPrefix.Length).Replace('/', Path.DirectorySeparatorChar));
-						ZipFile.ExtractToDirectory(LocalPath, ApplicationFolder);
+						string localPath = Path.Combine(applicationFolder, syncPath.Substring(depotPathPrefix.Length).Replace('/', Path.DirectorySeparatorChar));
+						ZipFile.ExtractToDirectory(localPath, applicationFolder);
 					}
 
 					// Check the application exists
-					if (!File.Exists(ApplicationExe))
+					if (!File.Exists(applicationExe))
 					{
-						throw new UserErrorException($"Application was not synced from Perforce. Check that UnrealGameSync exists at {SyncPath}/UnrealGameSync.exe, and you have access to it.");
+						throw new UserErrorException($"Application was not synced from Perforce. Check that UnrealGameSync exists at {syncPath}/UnrealGameSync.exe, and you have access to it.");
 					}
 
 					// Update the version
-					if (!TryWriteAllText(SyncVersionFile, RequiredSyncText))
+					if (!TryWriteAllText(syncVersionFile, requiredSyncText))
 					{
 						throw new UserErrorException("Couldn't write sync text to {SyncVersionFile}");
 					}
 				}
-				Logger.LogInformation("");
+				logger.LogInformation("");
 
 				// Build the command line for the synced application, including the sync path to monitor for updates
-				string OriginalExecutable = Assembly.GetEntryAssembly()!.Location;
-                if (Path.GetExtension(OriginalExecutable).Equals(".dll", StringComparison.OrdinalIgnoreCase))
+				string originalExecutable = Assembly.GetEntryAssembly()!.Location;
+                if (Path.GetExtension(originalExecutable).Equals(".dll", StringComparison.OrdinalIgnoreCase))
                 {
-                    string NewExecutable = Path.ChangeExtension(OriginalExecutable, ".exe");
-                    if (File.Exists(NewExecutable))
+                    string newExecutable = Path.ChangeExtension(originalExecutable, ".exe");
+                    if (File.Exists(newExecutable))
                     {
-                        OriginalExecutable = NewExecutable;
+                        originalExecutable = newExecutable;
                     }
                 }
 
-				StringBuilder NewCommandLine = new StringBuilder(String.Format("-updatepath=\"{0}@>{1}\" -updatespawn=\"{2}\"{3}", SyncPath, RequiredChangeNumber, OriginalExecutable, bPreview ? " -unstable" : ""));
-				foreach (string Arg in Args)
+				StringBuilder newCommandLine = new StringBuilder(String.Format("-updatepath=\"{0}@>{1}\" -updatespawn=\"{2}\"{3}", syncPath, requiredChangeNumber, originalExecutable, preview ? " -unstable" : ""));
+				foreach (string arg in args)
 				{
-					NewCommandLine.AppendFormat(" {0}", QuoteArgument(Arg));
+					newCommandLine.AppendFormat(" {0}", QuoteArgument(arg));
 				}
 
 				// Release the mutex now so that the new application can start up
-				InstanceMutex.Close();
+				instanceMutex.Close();
 
 				// Spawn the application
-				Logger.LogInformation("Spawning {App} with command line: {CmdLine}", ApplicationExe, NewCommandLine.ToString());
-				using (Process ChildProcess = new Process())
+				logger.LogInformation("Spawning {App} with command line: {CmdLine}", applicationExe, newCommandLine.ToString());
+				using (Process childProcess = new Process())
 				{
-					ChildProcess.StartInfo.FileName = ApplicationExe;
-					ChildProcess.StartInfo.Arguments = NewCommandLine.ToString();
-					ChildProcess.StartInfo.UseShellExecute = false;
-					ChildProcess.StartInfo.CreateNoWindow = false;
-					if (!ChildProcess.Start())
+					childProcess.StartInfo.FileName = applicationExe;
+					childProcess.StartInfo.Arguments = newCommandLine.ToString();
+					childProcess.StartInfo.UseShellExecute = false;
+					childProcess.StartInfo.CreateNoWindow = false;
+					if (!childProcess.Start())
 					{
 						throw new UserErrorException("Failed to start process");
 					}
 				}
 			}
-			catch (UserErrorException Ex)
+			catch (UserErrorException ex)
 			{
-				Logger.LogError("{Message}", Ex.Message);
+				logger.LogError("{Message}", ex.Message);
 				throw;
 			}
-			catch (Exception Ex)
+			catch (Exception ex)
 			{
-				Logger.LogError(Ex, "Error while syncing application.");
-				foreach (string Line in Ex.ToString().Split('\n'))
+				logger.LogError(ex, "Error while syncing application.");
+				foreach (string line in ex.ToString().Split('\n'))
 				{
-					Logger.LogError("{Line}", Line);
+					logger.LogError("{Line}", line);
 				}
 				throw;
 			}
 		}
 
-		static string QuoteArgument(string Arg)
+		static string QuoteArgument(string arg)
 		{
-			if(Arg.IndexOf(' ') != -1 && !Arg.StartsWith("\""))
+			if(arg.IndexOf(' ') != -1 && !arg.StartsWith("\""))
 			{
-				return String.Format("\"{0}\"", Arg);
+				return String.Format("\"{0}\"", arg);
 			}
 			else
 			{
-				return Arg;
+				return arg;
 			}
 		}
 
-		static bool TryReadAllText(string FileName, [NotNullWhen(true)] out string? Text)
+		static bool TryReadAllText(string fileName, [NotNullWhen(true)] out string? text)
 		{
 			try
 			{
-				Text = File.ReadAllText(FileName);
+				text = File.ReadAllText(fileName);
 				return true;
 			}
 			catch(Exception)
 			{
-				Text = null;
+				text = null;
 				return false;
 			}
 		}
 
-		static bool TryWriteAllText(string FileName, string Text)
+		static bool TryWriteAllText(string fileName, string text)
 		{
 			try
 			{
-				File.WriteAllText(FileName, Text);
+				File.WriteAllText(fileName, text);
 				return true;
 			}
 			catch(Exception)
@@ -277,11 +277,11 @@ namespace UnrealGameSyncLauncher
 			}
 		}
 
-		static bool SafeCreateDirectory(string DirectoryName)
+		static bool SafeCreateDirectory(string directoryName)
 		{
 			try
 			{
-				Directory.CreateDirectory(DirectoryName);
+				Directory.CreateDirectory(directoryName);
 				return true;
 			}
 			catch(Exception)
@@ -290,11 +290,11 @@ namespace UnrealGameSyncLauncher
 			}
 		}
 
-		static bool SafeDeleteDirectory(string DirectoryName)
+		static bool SafeDeleteDirectory(string directoryName)
 		{
 			try
 			{
-				Directory.Delete(DirectoryName, true);
+				Directory.Delete(directoryName, true);
 				return true;
 			}
 			catch(Exception)
@@ -303,19 +303,19 @@ namespace UnrealGameSyncLauncher
 			}
 		}
 
-		static bool SafeDeleteDirectoryContents(string DirectoryName)
+		static bool SafeDeleteDirectoryContents(string directoryName)
 		{
 			try
 			{
-				DirectoryInfo Directory = new DirectoryInfo(DirectoryName);
-				foreach(FileInfo ChildFile in Directory.EnumerateFiles("*", SearchOption.AllDirectories))
+				DirectoryInfo directory = new DirectoryInfo(directoryName);
+				foreach(FileInfo childFile in directory.EnumerateFiles("*", SearchOption.AllDirectories))
 				{
-					ChildFile.Attributes = ChildFile.Attributes & ~FileAttributes.ReadOnly;
-					ChildFile.Delete();
+					childFile.Attributes = childFile.Attributes & ~FileAttributes.ReadOnly;
+					childFile.Delete();
 				}
-				foreach(DirectoryInfo ChildDirectory in Directory.EnumerateDirectories())
+				foreach(DirectoryInfo childDirectory in directory.EnumerateDirectories())
 				{
-					SafeDeleteDirectory(ChildDirectory.FullName);
+					SafeDeleteDirectory(childDirectory.FullName);
 				}
 				return true;
 			}
diff --git a/Engine/Source/Programs/UnrealGameSync/UnrealGameSyncShared/ArchiveManifest.cs b/Engine/Source/Programs/UnrealGameSync/UnrealGameSyncShared/ArchiveManifest.cs
--- a/Engine/Source/Programs/UnrealGameSync/UnrealGameSyncShared/ArchiveManifest.cs
+++ b/Engine/Source/Programs/UnrealGameSync/UnrealGameSyncShared/ArchiveManifest.cs
@@ -20,25 +20,25 @@ namespace UnrealGameSync
 		public long Length;
 		public DateTime LastWriteTimeUtc;
 
-		public ArchiveManifestFile(BinaryReader Reader)
+		public ArchiveManifestFile(BinaryReader reader)
 		{
-			FileName = Reader.ReadString();
-			Length = Reader.ReadInt64();
-			LastWriteTimeUtc = new DateTime(Reader.ReadInt64());
+			FileName = reader.ReadString();
+			Length = reader.ReadInt64();
+			LastWriteTimeUtc = new DateTime(reader.ReadInt64());
 		}
 
-		public ArchiveManifestFile(string InFileName, long InLength, DateTime InLastWriteTimeUtc)
+		public ArchiveManifestFile(string inFileName, long inLength, DateTime inLastWriteTimeUtc)
 		{
-			FileName = InFileName;
-			Length = InLength;
-			LastWriteTimeUtc = InLastWriteTimeUtc;
+			FileName = inFileName;
+			Length = inLength;
+			LastWriteTimeUtc = inLastWriteTimeUtc;
 		}
 
-		public void Write(BinaryWriter Writer)
+		public void Write(BinaryWriter writer)
 		{
-			Writer.Write(FileName);
-			Writer.Write(Length);
-			Writer.Write(LastWriteTimeUtc.Ticks);
+			writer.Write(FileName);
+			writer.Write(Length);
+			writer.Write(LastWriteTimeUtc.Ticks);
 		}
 	}
 
@@ -52,111 +52,111 @@ namespace UnrealGameSync
 		{
 		}
 
-		public ArchiveManifest(FileStream InputStream)
+		public ArchiveManifest(FileStream inputStream)
 		{
-			BinaryReader Reader = new BinaryReader(InputStream);
-			if(Reader.ReadInt32() != Signature)
+			BinaryReader reader = new BinaryReader(inputStream);
+			if(reader.ReadInt32() != Signature)
 			{
 				throw new Exception("Archive manifest signature does not match");
 			}
 
-			int NumFiles = Reader.ReadInt32();
-			for(int Idx = 0; Idx < NumFiles; Idx++)
+			int numFiles = reader.ReadInt32();
+			for(int idx = 0; idx < numFiles; idx++)
 			{
-				Files.Add(new ArchiveManifestFile(Reader));
+				Files.Add(new ArchiveManifestFile(reader));
 			}
 		}
 
-		public void Write(FileStream OutputStream)
+		public void Write(FileStream outputStream)
 		{
-			BinaryWriter Writer = new BinaryWriter(OutputStream);
-			Writer.Write(Signature);
-			Writer.Write(Files.Count);
-			foreach(ArchiveManifestFile File in Files)
+			BinaryWriter writer = new BinaryWriter(outputStream);
+			writer.Write(Signature);
+			writer.Write(Files.Count);
+			foreach(ArchiveManifestFile file in Files)
 			{
-				File.Write(Writer);
+				file.Write(writer);
 			}
 		}
 	}
 
 	public static class ArchiveUtils
 	{
-		public static void ExtractFiles(FileReference ArchiveFileName, DirectoryReference BaseDirectoryName, FileReference? ManifestFileName, ProgressValue Progress, ILogger Logger)
+		public static void ExtractFiles(FileReference archiveFileName, DirectoryReference baseDirectoryName, FileReference? manifestFileName, ProgressValue progress, ILogger logger)
 		{
-			DateTime TimeStamp = DateTime.UtcNow;
-			using (ZipArchive Zip = new ZipArchive(File.OpenRead(ArchiveFileName.FullName)))
+			DateTime timeStamp = DateTime.UtcNow;
+			using (ZipArchive zip = new ZipArchive(File.OpenRead(archiveFileName.FullName)))
 			{
-				if (ManifestFileName != null)
+				if (manifestFileName != null)
 				{
-					FileReference.Delete(ManifestFileName);
+					FileReference.Delete(manifestFileName);
 
 					// Create the manifest
-					ArchiveManifest Manifest = new ArchiveManifest();
-					foreach (ZipArchiveEntry Entry in Zip.Entries)
+					ArchiveManifest manifest = new ArchiveManifest();
+					foreach (ZipArchiveEntry entry in zip.Entries)
 					{
-						if (!Entry.FullName.EndsWith("/") && !Entry.FullName.EndsWith("\\"))
+						if (!entry.FullName.EndsWith("/") && !entry.FullName.EndsWith("\\"))
 						{
-							Manifest.Files.Add(new ArchiveManifestFile(Entry.FullName, Entry.Length, TimeStamp));
+							manifest.Files.Add(new ArchiveManifestFile(entry.FullName, entry.Length, timeStamp));
 						}
 					}
 
 					// Write it out to a temporary file, then move it into place
-					FileReference TempManifestFileName = ManifestFileName + ".tmp";
-					using (FileStream OutputStream = FileReference.Open(TempManifestFileName, FileMode.Create, FileAccess.Write))
+					FileReference tempManifestFileName = manifestFileName + ".tmp";
+					using (FileStream outputStream = FileReference.Open(tempManifestFileName, FileMode.Create, FileAccess.Write))
 					{
-						Manifest.Write(OutputStream);
+						manifest.Write(outputStream);
 					}
-					FileReference.Move(TempManifestFileName, ManifestFileName);
+					FileReference.Move(tempManifestFileName, manifestFileName);
 				}
 
 				// Extract all the files
-				int EntryIdx = 0;
-				foreach(ZipArchiveEntry Entry in Zip.Entries)
+				int entryIdx = 0;
+				foreach(ZipArchiveEntry entry in zip.Entries)
 				{
-					if(!Entry.FullName.EndsWith("/") && !Entry.FullName.EndsWith("\\"))
+					if(!entry.FullName.EndsWith("/") && !entry.FullName.EndsWith("\\"))
 					{
-						FileReference FileName = FileReference.Combine(BaseDirectoryName, Entry.FullName);
-						DirectoryReference.CreateDirectory(FileName.Directory);
-						Logger.LogInformation("Writing {0}", FileName);
+						FileReference fileName = FileReference.Combine(baseDirectoryName, entry.FullName);
+						DirectoryReference.CreateDirectory(fileName.Directory);
+						logger.LogInformation("Writing {0}", fileName);
 
-						Entry.ExtractToFile(FileName.FullName, true);
-						FileReference.SetLastWriteTimeUtc(FileName, TimeStamp);
+						entry.ExtractToFile(fileName.FullName, true);
+						FileReference.SetLastWriteTimeUtc(fileName, timeStamp);
 					}
-					Progress.Set((float)++EntryIdx / (float)Zip.Entries.Count);
+					progress.Set((float)++entryIdx / (float)zip.Entries.Count);
 				}
 			}
 		}
 
-		public static void RemoveExtractedFiles(DirectoryReference BaseDirectoryName, FileReference ManifestFileName, ProgressValue Progress, ILogger LogWriter)
+		public static void RemoveExtractedFiles(DirectoryReference baseDirectoryName, FileReference manifestFileName, ProgressValue progress, ILogger logWriter)
 		{
 			// Read the manifest in
-			ArchiveManifest Manifest;
-			using(FileStream InputStream = FileReference.Open(ManifestFileName, FileMode.Open, FileAccess.Read))
+			ArchiveManifest manifest;
+			using(FileStream inputStream = FileReference.Open(manifestFileName, FileMode.Open, FileAccess.Read))
 			{
-				Manifest = new ArchiveManifest(InputStream);
+				manifest = new ArchiveManifest(inputStream);
 			}
 
 			// Remove all the files that haven't been modified match
-			for(int Idx = 0; Idx < Manifest.Files.Count; Idx++)
+			for(int idx = 0; idx < manifest.Files.Count; idx++)
 			{
-				FileInfo File = FileReference.Combine(BaseDirectoryName, Manifest.Files[Idx].FileName).ToFileInfo();
-				if(File.Exists)
+				FileInfo file = FileReference.Combine(baseDirectoryName, manifest.Files[idx].FileName).ToFileInfo();
+				if(file.Exists)
 				{
-					if(File.Length != Manifest.Files[Idx].Length)
+					if(file.Length != manifest.Files[idx].Length)
 					{
-						LogWriter.LogInformation("Skipping {FileName} due to modified length", File.FullName);
+						logWriter.LogInformation("Skipping {FileName} due to modified length", file.FullName);
 					}
-					else if(Math.Abs((File.LastWriteTimeUtc - Manifest.Files[Idx].LastWriteTimeUtc).TotalSeconds) > 2.0)
+					else if(Math.Abs((file.LastWriteTimeUtc - manifest.Files[idx].LastWriteTimeUtc).TotalSeconds) > 2.0)
 					{
-						LogWriter.LogInformation("Skipping {FileName} due to modified timestamp", File.FullName);
+						logWriter.LogInformation("Skipping {FileName} due to modified timestamp", file.FullName);
 					}
 					else
 					{
-						LogWriter.LogInformation("Removing {FileName}", File.FullName);
-						File.Delete();
+						logWriter.LogInformation("Removing {FileName}", file.FullName);
+						file.Delete();
 					}
 				}
-				Progress.Set((float)(Idx + 1) / (float)Manifest.Files.Count);
+				progress.Set((float)(idx + 1) / (float)manifest.Files.Count);
 			}
 		}
 	}
diff --git a/Engine/Source/Programs/UnrealGameSync/UnrealGameSyncShared/BuildStep.cs b/Engine/Source/Programs/UnrealGameSync/UnrealGameSyncShared/BuildStep.cs
--- a/Engine/Source/Programs/UnrealGameSync/UnrealGameSyncShared/BuildStep.cs
+++ b/Engine/Source/Programs/UnrealGameSync/UnrealGameSyncShared/BuildStep.cs
@@ -33,109 +33,109 @@ namespace UnrealGameSync
 		public string? FileName;
 		public string? WorkingDir;
 		public string? Arguments;
-		public bool bUseLogWindow;
-		public bool bNormalSync;
-		public bool bScheduledSync;
+		public bool UseLogWindow;
+		public bool NormalSync;
+		public bool ScheduledSync;
 		public string? StatusPanelLink;
 		public Guid[] Requires = Array.Empty<Guid>();
 		public Guid ToolId;
 
-		public BuildStep(Guid InUniqueId, int InOrderIndex, string? InDescription, string? InStatusText, int InEstimatedDuration, string? InFileName, string? InArguments, string? InWorkingDir, bool bInUseLogWindow)
+		public BuildStep(Guid inUniqueId, int inOrderIndex, string? inDescription, string? inStatusText, int inEstimatedDuration, string? inFileName, string? inArguments, string? inWorkingDir, bool inUseLogWindow)
 		{
-			UniqueId = InUniqueId;
-			OrderIndex = InOrderIndex;
-			Description = InDescription;
-			StatusText = InStatusText;
-			EstimatedDuration = InEstimatedDuration;
+			UniqueId = inUniqueId;
+			OrderIndex = inOrderIndex;
+			Description = inDescription;
+			StatusText = inStatusText;
+			EstimatedDuration = inEstimatedDuration;
 			Type = BuildStepType.Other;
-			FileName = InFileName;
-			Arguments = InArguments;
-			WorkingDir = InWorkingDir;
-			bUseLogWindow = bInUseLogWindow;
+			FileName = inFileName;
+			Arguments = inArguments;
+			WorkingDir = inWorkingDir;
+			UseLogWindow = inUseLogWindow;
 		}
 
-		public BuildStep(Guid InUniqueId, int InOrderIndex, string? InDescription, string? InStatusText, int InEstimatedDuration, string? InTarget, string? InPlatform, string? InConfiguration, string? InArguments, bool bInSyncDefault)
+		public BuildStep(Guid inUniqueId, int inOrderIndex, string? inDescription, string? inStatusText, int inEstimatedDuration, string? inTarget, string? inPlatform, string? inConfiguration, string? inArguments, bool inSyncDefault)
 		{
-			UniqueId = InUniqueId;
-			OrderIndex = InOrderIndex;
-			Description = InDescription;
-			StatusText = InStatusText;
-			EstimatedDuration = InEstimatedDuration;
+			UniqueId = inUniqueId;
+			OrderIndex = inOrderIndex;
+			Description = inDescription;
+			StatusText = inStatusText;
+			EstimatedDuration = inEstimatedDuration;
 			Type = BuildStepType.Compile;
-			Target = InTarget;
-			Platform = InPlatform;
-			Configuration = InConfiguration;
-			Arguments = InArguments;
-			bUseLogWindow = true;
-			bNormalSync = bInSyncDefault;
-			bScheduledSync = bInSyncDefault;
+			Target = inTarget;
+			Platform = inPlatform;
+			Configuration = inConfiguration;
+			Arguments = inArguments;
+			UseLogWindow = true;
+			NormalSync = inSyncDefault;
+			ScheduledSync = inSyncDefault;
 		}
 
-		public BuildStep(ConfigObject Object)
+		public BuildStep(ConfigObject obj)
 		{
-			if(!Guid.TryParse(Object.GetValue(UniqueIdKey, ""), out UniqueId))
+			if(!Guid.TryParse(obj.GetValue(UniqueIdKey, ""), out UniqueId))
 			{
 				UniqueId = Guid.NewGuid();
 			}
-			if(!Int32.TryParse(Object.GetValue("OrderIndex", ""), out OrderIndex))
+			if(!Int32.TryParse(obj.GetValue("OrderIndex", ""), out OrderIndex))
 			{
 				OrderIndex = -1;
 			}
 
-			Description = Object.GetValue("Description", "Untitled");
-			StatusText = Object.GetValue("StatusText", "Untitled");
+			Description = obj.GetValue("Description", "Untitled");
+			StatusText = obj.GetValue("StatusText", "Untitled");
 
-			if(!int.TryParse(Object.GetValue("EstimatedDuration", ""), out EstimatedDuration) || EstimatedDuration < 1)
+			if(!int.TryParse(obj.GetValue("EstimatedDuration", ""), out EstimatedDuration) || EstimatedDuration < 1)
 			{
 				EstimatedDuration = 1;
 			}
 
-			if(!Enum.TryParse(Object.GetValue("Type", ""), true, out Type))
+			if(!Enum.TryParse(obj.GetValue("Type", ""), true, out Type))
 			{
 				Type = BuildStepType.Other;
 			}
 
-			Target = Object.GetValue("Target");
-			Platform = Object.GetValue("Platform");
-			Configuration = Object.GetValue("Configuration");
-			FileName = Object.GetValue("FileName");
-			WorkingDir = Object.GetValue("WorkingDir");
-			Arguments = Object.GetValue("Arguments");
+			Target = obj.GetValue("Target");
+			Platform = obj.GetValue("Platform");
+			Configuration = obj.GetValue("Configuration");
+			FileName = obj.GetValue("FileName");
+			WorkingDir = obj.GetValue("WorkingDir");
+			Arguments = obj.GetValue("Arguments");
 
-			if(!Boolean.TryParse(Object.GetValue("bUseLogWindow", ""), out bUseLogWindow))
+			if(!Boolean.TryParse(obj.GetValue("bUseLogWindow", ""), out UseLogWindow))
 			{
-				bUseLogWindow = true;
+				UseLogWindow = true;
 			}
-			if(!Boolean.TryParse(Object.GetValue("bNormalSync", ""), out bNormalSync))
+			if(!Boolean.TryParse(obj.GetValue("bNormalSync", ""), out NormalSync))
 			{
-				bNormalSync = true;
+				NormalSync = true;
 			}
-			if(!Boolean.TryParse(Object.GetValue("bScheduledSync", ""), out bScheduledSync))
+			if(!Boolean.TryParse(obj.GetValue("bScheduledSync", ""), out ScheduledSync))
 			{
-				bScheduledSync = bNormalSync;
+				ScheduledSync = NormalSync;
 			}
 
-			StatusPanelLink = Object.GetValue("Link", null);
+			StatusPanelLink = obj.GetValue("Link", null);
 			if (String.IsNullOrEmpty(StatusPanelLink))
 			{
-				bool bShowAsTool;
-				if (Boolean.TryParse(Object.GetValue("bShowAsTool", ""), out bShowAsTool) && bShowAsTool)
+				bool showAsTool;
+				if (Boolean.TryParse(obj.GetValue("bShowAsTool", ""), out showAsTool) && showAsTool)
 				{
 					StatusPanelLink = $"More...|{Description}";
 				}
 			}
 
-			List<Guid> Requires = new List<Guid>();
-			foreach (string RequireString in Object.GetValue("Requires", String.Empty).Split(';', StringSplitOptions.RemoveEmptyEntries))
+			List<Guid> requires = new List<Guid>();
+			foreach (string requireString in obj.GetValue("Requires", String.Empty).Split(';', StringSplitOptions.RemoveEmptyEntries))
 			{
-				if (Guid.TryParse(RequireString, out Guid Require))
+				if (Guid.TryParse(requireString, out Guid require))
 				{
-					Requires.Add(Require);
+					requires.Add(require);
 				}
 			}
-			this.Requires = Requires.ToArray();
+			this.Requires = requires.ToArray();
 	
-			if (!Guid.TryParse(Object.GetValue("Tool", ""), out ToolId))
+			if (!Guid.TryParse(obj.GetValue("Tool", ""), out ToolId))
 			{
 				ToolId = Guid.Empty;
 			}
@@ -156,69 +156,69 @@ namespace UnrealGameSync
 			}
 		}
 
-		public static void MergeBuildStepObjects(Dictionary<Guid, ConfigObject> BuildStepObjects, IEnumerable<ConfigObject> ModifyObjects)
+		public static void MergeBuildStepObjects(Dictionary<Guid, ConfigObject> buildStepObjects, IEnumerable<ConfigObject> modifyObjects)
 		{
-			foreach(ConfigObject ModifyObject in ModifyObjects)
+			foreach(ConfigObject modifyObject in modifyObjects)
 			{
-				Guid UniqueId;
-				if(Guid.TryParse(ModifyObject.GetValue(UniqueIdKey, ""), out UniqueId))
+				Guid uniqueId;
+				if(Guid.TryParse(modifyObject.GetValue(UniqueIdKey, ""), out uniqueId))
 				{
-					ConfigObject? DefaultObject;
-					if(BuildStepObjects.TryGetValue(UniqueId, out DefaultObject))
+					ConfigObject? defaultObject;
+					if(buildStepObjects.TryGetValue(uniqueId, out defaultObject))
 					{
-						ModifyObject.SetDefaults(DefaultObject);
+						modifyObject.SetDefaults(defaultObject);
 					}
-					BuildStepObjects[UniqueId] = ModifyObject;
+					buildStepObjects[uniqueId] = modifyObject;
 				}
 			}
 		}
 
 		public ConfigObject ToConfigObject()
 		{
-			ConfigObject Result = new ConfigObject();
-			Result["UniqueId"] = UniqueId.ToString();
-			Result["Description"] = Description;
-			Result["StatusText"] = StatusText;
-			Result["EstimatedDuration"] = EstimatedDuration.ToString();
-			Result["Type"] = Type.ToString();
+			ConfigObject result = new ConfigObject();
+			result["UniqueId"] = UniqueId.ToString();
+			result["Description"] = Description;
+			result["StatusText"] = StatusText;
+			result["EstimatedDuration"] = EstimatedDuration.ToString();
+			result["Type"] = Type.ToString();
 			switch(Type)
 			{
 				case BuildStepType.Compile:
-					Result["Target"] = Target;
-					Result["Platform"] = Platform;
-					Result["Configuration"] = Configuration;
-					Result["Arguments"] = Arguments;
+					result["Target"] = Target;
+					result["Platform"] = Platform;
+					result["Configuration"] = Configuration;
+					result["Arguments"] = Arguments;
 					break;
 				case BuildStepType.Cook:
-					Result["FileName"] = FileName;
+					result["FileName"] = FileName;
 					break;
 				case BuildStepType.Other:
-					Result["FileName"] = FileName;
-					Result["WorkingDir"] = WorkingDir;
-					Result["Arguments"] = Arguments;
-					Result["bUseLogWindow"] = bUseLogWindow.ToString();
+					result["FileName"] = FileName;
+					result["WorkingDir"] = WorkingDir;
+					result["Arguments"] = Arguments;
+					result["bUseLogWindow"] = UseLogWindow.ToString();
 					break;
 			}
-			Result["OrderIndex"] = OrderIndex.ToString();
-			Result["bNormalSync"] = bNormalSync.ToString();
-			Result["bScheduledSync"] = bScheduledSync.ToString();
+			result["OrderIndex"] = OrderIndex.ToString();
+			result["bNormalSync"] = NormalSync.ToString();
+			result["bScheduledSync"] = ScheduledSync.ToString();
 			if (!String.IsNullOrEmpty(StatusPanelLink))
 			{
-				Result["Link"] = StatusPanelLink;
+				result["Link"] = StatusPanelLink;
 			}
 			if (ToolId != Guid.Empty)
 			{
-				Result["Tool"] = ToolId.ToString();
+				result["Tool"] = ToolId.ToString();
 			}
-			return Result;
+			return result;
 		}
 
-		public ConfigObject? ToConfigObject(ConfigObject? DefaultObject)
+		public ConfigObject? ToConfigObject(ConfigObject? defaultObject)
 		{
-			ConfigObject Result = new ConfigObject();
-			Result[UniqueIdKey] = UniqueId.ToString();
-			Result.AddOverrides(ToConfigObject(), DefaultObject);
-			return (Result.Pairs.Count <= 1)? null : Result;
+			ConfigObject result = new ConfigObject();
+			result[UniqueIdKey] = UniqueId.ToString();
+			result.AddOverrides(ToConfigObject(), defaultObject);
+			return (result.Pairs.Count <= 1)? null : result;
 		}
 	}
 }
diff --git a/Engine/Source/Programs/UnrealGameSync/UnrealGameSyncShared/ConfigFile.cs b/Engine/Source/Programs/UnrealGameSync/UnrealGameSyncShared/ConfigFile.cs
--- a/Engine/Source/Programs/UnrealGameSync/UnrealGameSyncShared/ConfigFile.cs
+++ b/Engine/Source/Programs/UnrealGameSync/UnrealGameSyncShared/ConfigFile.cs
@@ -27,55 +27,55 @@ namespace UnrealGameSync
 			Pairs = new List<KeyValuePair<string,string>>();
 		}
 
-		public ConfigObject(ConfigObject Other)
+		public ConfigObject(ConfigObject other)
 		{
-			Pairs = new List<KeyValuePair<string,string>>(Other.Pairs);
+			Pairs = new List<KeyValuePair<string,string>>(other.Pairs);
 		}
 
-		public ConfigObject(string Text)
+		public ConfigObject(string text)
 		{
 			Pairs = new List<KeyValuePair<string,string>>();
-			ParseConfigString(Text);
+			ParseConfigString(text);
 		}
 
-		public ConfigObject(ConfigObject BaseObject, string Text)
+		public ConfigObject(ConfigObject baseObject, string text)
 		{
-			Pairs = new List<KeyValuePair<string,string>>(BaseObject.Pairs);
-			ParseConfigString(Text);
+			Pairs = new List<KeyValuePair<string,string>>(baseObject.Pairs);
+			ParseConfigString(text);
 		}
 
-		void ParseConfigString(string Text)
+		void ParseConfigString(string text)
 		{
-			int Idx = 0;
-			if(ParseConfigToken(Text, ref Idx) == "(")
+			int idx = 0;
+			if(ParseConfigToken(text, ref idx) == "(")
 			{
-				while(Idx < Text.Length)
+				while(idx < text.Length)
 				{
 					// Read the next key/value pair
-					string? Key = ParseConfigToken(Text, ref Idx);
-					if (Key == null)
+					string? key = ParseConfigToken(text, ref idx);
+					if (key == null)
 					{
 						return;
 					}
-					if(ParseConfigToken(Text, ref Idx) == "=")
+					if(ParseConfigToken(text, ref idx) == "=")
 					{
-						string? Value = ParseConfigValueToken(Text, ref Idx);
-						if (Value == null)
+						string? value = ParseConfigValueToken(text, ref idx);
+						if (value == null)
 						{
 							return;
 						}
-						SetValue(Key, Value);
+						SetValue(key, value);
 					}
 
 					// Check for the end of the list, or a comma before the next pair
 					for(;;)
 					{
-						string? Token = ParseConfigValueToken(Text, ref Idx);
-						if(Token == ",")
+						string? token = ParseConfigValueToken(text, ref idx);
+						if(token == ",")
 						{
 							break;
 						}
-						if(Token == ")" || Token == null)
+						if(token == ")" || token == null)
 						{
 							return;
 						}
@@ -84,215 +84,215 @@ namespace UnrealGameSync
 			}
 		}
 
-		static string? ParseConfigToken(string Text, ref int Idx)
+		static string? ParseConfigToken(string text, ref int idx)
 		{
 			// Skip whitespace
-			while(Idx < Text.Length && Char.IsWhiteSpace(Text[Idx]))
+			while(idx < text.Length && Char.IsWhiteSpace(text[idx]))
 			{
-				Idx++;
+				idx++;
 			}
-			if(Idx == Text.Length)
+			if(idx == text.Length)
 			{
 				return null;
 			}
 
 			// Read the token
-			if (Text[Idx] == '\"')
+			if (text[idx] == '\"')
 			{
-				StringBuilder Token = new StringBuilder();
-				while (++Idx < Text.Length)
+				StringBuilder token = new StringBuilder();
+				while (++idx < text.Length)
 				{
-					if (Text[Idx] == '\"')
+					if (text[idx] == '\"')
 					{
-						Idx++;
+						idx++;
 						break;
 					}
-					if (Text[Idx] == '\\' && Idx + 1 < Text.Length)
+					if (text[idx] == '\\' && idx + 1 < text.Length)
 					{
-						Idx++;
+						idx++;
 					}
-					Token.Append(Text[Idx]);
+					token.Append(text[idx]);
 				}
-				return Token.ToString();
+				return token.ToString();
 			}
-			else if (ConfigSeparatorCharacters.IndexOf(Text[Idx]) != -1)
+			else if (ConfigSeparatorCharacters.IndexOf(text[idx]) != -1)
 			{
-				return Text[Idx++].ToString();
+				return text[idx++].ToString();
 			}
 			else
 			{
-				int StartIdx = Idx;
-				while (Idx < Text.Length && ConfigSeparatorCharacters.IndexOf(Text[Idx]) == -1)
+				int startIdx = idx;
+				while (idx < text.Length && ConfigSeparatorCharacters.IndexOf(text[idx]) == -1)
 				{
-					Idx++;
+					idx++;
 				}
-				return Text.Substring(StartIdx, Idx - StartIdx);
+				return text.Substring(startIdx, idx - startIdx);
 			}
 		}
 
-		static string? ParseConfigValueToken(string Text, ref int Idx)
+		static string? ParseConfigValueToken(string text, ref int idx)
 		{
-			string? Token = ParseConfigToken(Text, ref Idx);
-			if (Token == "(")
+			string? token = ParseConfigToken(text, ref idx);
+			if (token == "(")
 			{
-				int StartIdx = Idx - 1;
+				int startIdx = idx - 1;
 				for (; ; )
 				{
-					string? NextToken = ParseConfigValueToken(Text, ref Idx);
-					if (NextToken == null || NextToken == ")")
+					string? nextToken = ParseConfigValueToken(text, ref idx);
+					if (nextToken == null || nextToken == ")")
 					{
 						break;
 					}
 				}
-				Token = Text.Substring(StartIdx, Idx - StartIdx);
+				token = text.Substring(startIdx, idx - startIdx);
 			}
-			return Token;
+			return token;
 		}
 
-		[return: NotNullIfNotNull("DefaultValue")]
-		public string? GetValue(string Key, string? DefaultValue = null)
+		[return: NotNullIfNotNull("defaultValue")]
+		public string? GetValue(string key, string? defaultValue = null)
 		{
-			for(int Idx = 0; Idx < Pairs.Count; Idx++)
+			for(int idx = 0; idx < Pairs.Count; idx++)
 			{
-				if(Pairs[Idx].Key.Equals(Key, StringComparison.InvariantCultureIgnoreCase))
+				if(Pairs[idx].Key.Equals(key, StringComparison.InvariantCultureIgnoreCase))
 				{
-					return Pairs[Idx].Value;
+					return Pairs[idx].Value;
 				}
 			}
-			return DefaultValue;
+			return defaultValue;
 		}
 
-		public Guid GetValue(string Key, Guid DefaultValue)
+		public Guid GetValue(string key, Guid defaultValue)
 		{
-			string? StringValue = GetValue(Key);
-			if(StringValue != null)
+			string? stringValue = GetValue(key);
+			if(stringValue != null)
 			{
-				Guid Value;
-				if(Guid.TryParse(StringValue, out Value))
+				Guid value;
+				if(Guid.TryParse(stringValue, out value))
 				{
-					return Value;
+					return value;
 				}
 			}
-			return DefaultValue;
+			return defaultValue;
 		}
 
-		public int GetValue(string Key, int DefaultValue)
+		public int GetValue(string key, int defaultValue)
 		{
-			string? StringValue = GetValue(Key);
-			if(StringValue != null)
+			string? stringValue = GetValue(key);
+			if(stringValue != null)
 			{
-				int Value;
-				if(int.TryParse(StringValue, out Value))
+				int value;
+				if(int.TryParse(stringValue, out value))
 				{
-					return Value;
+					return value;
 				}
 			}
-			return DefaultValue;
+			return defaultValue;
 		}
 
-		public bool GetValue(string Key, bool DefaultValue)
+		public bool GetValue(string key, bool defaultValue)
 		{
-			string? StringValue = GetValue(Key);
-			if(StringValue != null)
+			string? stringValue = GetValue(key);
+			if(stringValue != null)
 			{
-				bool Value;
-				if(bool.TryParse(StringValue, out Value))
+				bool value;
+				if(bool.TryParse(stringValue, out value))
 				{
-					return Value;
+					return value;
 				}
 			}
-			return DefaultValue;
+			return defaultValue;
 		}
 
-		public void SetValue(string Key, string? Value)
+		public void SetValue(string key, string? value)
 		{
-			if (Value == null)
+			if (value == null)
 			{
-				Pairs.RemoveAll(x => x.Key.Equals(Key, StringComparison.OrdinalIgnoreCase));
+				Pairs.RemoveAll(x => x.Key.Equals(key, StringComparison.OrdinalIgnoreCase));
 			}
 			else
 			{
-				for (int Idx = 0; Idx < Pairs.Count; Idx++)
+				for (int idx = 0; idx < Pairs.Count; idx++)
 				{
-					if (Pairs[Idx].Key.Equals(Key, StringComparison.OrdinalIgnoreCase))
+					if (Pairs[idx].Key.Equals(key, StringComparison.OrdinalIgnoreCase))
 					{
-						Pairs[Idx] = new KeyValuePair<string, string>(Key, Value);
+						Pairs[idx] = new KeyValuePair<string, string>(key, value);
 						return;
 					}
 				}
-				Pairs.Add(new KeyValuePair<string, string>(Key, Value));
+				Pairs.Add(new KeyValuePair<string, string>(key, value));
 			}
 		}
 
-		public void SetValue(string Key, Guid Value)
+		public void SetValue(string key, Guid value)
 		{
-			SetValue(Key, Value.ToString());
+			SetValue(key, value.ToString());
 		}
 
-		public void SetValue(string Key, int Value)
+		public void SetValue(string key, int value)
 		{
-			SetValue(Key, Value.ToString());
+			SetValue(key, value.ToString());
 		}
 
-		public void SetValue(string Key, bool Value)
+		public void SetValue(string key, bool value)
 		{
-			SetValue(Key, Value.ToString());
+			SetValue(key, value.ToString());
 		}
 
-		public string? this[string Key]
+		public string? this[string key]
 		{
-			get { return GetValue(Key); }
-			set { SetValue(Key, value); }
+			get { return GetValue(key); }
+			set { SetValue(key, value); }
 		}
 
-		public void SetDefaults(ConfigObject Other)
+		public void SetDefaults(ConfigObject other)
 		{
-			foreach(KeyValuePair<string, string> Pair in Other.Pairs)
+			foreach(KeyValuePair<string, string> pair in other.Pairs)
 			{
-				if(GetValue(Pair.Key) == null)
+				if(GetValue(pair.Key) == null)
 				{
-					SetValue(Pair.Key, Pair.Value);
+					SetValue(pair.Key, pair.Value);
 				}
 			}
 		}
 
-		public void AddOverrides(ConfigObject Object, ConfigObject? DefaultObject)
+		public void AddOverrides(ConfigObject obj, ConfigObject? defaultObject)
 		{
-			foreach(KeyValuePair<string, string> Pair in Object.Pairs)
+			foreach(KeyValuePair<string, string> pair in obj.Pairs)
 			{
-				if(DefaultObject == null || DefaultObject.GetValue(Pair.Key) != Pair.Value)
+				if(defaultObject == null || defaultObject.GetValue(pair.Key) != pair.Value)
 				{
-					SetValue(Pair.Key, Pair.Value);
+					SetValue(pair.Key, pair.Value);
 				}
 			}
 		}
 
-		public string ToString(ConfigObject? BaseObject)
+		public string ToString(ConfigObject? baseObject)
 		{
-			StringBuilder Result = new StringBuilder();
-			Result.Append("(");
-			foreach(KeyValuePair<string, string> Pair in Pairs)
+			StringBuilder result = new StringBuilder();
+			result.Append("(");
+			foreach(KeyValuePair<string, string> pair in Pairs)
 			{
-				if(BaseObject == null || BaseObject.GetValue(Pair.Key) != Pair.Value)
+				if(baseObject == null || baseObject.GetValue(pair.Key) != pair.Value)
 				{
-					if(Result.Length > 1)
+					if(result.Length > 1)
 					{
-						Result.Append(", ");
+						result.Append(", ");
 					}
-					Result.Append(Pair.Key);
-					Result.Append("=");
-					if(Pair.Value == null)
+					result.Append(pair.Key);
+					result.Append("=");
+					if(pair.Value == null)
 					{
-						Result.Append("\"\"");
+						result.Append("\"\"");
 					}
 					else
 					{
-						Result.AppendFormat("\"{0}\"", Pair.Value.Replace("\\", "\\\\").Replace("\"", "\\\""));
+						result.AppendFormat("\"{0}\"", pair.Value.Replace("\\", "\\\\").Replace("\"", "\\\""));
 					}
 				}
 			}
-			Result.Append(")");
-			return Result.ToString();
+			result.Append(")");
+			return result.ToString();
 		}
  
 		public override string ToString()
@@ -303,34 +303,34 @@ namespace UnrealGameSync
 
 	public class ConfigObjectJsonConverter : JsonConverter<ConfigObject>
 	{
-		public override ConfigObject Read(ref Utf8JsonReader Reader, Type TypeToConvert, JsonSerializerOptions Options)
+		public override ConfigObject Read(ref Utf8JsonReader reader, Type typeToConvert, JsonSerializerOptions options)
 		{
-			ConfigObject Object = new ConfigObject();
-			if (Reader.TokenType == JsonTokenType.StartObject)
+			ConfigObject obj = new ConfigObject();
+			if (reader.TokenType == JsonTokenType.StartObject)
 			{
-				while (Reader.Read() && Reader.TokenType == JsonTokenType.PropertyName)
+				while (reader.Read() && reader.TokenType == JsonTokenType.PropertyName)
 				{
-					string Name = Reader.GetString();
-					Reader.Read();
-					string Value = Reader.GetString();
-					Object.Pairs.Add(KeyValuePair.Create(Name, Value));
+					string name = reader.GetString();
+					reader.Read();
+					string value = reader.GetString();
+					obj.Pairs.Add(KeyValuePair.Create(name, value));
 				}
 			}
 			else
 			{
-				Reader.Skip();
+				reader.Skip();
 			}
-			return Object;
+			return obj;
 		}
 
-		public override void Write(Utf8JsonWriter Writer, ConfigObject Value, JsonSerializerOptions Options)
+		public override void Write(Utf8JsonWriter writer, ConfigObject value, JsonSerializerOptions options)
 		{
-			Writer.WriteStartObject();
-			foreach (KeyValuePair<string, string> Pair in Value.Pairs)
+			writer.WriteStartObject();
+			foreach (KeyValuePair<string, string> pair in value.Pairs)
 			{
-				Writer.WriteString(Pair.Key, Pair.Value);
+				writer.WriteString(pair.Key, pair.Value);
 			}
-			Writer.WriteEndObject();
+			writer.WriteEndObject();
 		}
 	}
 
@@ -340,9 +340,9 @@ namespace UnrealGameSync
 		public string Name;
 		public Dictionary<string, string> Pairs = new Dictionary<string,string>();
 
-		public ConfigSection(string InName)
+		public ConfigSection(string inName)
 		{
-			Name = InName;
+			Name = inName;
 		}
 
 		public void Clear()
@@ -350,205 +350,205 @@ namespace UnrealGameSync
 			Pairs.Clear();
 		}
 
-		public void SetValue(string Key, int Value)
+		public void SetValue(string key, int value)
 		{
-			Pairs[Key] = Value.ToString();
+			Pairs[key] = value.ToString();
 		}
 
-		public void SetValue(string Key, bool Value)
+		public void SetValue(string key, bool value)
 		{
-			Pairs[Key] = Value? "1" : "0";
+			Pairs[key] = value? "1" : "0";
 		}
 
-		public void SetValue(string Key, string Value)
+		public void SetValue(string key, string value)
 		{
-			if(Value == null)
+			if(value == null)
 			{
-				RemoveValue(Key);
+				RemoveValue(key);
 			}
 			else
 			{
-				Pairs[Key] = Value;
+				Pairs[key] = value;
 			}
 		}
 
-		public void SetValues(string Key, string[] Values)
+		public void SetValues(string key, string[] values)
 		{
-			if(Values == null || Values.Length == 0)
+			if(values == null || values.Length == 0)
 			{
-				RemoveValue(Key);
+				RemoveValue(key);
 			}
 			else
 			{
-				Pairs[Key] = String.Join("\n", Values);
+				Pairs[key] = String.Join("\n", values);
 			}
 		}
 
-		public void SetValues(string Key, Guid[] Values)
+		public void SetValues(string key, Guid[] values)
 		{
-			if(Values == null)
+			if(values == null)
 			{
-				RemoveValue(Key);
+				RemoveValue(key);
 			}
 			else
 			{
-				Pairs[Key] = String.Join("\n", Values.Select(x => x.ToString()));
+				Pairs[key] = String.Join("\n", values.Select(x => x.ToString()));
 			}
 		}
 
-		public void AppendValue(string Key, string Value)
+		public void AppendValue(string key, string value)
 		{
-			string? CurrentValue;
-			if(Pairs.TryGetValue(Key, out CurrentValue))
+			string? currentValue;
+			if(Pairs.TryGetValue(key, out currentValue))
 			{
-				Pairs[Key] = CurrentValue + "\n" + Value;
+				Pairs[key] = currentValue + "\n" + value;
 			}
 			else
 			{
-				Pairs[Key] = Value;
+				Pairs[key] = value;
 			}
 		}
 
-		public void RemoveValue(string Key)
+		public void RemoveValue(string key)
 		{
-			Pairs.Remove(Key);
+			Pairs.Remove(key);
 		}
 
-		public int GetValue(string Key, int DefaultValue)
+		public int GetValue(string key, int defaultValue)
 		{
-			string? ValueString = GetValue(Key);
-			if(ValueString != null)
+			string? valueString = GetValue(key);
+			if(valueString != null)
 			{
-				int Value;
-				if(int.TryParse(ValueString, out Value))
+				int value;
+				if(int.TryParse(valueString, out value))
 				{
-					return Value;
+					return value;
 				}
 			}
-			return DefaultValue;
+			return defaultValue;
 		}
 
-		public bool GetValue(string Key, bool DefaultValue)
+		public bool GetValue(string key, bool defaultValue)
 		{
-			return GetValue(Key, DefaultValue? 1 : 0) != 0;
+			return GetValue(key, defaultValue? 1 : 0) != 0;
 		}
 
-		[return: NotNullIfNotNull("DefaultValue")]
-		public string? GetValue(string Key, string? DefaultValue = null)
+		[return: NotNullIfNotNull("defaultValue")]
+		public string? GetValue(string key, string? defaultValue = null)
 		{
-			string? Value;
-			if(!Pairs.TryGetValue(Key, out Value))
+			string? value;
+			if(!Pairs.TryGetValue(key, out value))
 			{
-				Value = DefaultValue;
+				value = defaultValue;
 			}
-			return Value;
+			return value;
 		}
 
-		[return: NotNullIfNotNull("DefaultValue")]
-		public string[]? GetValues(string Key, string[]? DefaultValue = null)
+		[return: NotNullIfNotNull("defaultValue")]
+		public string[]? GetValues(string key, string[]? defaultValue = null)
 		{
-			string? Value = GetValue(Key, null);
-			if(Value == null)
+			string? value = GetValue(key, null);
+			if(value == null)
 			{
-				return DefaultValue;
+				return defaultValue;
 			}
 			else
 			{
-				return Value.Split('\n');
+				return value.Split('\n');
 			}
 		}
 
-		[return: NotNullIfNotNull("DefaultValue")]
-		public Guid[]? GetValues(string Key, Guid[]? DefaultValue = null)
+		[return: NotNullIfNotNull("defaultValue")]
+		public Guid[]? GetValues(string key, Guid[]? defaultValue = null)
 		{
-			string[]? StringValues = GetValues(Key, (string[]?)null);
-			if(StringValues == null)
+			string[]? stringValues = GetValues(key, (string[]?)null);
+			if(stringValues == null)
 			{
-				return DefaultValue;
+				return defaultValue;
 			}
 
-			List<Guid> GuidValues = new List<Guid>();
-			foreach(string StringValue in StringValues)
+			List<Guid> guidValues = new List<Guid>();
+			foreach(string stringValue in stringValues)
 			{
-				Guid GuidValue;
-				if(Guid.TryParse(StringValue, out GuidValue))
+				Guid guidValue;
+				if(Guid.TryParse(stringValue, out guidValue))
 				{
-					GuidValues.Add(GuidValue);
+					guidValues.Add(guidValue);
 				}
 			}
 
-			return GuidValues.ToArray();
+			return guidValues.ToArray();
 		}
 	}
 
 	public class ConfigFile
 	{
-		List<ConfigSection> Sections = new List<ConfigSection>();
+		List<ConfigSection> _sections = new List<ConfigSection>();
 
 		public ConfigFile()
 		{
 		}
 
-		FileReference GetTempFileName(FileReference FileName)
+		FileReference GetTempFileName(FileReference fileName)
 		{
-			return FileName + ".tmp";
+			return fileName + ".tmp";
 		}
 
-		public void Load(FileReference FileName)
+		public void Load(FileReference fileName)
 		{
-			Parse(FileReference.ReadAllLines(FileName));
+			Parse(FileReference.ReadAllLines(fileName));
 		}
 
-		public bool TryLoad(FileReference FileName, ILogger Logger)
+		public bool TryLoad(FileReference fileName, ILogger logger)
 		{
-			FileInfo FileInfo = FileName.ToFileInfo();
-			if (FileInfo.Exists)
+			FileInfo fileInfo = fileName.ToFileInfo();
+			if (fileInfo.Exists)
 			{
-				Logger.LogDebug("Loading config file from {File} ({Size} bytes)", FileInfo.FullName, FileInfo.Length);
-				Load(FileName);
+				logger.LogDebug("Loading config file from {File} ({Size} bytes)", fileInfo.FullName, fileInfo.Length);
+				Load(fileName);
 				return true;
 			}
 
-			FileReference TempFileName = GetTempFileName(FileName);
+			FileReference tempFileName = GetTempFileName(fileName);
 
-			FileInfo TempFileInfo = TempFileName.ToFileInfo();
-			if (TempFileInfo.Exists)
+			FileInfo tempFileInfo = tempFileName.ToFileInfo();
+			if (tempFileInfo.Exists)
 			{
-				Logger.LogDebug("Loading temporary config file from {File} ({Size} bytes)", TempFileInfo.FullName, TempFileInfo.Length);
-				Load(TempFileName);
+				logger.LogDebug("Loading temporary config file from {File} ({Size} bytes)", tempFileInfo.FullName, tempFileInfo.Length);
+				Load(tempFileName);
 				return true;
 			}
 
-			Logger.LogDebug("No existing config file at {File}", FileName);
+			logger.LogDebug("No existing config file at {File}", fileName);
 			return false;
 		}
 
-		public void Parse(string[] Lines)
+		public void Parse(string[] lines)
 		{
-			ConfigSection? CurrentSection = null;
-			foreach(string Line in Lines)
+			ConfigSection? currentSection = null;
+			foreach(string line in lines)
 			{
-				string TrimLine = Line.Trim();
-				if(!TrimLine.StartsWith(";"))
+				string trimLine = line.Trim();
+				if(!trimLine.StartsWith(";"))
 				{
-					if(TrimLine.StartsWith("[") && TrimLine.EndsWith("]"))
+					if(trimLine.StartsWith("[") && trimLine.EndsWith("]"))
 					{
-						string SectionName = TrimLine.Substring(1, TrimLine.Length - 2).Trim();
-						CurrentSection = FindOrAddSection(SectionName);
+						string sectionName = trimLine.Substring(1, trimLine.Length - 2).Trim();
+						currentSection = FindOrAddSection(sectionName);
 					}
-					else if(CurrentSection != null)
+					else if(currentSection != null)
 					{
-						int EqualsIdx = TrimLine.IndexOf('=');
-						if(EqualsIdx != -1)
+						int equalsIdx = trimLine.IndexOf('=');
+						if(equalsIdx != -1)
 						{
-							string Value = Line.Substring(EqualsIdx + 1).TrimStart();
-							if(TrimLine.StartsWith("+"))
+							string value = line.Substring(equalsIdx + 1).TrimStart();
+							if(trimLine.StartsWith("+"))
 							{
-								CurrentSection.AppendValue(TrimLine.Substring(1, EqualsIdx - 1).Trim(), Value);
+								currentSection.AppendValue(trimLine.Substring(1, equalsIdx - 1).Trim(), value);
 							}
 							else
 							{
-								CurrentSection.SetValue(TrimLine.Substring(0, EqualsIdx).TrimEnd(), Value);
+								currentSection.SetValue(trimLine.Substring(0, equalsIdx).TrimEnd(), value);
 							}
 						}
 					}
@@ -556,121 +556,121 @@ namespace UnrealGameSync
 			}
 		}
 
-		public void Save(FileReference FileName)
+		public void Save(FileReference fileName)
 		{
-			FileReference TempFileName = GetTempFileName(FileName);
-			FileReference.Delete(TempFileName);
+			FileReference tempFileName = GetTempFileName(fileName);
+			FileReference.Delete(tempFileName);
 
-			using (StreamWriter Writer = new StreamWriter(TempFileName.FullName))
+			using (StreamWriter writer = new StreamWriter(tempFileName.FullName))
 			{
-				for(int Idx = 0; Idx < Sections.Count; Idx++)
+				for(int idx = 0; idx < _sections.Count; idx++)
 				{
-					Writer.WriteLine("[{0}]", Sections[Idx].Name);
-					foreach(KeyValuePair<string, string> Pair in Sections[Idx].Pairs)
+					writer.WriteLine("[{0}]", _sections[idx].Name);
+					foreach(KeyValuePair<string, string> pair in _sections[idx].Pairs)
 					{
-						if(Pair.Value.Contains('\n'))
+						if(pair.Value.Contains('\n'))
 						{
-							foreach(string Line in Pair.Value.Split('\n'))
+							foreach(string line in pair.Value.Split('\n'))
 							{
-								Writer.WriteLine("+{0}={1}", Pair.Key, Line);
+								writer.WriteLine("+{0}={1}", pair.Key, line);
 							}
 						}
 						else
 						{
-							Writer.WriteLine("{0}={1}", Pair.Key, Pair.Value);
+							writer.WriteLine("{0}={1}", pair.Key, pair.Value);
 						}
 					}
-					if(Idx < Sections.Count - 1)
+					if(idx < _sections.Count - 1)
 					{
-						Writer.WriteLine();
+						writer.WriteLine();
 					}
 				}
 			}
 
-			FileReference.Delete(FileName);
-			FileReference.Move(TempFileName, FileName);
+			FileReference.Delete(fileName);
+			FileReference.Move(tempFileName, fileName);
 		}
 
-		public ConfigSection FindSection(string Name)
+		public ConfigSection FindSection(string name)
 		{
-			return Sections.FirstOrDefault(x => String.Compare(x.Name, Name, true) == 0);
+			return _sections.FirstOrDefault(x => String.Compare(x.Name, name, true) == 0);
 		}
 
-		public ConfigSection FindOrAddSection(string Name)
+		public ConfigSection FindOrAddSection(string name)
 		{
-			ConfigSection Section = FindSection(Name);
-			if(Section == null)
+			ConfigSection section = FindSection(name);
+			if(section == null)
 			{
-				Section = new ConfigSection(Name);
-				Sections.Add(Section);
+				section = new ConfigSection(name);
+				_sections.Add(section);
 			}
-			return Section;
+			return section;
 		}
 
-		public void SetValue(string Key, int Value)
+		public void SetValue(string key, int value)
 		{
-			int DotIdx = Key.IndexOf('.');
-			ConfigSection Section = FindOrAddSection(Key.Substring(0, DotIdx));
-			Section.SetValue(Key.Substring(DotIdx + 1), Value);
+			int dotIdx = key.IndexOf('.');
+			ConfigSection section = FindOrAddSection(key.Substring(0, dotIdx));
+			section.SetValue(key.Substring(dotIdx + 1), value);
 		}
 
-		public void SetValue(string Key, bool Value)
+		public void SetValue(string key, bool value)
 		{
-			int DotIdx = Key.IndexOf('.');
-			ConfigSection Section = FindOrAddSection(Key.Substring(0, DotIdx));
-			Section.SetValue(Key.Substring(DotIdx + 1), Value);
+			int dotIdx = key.IndexOf('.');
+			ConfigSection section = FindOrAddSection(key.Substring(0, dotIdx));
+			section.SetValue(key.Substring(dotIdx + 1), value);
 		}
 
-		public void SetValue(string Key, string Value)
+		public void SetValue(string key, string value)
 		{
-			int DotIdx = Key.IndexOf('.');
-			ConfigSection Section = FindOrAddSection(Key.Substring(0, DotIdx));
-			Section.SetValue(Key.Substring(DotIdx + 1), Value);
+			int dotIdx = key.IndexOf('.');
+			ConfigSection section = FindOrAddSection(key.Substring(0, dotIdx));
+			section.SetValue(key.Substring(dotIdx + 1), value);
 		}
 
-		public void SetValues(string Key, string[] Values)
+		public void SetValues(string key, string[] values)
 		{
-			int DotIdx = Key.IndexOf('.');
-			ConfigSection Section = FindOrAddSection(Key.Substring(0, DotIdx));
-			Section.SetValues(Key.Substring(DotIdx + 1), Values);
+			int dotIdx = key.IndexOf('.');
+			ConfigSection section = FindOrAddSection(key.Substring(0, dotIdx));
+			section.SetValues(key.Substring(dotIdx + 1), values);
 		}
 
-		public bool GetValue(string Key, bool DefaultValue)
+		public bool GetValue(string key, bool defaultValue)
 		{
-			int DotIdx = Key.IndexOf('.');
-			ConfigSection Section = FindSection(Key.Substring(0, DotIdx));
-			return (Section == null)? DefaultValue : Section.GetValue(Key.Substring(DotIdx + 1), DefaultValue);
+			int dotIdx = key.IndexOf('.');
+			ConfigSection section = FindSection(key.Substring(0, dotIdx));
+			return (section == null)? defaultValue : section.GetValue(key.Substring(dotIdx + 1), defaultValue);
 		}
 
-		public int GetValue(string Key, int DefaultValue)
+		public int GetValue(string key, int defaultValue)
 		{
-			int DotIdx = Key.IndexOf('.');
-			ConfigSection Section = FindSection(Key.Substring(0, DotIdx));
-			return (Section == null)? DefaultValue : Section.GetValue(Key.Substring(DotIdx + 1), DefaultValue);
+			int dotIdx = key.IndexOf('.');
+			ConfigSection section = FindSection(key.Substring(0, dotIdx));
+			return (section == null)? defaultValue : section.GetValue(key.Substring(dotIdx + 1), defaultValue);
 		}
 
-		[return: NotNullIfNotNull("DefaultValue")]
-		public string? GetValue(string Key, string? DefaultValue)
+		[return: NotNullIfNotNull("defaultValue")]
+		public string? GetValue(string key, string? defaultValue)
 		{
-			int DotIdx = Key.IndexOf('.');
-			ConfigSection Section = FindSection(Key.Substring(0, DotIdx));
-			return (Section == null)? DefaultValue : Section.GetValue(Key.Substring(DotIdx + 1), DefaultValue);
+			int dotIdx = key.IndexOf('.');
+			ConfigSection section = FindSection(key.Substring(0, dotIdx));
+			return (section == null)? defaultValue : section.GetValue(key.Substring(dotIdx + 1), defaultValue);
 		}
 
-		[return: NotNullIfNotNull("DefaultValue")]
-		public string[]? GetValues(string Key, string[]? DefaultValue)
+		[return: NotNullIfNotNull("defaultValue")]
+		public string[]? GetValues(string key, string[]? defaultValue)
 		{
-			int DotIdx = Key.IndexOf('.');
-			ConfigSection Section = FindSection(Key.Substring(0, DotIdx));
-			return (Section == null)? DefaultValue : Section.GetValues(Key.Substring(DotIdx + 1), DefaultValue);
+			int dotIdx = key.IndexOf('.');
+			ConfigSection section = FindSection(key.Substring(0, dotIdx));
+			return (section == null)? defaultValue : section.GetValues(key.Substring(dotIdx + 1), defaultValue);
 		}
 
-		[return: NotNullIfNotNull("DefaultValue")]
-		public Guid[]? GetGuidValues(string Key, Guid[]? DefaultValue)
+		[return: NotNullIfNotNull("defaultValue")]
+		public Guid[]? GetGuidValues(string key, Guid[]? defaultValue)
 		{
-			int DotIdx = Key.IndexOf('.');
-			ConfigSection Section = FindSection(Key.Substring(0, DotIdx));
-			return (Section == null)? DefaultValue : Section.GetValues(Key.Substring(DotIdx + 1), DefaultValue);
+			int dotIdx = key.IndexOf('.');
+			ConfigSection section = FindSection(key.Substring(0, dotIdx));
+			return (section == null)? defaultValue : section.GetValues(key.Substring(dotIdx + 1), defaultValue);
 		}
 	}
 }
diff --git a/Engine/Source/Programs/UnrealGameSync/UnrealGameSyncShared/ConfigUtils.cs b/Engine/Source/Programs/UnrealGameSync/UnrealGameSyncShared/ConfigUtils.cs
--- a/Engine/Source/Programs/UnrealGameSync/UnrealGameSyncShared/ConfigUtils.cs
+++ b/Engine/Source/Programs/UnrealGameSync/UnrealGameSyncShared/ConfigUtils.cs
@@ -22,33 +22,33 @@ namespace UnrealGameSync
 		public string? Launch { get; set; }
 		public string? LaunchCmd { get; set; }
 
-		public static bool TryRead(FileReference Location, DirectoryReference? EngineDir, DirectoryReference? ProjectDir, [NotNullWhen(true)] out TargetReceipt? Receipt)
+		public static bool TryRead(FileReference location, DirectoryReference? engineDir, DirectoryReference? projectDir, [NotNullWhen(true)] out TargetReceipt? receipt)
 		{
-			if (Utility.TryLoadJson(Location, out Receipt))
+			if (Utility.TryLoadJson(location, out receipt))
 			{
-				Receipt.Launch = ExpandReceiptVariables(Receipt.Launch, EngineDir, ProjectDir);
-				Receipt.LaunchCmd = ExpandReceiptVariables(Receipt.LaunchCmd, EngineDir, ProjectDir);
+				receipt.Launch = ExpandReceiptVariables(receipt.Launch, engineDir, projectDir);
+				receipt.LaunchCmd = ExpandReceiptVariables(receipt.LaunchCmd, engineDir, projectDir);
 				return true;
 			}
 			return false;
 		}
 
-		[return: NotNullIfNotNull("Line")]
-		private static string? ExpandReceiptVariables(string? Line, DirectoryReference? EngineDir, DirectoryReference? ProjectDir)
+		[return: NotNullIfNotNull("line")]
+		private static string? ExpandReceiptVariables(string? line, DirectoryReference? engineDir, DirectoryReference? projectDir)
 		{
-			string? ExpandedLine = Line;
-			if (ExpandedLine != null)
+			string? expandedLine = line;
+			if (expandedLine != null)
 			{
-				if (EngineDir != null)
+				if (engineDir != null)
 				{
-					ExpandedLine = ExpandedLine.Replace("$(EngineDir)", EngineDir.FullName);
+					expandedLine = expandedLine.Replace("$(EngineDir)", engineDir.FullName);
 				}
-				if (ProjectDir != null)
+				if (projectDir != null)
 				{
-					ExpandedLine = ExpandedLine.Replace("$(ProjectDir)", ProjectDir.FullName);
+					expandedLine = expandedLine.Replace("$(ProjectDir)", projectDir.FullName);
 				}
 			}
-			return ExpandedLine;
+			return expandedLine;
 		}
 	}
 
@@ -74,381 +74,381 @@ namespace UnrealGameSync
 			}
 		}
 
-		public static Task<ConfigFile> ReadProjectConfigFileAsync(IPerforceConnection Perforce, ProjectInfo ProjectInfo, ILogger Logger, CancellationToken CancellationToken)
+		public static Task<ConfigFile> ReadProjectConfigFileAsync(IPerforceConnection perforce, ProjectInfo projectInfo, ILogger logger, CancellationToken cancellationToken)
 		{
-			return ReadProjectConfigFileAsync(Perforce, ProjectInfo, new List<KeyValuePair<FileReference, DateTime>>(), Logger, CancellationToken);
+			return ReadProjectConfigFileAsync(perforce, projectInfo, new List<KeyValuePair<FileReference, DateTime>>(), logger, cancellationToken);
 		}
 
-		public static Task<ConfigFile> ReadProjectConfigFileAsync(IPerforceConnection Perforce, ProjectInfo ProjectInfo, List<KeyValuePair<FileReference, DateTime>> LocalConfigFiles, ILogger Logger, CancellationToken CancellationToken)
+		public static Task<ConfigFile> ReadProjectConfigFileAsync(IPerforceConnection perforce, ProjectInfo projectInfo, List<KeyValuePair<FileReference, DateTime>> localConfigFiles, ILogger logger, CancellationToken cancellationToken)
 		{
-			return ReadProjectConfigFileAsync(Perforce, ProjectInfo.ClientRootPath, ProjectInfo.ClientFileName, ProjectInfo.CacheFolder, LocalConfigFiles, Logger, CancellationToken);
+			return ReadProjectConfigFileAsync(perforce, projectInfo.ClientRootPath, projectInfo.ClientFileName, projectInfo.CacheFolder, localConfigFiles, logger, cancellationToken);
 		}
 
-		public static async Task<ConfigFile> ReadProjectConfigFileAsync(IPerforceConnection Perforce, string BranchClientPath, string SelectedClientFileName, DirectoryReference CacheFolder, List<KeyValuePair<FileReference, DateTime>> LocalConfigFiles, ILogger Logger, CancellationToken CancellationToken)
+		public static async Task<ConfigFile> ReadProjectConfigFileAsync(IPerforceConnection perforce, string branchClientPath, string selectedClientFileName, DirectoryReference cacheFolder, List<KeyValuePair<FileReference, DateTime>> localConfigFiles, ILogger logger, CancellationToken cancellationToken)
 		{
-			List<string> ConfigFilePaths = Utility.GetDepotConfigPaths(BranchClientPath + "/Engine", SelectedClientFileName);
+			List<string> configFilePaths = Utility.GetDepotConfigPaths(branchClientPath + "/Engine", selectedClientFileName);
 
-			ConfigFile ProjectConfig = new ConfigFile();
+			ConfigFile projectConfig = new ConfigFile();
 
-			List<PerforceResponse<FStatRecord>> Responses = await Perforce.TryFStatAsync(FStatOptions.IncludeFileSizes, ConfigFilePaths, CancellationToken).ToListAsync(CancellationToken);
-			foreach (PerforceResponse<FStatRecord> Response in Responses)
+			List<PerforceResponse<FStatRecord>> responses = await perforce.TryFStatAsync(FStatOptions.IncludeFileSizes, configFilePaths, cancellationToken).ToListAsync(cancellationToken);
+			foreach (PerforceResponse<FStatRecord> response in responses)
 			{
-				if (Response.Succeeded)
+				if (response.Succeeded)
 				{
-					string[]? Lines = null;
+					string[]? lines = null;
 
 					// Skip file records which are still in the workspace, but were synced from a different branch. For these files, the action seems to be empty, so filter against that.
-					FStatRecord FileRecord = Response.Data;
-					if (FileRecord.HeadAction == FileAction.None)
+					FStatRecord fileRecord = response.Data;
+					if (fileRecord.HeadAction == FileAction.None)
 					{
 						continue;
 					}
 
 					// If this file is open for edit, read the local version
-					string? LocalFileName = FileRecord.ClientFile;
-					if (LocalFileName != null && File.Exists(LocalFileName) && (File.GetAttributes(LocalFileName) & FileAttributes.ReadOnly) == 0)
+					string? localFileName = fileRecord.ClientFile;
+					if (localFileName != null && File.Exists(localFileName) && (File.GetAttributes(localFileName) & FileAttributes.ReadOnly) == 0)
 					{
 						try
 						{
-							DateTime LastModifiedTime = File.GetLastWriteTimeUtc(LocalFileName);
-							LocalConfigFiles.Add(new KeyValuePair<FileReference, DateTime>(new FileReference(LocalFileName), LastModifiedTime));
-							Lines = await File.ReadAllLinesAsync(LocalFileName, CancellationToken);
+							DateTime lastModifiedTime = File.GetLastWriteTimeUtc(localFileName);
+							localConfigFiles.Add(new KeyValuePair<FileReference, DateTime>(new FileReference(localFileName), lastModifiedTime));
+							lines = await File.ReadAllLinesAsync(localFileName, cancellationToken);
 						}
-						catch (Exception Ex)
+						catch (Exception ex)
 						{
-							Logger.LogInformation(Ex, "Failed to read local config file for {Path}", LocalFileName);
+							logger.LogInformation(ex, "Failed to read local config file for {Path}", localFileName);
 						}
 					}
 
 					// Otherwise try to get it from perforce
-					if (Lines == null && FileRecord.DepotFile != null)
+					if (lines == null && fileRecord.DepotFile != null)
 					{
-						Lines = await Utility.TryPrintFileUsingCacheAsync(Perforce, FileRecord.DepotFile, CacheFolder, FileRecord.Digest, Logger, CancellationToken);
+						lines = await Utility.TryPrintFileUsingCacheAsync(perforce, fileRecord.DepotFile, cacheFolder, fileRecord.Digest, logger, cancellationToken);
 					}
 
 					// Merge the text with the config file
-					if (Lines != null)
+					if (lines != null)
 					{
 						try
 						{
-							ProjectConfig.Parse(Lines.ToArray());
-							Logger.LogDebug("Read config file from {DepotFile}", FileRecord.DepotFile);
+							projectConfig.Parse(lines.ToArray());
+							logger.LogDebug("Read config file from {DepotFile}", fileRecord.DepotFile);
 						}
-						catch (Exception Ex)
+						catch (Exception ex)
 						{
-							Logger.LogInformation(Ex, "Failed to read config file from {DepotFile}", FileRecord.DepotFile);
+							logger.LogInformation(ex, "Failed to read config file from {DepotFile}", fileRecord.DepotFile);
 						}
 					}
 				}
 			}
-			return ProjectConfig;
+			return projectConfig;
 		}
 
-		public static FileReference GetEditorTargetFile(ProjectInfo ProjectInfo, ConfigFile ProjectConfig)
+		public static FileReference GetEditorTargetFile(ProjectInfo projectInfo, ConfigFile projectConfig)
 		{
-			if (ProjectInfo.ProjectPath.EndsWith(".uproject", StringComparison.OrdinalIgnoreCase))
+			if (projectInfo.ProjectPath.EndsWith(".uproject", StringComparison.OrdinalIgnoreCase))
 			{
-				List<FileReference> TargetFiles = FindTargets(ProjectInfo.LocalFileName.Directory);
+				List<FileReference> targetFiles = FindTargets(projectInfo.LocalFileName.Directory);
 
-				FileReference? TargetFile = TargetFiles.OrderBy(x => x.FullName, StringComparer.OrdinalIgnoreCase).FirstOrDefault(x => x.FullName.EndsWith("Editor.target.cs", StringComparison.OrdinalIgnoreCase));
-				if (TargetFile != null)
+				FileReference? targetFile = targetFiles.OrderBy(x => x.FullName, StringComparer.OrdinalIgnoreCase).FirstOrDefault(x => x.FullName.EndsWith("Editor.target.cs", StringComparison.OrdinalIgnoreCase));
+				if (targetFile != null)
 				{
-					return TargetFile;
+					return targetFile;
 				}
 			}
 
-			string DefaultEditorTargetName = GetDefaultEditorTargetName(ProjectInfo, ProjectConfig);
-			return FileReference.Combine(ProjectInfo.LocalRootPath, "Engine", "Source", $"{DefaultEditorTargetName}.Target.cs");
+			string defaultEditorTargetName = GetDefaultEditorTargetName(projectInfo, projectConfig);
+			return FileReference.Combine(projectInfo.LocalRootPath, "Engine", "Source", $"{defaultEditorTargetName}.Target.cs");
 		}
 
-		public static FileReference GetEditorReceiptFile(ProjectInfo ProjectInfo, ConfigFile ProjectConfig, BuildConfig Config)
+		public static FileReference GetEditorReceiptFile(ProjectInfo projectInfo, ConfigFile projectConfig, BuildConfig config)
 		{
-			FileReference TargetFile = GetEditorTargetFile(ProjectInfo, ProjectConfig);
-			return GetReceiptFile(ProjectInfo, TargetFile, Config.ToString());
+			FileReference targetFile = GetEditorTargetFile(projectInfo, projectConfig);
+			return GetReceiptFile(projectInfo, targetFile, config.ToString());
 		}
 
-		private static List<FileReference> FindTargets(DirectoryReference EngineOrProjectDir)
+		private static List<FileReference> FindTargets(DirectoryReference engineOrProjectDir)
 		{
-			List<FileReference> Targets = new List<FileReference>();
+			List<FileReference> targets = new List<FileReference>();
 
-			DirectoryReference SourceDir = DirectoryReference.Combine(EngineOrProjectDir, "Source");
-			if (DirectoryReference.Exists(SourceDir))
+			DirectoryReference sourceDir = DirectoryReference.Combine(engineOrProjectDir, "Source");
+			if (DirectoryReference.Exists(sourceDir))
 			{
-				foreach (FileReference TargetFile in DirectoryReference.EnumerateFiles(SourceDir))
+				foreach (FileReference targetFile in DirectoryReference.EnumerateFiles(sourceDir))
 				{
-					const string Extension = ".target.cs";
-					if (TargetFile.FullName.EndsWith(Extension, StringComparison.OrdinalIgnoreCase))
+					const string extension = ".target.cs";
+					if (targetFile.FullName.EndsWith(extension, StringComparison.OrdinalIgnoreCase))
 					{
-						Targets.Add(TargetFile);
+						targets.Add(targetFile);
 					}
 				}
 			}
 
-			return Targets;
+			return targets;
 		}
 
-		public static string GetDefaultEditorTargetName(ProjectInfo ProjectInfo, ConfigFile ProjectConfigFile)
+		public static string GetDefaultEditorTargetName(ProjectInfo projectInfo, ConfigFile projectConfigFile)
 		{
-			string? EditorTarget;
-			if (!TryGetProjectSetting(ProjectConfigFile, ProjectInfo.ProjectIdentifier, "EditorTarget", out EditorTarget))
+			string? editorTarget;
+			if (!TryGetProjectSetting(projectConfigFile, projectInfo.ProjectIdentifier, "EditorTarget", out editorTarget))
 			{
-				if (ProjectInfo.bIsEnterpriseProject)
+				if (projectInfo.IsEnterpriseProject)
 				{
-					EditorTarget = "StudioEditor";
+					editorTarget = "StudioEditor";
 				}
 				else
 				{
-					EditorTarget = "UE4Editor";
+					editorTarget = "UE4Editor";
 				}
 			}
-			return EditorTarget;
+			return editorTarget;
 		}
 
-		public static bool TryReadEditorReceipt(ProjectInfo ProjectInfo, FileReference ReceiptFile, [NotNullWhen(true)] out TargetReceipt? Receipt)
+		public static bool TryReadEditorReceipt(ProjectInfo projectInfo, FileReference receiptFile, [NotNullWhen(true)] out TargetReceipt? receipt)
 		{
-			DirectoryReference EngineDir = DirectoryReference.Combine(ProjectInfo.LocalRootPath, "Engine");
-			DirectoryReference ProjectDir = ProjectInfo.LocalFileName.Directory;
+			DirectoryReference engineDir = DirectoryReference.Combine(projectInfo.LocalRootPath, "Engine");
+			DirectoryReference projectDir = projectInfo.LocalFileName.Directory;
 
-			if (ReceiptFile.IsUnderDirectory(ProjectDir))
+			if (receiptFile.IsUnderDirectory(projectDir))
 			{
-				return TargetReceipt.TryRead(ReceiptFile, EngineDir, ProjectDir, out Receipt);
+				return TargetReceipt.TryRead(receiptFile, engineDir, projectDir, out receipt);
 			}
 			else
 			{
-				return TargetReceipt.TryRead(ReceiptFile, EngineDir, null, out Receipt);
+				return TargetReceipt.TryRead(receiptFile, engineDir, null, out receipt);
 			}
 		}
 
-		public static TargetReceipt CreateDefaultEditorReceipt(ProjectInfo ProjectInfo, ConfigFile ProjectConfigFile, BuildConfig Configuration)
+		public static TargetReceipt CreateDefaultEditorReceipt(ProjectInfo projectInfo, ConfigFile projectConfigFile, BuildConfig configuration)
 		{
-			string BaseName = GetDefaultEditorTargetName(ProjectInfo, ProjectConfigFile);
-			if (Configuration != BuildConfig.Development || !String.IsNullOrEmpty(HostArchitectureSuffix))
+			string baseName = GetDefaultEditorTargetName(projectInfo, projectConfigFile);
+			if (configuration != BuildConfig.Development || !String.IsNullOrEmpty(HostArchitectureSuffix))
 			{
-				if (Configuration != BuildConfig.DebugGame || ProjectConfigFile.GetValue("Options.DebugGameHasSeparateExecutable", false))
+				if (configuration != BuildConfig.DebugGame || projectConfigFile.GetValue("Options.DebugGameHasSeparateExecutable", false))
 				{
-					BaseName += $"-{HostPlatform}-{Configuration}{HostArchitectureSuffix}";
+					baseName += $"-{HostPlatform}-{configuration}{HostArchitectureSuffix}";
 				}
 			}
 
-			string Extension = String.Empty;
+			string extension = String.Empty;
 			if (RuntimeInformation.IsOSPlatform(OSPlatform.Windows))
 			{
-				Extension = ".exe";
+				extension = ".exe";
 			}
 
-			TargetReceipt Receipt = new TargetReceipt();
-			Receipt.Configuration = Configuration.ToString();
-			Receipt.Launch = FileReference.Combine(ProjectInfo.LocalRootPath, "Engine", "Binaries", HostPlatform, $"{BaseName}{Extension}").FullName;
-			Receipt.LaunchCmd = FileReference.Combine(ProjectInfo.LocalRootPath, "Engine", "Binaries", HostPlatform, $"{BaseName}-Cmd{Extension}").FullName;
-			return Receipt;
+			TargetReceipt receipt = new TargetReceipt();
+			receipt.Configuration = configuration.ToString();
+			receipt.Launch = FileReference.Combine(projectInfo.LocalRootPath, "Engine", "Binaries", HostPlatform, $"{baseName}{extension}").FullName;
+			receipt.LaunchCmd = FileReference.Combine(projectInfo.LocalRootPath, "Engine", "Binaries", HostPlatform, $"{baseName}-Cmd{extension}").FullName;
+			return receipt;
 		}
 
-		public static FileReference GetReceiptFile(ProjectInfo ProjectInfo, FileReference TargetFile, string Configuration)
+		public static FileReference GetReceiptFile(ProjectInfo projectInfo, FileReference targetFile, string configuration)
 		{
-			string TargetName = TargetFile.GetFileNameWithoutAnyExtensions();
+			string targetName = targetFile.GetFileNameWithoutAnyExtensions();
 
-			DirectoryReference? ProjectDir = ProjectInfo.ProjectDir;
-			if (ProjectDir != null)
+			DirectoryReference? projectDir = projectInfo.ProjectDir;
+			if (projectDir != null)
 			{
-				return GetReceiptFile(ProjectDir, TargetName, Configuration);
+				return GetReceiptFile(projectDir, targetName, configuration);
 			}
 			else
 			{
-				return GetReceiptFile(ProjectInfo.EngineDir, TargetName, Configuration);
+				return GetReceiptFile(projectInfo.EngineDir, targetName, configuration);
 			}
 		}
 
-		public static FileReference GetReceiptFile(DirectoryReference BaseDir, string TargetName, string Configuration)
+		public static FileReference GetReceiptFile(DirectoryReference baseDir, string targetName, string configuration)
 		{
-			return GetReceiptFile(BaseDir, TargetName, HostPlatform, Configuration, HostArchitectureSuffix);
+			return GetReceiptFile(baseDir, targetName, HostPlatform, configuration, HostArchitectureSuffix);
 		}
 
-		public static FileReference GetReceiptFile(DirectoryReference BaseDir, string TargetName, string Platform, string Configuration, string ArchitectureSuffix)
+		public static FileReference GetReceiptFile(DirectoryReference baseDir, string targetName, string platform, string configuration, string architectureSuffix)
 		{
-			if (String.IsNullOrEmpty(ArchitectureSuffix) && Configuration.Equals("Development", StringComparison.OrdinalIgnoreCase))
+			if (String.IsNullOrEmpty(architectureSuffix) && configuration.Equals("Development", StringComparison.OrdinalIgnoreCase))
 			{
-				return FileReference.Combine(BaseDir, "Binaries", Platform, $"{TargetName}.target");
+				return FileReference.Combine(baseDir, "Binaries", platform, $"{targetName}.target");
 			}
 			else
 			{
-				return FileReference.Combine(BaseDir, "Binaries", Platform, $"{TargetName}-{Platform}-{Configuration}{ArchitectureSuffix}.target");
+				return FileReference.Combine(baseDir, "Binaries", platform, $"{targetName}-{platform}-{configuration}{architectureSuffix}.target");
 			}
 		}
 
-		public static Dictionary<Guid, ConfigObject> GetDefaultBuildStepObjects(ProjectInfo ProjectInfo, string EditorTarget, BuildConfig EditorConfig, ConfigFile LatestProjectConfigFile, bool ShouldSyncPrecompiledEditor)
+		public static Dictionary<Guid, ConfigObject> GetDefaultBuildStepObjects(ProjectInfo projectInfo, string editorTarget, BuildConfig editorConfig, ConfigFile latestProjectConfigFile, bool shouldSyncPrecompiledEditor)
 		{
-			string ProjectArgument = "";
-			if (ProjectInfo.LocalFileName.HasExtension(".uproject"))
+			string projectArgument = "";
+			if (projectInfo.LocalFileName.HasExtension(".uproject"))
 			{
-				ProjectArgument = String.Format("\"{0}\"", ProjectInfo.LocalFileName);
+				projectArgument = String.Format("\"{0}\"", projectInfo.LocalFileName);
 			}
 
-			bool bUseCrashReportClientEditor = LatestProjectConfigFile.GetValue("Options.UseCrashReportClientEditor", false);
+			bool useCrashReportClientEditor = latestProjectConfigFile.GetValue("Options.UseCrashReportClientEditor", false);
 
-			string HostPlatform;
+			string hostPlatform;
 			if(RuntimeInformation.IsOSPlatform(OSPlatform.OSX))
 			{
-				HostPlatform = "Mac";
+				hostPlatform = "Mac";
 			}
 			else if(RuntimeInformation.IsOSPlatform(OSPlatform.Linux))
 			{
-				HostPlatform = "Linux";
+				hostPlatform = "Linux";
 			}
 			else
 			{
-				HostPlatform = "Win64";
+				hostPlatform = "Win64";
 			}
 
-			List<BuildStep> DefaultBuildSteps = new List<BuildStep>();
-			DefaultBuildSteps.Add(new BuildStep(new Guid("{01F66060-73FA-4CC8-9CB3-E217FBBA954E}"), 0, "Compile UnrealHeaderTool", "Compiling UnrealHeaderTool...", 1, "UnrealHeaderTool", HostPlatform, "Development", "", !ShouldSyncPrecompiledEditor));
-			DefaultBuildSteps.Add(new BuildStep(new Guid("{F097FF61-C916-4058-8391-35B46C3173D5}"), 1, $"Compile {EditorTarget}", $"Compiling {EditorTarget}...", 10, EditorTarget, HostPlatform, EditorConfig.ToString(), ProjectArgument, !ShouldSyncPrecompiledEditor));
-			DefaultBuildSteps.Add(new BuildStep(new Guid("{C6E633A1-956F-4AD3-BC95-6D06D131E7B4}"), 2, "Compile ShaderCompileWorker", "Compiling ShaderCompileWorker...", 1, "ShaderCompileWorker", HostPlatform, "Development", "", !ShouldSyncPrecompiledEditor));
-			DefaultBuildSteps.Add(new BuildStep(new Guid("{24FFD88C-7901-4899-9696-AE1066B4B6E8}"), 3, "Compile UnrealLightmass", "Compiling UnrealLightmass...", 1, "UnrealLightmass", HostPlatform, "Development", "", !ShouldSyncPrecompiledEditor));
-			DefaultBuildSteps.Add(new BuildStep(new Guid("{FFF20379-06BF-4205-8A3E-C53427736688}"), 4, "Compile CrashReportClient", "Compiling CrashReportClient...", 1, "CrashReportClient", HostPlatform, "Shipping", "", !ShouldSyncPrecompiledEditor && !bUseCrashReportClientEditor));
-			DefaultBuildSteps.Add(new BuildStep(new Guid("{7143D861-58D3-4F83-BADC-BC5DCB2079F6}"), 5, "Compile CrashReportClientEditor", "Compiling CrashReportClientEditor...", 1, "CrashReportClientEditor", HostPlatform, "Shipping", "", !ShouldSyncPrecompiledEditor && bUseCrashReportClientEditor));
+			List<BuildStep> defaultBuildSteps = new List<BuildStep>();
+			defaultBuildSteps.Add(new BuildStep(new Guid("{01F66060-73FA-4CC8-9CB3-E217FBBA954E}"), 0, "Compile UnrealHeaderTool", "Compiling UnrealHeaderTool...", 1, "UnrealHeaderTool", hostPlatform, "Development", "", !shouldSyncPrecompiledEditor));
+			defaultBuildSteps.Add(new BuildStep(new Guid("{F097FF61-C916-4058-8391-35B46C3173D5}"), 1, $"Compile {editorTarget}", $"Compiling {editorTarget}...", 10, editorTarget, hostPlatform, editorConfig.ToString(), projectArgument, !shouldSyncPrecompiledEditor));
+			defaultBuildSteps.Add(new BuildStep(new Guid("{C6E633A1-956F-4AD3-BC95-6D06D131E7B4}"), 2, "Compile ShaderCompileWorker", "Compiling ShaderCompileWorker...", 1, "ShaderCompileWorker", hostPlatform, "Development", "", !shouldSyncPrecompiledEditor));
+			defaultBuildSteps.Add(new BuildStep(new Guid("{24FFD88C-7901-4899-9696-AE1066B4B6E8}"), 3, "Compile UnrealLightmass", "Compiling UnrealLightmass...", 1, "UnrealLightmass", hostPlatform, "Development", "", !shouldSyncPrecompiledEditor));
+			defaultBuildSteps.Add(new BuildStep(new Guid("{FFF20379-06BF-4205-8A3E-C53427736688}"), 4, "Compile CrashReportClient", "Compiling CrashReportClient...", 1, "CrashReportClient", hostPlatform, "Shipping", "", !shouldSyncPrecompiledEditor && !useCrashReportClientEditor));
+			defaultBuildSteps.Add(new BuildStep(new Guid("{7143D861-58D3-4F83-BADC-BC5DCB2079F6}"), 5, "Compile CrashReportClientEditor", "Compiling CrashReportClientEditor...", 1, "CrashReportClientEditor", hostPlatform, "Shipping", "", !shouldSyncPrecompiledEditor && useCrashReportClientEditor));
 
-			return DefaultBuildSteps.ToDictionary(x => x.UniqueId, x => x.ToConfigObject());
+			return defaultBuildSteps.ToDictionary(x => x.UniqueId, x => x.ToConfigObject());
 		}
 
-		public static Dictionary<string, string> GetWorkspaceVariables(ProjectInfo ProjectInfo, int ChangeNumber, int CodeChangeNumber, TargetReceipt? EditorTarget, ConfigFile? ProjectConfigFile)
+		public static Dictionary<string, string> GetWorkspaceVariables(ProjectInfo projectInfo, int changeNumber, int codeChangeNumber, TargetReceipt? editorTarget, ConfigFile? projectConfigFile)
 		{
-			Dictionary<string, string> Variables = new Dictionary<string, string>(StringComparer.OrdinalIgnoreCase);
+			Dictionary<string, string> variables = new Dictionary<string, string>(StringComparer.OrdinalIgnoreCase);
 
-			if (ProjectInfo.StreamName != null)
+			if (projectInfo.StreamName != null)
 			{
-				Variables.Add("Stream", ProjectInfo.StreamName);
+				variables.Add("Stream", projectInfo.StreamName);
 			}
 
-			Variables.Add("Change", ChangeNumber.ToString());
-			Variables.Add("CodeChange", CodeChangeNumber.ToString());
+			variables.Add("Change", changeNumber.ToString());
+			variables.Add("CodeChange", codeChangeNumber.ToString());
 
-			Variables.Add("ClientName", ProjectInfo.ClientName);
-			Variables.Add("BranchDir", ProjectInfo.LocalRootPath.FullName);
-			Variables.Add("ProjectDir", ProjectInfo.LocalFileName.Directory.FullName);
-			Variables.Add("ProjectFile", ProjectInfo.LocalFileName.FullName);
-			Variables.Add("UseIncrementalBuilds", "1");
+			variables.Add("ClientName", projectInfo.ClientName);
+			variables.Add("BranchDir", projectInfo.LocalRootPath.FullName);
+			variables.Add("ProjectDir", projectInfo.LocalFileName.Directory.FullName);
+			variables.Add("ProjectFile", projectInfo.LocalFileName.FullName);
+			variables.Add("UseIncrementalBuilds", "1");
 
-			string EditorConfig = EditorTarget?.Configuration ?? String.Empty;
-			Variables.Add("EditorConfig", EditorConfig);
+			string editorConfig = editorTarget?.Configuration ?? String.Empty;
+			variables.Add("EditorConfig", editorConfig);
 
-			string EditorLaunch = EditorTarget?.Launch ?? String.Empty;
-			Variables.Add("EditorExe", EditorLaunch);
+			string editorLaunch = editorTarget?.Launch ?? String.Empty;
+			variables.Add("EditorExe", editorLaunch);
 
-			string EditorLaunchCmd = EditorTarget?.LaunchCmd ?? EditorLaunch.Replace(".exe", "-Cmd.exe");
-			Variables.Add("EditorCmdExe", EditorLaunchCmd);
+			string editorLaunchCmd = editorTarget?.LaunchCmd ?? editorLaunch.Replace(".exe", "-Cmd.exe");
+			variables.Add("EditorCmdExe", editorLaunchCmd);
 
 			// Legacy
-			Variables.Add("UE4EditorConfig", EditorConfig);
-			Variables.Add("UE4EditorDebugArg", (EditorConfig.Equals("Debug", StringComparison.Ordinal) || EditorConfig.Equals("DebugGame", StringComparison.Ordinal)) ? " -debug" : "");
-			Variables.Add("UE4EditorExe", EditorLaunch);
-			Variables.Add("UE4EditorCmdExe", EditorLaunchCmd);
+			variables.Add("UE4EditorConfig", editorConfig);
+			variables.Add("UE4EditorDebugArg", (editorConfig.Equals("Debug", StringComparison.Ordinal) || editorConfig.Equals("DebugGame", StringComparison.Ordinal)) ? " -debug" : "");
+			variables.Add("UE4EditorExe", editorLaunch);
+			variables.Add("UE4EditorCmdExe", editorLaunchCmd);
 
-			if (ProjectConfigFile != null)
+			if (projectConfigFile != null)
 			{
-				if (TryGetProjectSetting(ProjectConfigFile, ProjectInfo.ProjectIdentifier, "SdkInstallerDir", out string? SdkInstallerDir))
+				if (TryGetProjectSetting(projectConfigFile, projectInfo.ProjectIdentifier, "SdkInstallerDir", out string? sdkInstallerDir))
 				{
-					Variables.Add("SdkInstallerDir", SdkInstallerDir);
+					variables.Add("SdkInstallerDir", sdkInstallerDir);
 				}
 			}
 
-			return Variables;
+			return variables;
 		}
 
-		public static Dictionary<string, string> GetWorkspaceVariables(ProjectInfo ProjectInfo, int ChangeNumber, int CodeChangeNumber, TargetReceipt? EditorTarget, ConfigFile? ProjectConfigFile, IEnumerable<KeyValuePair<string, string>> AdditionalVariables)
+		public static Dictionary<string, string> GetWorkspaceVariables(ProjectInfo projectInfo, int changeNumber, int codeChangeNumber, TargetReceipt? editorTarget, ConfigFile? projectConfigFile, IEnumerable<KeyValuePair<string, string>> additionalVariables)
 		{
-			Dictionary<string, string> Variables = GetWorkspaceVariables(ProjectInfo, ChangeNumber, CodeChangeNumber, EditorTarget, ProjectConfigFile);
-			foreach ((string Key, string Value) in AdditionalVariables)
+			Dictionary<string, string> variables = GetWorkspaceVariables(projectInfo, changeNumber, codeChangeNumber, editorTarget, projectConfigFile);
+			foreach ((string key, string value) in additionalVariables)
 			{
-				Variables[Key] = Value;
+				variables[key] = value;
 			}
-			return Variables;
+			return variables;
 		}
 
-		public static bool TryGetProjectSetting(ConfigFile ProjectConfigFile, string SelectedProjectIdentifier, string Name, [NotNullWhen(true)] out string? Value)
+		public static bool TryGetProjectSetting(ConfigFile projectConfigFile, string selectedProjectIdentifier, string name, [NotNullWhen(true)] out string? value)
 		{
-			string Path = SelectedProjectIdentifier;
+			string path = selectedProjectIdentifier;
 			for (; ; )
 			{
-				ConfigSection ProjectSection = ProjectConfigFile.FindSection(Path);
-				if (ProjectSection != null)
+				ConfigSection projectSection = projectConfigFile.FindSection(path);
+				if (projectSection != null)
 				{
-					string? NewValue = ProjectSection.GetValue(Name, null);
-					if (NewValue != null)
+					string? newValue = projectSection.GetValue(name, null);
+					if (newValue != null)
 					{
-						Value = NewValue;
+						value = newValue;
 						return true;
 					}
 				}
 
-				int LastSlash = Path.LastIndexOf('/');
-				if (LastSlash < 2)
+				int lastSlash = path.LastIndexOf('/');
+				if (lastSlash < 2)
 				{
 					break;
 				}
 
-				Path = Path.Substring(0, LastSlash);
+				path = path.Substring(0, lastSlash);
 			}
 
-			ConfigSection DefaultSection = ProjectConfigFile.FindSection("Default");
-			if (DefaultSection != null)
+			ConfigSection defaultSection = projectConfigFile.FindSection("Default");
+			if (defaultSection != null)
 			{
-				string? NewValue = DefaultSection.GetValue(Name, null);
-				if (NewValue != null)
+				string? newValue = defaultSection.GetValue(name, null);
+				if (newValue != null)
 				{
-					Value = NewValue;
+					value = newValue;
 					return true;
 				}
 			}
 
-			Value = null;
+			value = null;
 			return false;
 		}
 
-		public static Dictionary<Guid, WorkspaceSyncCategory> GetSyncCategories(ConfigFile ProjectConfigFile)
+		public static Dictionary<Guid, WorkspaceSyncCategory> GetSyncCategories(ConfigFile projectConfigFile)
 		{
-			Dictionary<Guid, WorkspaceSyncCategory> UniqueIdToCategory = new Dictionary<Guid, WorkspaceSyncCategory>();
-			if (ProjectConfigFile != null)
+			Dictionary<Guid, WorkspaceSyncCategory> uniqueIdToCategory = new Dictionary<Guid, WorkspaceSyncCategory>();
+			if (projectConfigFile != null)
 			{
-				string[] CategoryLines = ProjectConfigFile.GetValues("Options.SyncCategory", new string[0]);
-				foreach (string CategoryLine in CategoryLines)
+				string[] categoryLines = projectConfigFile.GetValues("Options.SyncCategory", new string[0]);
+				foreach (string categoryLine in categoryLines)
 				{
-					ConfigObject Object = new ConfigObject(CategoryLine);
+					ConfigObject obj = new ConfigObject(categoryLine);
 
-					Guid UniqueId;
-					if (Guid.TryParse(Object.GetValue("UniqueId", ""), out UniqueId))
+					Guid uniqueId;
+					if (Guid.TryParse(obj.GetValue("UniqueId", ""), out uniqueId))
 					{
-						WorkspaceSyncCategory? Category;
-						if (!UniqueIdToCategory.TryGetValue(UniqueId, out Category))
+						WorkspaceSyncCategory? category;
+						if (!uniqueIdToCategory.TryGetValue(uniqueId, out category))
 						{
-							Category = new WorkspaceSyncCategory(UniqueId);
-							UniqueIdToCategory.Add(UniqueId, Category);
+							category = new WorkspaceSyncCategory(uniqueId);
+							uniqueIdToCategory.Add(uniqueId, category);
 						}
 
-						if (Object.GetValue("Clear", false))
+						if (obj.GetValue("Clear", false))
 						{
-							Category.Paths = new string[0];
-							Category.Requires = new Guid[0];
+							category.Paths = new string[0];
+							category.Requires = new Guid[0];
 						}
 
-						Category.Name = Object.GetValue("Name", Category.Name);
-						Category.bEnable = Object.GetValue("Enable", Category.bEnable);
-						Category.Paths = Enumerable.Concat(Category.Paths, Object.GetValue("Paths", "").Split(';').Select(x => x.Trim())).Where(x => x.Length > 0).Distinct().OrderBy(x => x).ToArray();
-						Category.bHidden = Object.GetValue("Hidden", Category.bHidden);
-						Category.Requires = Enumerable.Concat(Category.Requires, ParseGuids(Object.GetValue("Requires", "").Split(';'))).Distinct().OrderBy(x => x).ToArray();
+						category.Name = obj.GetValue("Name", category.Name);
+						category.Enable = obj.GetValue("Enable", category.Enable);
+						category.Paths = Enumerable.Concat(category.Paths, obj.GetValue("Paths", "").Split(';').Select(x => x.Trim())).Where(x => x.Length > 0).Distinct().OrderBy(x => x).ToArray();
+						category.Hidden = obj.GetValue("Hidden", category.Hidden);
+						category.Requires = Enumerable.Concat(category.Requires, ParseGuids(obj.GetValue("Requires", "").Split(';'))).Distinct().OrderBy(x => x).ToArray();
 					}
 				}
 			}
-			return UniqueIdToCategory;
+			return uniqueIdToCategory;
 		}
 
-		static IEnumerable<Guid> ParseGuids(IEnumerable<string> Values)
+		static IEnumerable<Guid> ParseGuids(IEnumerable<string> values)
 		{
-			foreach (string Value in Values)
+			foreach (string value in values)
 			{
-				Guid Guid;
-				if (Guid.TryParse(Value, out Guid))
+				Guid guid;
+				if (Guid.TryParse(value, out guid))
 				{
-					yield return Guid;
+					yield return guid;
 				}
 			}
 		}
