UGS: Reformat to standard C# style conventions.

diff --git a/Engine/Source/Programs/UnrealGameSync/UnrealGameSyncShared/FileFilter.cs b/Engine/Source/Programs/UnrealGameSync/UnrealGameSyncShared/FileFilter.cs
--- a/Engine/Source/Programs/UnrealGameSync/UnrealGameSyncShared/FileFilter.cs
+++ b/Engine/Source/Programs/UnrealGameSync/UnrealGameSyncShared/FileFilter.cs
@@ -28,52 +28,52 @@ namespace UnrealGameSync
 		/// <summary>
 		/// Root node for the tree.
 		/// </summary>
-		FileFilterNode RootNode;
+		FileFilterNode _rootNode;
 
 		/// <summary>
 		/// The default node, which will match any path
 		/// </summary>
-		FileFilterNode DefaultNode;
+		FileFilterNode _defaultNode;
 
 		/// <summary>
 		/// Terminating nodes for each rule added to the filter
 		/// </summary>
-		List<FileFilterNode> Rules = new List<FileFilterNode>();
+		List<FileFilterNode> _rules = new List<FileFilterNode>();
 
 		/// <summary>
 		/// Default constructor
 		/// </summary>
-		public FileFilter(FileFilterType DefaultType = FileFilterType.Exclude)
+		public FileFilter(FileFilterType defaultType = FileFilterType.Exclude)
 		{
-			RootNode = new FileFilterNode(null, "");
+			_rootNode = new FileFilterNode(null, "");
 
-			DefaultNode = new FileFilterNode(RootNode, "...");
-			DefaultNode.Type = DefaultType;
+			_defaultNode = new FileFilterNode(_rootNode, "...");
+			_defaultNode.Type = defaultType;
 		}
 
 		/// <summary>
 		/// Copy constructor
 		/// </summary>
-		/// <param name="Other">Filter to copy from</param>
-		public FileFilter(FileFilter Other)
-			: this(Other.DefaultNode.Type)
+		/// <param name="other">Filter to copy from</param>
+		public FileFilter(FileFilter other)
+			: this(other._defaultNode.Type)
 		{
-			foreach (FileFilterNode OtherRule in Other.Rules)
+			foreach (FileFilterNode otherRule in other._rules)
 			{
-				AddRule(OtherRule.ToString(), OtherRule.Type);
+				AddRule(otherRule.ToString(), otherRule.Type);
 			}
 		}
 
 		/// <summary>
 		/// Constructs a file filter from a p4-style filespec. Exclude lines are prefixed with a - character.
 		/// </summary>
-		/// <param name="Lines">Lines to parse rules from</param>
-		public FileFilter(IEnumerable<string> Lines)
+		/// <param name="lines">Lines to parse rules from</param>
+		public FileFilter(IEnumerable<string> lines)
 			: this()
 		{
-			foreach (string Line in Lines)
+			foreach (string line in lines)
 			{
-				AddRule(Line);
+				AddRule(line);
 			}
 		}
 
@@ -81,15 +81,15 @@ namespace UnrealGameSync
 		/// Adds an include or exclude rule to the filter
 		/// </summary>
 		/// <param name="Pattern">Pattern to match. See CreateRegex() for details.</param>
-		public void AddRule(string Rule)
+		public void AddRule(string rule)
 		{
-			if (Rule.StartsWith("-"))
+			if (rule.StartsWith("-"))
 			{
-				Exclude(Rule.Substring(1).TrimStart());
+				Exclude(rule.Substring(1).TrimStart());
 			}
 			else
 			{
-				Include(Rule);
+				Include(rule);
 			}
 		}
 
@@ -97,40 +97,40 @@ namespace UnrealGameSync
 		/// Adds an include or exclude rule to the filter. The rule may be 
 		/// </summary>
 		/// <param name="Pattern">Pattern to match. See CreateRegex() for details.</param>
-		public void AddRule(string Rule, params string[] AllowTags)
+		public void AddRule(string rule, params string[] allowTags)
 		{
-			string CleanRule = Rule.Trim();
-			if(CleanRule.StartsWith("{"))
+			string cleanRule = rule.Trim();
+			if(cleanRule.StartsWith("{"))
 			{
 				// Find the end of the condition
-				int ConditionEnd = CleanRule.IndexOf('}');
-				if(ConditionEnd == -1)
+				int conditionEnd = cleanRule.IndexOf('}');
+				if(conditionEnd == -1)
 				{
-					throw new Exception(String.Format("Missing closing parenthesis in rule: {0}", CleanRule));
+					throw new Exception(String.Format("Missing closing parenthesis in rule: {0}", cleanRule));
 				}
 
 				// Check there's a matching tag
-				string[] RuleTags = CleanRule.Substring(1, ConditionEnd - 1).Split(',').Select(x => x.Trim()).ToArray();
-				if(!RuleTags.Any(x => AllowTags.Contains(x)))
+				string[] ruleTags = cleanRule.Substring(1, conditionEnd - 1).Split(',').Select(x => x.Trim()).ToArray();
+				if(!ruleTags.Any(x => allowTags.Contains(x)))
 				{
 					return;
 				}
 
 				// Strip the condition from the rule
-				CleanRule = CleanRule.Substring(ConditionEnd + 1).TrimStart(); 
+				cleanRule = cleanRule.Substring(conditionEnd + 1).TrimStart(); 
 			}
-			AddRule(CleanRule);
+			AddRule(cleanRule);
 		}
 
 		/// <summary>
 		/// Adds several rules to the filter
 		/// </summary>
 		/// <param name="Patterns">List of patterns to match.</param>
-		public void AddRules(IEnumerable<string> Rules)
+		public void AddRules(IEnumerable<string> rules)
 		{
-			foreach(string Rule in Rules)
+			foreach(string rule in rules)
 			{
-				AddRule(Rule);
+				AddRule(rule);
 			}
 		}
 
@@ -140,11 +140,11 @@ namespace UnrealGameSync
 		/// </summary>
 		/// <param name="Lines"></param>
 		/// <param name="Variables">Lookup for variables to test against</param>
-		public void AddRules(IEnumerable<string> Rules, params string[] Tags)
+		public void AddRules(IEnumerable<string> rules, params string[] tags)
 		{
-			foreach(string Rule in Rules)
+			foreach(string rule in rules)
 			{
-				AddRule(Rule, Tags);
+				AddRule(rule, tags);
 			}
 		}
 
@@ -154,21 +154,21 @@ namespace UnrealGameSync
 		/// <param name="Filter"></param>
 		/// <param name="RulesFileName"></param>
 		/// <param name="Conditions"></param>
-		public void ReadRulesFromFile(string FileName, string SectionName, params string[] AllowTags)
+		public void ReadRulesFromFile(string fileName, string sectionName, params string[] allowTags)
 		{
-			bool bInSection = false;
-			foreach(string Line in File.ReadAllLines(FileName))
+			bool inSection = false;
+			foreach(string line in File.ReadAllLines(fileName))
 			{
-				string TrimLine = Line.Trim();
-				if(!TrimLine.StartsWith(";") && TrimLine.Length > 0)
+				string trimLine = line.Trim();
+				if(!trimLine.StartsWith(";") && trimLine.Length > 0)
 				{
-					if(TrimLine.StartsWith("["))
+					if(trimLine.StartsWith("["))
 					{
-						bInSection = (TrimLine == "[" + SectionName + "]");
+						inSection = (trimLine == "[" + sectionName + "]");
 					}
-					else if(bInSection)
+					else if(inSection)
 					{
-						AddRule(Line, AllowTags);
+						AddRule(line, allowTags);
 					}
 				}
 			}
@@ -177,119 +177,119 @@ namespace UnrealGameSync
 		/// <summary>
 		/// Adds an include rule to the filter
 		/// </summary>
-		/// <param name="Pattern">Pattern to match. See CreateRegex() for details.</param>
-		public void Include(string Pattern)
+		/// <param name="pattern">Pattern to match. See CreateRegex() for details.</param>
+		public void Include(string pattern)
 		{
-			AddRule(Pattern, FileFilterType.Include);
+			AddRule(pattern, FileFilterType.Include);
 		}
 
 		/// <summary>
 		/// Adds several exclude rules to the filter
 		/// </summary>
-		/// <param name="Patterns">Patterns to match. See CreateRegex() for details.</param>
-		public void Include(IEnumerable<string> Patterns)
+		/// <param name="patterns">Patterns to match. See CreateRegex() for details.</param>
+		public void Include(IEnumerable<string> patterns)
 		{
-			foreach (string Pattern in Patterns)
+			foreach (string pattern in patterns)
 			{
-				Include(Pattern);
+				Include(pattern);
 			}
 		}
 
 		/// <summary>
 		/// Adds an exclude rule to the filter
 		/// </summary>
-		/// <param name="Pattern">Mask to match. See CreateRegex() for details.</param>
-		public void Exclude(string Pattern)
+		/// <param name="pattern">Mask to match. See CreateRegex() for details.</param>
+		public void Exclude(string pattern)
 		{
-			AddRule(Pattern, FileFilterType.Exclude);
+			AddRule(pattern, FileFilterType.Exclude);
 		}
 
 		/// <summary>
 		/// Adds several exclude rules to the filter
 		/// </summary>
-		/// <param name="Patterns">Patterns to match. See CreateRegex() for details.</param>
-		public void Exclude(IEnumerable<string> Patterns)
+		/// <param name="patterns">Patterns to match. See CreateRegex() for details.</param>
+		public void Exclude(IEnumerable<string> patterns)
 		{
-			foreach (string Pattern in Patterns)
+			foreach (string pattern in patterns)
 			{
-				Exclude(Pattern);
+				Exclude(pattern);
 			}
 		}
 
 		/// <summary>
 		/// Adds an include or exclude rule to the filter
 		/// </summary>
-		/// <param name="Pattern">The pattern which the rule should match</param>
+		/// <param name="pattern">The pattern which the rule should match</param>
 		/// <param name="bInclude">Whether to include or exclude files matching this rule</param>
-		public void AddRule(string Pattern, FileFilterType Type)
+		public void AddRule(string pattern, FileFilterType type)
 		{
-			string NormalizedPattern = Pattern.Replace('\\', '/');
+			string normalizedPattern = pattern.Replace('\\', '/');
 
 			// We don't want a slash at the start, but if there was not one specified, it's not anchored to the root of the tree.
-			if (NormalizedPattern.StartsWith("/"))
+			if (normalizedPattern.StartsWith("/"))
 			{
-				NormalizedPattern = NormalizedPattern.Substring(1);
+				normalizedPattern = normalizedPattern.Substring(1);
 			}
-			else if(!NormalizedPattern.StartsWith("..."))
+			else if(!normalizedPattern.StartsWith("..."))
 			{
-				NormalizedPattern = ".../" + NormalizedPattern;
+				normalizedPattern = ".../" + normalizedPattern;
 			}
 
 			// All directories indicate a wildcard match
-			if (NormalizedPattern.EndsWith("/"))
+			if (normalizedPattern.EndsWith("/"))
 			{
-				NormalizedPattern += "...";
+				normalizedPattern += "...";
 			}
 
 			// Replace any directory wildcards mid-string
-			for (int Idx = NormalizedPattern.IndexOf("..."); Idx != -1; Idx = NormalizedPattern.IndexOf("...", Idx))
+			for (int idx = normalizedPattern.IndexOf("..."); idx != -1; idx = normalizedPattern.IndexOf("...", idx))
 			{
-				if (Idx > 0 && NormalizedPattern[Idx - 1] != '/')
+				if (idx > 0 && normalizedPattern[idx - 1] != '/')
 				{
-					NormalizedPattern = NormalizedPattern.Insert(Idx, "*/");
-					Idx++;
+					normalizedPattern = normalizedPattern.Insert(idx, "*/");
+					idx++;
 				}
 
-				Idx += 3;
+				idx += 3;
 
-				if (Idx < NormalizedPattern.Length && NormalizedPattern[Idx] != '/')
+				if (idx < normalizedPattern.Length && normalizedPattern[idx] != '/')
 				{
-					NormalizedPattern = NormalizedPattern.Insert(Idx, "/*");
-					Idx += 2;
+					normalizedPattern = normalizedPattern.Insert(idx, "/*");
+					idx += 2;
 				}
 			}
 
 			// Split the pattern into fragments
-			string[] BranchPatterns = NormalizedPattern.Split('/');
+			string[] branchPatterns = normalizedPattern.Split('/');
 
 			// Add it into the tree
-			FileFilterNode LastNode = RootNode;
-			foreach (string BranchPattern in BranchPatterns)
+			FileFilterNode lastNode = _rootNode;
+			foreach (string branchPattern in branchPatterns)
 			{
-				FileFilterNode NextNode = LastNode.Branches.FirstOrDefault(x => x.Pattern == BranchPattern);
-				if (NextNode == null)
+				FileFilterNode nextNode = lastNode.Branches.FirstOrDefault(x => x.Pattern == branchPattern);
+				if (nextNode == null)
 				{
-					NextNode = new FileFilterNode(LastNode, BranchPattern);
-					LastNode.Branches.Add(NextNode);
+					nextNode = new FileFilterNode(lastNode, branchPattern);
+					lastNode.Branches.Add(nextNode);
 				}
-				LastNode = NextNode;
+				lastNode = nextNode;
 			}
 
 			// We've reached the end of the pattern, so mark it as a leaf node
-			Rules.Add(LastNode);
-			LastNode.RuleNumber = Rules.Count - 1;
-			LastNode.Type = Type;
+			_rules.Add(lastNode);
+			lastNode.RuleNumber = _rules.Count - 1;
+			lastNode.Type = type;
 
 			// Update the maximums along that path
-			for (FileFilterNode? UpdateNode = LastNode; UpdateNode != null; UpdateNode = UpdateNode.Parent)
+			for (FileFilterNode? updateNode = lastNode; updateNode != null; updateNode = updateNode.Parent)
 			{
-				if (Type == FileFilterType.Include)
+				if (type == FileFilterType.Include)
 				{
-					UpdateNode.MaxIncludeRuleNumber = LastNode.RuleNumber;
+					updateNode.MaxIncludeRuleNumber = lastNode.RuleNumber;
 				}
 				else
 				{
-					UpdateNode.MaxExcludeRuleNumber = LastNode.RuleNumber;
+					updateNode.MaxExcludeRuleNumber = lastNode.RuleNumber;
 				}
 			}
 		}
@@ -308,124 +308,124 @@ namespace UnrealGameSync
 		/// <summary>
 		/// Determines whether the given file matches the filter
 		/// </summary>
-		/// <param name="FileName">File to match</param>
+		/// <param name="fileName">File to match</param>
 		/// <returns>True if the file passes the filter</returns>
-		public bool Matches(string FileName)
+		public bool Matches(string fileName)
 		{
-			string[] Tokens = FileName.TrimStart('/', '\\').Split('/', '\\');
+			string[] tokens = fileName.TrimStart('/', '\\').Split('/', '\\');
 
-			FileFilterNode MatchingNode = FindMatchingNode(RootNode, Tokens, 0, DefaultNode);
+			FileFilterNode matchingNode = FindMatchingNode(_rootNode, tokens, 0, _defaultNode);
 
-			return MatchingNode.Type == FileFilterType.Include;
+			return matchingNode.Type == FileFilterType.Include;
 		}
 
 		/// <summary>
 		/// Determines whether it's possible for anything within the given folder name to match the filter. Useful to early out of recursive file searches.
 		/// </summary>
-		/// <param name="FolderName">File to match</param>
+		/// <param name="folderName">File to match</param>
 		/// <returns>True if the file passes the filter</returns>
-		public bool PossiblyMatches(string FolderName)
+		public bool PossiblyMatches(string folderName)
 		{
-			string[] Tokens = FolderName.Trim('/', '\\').Split('/', '\\');
+			string[] tokens = folderName.Trim('/', '\\').Split('/', '\\');
 
-			FileFilterNode MatchingNode = FindMatchingNode(RootNode, Tokens.Union(new string[]{ "" }).ToArray(), 0, DefaultNode);
+			FileFilterNode matchingNode = FindMatchingNode(_rootNode, tokens.Union(new string[]{ "" }).ToArray(), 0, _defaultNode);
 
-			return MatchingNode.Type == FileFilterType.Include || HighestPossibleIncludeMatch(RootNode, Tokens, 0, MatchingNode.RuleNumber) > MatchingNode.RuleNumber;
+			return matchingNode.Type == FileFilterType.Include || HighestPossibleIncludeMatch(_rootNode, tokens, 0, matchingNode.RuleNumber) > matchingNode.RuleNumber;
 		}
 
 		/// <summary>
 		/// Applies the filter to each element in a sequence, and returns the list of files that match
 		/// </summary>
-		/// <param name="FileNames">List of filenames</param>
+		/// <param name="fileNames">List of filenames</param>
 		/// <returns>List of filenames which match the filter</returns>
-		public IEnumerable<string> ApplyTo(IEnumerable<string> FileNames)
+		public IEnumerable<string> ApplyTo(IEnumerable<string> fileNames)
 		{
-			return FileNames.Where(x => Matches(x));
+			return fileNames.Where(x => Matches(x));
 		}
 
 		/// <summary>
 		/// Finds the node which matches a given list of tokens.
 		/// </summary>
-		/// <param name="CurrentNode"></param>
-		/// <param name="Tokens"></param>
-		/// <param name="TokenIdx"></param>
-		/// <param name="CurrentBestNode"></param>
+		/// <param name="currentNode"></param>
+		/// <param name="tokens"></param>
+		/// <param name="tokenIdx"></param>
+		/// <param name="currentBestNode"></param>
 		/// <returns></returns>
-		FileFilterNode FindMatchingNode(FileFilterNode CurrentNode, string[] Tokens, int TokenIdx, FileFilterNode CurrentBestNode)
+		FileFilterNode FindMatchingNode(FileFilterNode currentNode, string[] tokens, int tokenIdx, FileFilterNode currentBestNode)
 		{
 			// If we've matched all the input tokens, check if this rule is better than any other we've seen
-			if (TokenIdx == Tokens.Length)
+			if (tokenIdx == tokens.Length)
 			{
-				return (CurrentNode.RuleNumber > CurrentBestNode.RuleNumber) ? CurrentNode : CurrentBestNode;
+				return (currentNode.RuleNumber > currentBestNode.RuleNumber) ? currentNode : currentBestNode;
 			}
 
 			// If there is no rule under the current node which is better than the current best node, early out
-			if (CurrentNode.MaxIncludeRuleNumber <= CurrentBestNode.RuleNumber && CurrentNode.MaxExcludeRuleNumber <= CurrentBestNode.RuleNumber)
+			if (currentNode.MaxIncludeRuleNumber <= currentBestNode.RuleNumber && currentNode.MaxExcludeRuleNumber <= currentBestNode.RuleNumber)
 			{
-				return CurrentBestNode;
+				return currentBestNode;
 			}
 
 			// Test all the branches for one that matches
-			FileFilterNode BestNode = CurrentBestNode;
-			foreach (FileFilterNode Branch in CurrentNode.Branches)
+			FileFilterNode bestNode = currentBestNode;
+			foreach (FileFilterNode branch in currentNode.Branches)
 			{
-				if (Branch.Pattern == "...")
+				if (branch.Pattern == "...")
 				{
-					for (int NextTokenIdx = Tokens.Length; NextTokenIdx >= TokenIdx; NextTokenIdx--)
+					for (int nextTokenIdx = tokens.Length; nextTokenIdx >= tokenIdx; nextTokenIdx--)
 					{
-						BestNode = FindMatchingNode(Branch, Tokens, NextTokenIdx, BestNode);
+						bestNode = FindMatchingNode(branch, tokens, nextTokenIdx, bestNode);
 					}
 				}
 				else
 				{
-					if (Branch.IsMatch(Tokens[TokenIdx]))
+					if (branch.IsMatch(tokens[tokenIdx]))
 					{
-						BestNode = FindMatchingNode(Branch, Tokens, TokenIdx + 1, BestNode);
+						bestNode = FindMatchingNode(branch, tokens, tokenIdx + 1, bestNode);
 					}
 				}
 			}
-			return BestNode;
+			return bestNode;
 		}
 
 		/// <summary>
 		/// Returns the highest possible rule number which can match the given list of input tokens, assuming that the list of input tokens is incomplete.
 		/// </summary>
-		/// <param name="CurrentNode">The current node being checked</param>
-		/// <param name="Tokens">The tokens to match</param>
-		/// <param name="TokenIdx">Current token index</param>
-		/// <param name="CurrentBestRuleNumber">The highest rule number seen so far. Used to optimize tree traversals.</param>
+		/// <param name="currentNode">The current node being checked</param>
+		/// <param name="tokens">The tokens to match</param>
+		/// <param name="tokenIdx">Current token index</param>
+		/// <param name="currentBestRuleNumber">The highest rule number seen so far. Used to optimize tree traversals.</param>
 		/// <returns>New highest rule number</returns>
-		int HighestPossibleIncludeMatch(FileFilterNode CurrentNode, string[] Tokens, int TokenIdx, int CurrentBestRuleNumber)
+		int HighestPossibleIncludeMatch(FileFilterNode currentNode, string[] tokens, int tokenIdx, int currentBestRuleNumber)
 		{
 			// If we've matched all the input tokens, check if this rule is better than any other we've seen
-			if (TokenIdx == Tokens.Length)
+			if (tokenIdx == tokens.Length)
 			{
-				return Math.Max(CurrentBestRuleNumber, CurrentNode.MaxIncludeRuleNumber);
+				return Math.Max(currentBestRuleNumber, currentNode.MaxIncludeRuleNumber);
 			}
 
 			// Test all the branches for one that matches
-			int BestRuleNumber = CurrentBestRuleNumber;
-			if (CurrentNode.MaxIncludeRuleNumber > BestRuleNumber)
+			int bestRuleNumber = currentBestRuleNumber;
+			if (currentNode.MaxIncludeRuleNumber > bestRuleNumber)
 			{
-				foreach (FileFilterNode Branch in CurrentNode.Branches)
+				foreach (FileFilterNode branch in currentNode.Branches)
 				{
-					if (Branch.Pattern == "...")
+					if (branch.Pattern == "...")
 					{
-						if (Branch.MaxIncludeRuleNumber > BestRuleNumber)
+						if (branch.MaxIncludeRuleNumber > bestRuleNumber)
 						{
-							BestRuleNumber = Branch.MaxIncludeRuleNumber;
+							bestRuleNumber = branch.MaxIncludeRuleNumber;
 						}
 					}
 					else
 					{
-						if (Branch.IsMatch(Tokens[TokenIdx]))
+						if (branch.IsMatch(tokens[tokenIdx]))
 						{
-							BestRuleNumber = HighestPossibleIncludeMatch(Branch, Tokens, TokenIdx + 1, BestRuleNumber);
+							bestRuleNumber = HighestPossibleIncludeMatch(branch, tokens, tokenIdx + 1, bestRuleNumber);
 						}
 					}
 				}
 			}
-			return BestRuleNumber;
+			return bestRuleNumber;
 		}
 	}
 
@@ -472,10 +472,10 @@ namespace UnrealGameSync
 		/// <summary>
 		/// Default constructor.
 		/// </summary>
-		public FileFilterNode(FileFilterNode? InParent, string InPattern)
+		public FileFilterNode(FileFilterNode? inParent, string inPattern)
 		{
-			Parent = InParent;
-			Pattern = InPattern;
+			Parent = inParent;
+			Pattern = inPattern;
 			RuleNumber = -1;
 			Type = FileFilterType.Exclude;
 		}
@@ -483,44 +483,44 @@ namespace UnrealGameSync
 		/// <summary>
 		/// Determine if the given token matches the pattern in this node
 		/// </summary>
-		public bool IsMatch(string Token)
+		public bool IsMatch(string token)
 		{
 			if(Pattern.EndsWith("."))
 			{
-				return !Token.Contains('.') && IsMatch(Token, 0, Pattern.Substring(0, Pattern.Length - 1), 0);
+				return !token.Contains('.') && IsMatch(token, 0, Pattern.Substring(0, Pattern.Length - 1), 0);
 			}
 			else
 			{
-				return IsMatch(Token, 0, Pattern, 0);
+				return IsMatch(token, 0, Pattern, 0);
 			}
 		}
 
 		/// <summary>
 		/// Determine if a given token matches a pattern, starting from the given positions within each string.
 		/// </summary>
-		static bool IsMatch(string Token, int TokenIdx, string Pattern, int PatternIdx)
+		static bool IsMatch(string token, int tokenIdx, string pattern, int patternIdx)
 		{
 			for (; ; )
 			{
-				if (PatternIdx == Pattern.Length)
+				if (patternIdx == pattern.Length)
 				{
-					return (TokenIdx == Token.Length);
+					return (tokenIdx == token.Length);
 				}
-				else if (Pattern[PatternIdx] == '*')
+				else if (pattern[patternIdx] == '*')
 				{
-					for (int NextTokenIdx = Token.Length; NextTokenIdx >= TokenIdx; NextTokenIdx--)
+					for (int nextTokenIdx = token.Length; nextTokenIdx >= tokenIdx; nextTokenIdx--)
 					{
-						if (IsMatch(Token, NextTokenIdx, Pattern, PatternIdx + 1))
+						if (IsMatch(token, nextTokenIdx, pattern, patternIdx + 1))
 						{
 							return true;
 						}
 					}
 					return false;
 				}
-				else if (TokenIdx < Token.Length && (Char.ToLower(Token[TokenIdx]) == Char.ToLower(Pattern[PatternIdx]) || Pattern[PatternIdx] == '?'))
+				else if (tokenIdx < token.Length && (Char.ToLower(token[tokenIdx]) == Char.ToLower(pattern[patternIdx]) || pattern[patternIdx] == '?'))
 				{
-					TokenIdx++;
-					PatternIdx++;
+					tokenIdx++;
+					patternIdx++;
 				}
 				else
 				{
diff --git a/Engine/Source/Programs/UnrealGameSync/UnrealGameSyncShared/GlobalSettings.cs b/Engine/Source/Programs/UnrealGameSync/UnrealGameSyncShared/GlobalSettings.cs
--- a/Engine/Source/Programs/UnrealGameSync/UnrealGameSyncShared/GlobalSettings.cs
+++ b/Engine/Source/Programs/UnrealGameSync/UnrealGameSyncShared/GlobalSettings.cs
@@ -29,24 +29,24 @@ namespace UnrealGameSync
 			AllProjectsInSln = null;
 		}
 
-		public void SetCategories(Dictionary<Guid, bool> Categories)
+		public void SetCategories(Dictionary<Guid, bool> categories)
 		{
-			IncludeCategories = Categories.Where(x => x.Value).Select(x => x.Key).ToList();
-			ExcludeCategories = Categories.Where(x => !x.Value).Select(x => x.Key).ToList();
+			IncludeCategories = categories.Where(x => x.Value).Select(x => x.Key).ToList();
+			ExcludeCategories = categories.Where(x => !x.Value).Select(x => x.Key).ToList();
 		}
 
 		public Dictionary<Guid, bool> GetCategories()
 		{
-			Dictionary<Guid, bool> Categories = new Dictionary<Guid, bool>();
-			foreach (Guid IncludeCategory in IncludeCategories)
+			Dictionary<Guid, bool> categories = new Dictionary<Guid, bool>();
+			foreach (Guid includeCategory in IncludeCategories)
 			{
-				Categories[IncludeCategory] = true;
+				categories[includeCategory] = true;
 			}
-			foreach (Guid ExcludeCategory in ExcludeCategories)
+			foreach (Guid excludeCategory in ExcludeCategories)
 			{
-				Categories[ExcludeCategory] = false;
+				categories[excludeCategory] = false;
 			}
-			return Categories;
+			return categories;
 		}
 	}
 
@@ -60,175 +60,175 @@ namespace UnrealGameSync
 		public FileReference File { get; }
 		public GlobalSettings Global { get; }
 
-		public UserProjectSettings FindOrAddProjectSettings(ProjectInfo ProjectInfo, UserWorkspaceSettings Settings, ILogger Logger)
+		public UserProjectSettings FindOrAddProjectSettings(ProjectInfo projectInfo, UserWorkspaceSettings settings, ILogger logger)
 		{
-			FileReference ConfigFile;
-			if (ProjectInfo.LocalFileName.HasExtension(".uprojectdirs"))
+			FileReference configFile;
+			if (projectInfo.LocalFileName.HasExtension(".uprojectdirs"))
 			{
-				ConfigFile = FileReference.Combine(UserSettings.GetConfigDir(Settings.RootDir), "project.json");
+				configFile = FileReference.Combine(UserSettings.GetConfigDir(settings.RootDir), "project.json");
 			}
 			else
 			{
-				ConfigFile = FileReference.Combine(UserSettings.GetConfigDir(Settings.RootDir), $"project_{ProjectInfo.LocalFileName.GetFileNameWithoutExtension()}.json");
+				configFile = FileReference.Combine(UserSettings.GetConfigDir(settings.RootDir), $"project_{projectInfo.LocalFileName.GetFileNameWithoutExtension()}.json");
 			}
 
-			UserSettings.CreateConfigDir(ConfigFile.Directory);
+			UserSettings.CreateConfigDir(configFile.Directory);
 
-			UserProjectSettings? ProjectSettings;
-			if (!UserProjectSettings.TryLoad(ConfigFile, out ProjectSettings))
+			UserProjectSettings? projectSettings;
+			if (!UserProjectSettings.TryLoad(configFile, out projectSettings))
 			{
-				ProjectSettings = new UserProjectSettings(ConfigFile);
-				ImportProjectSettings(ProjectInfo, ProjectSettings);
-				ProjectSettings.Save(Logger);
+				projectSettings = new UserProjectSettings(configFile);
+				ImportProjectSettings(projectInfo, projectSettings);
+				projectSettings.Save(logger);
 			}
-			return ProjectSettings;
+			return projectSettings;
 		}
 
-		protected virtual void ImportProjectSettings(ProjectInfo ProjectInfo, UserProjectSettings ProjectSettings)
+		protected virtual void ImportProjectSettings(ProjectInfo projectInfo, UserProjectSettings projectSettings)
 		{
 		}
 
-		protected virtual void ImportWorkspaceSettings(DirectoryReference RootDir, string ClientName, string BranchPath, UserWorkspaceSettings WorkspaceSettings)
+		protected virtual void ImportWorkspaceSettings(DirectoryReference rootDir, string clientName, string branchPath, UserWorkspaceSettings workspaceSettings)
 		{
 		}
 
-		protected virtual void ImportWorkspaceState(DirectoryReference RootDir, string ClientName, string BranchPath, UserWorkspaceState WorkspaceState)
+		protected virtual void ImportWorkspaceState(DirectoryReference rootDir, string clientName, string branchPath, UserWorkspaceState workspaceState)
 		{
 		}
 
-		public GlobalSettingsFile(FileReference File, GlobalSettings Global)
+		public GlobalSettingsFile(FileReference file, GlobalSettings global)
 		{
-			this.File = File;
-			this.Global = Global;
+			this.File = file;
+			this.Global = global;
 		}
 
-		public static GlobalSettingsFile Create(FileReference File)
+		public static GlobalSettingsFile Create(FileReference file)
 		{
-			GlobalSettings? Data;
-			if (!Utility.TryLoadJson(File, out Data))
+			GlobalSettings? data;
+			if (!Utility.TryLoadJson(file, out data))
 			{
-				Data = new GlobalSettings();
+				data = new GlobalSettings();
 			}
-			return new GlobalSettingsFile(File, Data);
+			return new GlobalSettingsFile(file, data);
 		}
 
-		public virtual bool Save(ILogger Logger)
+		public virtual bool Save(ILogger logger)
 		{
 			try
 			{
 				Utility.SaveJson(File, Global);
 				return true;
 			}
-			catch (Exception Ex)
+			catch (Exception ex)
 			{
-				Logger.LogError(Ex, "Unable to save {File}: {Message}", File, Ex.Message);
+				logger.LogError(ex, "Unable to save {File}: {Message}", File, ex.Message);
 				return false;
 			}
 		}
 
-		public UserWorkspaceState FindOrAddWorkspaceState(UserWorkspaceSettings Settings, ILogger Logger)
+		public UserWorkspaceState FindOrAddWorkspaceState(UserWorkspaceSettings settings, ILogger logger)
 		{
-			return FindOrAddWorkspaceState(Settings.RootDir, Settings.ClientName, Settings.BranchPath, Logger);
+			return FindOrAddWorkspaceState(settings.RootDir, settings.ClientName, settings.BranchPath, logger);
 		}
 
-		public UserWorkspaceState FindOrAddWorkspaceState(DirectoryReference RootDir, string ClientName, string BranchPath, ILogger Logger)
+		public UserWorkspaceState FindOrAddWorkspaceState(DirectoryReference rootDir, string clientName, string branchPath, ILogger logger)
 		{
-			UserWorkspaceState? State;
-			if (!UserWorkspaceState.TryLoad(RootDir, out State))
+			UserWorkspaceState? state;
+			if (!UserWorkspaceState.TryLoad(rootDir, out state))
 			{
-				State = new UserWorkspaceState();
-				State.RootDir = RootDir;
-				ImportWorkspaceState(RootDir, ClientName, BranchPath, State);
-				State.Save(Logger);
+				state = new UserWorkspaceState();
+				state.RootDir = rootDir;
+				ImportWorkspaceState(rootDir, clientName, branchPath, state);
+				state.Save(logger);
 			}
-			return State;
+			return state;
 		}
 
-		public UserWorkspaceState FindOrAddWorkspaceState(ProjectInfo ProjectInfo, UserWorkspaceSettings Settings, ILogger Logger)
+		public UserWorkspaceState FindOrAddWorkspaceState(ProjectInfo projectInfo, UserWorkspaceSettings settings, ILogger logger)
 		{
-			UserWorkspaceState State = FindOrAddWorkspaceState(ProjectInfo.LocalRootPath, ProjectInfo.ClientName, ProjectInfo.BranchPath, Logger);
-			if (!State.IsValid(ProjectInfo))
+			UserWorkspaceState state = FindOrAddWorkspaceState(projectInfo.LocalRootPath, projectInfo.ClientName, projectInfo.BranchPath, logger);
+			if (!state.IsValid(projectInfo))
 			{
-				State = new UserWorkspaceState();
+				state = new UserWorkspaceState();
 			}
-			State.UpdateCachedProjectInfo(ProjectInfo, Settings.LastModifiedTimeUtc);
-			return State;
+			state.UpdateCachedProjectInfo(projectInfo, settings.LastModifiedTimeUtc);
+			return state;
 		}
 
-		public UserWorkspaceSettings FindOrAddWorkspaceSettings(DirectoryReference RootDir, string? ServerAndPort, string? UserName, string ClientName, string BranchPath, string ProjectPath, ILogger Logger)
+		public UserWorkspaceSettings FindOrAddWorkspaceSettings(DirectoryReference rootDir, string? serverAndPort, string? userName, string clientName, string branchPath, string projectPath, ILogger logger)
 		{
-			ProjectInfo.ValidateBranchPath(BranchPath);
-			ProjectInfo.ValidateProjectPath(ProjectPath);
+			ProjectInfo.ValidateBranchPath(branchPath);
+			ProjectInfo.ValidateProjectPath(projectPath);
 
-			UserWorkspaceSettings? Settings;
-			if (!UserWorkspaceSettings.TryLoad(RootDir, out Settings))
+			UserWorkspaceSettings? settings;
+			if (!UserWorkspaceSettings.TryLoad(rootDir, out settings))
 			{
-				Settings = new UserWorkspaceSettings();
-				Settings.RootDir = RootDir;
-				ImportWorkspaceSettings(RootDir, ClientName, BranchPath, Settings);
+				settings = new UserWorkspaceSettings();
+				settings.RootDir = rootDir;
+				ImportWorkspaceSettings(rootDir, clientName, branchPath, settings);
 			}
 
-			Settings.Init(ServerAndPort, UserName, ClientName, BranchPath, ProjectPath);
-			Settings.Save(Logger);
+			settings.Init(serverAndPort, userName, clientName, branchPath, projectPath);
+			settings.Save(logger);
 
-			return Settings;
+			return settings;
 		}
 
-		public static string[] GetCombinedSyncFilter(Dictionary<Guid, WorkspaceSyncCategory> UniqueIdToFilter, FilterSettings GlobalFilter, FilterSettings WorkspaceFilter)
+		public static string[] GetCombinedSyncFilter(Dictionary<Guid, WorkspaceSyncCategory> uniqueIdToFilter, FilterSettings globalFilter, FilterSettings workspaceFilter)
 		{
-			List<string> Lines = new List<string>();
-			foreach (string ViewLine in Enumerable.Concat(GlobalFilter.View, WorkspaceFilter.View).Select(x => x.Trim()).Where(x => x.Length > 0 && !x.StartsWith(";")))
+			List<string> lines = new List<string>();
+			foreach (string viewLine in Enumerable.Concat(globalFilter.View, workspaceFilter.View).Select(x => x.Trim()).Where(x => x.Length > 0 && !x.StartsWith(";")))
 			{
-				Lines.Add(ViewLine);
+				lines.Add(viewLine);
 			}
 
-			Dictionary<Guid, bool> GlobalCategoryIdToSetting = GlobalFilter.GetCategories();
-			Dictionary<Guid, bool> WorkspaceCategoryIdToSetting = WorkspaceFilter.GetCategories();
+			Dictionary<Guid, bool> globalCategoryIdToSetting = globalFilter.GetCategories();
+			Dictionary<Guid, bool> workspaceCategoryIdToSetting = workspaceFilter.GetCategories();
 
-			HashSet<Guid> Enabled = new HashSet<Guid>();
-			foreach (WorkspaceSyncCategory Filter in UniqueIdToFilter.Values)
+			HashSet<Guid> enabled = new HashSet<Guid>();
+			foreach (WorkspaceSyncCategory filter in uniqueIdToFilter.Values)
 			{
-				bool bEnable = Filter.bEnable;
+				bool enable = filter.Enable;
 
-				bool bGlobalEnable;
-				if (GlobalCategoryIdToSetting.TryGetValue(Filter.UniqueId, out bGlobalEnable))
+				bool globalEnable;
+				if (globalCategoryIdToSetting.TryGetValue(filter.UniqueId, out globalEnable))
 				{
-					bEnable = bGlobalEnable;
+					enable = globalEnable;
 				}
 
-				bool bWorkspaceEnable;
-				if (WorkspaceCategoryIdToSetting.TryGetValue(Filter.UniqueId, out bWorkspaceEnable))
+				bool workspaceEnable;
+				if (workspaceCategoryIdToSetting.TryGetValue(filter.UniqueId, out workspaceEnable))
 				{
-					bEnable = bWorkspaceEnable;
+					enable = workspaceEnable;
 				}
 
-				if (bEnable)
+				if (enable)
 				{
-					EnableFilter(Filter.UniqueId, Enabled, UniqueIdToFilter);
+					EnableFilter(filter.UniqueId, enabled, uniqueIdToFilter);
 				}
 			}
 
-			foreach (WorkspaceSyncCategory Filter in UniqueIdToFilter.Values.OrderBy(x => x.Name))
+			foreach (WorkspaceSyncCategory filter in uniqueIdToFilter.Values.OrderBy(x => x.Name))
 			{
-				if (!Enabled.Contains(Filter.UniqueId))
+				if (!enabled.Contains(filter.UniqueId))
 				{
-					Lines.AddRange(Filter.Paths.Select(x => "-" + x.Trim()));
+					lines.AddRange(filter.Paths.Select(x => "-" + x.Trim()));
 				}
 			}
 
-			return Lines.ToArray();
+			return lines.ToArray();
 		}
 
-		static void EnableFilter(Guid UniqueId, HashSet<Guid> Enabled, Dictionary<Guid, WorkspaceSyncCategory> UniqueIdToFilter)
+		static void EnableFilter(Guid uniqueId, HashSet<Guid> enabled, Dictionary<Guid, WorkspaceSyncCategory> uniqueIdToFilter)
 		{
-			if (Enabled.Add(UniqueId))
+			if (enabled.Add(uniqueId))
 			{
-				WorkspaceSyncCategory? Category;
-				if (UniqueIdToFilter.TryGetValue(UniqueId, out Category))
+				WorkspaceSyncCategory? category;
+				if (uniqueIdToFilter.TryGetValue(uniqueId, out category))
 				{
-					foreach (Guid RequiresUniqueId in Category.Requires)
+					foreach (Guid requiresUniqueId in category.Requires)
 					{
-						EnableFilter(RequiresUniqueId, Enabled, UniqueIdToFilter);
+						EnableFilter(requiresUniqueId, enabled, uniqueIdToFilter);
 					}
 				}
 			}
diff --git a/Engine/Source/Programs/UnrealGameSync/UnrealGameSyncShared/OpenProjectInfo.cs b/Engine/Source/Programs/UnrealGameSync/UnrealGameSyncShared/OpenProjectInfo.cs
--- a/Engine/Source/Programs/UnrealGameSync/UnrealGameSyncShared/OpenProjectInfo.cs
+++ b/Engine/Source/Programs/UnrealGameSync/UnrealGameSyncShared/OpenProjectInfo.cs
@@ -27,264 +27,264 @@ namespace UnrealGameSync
 		public IReadOnlyList<string>? WorkspaceProjectStreamFilter { get; }
 		public List<KeyValuePair<FileReference, DateTime>> LocalConfigFiles { get; }
 
-		public OpenProjectInfo(UserSelectedProjectSettings SelectedProject, IPerforceSettings PerforceSettings, ProjectInfo ProjectInfo, UserWorkspaceSettings WorkspaceSettings, UserWorkspaceState WorkspaceState, ConfigFile LatestProjectConfigFile, ConfigFile WorkspaceProjectConfigFile, IReadOnlyList<string>? WorkspaceProjectStreamFilter, List<KeyValuePair<FileReference, DateTime>> LocalConfigFiles)
+		public OpenProjectInfo(UserSelectedProjectSettings selectedProject, IPerforceSettings perforceSettings, ProjectInfo projectInfo, UserWorkspaceSettings workspaceSettings, UserWorkspaceState workspaceState, ConfigFile latestProjectConfigFile, ConfigFile workspaceProjectConfigFile, IReadOnlyList<string>? workspaceProjectStreamFilter, List<KeyValuePair<FileReference, DateTime>> localConfigFiles)
 		{
-			this.SelectedProject = SelectedProject;
-
-			this.PerforceSettings = PerforceSettings;
-			this.ProjectInfo = ProjectInfo;
-			this.WorkspaceSettings = WorkspaceSettings;
-			this.WorkspaceState = WorkspaceState;
-			this.LatestProjectConfigFile = LatestProjectConfigFile;
-			this.WorkspaceProjectConfigFile = WorkspaceProjectConfigFile;
-			this.WorkspaceProjectStreamFilter = WorkspaceProjectStreamFilter;
-			this.LocalConfigFiles = LocalConfigFiles;
+			this.SelectedProject = selectedProject;
+
+			this.PerforceSettings = perforceSettings;
+			this.ProjectInfo = projectInfo;
+			this.WorkspaceSettings = workspaceSettings;
+			this.WorkspaceState = workspaceState;
+			this.LatestProjectConfigFile = latestProjectConfigFile;
+			this.WorkspaceProjectConfigFile = workspaceProjectConfigFile;
+			this.WorkspaceProjectStreamFilter = workspaceProjectStreamFilter;
+			this.LocalConfigFiles = localConfigFiles;
 		}
 
-		public static async Task<OpenProjectInfo> CreateAsync(IPerforceSettings DefaultPerforceSettings, UserSelectedProjectSettings SelectedProject, UserSettings UserSettings, ILogger<OpenProjectInfo> Logger, CancellationToken CancellationToken)
+		public static async Task<OpenProjectInfo> CreateAsync(IPerforceSettings defaultPerforceSettings, UserSelectedProjectSettings selectedProject, UserSettings userSettings, ILogger<OpenProjectInfo> logger, CancellationToken cancellationToken)
 		{
-			PerforceSettings PerforceSettings = Utility.OverridePerforceSettings(DefaultPerforceSettings, SelectedProject.ServerAndPort, SelectedProject.UserName);
-			using IPerforceConnection Perforce = await PerforceConnection.CreateAsync(PerforceSettings, Logger);
+			PerforceSettings perforceSettings = Utility.OverridePerforceSettings(defaultPerforceSettings, selectedProject.ServerAndPort, selectedProject.UserName);
+			using IPerforceConnection perforce = await PerforceConnection.CreateAsync(perforceSettings, logger);
 
 			// Make sure we're logged in
-			PerforceResponse<LoginRecord> LoginState = await Perforce.TryGetLoginStateAsync(CancellationToken);
-			if (!LoginState.Succeeded)
+			PerforceResponse<LoginRecord> loginState = await perforce.TryGetLoginStateAsync(cancellationToken);
+			if (!loginState.Succeeded)
 			{
 				throw new UserErrorException("User is not logged in to Perforce.");
 			}
 
 			// Execute like a regular task
-			return await CreateAsync(Perforce, SelectedProject, UserSettings, Logger, CancellationToken);
+			return await CreateAsync(perforce, selectedProject, userSettings, logger, cancellationToken);
 		}
 
-		public static async Task<OpenProjectInfo> CreateAsync(IPerforceConnection DefaultConnection, UserSelectedProjectSettings SelectedProject, UserSettings UserSettings, ILogger<OpenProjectInfo> Logger, CancellationToken CancellationToken)
+		public static async Task<OpenProjectInfo> CreateAsync(IPerforceConnection defaultConnection, UserSelectedProjectSettings selectedProject, UserSettings userSettings, ILogger<OpenProjectInfo> logger, CancellationToken cancellationToken)
 		{
-			using IDisposable LoggerScope = Logger.BeginScope("Project {SelectedProject}", SelectedProject.ToString());
-			Logger.LogInformation("Detecting settings for {Project}", SelectedProject);
+			using IDisposable loggerScope = logger.BeginScope("Project {SelectedProject}", selectedProject.ToString());
+			logger.LogInformation("Detecting settings for {Project}", selectedProject);
 
 			// Use the cached client path to the file if it's available; it's much quicker than trying to find the correct workspace.
-			IPerforceConnection? PerforceClient = null;
+			IPerforceConnection? perforceClient = null;
 			try
 			{
-				IPerforceSettings PerforceSettings;
+				IPerforceSettings perforceSettings;
 
-				FileReference NewSelectedFileName;
-				string NewSelectedClientFileName;
-				if (!String.IsNullOrEmpty(SelectedProject.ClientPath))
+				FileReference newSelectedFileName;
+				string newSelectedClientFileName;
+				if (!String.IsNullOrEmpty(selectedProject.ClientPath))
 				{
 					// Get the client path
-					NewSelectedClientFileName = SelectedProject.ClientPath;
+					newSelectedClientFileName = selectedProject.ClientPath;
 
 					// Get the client name
-					string? ClientName;
-					if (!PerforceUtils.TryGetClientName(NewSelectedClientFileName, out ClientName))
+					string? clientName;
+					if (!PerforceUtils.TryGetClientName(newSelectedClientFileName, out clientName))
 					{
-						throw new UserErrorException($"Couldn't get client name from {NewSelectedClientFileName}");
+						throw new UserErrorException($"Couldn't get client name from {newSelectedClientFileName}");
 					}
 
 					// Create the client
-					PerforceSettings = new PerforceSettings(DefaultConnection.Settings) { ClientName = ClientName };
-					PerforceClient = await PerforceConnection.CreateAsync(PerforceSettings, Logger);
+					perforceSettings = new PerforceSettings(defaultConnection.Settings) { ClientName = clientName };
+					perforceClient = await PerforceConnection.CreateAsync(perforceSettings, logger);
 
 					// Figure out the path on the client. Use the cached location if it's valid.
-					string? LocalPath = SelectedProject.LocalPath;
-					if (LocalPath == null || !File.Exists(LocalPath))
+					string? localPath = selectedProject.LocalPath;
+					if (localPath == null || !File.Exists(localPath))
 					{
-						List<WhereRecord> Records = await PerforceClient.WhereAsync(NewSelectedClientFileName, CancellationToken).Where(x => !x.Unmap).ToListAsync(CancellationToken);
-						if (Records.Count != 1)
+						List<WhereRecord> records = await perforceClient.WhereAsync(newSelectedClientFileName, cancellationToken).Where(x => !x.Unmap).ToListAsync(cancellationToken);
+						if (records.Count != 1)
 						{
-							throw new UserErrorException($"Couldn't get client path for {NewSelectedClientFileName}");
+							throw new UserErrorException($"Couldn't get client path for {newSelectedClientFileName}");
 						}
-						LocalPath = Path.GetFullPath(Records[0].Path);
+						localPath = Path.GetFullPath(records[0].Path);
 					}
-					NewSelectedFileName = new FileReference(LocalPath);
+					newSelectedFileName = new FileReference(localPath);
 				}
 				else
 				{
 					// Get the perforce server settings
-					InfoRecord PerforceInfo = await DefaultConnection.GetInfoAsync(InfoOptions.ShortOutput, CancellationToken);
+					InfoRecord perforceInfo = await defaultConnection.GetInfoAsync(InfoOptions.ShortOutput, cancellationToken);
 
 					// Use the path as the selected filename
-					NewSelectedFileName = new FileReference(SelectedProject.LocalPath!);
+					newSelectedFileName = new FileReference(selectedProject.LocalPath!);
 
 					// Make sure the project exists
-					if (!FileReference.Exists(NewSelectedFileName))
+					if (!FileReference.Exists(newSelectedFileName))
 					{
-						throw new UserErrorException($"{SelectedProject.LocalPath} does not exist.");
+						throw new UserErrorException($"{selectedProject.LocalPath} does not exist.");
 					}
 
 					// Find all the clients for this user
-					Logger.LogInformation("Enumerating clients for {UserName}...", PerforceInfo.UserName);
+					logger.LogInformation("Enumerating clients for {UserName}...", perforceInfo.UserName);
 
-					List<ClientsRecord> Clients = await DefaultConnection.GetClientsAsync(ClientsOptions.None, DefaultConnection.Settings.UserName, CancellationToken);
+					List<ClientsRecord> clients = await defaultConnection.GetClientsAsync(ClientsOptions.None, defaultConnection.Settings.UserName, cancellationToken);
 
-					List<IPerforceSettings> CandidateClients = await FilterClients(Clients, NewSelectedFileName, DefaultConnection.Settings, PerforceInfo.ClientHost, Logger, CancellationToken);
-					if (CandidateClients.Count == 0)
+					List<IPerforceSettings> candidateClients = await FilterClients(clients, newSelectedFileName, defaultConnection.Settings, perforceInfo.ClientHost, logger, cancellationToken);
+					if (candidateClients.Count == 0)
 					{
 						// Search through all workspaces. We may find a suitable workspace which is for any user.
-						Logger.LogInformation("Enumerating shared clients...");
-						Clients = await DefaultConnection.GetClientsAsync(ClientsOptions.None, "", CancellationToken);
+						logger.LogInformation("Enumerating shared clients...");
+						clients = await defaultConnection.GetClientsAsync(ClientsOptions.None, "", cancellationToken);
 
 						// Filter this list of clients
-						CandidateClients = await FilterClients(Clients, NewSelectedFileName, DefaultConnection.Settings, PerforceInfo.ClientHost, Logger, CancellationToken);
+						candidateClients = await FilterClients(clients, newSelectedFileName, defaultConnection.Settings, perforceInfo.ClientHost, logger, cancellationToken);
 
 						// If we still couldn't find any, fail.
-						if (CandidateClients.Count == 0)
+						if (candidateClients.Count == 0)
 						{
-							throw new UserErrorException($"Couldn't find any Perforce workspace containing {NewSelectedFileName}. Check your connection settings.");
+							throw new UserErrorException($"Couldn't find any Perforce workspace containing {newSelectedFileName}. Check your connection settings.");
 						}
 					}
 
 					// Check there's only one client
-					if (CandidateClients.Count > 1)
+					if (candidateClients.Count > 1)
 					{
-						throw new UserErrorException(String.Format("Found multiple workspaces containing {0}:\n\n{1}\n\nCannot determine which to use.", Path.GetFileName(NewSelectedFileName.GetFileName()), String.Join("\n", CandidateClients.Select(x => x.ClientName))));
+						throw new UserErrorException(String.Format("Found multiple workspaces containing {0}:\n\n{1}\n\nCannot determine which to use.", Path.GetFileName(newSelectedFileName.GetFileName()), String.Join("\n", candidateClients.Select(x => x.ClientName))));
 					}
 
 					// Take the client we've chosen
-					PerforceSettings = CandidateClients[0];
-					PerforceClient = await PerforceConnection.CreateAsync(PerforceSettings, Logger);
+					perforceSettings = candidateClients[0];
+					perforceClient = await PerforceConnection.CreateAsync(perforceSettings, logger);
 
 					// Get the client path for the project file
-					List<WhereRecord> Records = await PerforceClient.WhereAsync(NewSelectedFileName.FullName, CancellationToken).Where(x => !x.Unmap).ToListAsync(CancellationToken);
-					if (Records.Count == 0)
+					List<WhereRecord> records = await perforceClient.WhereAsync(newSelectedFileName.FullName, cancellationToken).Where(x => !x.Unmap).ToListAsync(cancellationToken);
+					if (records.Count == 0)
 					{
 						throw new UserErrorException("File is not mapped to any client");
 					}
-					else if (Records.Count > 1)
+					else if (records.Count > 1)
 					{
-						throw new UserErrorException($"File is mapped to {Records.Count} locations: {String.Join(", ", Records.Select(x => x.Path))}");
+						throw new UserErrorException($"File is mapped to {records.Count} locations: {String.Join(", ", records.Select(x => x.Path))}");
 					}
 
-					NewSelectedClientFileName = Records[0].ClientFile;
+					newSelectedClientFileName = records[0].ClientFile;
 				}
 
 				// Make sure the drive containing the project exists, to prevent other errors down the line
-				string PathRoot = Path.GetPathRoot(NewSelectedFileName.FullName)!;
-				if (!Directory.Exists(PathRoot))
+				string pathRoot = Path.GetPathRoot(newSelectedFileName.FullName)!;
+				if (!Directory.Exists(pathRoot))
 				{
-					throw new UserErrorException($"Path '{NewSelectedFileName}' is invalid");
+					throw new UserErrorException($"Path '{newSelectedFileName}' is invalid");
 				}
 
 				// Make sure the path case is correct. This can cause UBT intermediates to be out of date if the case mismatches.
-				NewSelectedFileName = FileReference.FindCorrectCase(NewSelectedFileName);
+				newSelectedFileName = FileReference.FindCorrectCase(newSelectedFileName);
 
 				// Update the selected project with all the data we've found
-				SelectedProject = new UserSelectedProjectSettings(SelectedProject.ServerAndPort, SelectedProject.UserName, SelectedProject.Type, NewSelectedClientFileName, NewSelectedFileName.FullName);
+				selectedProject = new UserSelectedProjectSettings(selectedProject.ServerAndPort, selectedProject.UserName, selectedProject.Type, newSelectedClientFileName, newSelectedFileName.FullName);
 
 				// Get the local branch root
-				string? BranchClientPath = null;
-				DirectoryReference? BranchDirectoryName = null;
+				string? branchClientPath = null;
+				DirectoryReference? branchDirectoryName = null;
 
 				// Figure out where the engine is in relation to it
-				int EndIdx = NewSelectedClientFileName.Length - 1;
-				if (EndIdx != -1 && NewSelectedClientFileName.EndsWith(".uproject", StringComparison.InvariantCultureIgnoreCase))
+				int endIdx = newSelectedClientFileName.Length - 1;
+				if (endIdx != -1 && newSelectedClientFileName.EndsWith(".uproject", StringComparison.InvariantCultureIgnoreCase))
 				{
-					EndIdx = NewSelectedClientFileName.LastIndexOf('/') - 1;
+					endIdx = newSelectedClientFileName.LastIndexOf('/') - 1;
 				}
-				for (; EndIdx >= 2; EndIdx--)
+				for (; endIdx >= 2; endIdx--)
 				{
-					if (NewSelectedClientFileName[EndIdx] == '/')
+					if (newSelectedClientFileName[endIdx] == '/')
 					{
-						List<PerforceResponse<FStatRecord>> FileRecords = await PerforceClient.TryFStatAsync(FStatOptions.None, NewSelectedClientFileName.Substring(0, EndIdx) + "/Engine/Build/Build.version", CancellationToken).ToListAsync();
-						if (FileRecords.Succeeded() && FileRecords.Count > 0)
+						List<PerforceResponse<FStatRecord>> fileRecords = await perforceClient.TryFStatAsync(FStatOptions.None, newSelectedClientFileName.Substring(0, endIdx) + "/Engine/Build/Build.version", cancellationToken).ToListAsync();
+						if (fileRecords.Succeeded() && fileRecords.Count > 0)
 						{
-							FStatRecord FileRecord = FileRecords[0].Data;
-							if (FileRecord.ClientFile == null)
+							FStatRecord fileRecord = fileRecords[0].Data;
+							if (fileRecord.ClientFile == null)
 							{
-								throw new UserErrorException($"Missing client path for {FileRecord.DepotFile}");
+								throw new UserErrorException($"Missing client path for {fileRecord.DepotFile}");
 							}
 
-							BranchClientPath = NewSelectedClientFileName.Substring(0, EndIdx);
-							BranchDirectoryName = new FileReference(FileRecord.ClientFile).Directory.ParentDirectory?.ParentDirectory;
+							branchClientPath = newSelectedClientFileName.Substring(0, endIdx);
+							branchDirectoryName = new FileReference(fileRecord.ClientFile).Directory.ParentDirectory?.ParentDirectory;
 							break;
 						}
 					}
 				}
-				if(BranchClientPath == null || BranchDirectoryName == null)
+				if(branchClientPath == null || branchDirectoryName == null)
 				{
-					throw new UserErrorException($"Could not find engine in Perforce relative to project path ({NewSelectedClientFileName})");
+					throw new UserErrorException($"Could not find engine in Perforce relative to project path ({newSelectedClientFileName})");
 				}
 
-				Logger.LogInformation("Found branch root at {RootPath}", BranchClientPath);
+				logger.LogInformation("Found branch root at {RootPath}", branchClientPath);
 
 				// Read the existing workspace settings from disk, and update them with any info computed here
-				int BranchIdx = BranchClientPath.IndexOf('/', 2);
-				string BranchPath = (BranchIdx == -1) ? String.Empty : BranchClientPath.Substring(BranchIdx);
-				string ProjectPath = NewSelectedClientFileName.Substring(BranchClientPath.Length);
-				UserWorkspaceSettings UserWorkspaceSettings = UserSettings.FindOrAddWorkspaceSettings(BranchDirectoryName, PerforceSettings.ServerAndPort, PerforceSettings.UserName, PerforceSettings.ClientName!, BranchPath, ProjectPath, Logger);
+				int branchIdx = branchClientPath.IndexOf('/', 2);
+				string branchPath = (branchIdx == -1) ? String.Empty : branchClientPath.Substring(branchIdx);
+				string projectPath = newSelectedClientFileName.Substring(branchClientPath.Length);
+				UserWorkspaceSettings userWorkspaceSettings = userSettings.FindOrAddWorkspaceSettings(branchDirectoryName, perforceSettings.ServerAndPort, perforceSettings.UserName, perforceSettings.ClientName!, branchPath, projectPath, logger);
 
 				// Now compute the updated project info
-				ProjectInfo ProjectInfo = await ProjectInfo.CreateAsync(PerforceClient, UserWorkspaceSettings, CancellationToken);
+				ProjectInfo projectInfo = await ProjectInfo.CreateAsync(perforceClient, userWorkspaceSettings, cancellationToken);
 
 				// Update the cached workspace state
-				UserWorkspaceState UserWorkspaceState = UserSettings.FindOrAddWorkspaceState(ProjectInfo, UserWorkspaceSettings, Logger);
+				UserWorkspaceState userWorkspaceState = userSettings.FindOrAddWorkspaceState(projectInfo, userWorkspaceSettings, logger);
 
 				// Read the initial config file
-				List<KeyValuePair<FileReference, DateTime>> LocalConfigFiles = new List<KeyValuePair<FileReference, DateTime>>();
-				ConfigFile LatestProjectConfigFile = await ConfigUtils.ReadProjectConfigFileAsync(PerforceClient, ProjectInfo, LocalConfigFiles, Logger, CancellationToken);
+				List<KeyValuePair<FileReference, DateTime>> localConfigFiles = new List<KeyValuePair<FileReference, DateTime>>();
+				ConfigFile latestProjectConfigFile = await ConfigUtils.ReadProjectConfigFileAsync(perforceClient, projectInfo, localConfigFiles, logger, cancellationToken);
 
 				// Get the local config file and stream filter
-				ConfigFile WorkspaceProjectConfigFile = await WorkspaceUpdate.ReadProjectConfigFile(BranchDirectoryName, NewSelectedFileName, Logger);
-				IReadOnlyList<string>? WorkspaceProjectStreamFilter = await WorkspaceUpdate.ReadProjectStreamFilter(PerforceClient, WorkspaceProjectConfigFile, Logger, CancellationToken);
+				ConfigFile workspaceProjectConfigFile = await WorkspaceUpdate.ReadProjectConfigFile(branchDirectoryName, newSelectedFileName, logger);
+				IReadOnlyList<string>? workspaceProjectStreamFilter = await WorkspaceUpdate.ReadProjectStreamFilter(perforceClient, workspaceProjectConfigFile, logger, cancellationToken);
 
-				OpenProjectInfo WorkspaceSettings = new OpenProjectInfo(SelectedProject, PerforceSettings, ProjectInfo, UserWorkspaceSettings, UserWorkspaceState, LatestProjectConfigFile, WorkspaceProjectConfigFile, WorkspaceProjectStreamFilter, LocalConfigFiles);
+				OpenProjectInfo workspaceSettings = new OpenProjectInfo(selectedProject, perforceSettings, projectInfo, userWorkspaceSettings, userWorkspaceState, latestProjectConfigFile, workspaceProjectConfigFile, workspaceProjectStreamFilter, localConfigFiles);
 
-				return WorkspaceSettings;
+				return workspaceSettings;
 			}
 			finally
 			{
-				PerforceClient?.Dispose();
+				perforceClient?.Dispose();
 			}
 		}
 
-		static async Task<List<IPerforceSettings>> FilterClients(List<ClientsRecord> Clients, FileReference NewSelectedFileName, IPerforceSettings DefaultPerforceSettings, string? HostName, ILogger Logger, CancellationToken CancellationToken)
+		static async Task<List<IPerforceSettings>> FilterClients(List<ClientsRecord> clients, FileReference newSelectedFileName, IPerforceSettings defaultPerforceSettings, string? hostName, ILogger logger, CancellationToken cancellationToken)
 		{
-			List<IPerforceSettings> CandidateClients = new List<IPerforceSettings>();
-			foreach(ClientsRecord Client in Clients)
+			List<IPerforceSettings> candidateClients = new List<IPerforceSettings>();
+			foreach(ClientsRecord client in clients)
 			{
 				// Make sure the client is well formed
-				if(!String.IsNullOrEmpty(Client.Name) && (!String.IsNullOrEmpty(Client.Host) || !String.IsNullOrEmpty(Client.Owner)) && !String.IsNullOrEmpty(Client.Root))
+				if(!String.IsNullOrEmpty(client.Name) && (!String.IsNullOrEmpty(client.Host) || !String.IsNullOrEmpty(client.Owner)) && !String.IsNullOrEmpty(client.Root))
 				{
 					// Require either a username or host name match
-					if((String.IsNullOrEmpty(Client.Host) || String.Compare(Client.Host, HostName, StringComparison.OrdinalIgnoreCase) == 0) && (String.IsNullOrEmpty(Client.Owner) || String.Compare(Client.Owner, DefaultPerforceSettings.UserName, StringComparison.OrdinalIgnoreCase) == 0))
+					if((String.IsNullOrEmpty(client.Host) || String.Compare(client.Host, hostName, StringComparison.OrdinalIgnoreCase) == 0) && (String.IsNullOrEmpty(client.Owner) || String.Compare(client.Owner, defaultPerforceSettings.UserName, StringComparison.OrdinalIgnoreCase) == 0))
 					{
-						if(!Utility.SafeIsFileUnderDirectory(NewSelectedFileName.FullName, Client.Root))
+						if(!Utility.SafeIsFileUnderDirectory(newSelectedFileName.FullName, client.Root))
 						{
-							Logger.LogInformation("Rejecting {ClientName} due to root mismatch ({RootPath})", Client.Name, Client.Root);
+							logger.LogInformation("Rejecting {ClientName} due to root mismatch ({RootPath})", client.Name, client.Root);
 							continue;
 						}
 
-						PerforceSettings CandidateSettings = new PerforceSettings(DefaultPerforceSettings) { ClientName = Client.Name };
-						using IPerforceConnection CandidateClient = await PerforceConnection.CreateAsync(CandidateSettings, Logger);
+						PerforceSettings candidateSettings = new PerforceSettings(defaultPerforceSettings) { ClientName = client.Name };
+						using IPerforceConnection candidateClient = await PerforceConnection.CreateAsync(candidateSettings, logger);
 
-						List<PerforceResponse<WhereRecord>> WhereRecords = await CandidateClient.TryWhereAsync(NewSelectedFileName.FullName, CancellationToken).Where(x => x.Failed || !x.Data.Unmap).ToListAsync(CancellationToken);
-						if(!WhereRecords.Succeeded() || WhereRecords.Count != 1)
+						List<PerforceResponse<WhereRecord>> whereRecords = await candidateClient.TryWhereAsync(newSelectedFileName.FullName, cancellationToken).Where(x => x.Failed || !x.Data.Unmap).ToListAsync(cancellationToken);
+						if(!whereRecords.Succeeded() || whereRecords.Count != 1)
 						{
-							Logger.LogInformation("Rejecting {ClientName} due to file not existing in workspace", Client.Name);
+							logger.LogInformation("Rejecting {ClientName} due to file not existing in workspace", client.Name);
 							continue;
 						}
 
-						List<PerforceResponse<FStatRecord>> Records = await CandidateClient.TryFStatAsync(FStatOptions.None, NewSelectedFileName.FullName, CancellationToken).ToListAsync(CancellationToken);
-						if (!Records.Succeeded())
+						List<PerforceResponse<FStatRecord>> records = await candidateClient.TryFStatAsync(FStatOptions.None, newSelectedFileName.FullName, cancellationToken).ToListAsync(cancellationToken);
+						if (!records.Succeeded())
 						{
-							Logger.LogInformation("Rejecting {ClientName} due to {FileName} not in depot", Client.Name, NewSelectedFileName);
+							logger.LogInformation("Rejecting {ClientName} due to {FileName} not in depot", client.Name, newSelectedFileName);
 							continue;
 						}
 
-						Records.RemoveAll(x => !x.Data.IsMapped);
-						if(Records.Count == 0)
+						records.RemoveAll(x => !x.Data.IsMapped);
+						if(records.Count == 0)
 						{
-							Logger.LogInformation("Rejecting {ClientName} due to {NumRecords} matching records", Client.Name, Records.Count);
+							logger.LogInformation("Rejecting {ClientName} due to {NumRecords} matching records", client.Name, records.Count);
 							continue;
 						}
 
-						Logger.LogInformation("Found valid client {ClientName}", Client.Name);
-						CandidateClients.Add(CandidateSettings);
+						logger.LogInformation("Found valid client {ClientName}", client.Name);
+						candidateClients.Add(candidateSettings);
 					}
 				}
 			}
-			return CandidateClients;
+			return candidateClients;
 		}
 	}
 }
diff --git a/Engine/Source/Programs/UnrealGameSync/UnrealGameSyncShared/OutputAdapters.cs b/Engine/Source/Programs/UnrealGameSync/UnrealGameSyncShared/OutputAdapters.cs
--- a/Engine/Source/Programs/UnrealGameSync/UnrealGameSyncShared/OutputAdapters.cs
+++ b/Engine/Source/Programs/UnrealGameSync/UnrealGameSyncShared/OutputAdapters.cs
@@ -15,29 +15,29 @@ namespace UnrealGameSync
 {
 	class PrefixedTextWriter : ILogger
 	{
-		string Prefix;
-		ILogger Inner;
+		string _prefix;
+		ILogger _inner;
 
-		public PrefixedTextWriter(string InPrefix, ILogger InInner)
+		public PrefixedTextWriter(string inPrefix, ILogger inInner)
 		{
-			Prefix = InPrefix;
-			Inner = InInner;
+			_prefix = inPrefix;
+			_inner = inInner;
 		}
 
-		public IDisposable BeginScope<TState>(TState State) => Inner.BeginScope(State);
+		public IDisposable BeginScope<TState>(TState state) => _inner.BeginScope(state);
 
-		public bool IsEnabled(LogLevel LogLevel) => Inner.IsEnabled(LogLevel);
+		public bool IsEnabled(LogLevel logLevel) => _inner.IsEnabled(logLevel);
 
-		public void Log<TState>(LogLevel LogLevel, EventId EventId, TState State, Exception Exception, Func<TState, Exception, string> Formatter)
+		public void Log<TState>(LogLevel logLevel, EventId eventId, TState state, Exception exception, Func<TState, Exception, string> formatter)
 		{
-			Inner.Log(LogLevel, EventId, State, Exception, (State, Exception) => Prefix + Formatter(State, Exception));
+			_inner.Log(logLevel, eventId, state, exception, (state, exception) => _prefix + formatter(state, exception));
 		}
 	}
 
 	public class ProgressValue
 	{
-		Tuple<string, float> State = null!;
-		Stack<Tuple<float, float>> Ranges = new Stack<Tuple<float,float>>();
+		Tuple<string, float> _state = null!;
+		Stack<Tuple<float, float>> _ranges = new Stack<Tuple<float,float>>();
 
 		public ProgressValue()
 		{
@@ -46,64 +46,64 @@ namespace UnrealGameSync
 
 		public void Clear()
 		{
-			State = new Tuple<string,float>("Starting...", 0.0f);
+			_state = new Tuple<string,float>("Starting...", 0.0f);
 
-			Ranges.Clear();
-			Ranges.Push(new Tuple<float, float>(0.0f, 1.0f));
+			_ranges.Clear();
+			_ranges.Push(new Tuple<float, float>(0.0f, 1.0f));
 		}
 
 		public Tuple<string, float> Current
 		{
-			get { return State; }
+			get { return _state; }
 		}
 
-		public void Set(string Message)
+		public void Set(string message)
 		{
-			if(Ranges.Count == 1)
+			if(_ranges.Count == 1)
 			{
-				State = new Tuple<string,float>(Message, State.Item2);
+				_state = new Tuple<string,float>(message, _state.Item2);
 			}
 		}
 
-		public void Set(string Message, float Fraction)
+		public void Set(string message, float fraction)
 		{
-			if(Ranges.Count == 1)
+			if(_ranges.Count == 1)
 			{
-				State = new Tuple<string, float>(Message, RelativeToAbsoluteFraction(Fraction));
+				_state = new Tuple<string, float>(message, RelativeToAbsoluteFraction(fraction));
 			}
 			else
 			{
-				State = new Tuple<string, float>(State.Item1, RelativeToAbsoluteFraction(Fraction));
+				_state = new Tuple<string, float>(_state.Item1, RelativeToAbsoluteFraction(fraction));
 			}
 		}
 
-		public void Set(float Fraction)
+		public void Set(float fraction)
 		{
-			State = new Tuple<string, float>(State.Item1, RelativeToAbsoluteFraction(Fraction));
+			_state = new Tuple<string, float>(_state.Item1, RelativeToAbsoluteFraction(fraction));
 		}
 
-		public void Increment(float Fraction)
+		public void Increment(float fraction)
 		{
-			Set(State.Item2 + RelativeToAbsoluteFraction(Fraction));
+			Set(_state.Item2 + RelativeToAbsoluteFraction(fraction));
 		}
 
-		public void Push(float MaxFraction)
+		public void Push(float maxFraction)
 		{
-			Ranges.Push(new Tuple<float,float>(State.Item2, RelativeToAbsoluteFraction(MaxFraction)));
+			_ranges.Push(new Tuple<float,float>(_state.Item2, RelativeToAbsoluteFraction(maxFraction)));
 		}
 
 		public void Pop()
 		{
-			if(Ranges.Count > 1)
+			if(_ranges.Count > 1)
 			{
-				State = new Tuple<string,float>(State.Item1, Ranges.Pop().Item2);
+				_state = new Tuple<string,float>(_state.Item1, _ranges.Pop().Item2);
 			}
 		}
 
-		float RelativeToAbsoluteFraction(float Fraction)
+		float RelativeToAbsoluteFraction(float fraction)
 		{
-			Tuple<float, float> Range = Ranges.Peek();
-			return Range.Item1 + (Range.Item2 - Range.Item1) * Fraction;
+			Tuple<float, float> range = _ranges.Peek();
+			return range.Item1 + (range.Item2 - range.Item1) * fraction;
 		}
 	}
 
@@ -111,167 +111,167 @@ namespace UnrealGameSync
 	{
 		const string DirectivePrefix = "@progress ";
 
-		public static string? ParseLine(string Line, ProgressValue Value)
+		public static string? ParseLine(string line, ProgressValue value)
 		{
-			string TrimLine = Line.Trim();
-			if(TrimLine.StartsWith(DirectivePrefix))
+			string trimLine = line.Trim();
+			if(trimLine.StartsWith(DirectivePrefix))
 			{
 				// Line that just contains a progress directive
-				bool bSkipLine = false;
-				ProcessInternal(TrimLine.Substring(DirectivePrefix.Length), ref bSkipLine, Value);
+				bool skipLine = false;
+				ProcessInternal(trimLine.Substring(DirectivePrefix.Length), ref skipLine, value);
 				return null;
 			}
 			else
 			{
-				bool bSkipLine = false;
-				string RemainingLine = Line;
+				bool skipLine = false;
+				string remainingLine = line;
 
 				// Look for a progress directive at the end of a line, in square brackets
-				if(TrimLine.EndsWith("]"))
+				if(trimLine.EndsWith("]"))
 				{
-					for(int LastIdx = TrimLine.Length - 2; LastIdx >= 0 && TrimLine[LastIdx] != ']'; LastIdx--)
+					for(int lastIdx = trimLine.Length - 2; lastIdx >= 0 && trimLine[lastIdx] != ']'; lastIdx--)
 					{
-						if(TrimLine[LastIdx] == '[')
+						if(trimLine[lastIdx] == '[')
 						{
-							string DirectiveSubstring = TrimLine.Substring(LastIdx + 1, TrimLine.Length - LastIdx - 2);
-							if(DirectiveSubstring.StartsWith(DirectivePrefix))
+							string directiveSubstring = trimLine.Substring(lastIdx + 1, trimLine.Length - lastIdx - 2);
+							if(directiveSubstring.StartsWith(DirectivePrefix))
 							{
-								ProcessInternal(DirectiveSubstring.Substring(DirectivePrefix.Length), ref bSkipLine, Value);
-								RemainingLine = Line.Substring(0, LastIdx).TrimEnd();
+								ProcessInternal(directiveSubstring.Substring(DirectivePrefix.Length), ref skipLine, value);
+								remainingLine = line.Substring(0, lastIdx).TrimEnd();
 							}
 							break;
 						}
 					}
 				}
 
-				if (bSkipLine)
+				if (skipLine)
 				{
 					return null;
 				}
 				else
 				{
-					return RemainingLine;
+					return remainingLine;
 				}
 			}
 		}
 
-		static void ProcessInternal(string Line, ref bool bSkipLine, ProgressValue Value)
+		static void ProcessInternal(string line, ref bool skipLine, ProgressValue value)
 		{
-			List<string> Tokens = ParseTokens(Line);
-			for(int TokenIdx = 0; TokenIdx < Tokens.Count; )
+			List<string> tokens = ParseTokens(line);
+			for(int tokenIdx = 0; tokenIdx < tokens.Count; )
 			{
-				float Fraction;
-				if(ReadFraction(Tokens, ref TokenIdx, out Fraction))
+				float fraction;
+				if(ReadFraction(tokens, ref tokenIdx, out fraction))
 				{
-					Value.Set(Fraction);
+					value.Set(fraction);
 				}
-				else if(Tokens[TokenIdx] == "push")
+				else if(tokens[tokenIdx] == "push")
 				{
-					TokenIdx++;
-					if(ReadFraction(Tokens, ref TokenIdx, out Fraction))
+					tokenIdx++;
+					if(ReadFraction(tokens, ref tokenIdx, out fraction))
 					{
-						Value.Push(Fraction);
+						value.Push(fraction);
 					}
 				}
-				else if(Tokens[TokenIdx] == "pop")
+				else if(tokens[tokenIdx] == "pop")
 				{
-					TokenIdx++;
-					Value.Pop();
+					tokenIdx++;
+					value.Pop();
 				}
-				else if(Tokens[TokenIdx] == "increment")
+				else if(tokens[tokenIdx] == "increment")
 				{
-					TokenIdx++;
-					if(ReadFraction(Tokens, ref TokenIdx, out Fraction))
+					tokenIdx++;
+					if(ReadFraction(tokens, ref tokenIdx, out fraction))
 					{
-						Value.Increment(Fraction);
+						value.Increment(fraction);
 					}
 				}
-				else if(Tokens[TokenIdx] == "skipline")
+				else if(tokens[tokenIdx] == "skipline")
 				{
-					TokenIdx++;
-					bSkipLine = true;
+					tokenIdx++;
+					skipLine = true;
 				}
-				else if(Tokens[TokenIdx].Length >= 2 && (Tokens[TokenIdx][0] == '\'' || Tokens[TokenIdx][0] == '\"') && Tokens[TokenIdx].Last() == Tokens[TokenIdx].First())
+				else if(tokens[tokenIdx].Length >= 2 && (tokens[tokenIdx][0] == '\'' || tokens[tokenIdx][0] == '\"') && tokens[tokenIdx].Last() == tokens[tokenIdx].First())
 				{
-					string Message = Tokens[TokenIdx++];
-					Value.Set(Message.Substring(1, Message.Length - 2));
+					string message = tokens[tokenIdx++];
+					value.Set(message.Substring(1, message.Length - 2));
 				}
 				else
 				{
-					TokenIdx++;
+					tokenIdx++;
 				}
 			}
 		}
 
-		static List<string> ParseTokens(string Line)
+		static List<string> ParseTokens(string line)
 		{
-			List<string> Tokens = new List<string>();
-			for(int Idx = 0;;)
+			List<string> tokens = new List<string>();
+			for(int idx = 0;;)
 			{
 				// Skip whitespace
-				while(Idx < Line.Length && Char.IsWhiteSpace(Line[Idx]))
+				while(idx < line.Length && Char.IsWhiteSpace(line[idx]))
 				{
-					Idx++;
+					idx++;
 				}
-				if(Idx == Line.Length)
+				if(idx == line.Length)
 				{
 					break;
 				}
 
 				// Read the next token
-				if(Char.IsLetterOrDigit(Line[Idx]))
+				if(Char.IsLetterOrDigit(line[idx]))
 				{
-					int StartIdx = Idx++;
-					while(Idx < Line.Length && Char.IsLetterOrDigit(Line[Idx]))
+					int startIdx = idx++;
+					while(idx < line.Length && Char.IsLetterOrDigit(line[idx]))
 					{
-						Idx++;
+						idx++;
 					}
-					Tokens.Add(Line.Substring(StartIdx, Idx - StartIdx));
+					tokens.Add(line.Substring(startIdx, idx - startIdx));
 				}
-				else if(Line[Idx] == '\'' || Line[Idx] == '\"')
+				else if(line[idx] == '\'' || line[idx] == '\"')
 				{
-					int StartIdx = Idx++;
-					while(Idx < Line.Length && Line[Idx] != Line[StartIdx])
+					int startIdx = idx++;
+					while(idx < line.Length && line[idx] != line[startIdx])
 					{
-						Idx++;
+						idx++;
 					}
-					Tokens.Add(Line.Substring(StartIdx, ++Idx - StartIdx));
+					tokens.Add(line.Substring(startIdx, ++idx - startIdx));
 				}
 				else
 				{
-					Tokens.Add(Line.Substring(Idx++, 1));
+					tokens.Add(line.Substring(idx++, 1));
 				}
 			}
-			return Tokens;
+			return tokens;
 		}
 
-		static bool ReadFraction(List<string> Tokens, ref int TokenIdx, out float Fraction)
+		static bool ReadFraction(List<string> tokens, ref int tokenIdx, out float fraction)
 		{
 			// Read a fraction in the form x%
-			if(TokenIdx + 2 <= Tokens.Count && Tokens[TokenIdx + 1] == "%")
+			if(tokenIdx + 2 <= tokens.Count && tokens[tokenIdx + 1] == "%")
 			{
-				int Numerator;
-				if(int.TryParse(Tokens[TokenIdx], out Numerator))
+				int numerator;
+				if(int.TryParse(tokens[tokenIdx], out numerator))
 				{
-					Fraction = (float)Numerator / 100.0f;
-					TokenIdx += 2;
+					fraction = (float)numerator / 100.0f;
+					tokenIdx += 2;
 					return true;
 				}
 			}
 			
 			// Read a fraction in the form x/y
-			if(TokenIdx + 3 <= Tokens.Count && Tokens[TokenIdx + 1] == "/")
+			if(tokenIdx + 3 <= tokens.Count && tokens[tokenIdx + 1] == "/")
 			{
-				int Numerator, Denominator;
-				if(int.TryParse(Tokens[TokenIdx], out Numerator) && int.TryParse(Tokens[TokenIdx + 2], out Denominator))
+				int numerator, denominator;
+				if(int.TryParse(tokens[tokenIdx], out numerator) && int.TryParse(tokens[tokenIdx + 2], out denominator))
 				{
-					Fraction = (float)Numerator / (float)Denominator;
-					TokenIdx += 3;
+					fraction = (float)numerator / (float)denominator;
+					tokenIdx += 3;
 					return true;
 				}
 			}
 
-			Fraction = 0.0f;
+			fraction = 0.0f;
 			return false;
 		}
 	}
diff --git a/Engine/Source/Programs/UnrealGameSync/UnrealGameSyncShared/Telemetry.cs b/Engine/Source/Programs/UnrealGameSync/UnrealGameSyncShared/Telemetry.cs
--- a/Engine/Source/Programs/UnrealGameSync/UnrealGameSyncShared/Telemetry.cs
+++ b/Engine/Source/Programs/UnrealGameSync/UnrealGameSyncShared/Telemetry.cs
@@ -20,9 +20,9 @@ namespace UnrealGameSync
 		/// <summary>
 		/// Sends a telemetry event with the given information
 		/// </summary>
-		/// <param name="EventName">Name of the event</param>
-		/// <param name="Attributes">Arbitrary object to include in the payload</param>
-		void SendEvent(string EventName, object Attributes);
+		/// <param name="eventName">Name of the event</param>
+		/// <param name="attributes">Arbitrary object to include in the payload</param>
+		void SendEvent(string eventName, object attributes);
 	}
 
 	/// <summary>
@@ -36,7 +36,7 @@ namespace UnrealGameSync
 		}
 
 		/// <inheritdoc/>
-		public void SendEvent(string EventName, object Attributes)
+		public void SendEvent(string eventName, object attributes)
 		{
 		}
 	}
@@ -49,37 +49,37 @@ namespace UnrealGameSync
 		/// <summary>
 		/// Combined url to post event streams to
 		/// </summary>
-		string Url;
+		string _url;
 
 		/// <summary>
 		/// Lock used to modify the event queue
 		/// </summary>
-		object LockObject = new object();
+		object _lockObject = new object();
 
 		/// <summary>
 		/// Whether a flush is queued
 		/// </summary>
-		bool bHasPendingFlush = false;
+		bool _hasPendingFlush = false;
 
 		/// <summary>
 		/// List of pending events
 		/// </summary>
-		List<string> PendingEvents = new List<string>();
+		List<string> _pendingEvents = new List<string>();
 
 		/// <summary>
 		/// The log writer to use
 		/// </summary>
-		ILogger Logger;
+		ILogger _logger;
 
 		/// <summary>
 		/// Constructor
 		/// </summary>
-		public EpicTelemetrySink(string Url, ILogger Logger)
+		public EpicTelemetrySink(string url, ILogger logger)
 		{
-			this.Url = Url;
-			this.Logger = Logger;
+			this._url = url;
+			this._logger = logger;
 
-			Logger.LogInformation("Posting to URL: {Url}", Url);
+			logger.LogInformation("Posting to URL: {Url}", url);
 		}
 
 		/// <inheritdoc/>
@@ -93,9 +93,9 @@ namespace UnrealGameSync
 		{
 			for (; ; )
 			{
-				lock (LockObject)
+				lock (_lockObject)
 				{
-					if (!bHasPendingFlush)
+					if (!_hasPendingFlush)
 					{
 						break;
 					}
@@ -105,22 +105,22 @@ namespace UnrealGameSync
 		}
 
 		/// <inheritdoc/>
-		public void SendEvent(string EventName, object Attributes)
+		public void SendEvent(string eventName, object attributes)
 		{
-			string AttributesText = JsonSerializer.Serialize(Attributes);
-			if (AttributesText[0] != '{')
+			string attributesText = JsonSerializer.Serialize(attributes);
+			if (attributesText[0] != '{')
 			{
 				throw new Exception("Expected event data with named properties");
 			}
 
-			string EventText = AttributesText.Insert(1, String.Format("\"EventName\":\"{0}\",", HttpUtility.JavaScriptStringEncode(EventName)));
-			lock (PendingEvents)
+			string eventText = attributesText.Insert(1, String.Format("\"EventName\":\"{0}\",", HttpUtility.JavaScriptStringEncode(eventName)));
+			lock (_pendingEvents)
 			{
-				PendingEvents.Add(EventText);
-				if (!bHasPendingFlush)
+				_pendingEvents.Add(eventText);
+				if (!_hasPendingFlush)
 				{
-					ThreadPool.QueueUserWorkItem(Obj => BackgroundFlush());
-					bHasPendingFlush = true;
+					ThreadPool.QueueUserWorkItem(obj => BackgroundFlush());
+					_hasPendingFlush = true;
 				}
 			}
 		}
@@ -135,71 +135,71 @@ namespace UnrealGameSync
 				try
 				{
 					// Generate the content for this event
-					List<string> Events = new List<string>();
-					lock (LockObject)
+					List<string> events = new List<string>();
+					lock (_lockObject)
 					{
-						if (PendingEvents.Count == 0)
+						if (_pendingEvents.Count == 0)
 						{
-							bHasPendingFlush = false;
+							_hasPendingFlush = false;
 							break;
 						}
 
-						Events.AddRange(PendingEvents);
-						PendingEvents.Clear();
+						events.AddRange(_pendingEvents);
+						_pendingEvents.Clear();
 					}
 
 					// Print all the events we're sending
-					foreach (string Event in Events)
+					foreach (string evt in events)
 					{
-						Logger.LogInformation("Sending Event: {0}", Event);
+						_logger.LogInformation("Sending Event: {0}", evt);
 					}
 
 					// Convert the content to UTF8
-					string ContentText = String.Format("{{\"Events\":[{0}]}}", String.Join(",", Events));
-					byte[] Content = Encoding.UTF8.GetBytes(ContentText);
+					string contentText = String.Format("{{\"Events\":[{0}]}}", String.Join(",", events));
+					byte[] content = Encoding.UTF8.GetBytes(contentText);
 
 					// Post the event data
-					HttpWebRequest Request = (HttpWebRequest)WebRequest.Create(Url);
-					Request.Method = "POST";
-					Request.ContentType = "application/json";
-					Request.UserAgent = "ue/ugs";
-					Request.Timeout = 5000;
-					Request.ContentLength = Content.Length;
-					Request.ContentType = "application/json";
-					using (Stream RequestStream = Request.GetRequestStream())
+					HttpWebRequest request = (HttpWebRequest)WebRequest.Create(_url);
+					request.Method = "POST";
+					request.ContentType = "application/json";
+					request.UserAgent = "ue/ugs";
+					request.Timeout = 5000;
+					request.ContentLength = content.Length;
+					request.ContentType = "application/json";
+					using (Stream requestStream = request.GetRequestStream())
 					{
-						RequestStream.Write(Content, 0, Content.Length);
+						requestStream.Write(content, 0, content.Length);
 					}
 
 					// Wait for the response and dispose of it immediately
-					using (HttpWebResponse Response = (HttpWebResponse)Request.GetResponse())
+					using (HttpWebResponse response = (HttpWebResponse)request.GetResponse())
 					{
-						Logger.LogInformation("Response: {StatusCode}", (int)Response.StatusCode);
+						_logger.LogInformation("Response: {StatusCode}", (int)response.StatusCode);
 					}
 				}
-				catch (WebException Ex)
+				catch (WebException ex)
 				{
 					// Handle errors. Any non-200 responses automatically generate a WebException.
-					HttpWebResponse Response = (HttpWebResponse)Ex.Response;
-					if (Response == null)
+					HttpWebResponse response = (HttpWebResponse)ex.Response;
+					if (response == null)
 					{
-						Logger.LogError(Ex, "Exception while attempting to send event");
+						_logger.LogError(ex, "Exception while attempting to send event");
 					}
 					else
 					{
-						string ResponseText;
-						using (Stream ResponseStream = Response.GetResponseStream())
+						string responseText;
+						using (Stream responseStream = response.GetResponseStream())
 						{
-							MemoryStream MemoryStream = new MemoryStream();
-							ResponseStream.CopyTo(MemoryStream);
-							ResponseText = Encoding.UTF8.GetString(MemoryStream.ToArray());
+							MemoryStream memoryStream = new MemoryStream();
+							responseStream.CopyTo(memoryStream);
+							responseText = Encoding.UTF8.GetString(memoryStream.ToArray());
 						}
-						Logger.LogError("Failed to send analytics event. Code = {Code}. Desc = {Dec}. Response = {Response}.", (int)Response.StatusCode, Response.StatusDescription, ResponseText);
+						_logger.LogError("Failed to send analytics event. Code = {Code}. Desc = {Dec}. Response = {Response}.", (int)response.StatusCode, response.StatusDescription, responseText);
 					}
 				}
-				catch (Exception Ex)
+				catch (Exception ex)
 				{
-					Logger.LogError(Ex, "Exception while attempting to send event");
+					_logger.LogError(ex, "Exception while attempting to send event");
 				}
 			}
 		}
@@ -221,11 +221,11 @@ namespace UnrealGameSync
 		/// <summary>
 		/// Sends a telemetry event with the given information
 		/// </summary>
-		/// <param name="EventName">Name of the event</param>
-		/// <param name="Attributes">Arbitrary object to include in the payload</param>
-		public static void SendEvent(string EventName, object Attributes)
+		/// <param name="eventName">Name of the event</param>
+		/// <param name="attributes">Arbitrary object to include in the payload</param>
+		public static void SendEvent(string eventName, object attributes)
 		{
-			ActiveSink?.SendEvent(EventName, Attributes);
+			ActiveSink?.SendEvent(eventName, attributes);
 		}
 	}
 }
diff --git a/Engine/Source/Programs/UnrealGameSync/UnrealGameSyncShared/TelemetryStopwatch.cs b/Engine/Source/Programs/UnrealGameSync/UnrealGameSyncShared/TelemetryStopwatch.cs
--- a/Engine/Source/Programs/UnrealGameSync/UnrealGameSyncShared/TelemetryStopwatch.cs
+++ b/Engine/Source/Programs/UnrealGameSync/UnrealGameSyncShared/TelemetryStopwatch.cs
@@ -9,52 +9,52 @@ namespace UnrealGameSync
 {
 	class TelemetryStopwatch : IDisposable
 	{
-		readonly string EventName;
-		readonly Dictionary<string, object?> EventData;
-		readonly Stopwatch Timer;
+		readonly string _eventName;
+		readonly Dictionary<string, object?> _eventData;
+		readonly Stopwatch _timer;
 
-		public TelemetryStopwatch(string EventName, string Project)
+		public TelemetryStopwatch(string eventName, string project)
 		{
-			this.EventName = EventName;
+			this._eventName = eventName;
 
-			EventData = new Dictionary<string, object?>();
-			EventData["Project"] = Project;
+			_eventData = new Dictionary<string, object?>();
+			_eventData["Project"] = project;
 
-			Timer = Stopwatch.StartNew();
+			_timer = Stopwatch.StartNew();
 		}
 
-		public void AddData(object Data)
+		public void AddData(object data)
 		{
-			foreach (PropertyInfo Property in Data.GetType().GetProperties())
+			foreach (PropertyInfo property in data.GetType().GetProperties())
 			{
-				EventData[Property.Name] = Property.GetValue(Data);
+				_eventData[property.Name] = property.GetValue(data);
 			}
 		}
 
-		public TimeSpan Stop(string InResult)
+		public TimeSpan Stop(string inResult)
 		{
-			if (Timer.IsRunning)
+			if (_timer.IsRunning)
 			{
-				Timer.Stop();
+				_timer.Stop();
 
-				EventData["Result"] = InResult;
-				EventData["TimeSeconds"] = Timer.Elapsed.TotalSeconds;
+				_eventData["Result"] = inResult;
+				_eventData["TimeSeconds"] = _timer.Elapsed.TotalSeconds;
 			}
 			return Elapsed;
 		}
 
 		public void Dispose()
 		{
-			if (Timer.IsRunning)
+			if (_timer.IsRunning)
 			{
 				Stop("Aborted");
 			}
-			Telemetry.SendEvent(EventName, EventData);
+			Telemetry.SendEvent(_eventName, _eventData);
 		}
 
 		public TimeSpan Elapsed
 		{
-			get { return Timer.Elapsed; }
+			get { return _timer.Elapsed; }
 		}
 	}
 }
