UGS: Reformat to standard C# style conventions.

diff --git a/Engine/Source/Programs/UnrealGameSync/UnrealGameSyncShared/UserSettings.cs b/Engine/Source/Programs/UnrealGameSync/UnrealGameSyncShared/UserSettings.cs
--- a/Engine/Source/Programs/UnrealGameSync/UnrealGameSyncShared/UserSettings.cs
+++ b/Engine/Source/Programs/UnrealGameSync/UnrealGameSyncShared/UserSettings.cs
@@ -63,32 +63,32 @@ namespace UnrealGameSync
 		public int OrderIndex = Int32.MaxValue;
 
 		// What Rules to Check for.
-		public bool bGood = false;
-		public bool bStarred = false;
-		public bool bFindNewestGoodContent = false;
+		public bool Good = false;
+		public bool Starred = false;
+		public bool FindNewestGoodContent = false;
 		public List<string> RequiredBadges = new List<string>();
 
-		public static bool TryParseConfigEntry(string Text, [NotNullWhen(true)] out LatestChangeType? ChangeType)
+		public static bool TryParseConfigEntry(string text, [NotNullWhen(true)] out LatestChangeType? changeType)
 		{
-			ConfigObject DefinitionObject = new ConfigObject(Text);
+			ConfigObject definitionObject = new ConfigObject(text);
 
-			string LatestChangeTypeName = DefinitionObject.GetValue("Name", "");
-			if (LatestChangeTypeName != "")
+			string latestChangeTypeName = definitionObject.GetValue("Name", "");
+			if (latestChangeTypeName != "")
 			{
-				ChangeType = new LatestChangeType();
-				ChangeType.Name = LatestChangeTypeName;
-				ChangeType.Description = DefinitionObject.GetValue("Description", "No \"Description\" for LatestChangeType given.");
-				ChangeType.OrderIndex = DefinitionObject.GetValue("OrderIndex", Int32.MaxValue);
+				changeType = new LatestChangeType();
+				changeType.Name = latestChangeTypeName;
+				changeType.Description = definitionObject.GetValue("Description", "No \"Description\" for LatestChangeType given.");
+				changeType.OrderIndex = definitionObject.GetValue("OrderIndex", Int32.MaxValue);
 
-				ChangeType.bGood = DefinitionObject.GetValue("bGood", false);
-				ChangeType.bStarred = DefinitionObject.GetValue("bStarred", false);
-				ChangeType.bFindNewestGoodContent = DefinitionObject.GetValue("bFindNewestGoodContent", false);
-				ChangeType.RequiredBadges = new List<string>(DefinitionObject.GetValue("RequiredBadges", "").Split(new char[] { ',',' '}, StringSplitOptions.RemoveEmptyEntries));
+				changeType.Good = definitionObject.GetValue("bGood", false);
+				changeType.Starred = definitionObject.GetValue("bStarred", false);
+				changeType.FindNewestGoodContent = definitionObject.GetValue("bFindNewestGoodContent", false);
+				changeType.RequiredBadges = new List<string>(definitionObject.GetValue("RequiredBadges", "").Split(new char[] { ',',' '}, StringSplitOptions.RemoveEmptyEntries));
 				return true;
 			}
 			else
 			{
-				ChangeType = null;
+				changeType = null;
 			}
 
 			return false;
@@ -97,39 +97,39 @@ namespace UnrealGameSync
 		// Default Template for Latest Change
 		public static LatestChangeType LatestChange()
 		{
-			LatestChangeType Latest = new LatestChangeType();
-			Latest.Name = "LatestChange";
-			Latest.Description = "Latest Change";
-			Latest.OrderIndex = -3;
-			Latest.bFindNewestGoodContent = true;
+			LatestChangeType latest = new LatestChangeType();
+			latest.Name = "LatestChange";
+			latest.Description = "Latest Change";
+			latest.OrderIndex = -3;
+			latest.FindNewestGoodContent = true;
 
-			return Latest;
+			return latest;
 		}
 
 		// Default Template for Latest Good Change
 		public static LatestChangeType LatestGoodChange()
 		{
-			LatestChangeType LatestGood = new LatestChangeType();
-			LatestGood.Name = "LatestGoodChange";
-			LatestGood.Description = "Latest Good Change";
-			LatestGood.OrderIndex = -2;
-			LatestGood.bGood = true;
-			LatestGood.bFindNewestGoodContent = true;
-
-			return LatestGood;
+			LatestChangeType latestGood = new LatestChangeType();
+			latestGood.Name = "LatestGoodChange";
+			latestGood.Description = "Latest Good Change";
+			latestGood.OrderIndex = -2;
+			latestGood.Good = true;
+			latestGood.FindNewestGoodContent = true;
+
+			return latestGood;
 		}
 
 		// Default Template for Latest Starred Change
 		public static LatestChangeType LatestStarredChange()
 		{
-			LatestChangeType LatestStarred = new LatestChangeType();
-			LatestStarred.Name = "LatestStarredChange";
-			LatestStarred.Description = "Latest Starred Change";
-			LatestStarred.OrderIndex = -1;
-			LatestStarred.bStarred = true;
-			LatestStarred.bFindNewestGoodContent = true;
-
-			return LatestStarred;
+			LatestChangeType latestStarred = new LatestChangeType();
+			latestStarred.Name = "LatestStarredChange";
+			latestStarred.Description = "Latest Starred Change";
+			latestStarred.OrderIndex = -1;
+			latestStarred.Starred = true;
+			latestStarred.FindNewestGoodContent = true;
+
+			return latestStarred;
 		}
 	}
 
@@ -144,46 +144,46 @@ namespace UnrealGameSync
 
 	public class ArchiveSettings
 	{
-		public bool bEnabled;
+		public bool Enabled;
 		public string Type;
 		public List<string> Order;
 
-		public ArchiveSettings(bool bEnabled, string Type, IEnumerable<string> Order)
+		public ArchiveSettings(bool enabled, string type, IEnumerable<string> order)
 		{
-			this.bEnabled = bEnabled;
-			this.Type = Type;
-			this.Order = new List<string>(Order);
+			this.Enabled = enabled;
+			this.Type = type;
+			this.Order = new List<string>(order);
 		}
 
-		public static bool TryParseConfigEntry(string Text, [NotNullWhen(true)] out ArchiveSettings? Settings)
+		public static bool TryParseConfigEntry(string text, [NotNullWhen(true)] out ArchiveSettings? settings)
 		{
-			ConfigObject Object = new ConfigObject(Text);
+			ConfigObject obj = new ConfigObject(text);
 
-			string? Type = Object.GetValue("Type", null);
-			if (Type == null)
+			string? type = obj.GetValue("Type", null);
+			if (type == null)
 			{
-				Settings = null;
+				settings = null;
 				return false;
 			}
 			else
 			{
-				string[] Order = Object.GetValue("Order", "").Split(new char[] { ';' }, StringSplitOptions.RemoveEmptyEntries);
-				bool bEnabled = Object.GetValue("Enabled", 0) != 0;
+				string[] order = obj.GetValue("Order", "").Split(new char[] { ';' }, StringSplitOptions.RemoveEmptyEntries);
+				bool enabled = obj.GetValue("Enabled", 0) != 0;
 
-				Settings = new ArchiveSettings(bEnabled, Type, Order);
+				settings = new ArchiveSettings(enabled, type, order);
 				return true;
 			}
 		}
 
 		public string ToConfigEntry()
 		{
-			ConfigObject Object = new ConfigObject();
+			ConfigObject obj = new ConfigObject();
 
-			Object.SetValue("Enabled", bEnabled ? 1 : 0);
-			Object.SetValue("Type", Type);
-			Object.SetValue("Order", String.Join(";", Order));
+			obj.SetValue("Enabled", Enabled ? 1 : 0);
+			obj.SetValue("Type", Type);
+			obj.SetValue("Order", String.Join(";", Order));
 
-			return Object.ToString();
+			return obj.ToString();
 		}
 
 		public override string ToString()
@@ -199,101 +199,101 @@ namespace UnrealGameSync
 		public UserSelectedProjectType Type { get; set; }
 		public string? ClientPath { get; set; }
 		public string? LocalPath { get; set; }
-		public string? ScheduledSyncTypeID { get; set; }
+		public string? ScheduledSyncTypeId { get; set; }
 
-		public UserSelectedProjectSettings(string? ServerAndPort, string? UserName, UserSelectedProjectType Type, string? ClientPath, string? LocalPath)
+		public UserSelectedProjectSettings(string? serverAndPort, string? userName, UserSelectedProjectType type, string? clientPath, string? localPath)
 		{
-			this.ServerAndPort = ServerAndPort;
-			this.UserName = UserName;
-			this.Type = Type;
-			this.ClientPath = ClientPath;
-			this.LocalPath = LocalPath;
+			this.ServerAndPort = serverAndPort;
+			this.UserName = userName;
+			this.Type = type;
+			this.ClientPath = clientPath;
+			this.LocalPath = localPath;
 		}
 
-		public static bool TryParseConfigEntry(string Text, [NotNullWhen(true)] out UserSelectedProjectSettings? Project)
+		public static bool TryParseConfigEntry(string text, [NotNullWhen(true)] out UserSelectedProjectSettings? project)
 		{
-			ConfigObject Object = new ConfigObject(Text);
+			ConfigObject obj = new ConfigObject(text);
 
-			UserSelectedProjectType Type;
-			if(Enum.TryParse(Object.GetValue("Type", ""), out Type))
+			UserSelectedProjectType type;
+			if(Enum.TryParse(obj.GetValue("Type", ""), out type))
 			{
-				string? ServerAndPort = Object.GetValue("ServerAndPort", null);
-				if(String.IsNullOrWhiteSpace(ServerAndPort))
+				string? serverAndPort = obj.GetValue("ServerAndPort", null);
+				if(String.IsNullOrWhiteSpace(serverAndPort))
 				{
-					ServerAndPort = null;
+					serverAndPort = null;
 				}
 
 				// Fixup for code that was saving server host name rather than DNS entry
-				if(ServerAndPort != null && ServerAndPort.Equals("p4-nodeb.epicgames.net:1666", StringComparison.OrdinalIgnoreCase))
+				if(serverAndPort != null && serverAndPort.Equals("p4-nodeb.epicgames.net:1666", StringComparison.OrdinalIgnoreCase))
 				{
-					ServerAndPort = "perforce:1666";
+					serverAndPort = "perforce:1666";
 				}
 
-				string? UserName = Object.GetValue("UserName", null);
-				if(String.IsNullOrWhiteSpace(UserName))
+				string? userName = obj.GetValue("UserName", null);
+				if(String.IsNullOrWhiteSpace(userName))
 				{
-					UserName = null;
+					userName = null;
 				}
 
-				string? LocalPath = Object.GetValue("LocalPath", null);
-				if(String.IsNullOrWhiteSpace(LocalPath))
+				string? localPath = obj.GetValue("LocalPath", null);
+				if(String.IsNullOrWhiteSpace(localPath))
 				{
-					LocalPath = null;
+					localPath = null;
 				}
 
-				string? ClientPath = Object.GetValue("ClientPath", null);
-				if(String.IsNullOrWhiteSpace(ClientPath))
+				string? clientPath = obj.GetValue("ClientPath", null);
+				if(String.IsNullOrWhiteSpace(clientPath))
 				{
-					ClientPath = null;
+					clientPath = null;
 				}
 
-				string? ScheduledSyncTypeID = Object.GetValue("ScheduledSyncTypeID", null);
-				if (String.IsNullOrWhiteSpace(ScheduledSyncTypeID))
+				string? scheduledSyncTypeId = obj.GetValue("ScheduledSyncTypeID", null);
+				if (String.IsNullOrWhiteSpace(scheduledSyncTypeId))
 				{
-					ScheduledSyncTypeID = null;
+					scheduledSyncTypeId = null;
 				}
 
-				if ((Type == UserSelectedProjectType.Client && ClientPath != null) || (Type == UserSelectedProjectType.Local && LocalPath != null))
+				if ((type == UserSelectedProjectType.Client && clientPath != null) || (type == UserSelectedProjectType.Local && localPath != null))
 				{
-					Project = new UserSelectedProjectSettings(ServerAndPort, UserName, Type, ClientPath, LocalPath);
-					Project.ScheduledSyncTypeID = ScheduledSyncTypeID;
+					project = new UserSelectedProjectSettings(serverAndPort, userName, type, clientPath, localPath);
+					project.ScheduledSyncTypeId = scheduledSyncTypeId;
 					return true;
 				}
 			}
 
-			Project = null;
+			project = null;
 			return false;
 		}
 
 		public string ToConfigEntry()
 		{
-			ConfigObject Object = new ConfigObject();
+			ConfigObject obj = new ConfigObject();
 
 			if(ServerAndPort != null)
 			{
-				Object.SetValue("ServerAndPort", ServerAndPort);
+				obj.SetValue("ServerAndPort", ServerAndPort);
 			}
 			if(UserName != null)
 			{
-				Object.SetValue("UserName", UserName);
+				obj.SetValue("UserName", UserName);
 			}
 
-			Object.SetValue("Type", Type.ToString());
+			obj.SetValue("Type", Type.ToString());
 
 			if(ClientPath != null)
 			{
-				Object.SetValue("ClientPath", ClientPath);
+				obj.SetValue("ClientPath", ClientPath);
 			}
 			if(LocalPath != null)
 			{
-				Object.SetValue("LocalPath", LocalPath);
+				obj.SetValue("LocalPath", LocalPath);
 			}
-			if (ScheduledSyncTypeID != null)
+			if (ScheduledSyncTypeId != null)
 			{
-				Object.SetValue("ScheduledSyncTypeID", ScheduledSyncTypeID);
+				obj.SetValue("ScheduledSyncTypeID", ScheduledSyncTypeId);
 			}
 
-			return Object.ToString();
+			return obj.ToString();
 		}
 
 		public override string? ToString()
@@ -345,53 +345,53 @@ namespace UnrealGameSync
 		[JsonIgnore]
 		public FileReference LocalProjectPath => new FileReference(RootDir.FullName + ProjectPath);
 
-		public void Init(string? ServerAndPort, string? UserName, string ClientName, string BranchPath, string ProjectPath)
+		public void Init(string? serverAndPort, string? userName, string clientName, string branchPath, string projectPath)
 		{
-			ProjectInfo.ValidateBranchPath(BranchPath);
-			ProjectInfo.ValidateProjectPath(ProjectPath);
-
-			this.ServerAndPort = ServerAndPort;
-			this.UserName = UserName;
-			this.ClientName = ClientName;
-			this.BranchPath = BranchPath;
-			this.ProjectPath = ProjectPath;
+			ProjectInfo.ValidateBranchPath(branchPath);
+			ProjectInfo.ValidateProjectPath(projectPath);
+
+			this.ServerAndPort = serverAndPort;
+			this.UserName = userName;
+			this.ClientName = clientName;
+			this.BranchPath = branchPath;
+			this.ProjectPath = projectPath;
 		}
 
-		public static bool TryLoad(DirectoryReference RootDir, [NotNullWhen(true)] out UserWorkspaceSettings? Settings)
+		public static bool TryLoad(DirectoryReference rootDir, [NotNullWhen(true)] out UserWorkspaceSettings? settings)
 		{
-			FileReference ConfigFile = GetConfigFile(RootDir);
-			if (Utility.TryLoadJson(ConfigFile, out Settings))
+			FileReference configFile = GetConfigFile(rootDir);
+			if (Utility.TryLoadJson(configFile, out settings))
 			{
-				Settings.RootDir = RootDir;
-				Settings.LastModifiedTimeUtc = FileReference.GetLastWriteTimeUtc(ConfigFile).Ticks;
+				settings.RootDir = rootDir;
+				settings.LastModifiedTimeUtc = FileReference.GetLastWriteTimeUtc(configFile).Ticks;
 				return true;
 			}
 			else
 			{
-				Settings = null;
+				settings = null;
 				return false;
 			}
 		}
 
-		static object SyncRoot = new object();
+		static object _syncRoot = new object();
 
-		public bool Save(ILogger Logger)
+		public bool Save(ILogger logger)
 		{
 			try
 			{
 				SaveInternal();
 				return true;
 			}
-			catch (Exception Ex)
+			catch (Exception ex)
 			{
-				Logger.LogError(Ex, "Unable to save {ConfigFile}: {Message}", ConfigFile, Ex.Message);
+				logger.LogError(ex, "Unable to save {ConfigFile}: {Message}", ConfigFile, ex.Message);
 				return false;
 			}
 		}
 
 		private void SaveInternal()
 		{
-			lock (SyncRoot)
+			lock (_syncRoot)
 			{
 				UserSettings.CreateConfigDir(ConfigFile.Directory);
 				Utility.SaveJson(ConfigFile, this);
@@ -399,9 +399,9 @@ namespace UnrealGameSync
 			}
 		}
 
-		public static FileReference GetConfigFile(DirectoryReference RootDir)
+		public static FileReference GetConfigFile(DirectoryReference rootDir)
 		{
-			return FileReference.Combine(UserSettings.GetConfigDir(RootDir), "settings.json");
+			return FileReference.Combine(UserSettings.GetConfigDir(rootDir), "settings.json");
 		}
 	}
 
@@ -445,11 +445,11 @@ namespace UnrealGameSync
 		[JsonIgnore]
 		public FileReference ConfigFile => GetConfigFile(RootDir);
 
-		public static UserWorkspaceState CreateNew(DirectoryReference RootDir)
+		public static UserWorkspaceState CreateNew(DirectoryReference rootDir)
 		{
-			UserWorkspaceState State = new UserWorkspaceState();
-			State.RootDir = RootDir;
-			return State;
+			UserWorkspaceState state = new UserWorkspaceState();
+			state.RootDir = rootDir;
+			return state;
 		}
 
 		public ProjectInfo CreateProjectInfo()
@@ -457,86 +457,86 @@ namespace UnrealGameSync
 			return new ProjectInfo(RootDir, ClientName, BranchPath, ProjectPath, StreamName, ProjectIdentifier, IsEnterpriseProject);
 		}
 
-		public void UpdateCachedProjectInfo(ProjectInfo ProjectInfo, long SettingsTimeUtc)
+		public void UpdateCachedProjectInfo(ProjectInfo projectInfo, long settingsTimeUtc)
 		{
-			this.SettingsTimeUtc = SettingsTimeUtc;
-
-			RootDir = ProjectInfo.LocalRootPath;
-			ClientName = ProjectInfo.ClientName;
-			BranchPath = ProjectInfo.BranchPath;
-			ProjectPath = ProjectInfo.ProjectPath;
-			StreamName = ProjectInfo.StreamName;
-			ProjectIdentifier = ProjectInfo.ProjectIdentifier;
-			IsEnterpriseProject = ProjectInfo.bIsEnterpriseProject;
+			this.SettingsTimeUtc = settingsTimeUtc;
+
+			RootDir = projectInfo.LocalRootPath;
+			ClientName = projectInfo.ClientName;
+			BranchPath = projectInfo.BranchPath;
+			ProjectPath = projectInfo.ProjectPath;
+			StreamName = projectInfo.StreamName;
+			ProjectIdentifier = projectInfo.ProjectIdentifier;
+			IsEnterpriseProject = projectInfo.IsEnterpriseProject;
 		}
 
-		public bool IsValid(ProjectInfo ProjectInfo)
+		public bool IsValid(ProjectInfo projectInfo)
 		{
-			return ProjectIdentifier.Equals(ProjectInfo.ProjectIdentifier, StringComparison.OrdinalIgnoreCase);
+			return ProjectIdentifier.Equals(projectInfo.ProjectIdentifier, StringComparison.OrdinalIgnoreCase);
 		}
 
-		public void SetBisectState(int Change, BisectState State)
+		public void SetBisectState(int change, BisectState state)
 		{
-			BisectEntry Entry = BisectChanges.FirstOrDefault(x => x.Change == Change);
-			if (Entry == null)
+			BisectEntry entry = BisectChanges.FirstOrDefault(x => x.Change == change);
+			if (entry == null)
 			{
-				Entry = new BisectEntry();
-				Entry.Change = Change;
-				BisectChanges.Add(Entry);
+				entry = new BisectEntry();
+				entry.Change = change;
+				BisectChanges.Add(entry);
 			}
-			Entry.State = State;
+			entry.State = state;
 		}
 
-		public static bool TryLoad(DirectoryReference RootDir, [NotNullWhen(true)] out UserWorkspaceState? State)
+		public static bool TryLoad(DirectoryReference rootDir, [NotNullWhen(true)] out UserWorkspaceState? state)
 		{
-			FileReference ConfigFile = GetConfigFile(RootDir);
-			if (Utility.TryLoadJson(ConfigFile, out State))
+			FileReference configFile = GetConfigFile(rootDir);
+			if (Utility.TryLoadJson(configFile, out state))
 			{
-				State.RootDir = RootDir;
+				state.RootDir = rootDir;
 				return true;
 			}
 			else
 			{
-				State = null;
+				state = null;
 				return false;
 			}
 		}
 
-		public void SetLastSyncState(WorkspaceUpdateResult Result, WorkspaceUpdateContext Context, string StatusMessage)
+		public void SetLastSyncState(WorkspaceUpdateResult result, WorkspaceUpdateContext context, string statusMessage)
 		{
-			LastSyncChangeNumber = Context.ChangeNumber;
-			LastSyncResult = Result;
-			LastSyncResultMessage = StatusMessage;
+			LastSyncChangeNumber = context.ChangeNumber;
+			LastSyncResult = result;
+			LastSyncResultMessage = statusMessage;
 			LastSyncTime = DateTime.UtcNow;
-			LastSyncDurationSeconds = (int)(LastSyncTime.Value - Context.StartTime).TotalSeconds;
+			LastSyncDurationSeconds = (int)(LastSyncTime.Value - context.StartTime).TotalSeconds;
 		}
 
-		static object SyncRoot = new object();
+		static object _syncRoot = new object();
 
-		public bool Save(ILogger Logger)
+		public bool Save(ILogger logger)
 		{
 			try
 			{
 				SaveInternal();
 				return true;
 			}
-			catch (Exception Ex)
+			catch (Exception ex)
 			{
-				Logger.LogError(Ex, "Unable to save {ConfigFile}: {Message}", ConfigFile, Ex.Message);
+				logger.LogError(ex, "Unable to save {ConfigFile}: {Message}", ConfigFile, ex.Message);
 				return false;
 			}
 		}
 
 		private void SaveInternal()
 		{
-			lock (SyncRoot)
+			lock (_syncRoot)
 			{
 				UserSettings.CreateConfigDir(ConfigFile.Directory);
 				Utility.SaveJson(ConfigFile, this);
 			}
 		}
 
-		public static FileReference GetConfigFile(DirectoryReference RootDir) => FileReference.Combine(UserSettings.GetConfigDir(RootDir), "state.json");
+		public static FileReference GetConfigFile(DirectoryReference rootDir) => FileReference.Combine(UserSettings.GetConfigDir(rootDir), "state.json");
 	}
 
 	public class UserProjectSettings
@@ -548,47 +548,47 @@ namespace UnrealGameSync
 		public FilterType FilterType { get; set; }
 		public HashSet<string> FilterBadges { get; set; } = new HashSet<string>(StringComparer.OrdinalIgnoreCase);
 
-		static object SyncRoot = new object();
+		static object _syncRoot = new object();
 
 		private UserProjectSettings()
 		{
 		}
 
-		public UserProjectSettings(FileReference ConfigFile)
+		public UserProjectSettings(FileReference configFile)
 		{
-			this.ConfigFile = ConfigFile;
+			this.ConfigFile = configFile;
 		}
 
-		public static bool TryLoad(FileReference ConfigFile, [NotNullWhen(true)] out UserProjectSettings? Settings)
+		public static bool TryLoad(FileReference configFile, [NotNullWhen(true)] out UserProjectSettings? settings)
 		{
-			if (FileReference.Exists(ConfigFile))
+			if (FileReference.Exists(configFile))
 			{
-				Settings = Utility.LoadJson<UserProjectSettings>(ConfigFile);
-				Settings.ConfigFile = ConfigFile;
+				settings = Utility.LoadJson<UserProjectSettings>(configFile);
+				settings.ConfigFile = configFile;
 				return true;
 			}
 
-			Settings = null;
+			settings = null;
 			return false;
 		}
 
-		public bool Save(ILogger Logger)
+		public bool Save(ILogger logger)
 		{
 			try
 			{
 				SaveInternal();
 				return true;
 			}
-			catch (Exception Ex)
+			catch (Exception ex)
 			{
-				Logger.LogError(Ex, "Unable to save {ConfigFile}: {Message}", ConfigFile, Ex.Message);
+				logger.LogError(ex, "Unable to save {ConfigFile}: {Message}", ConfigFile, ex.Message);
 				return false;
 			}
 		}
 		
 		void SaveInternal()
 		{
-			lock (SyncRoot)
+			lock (_syncRoot)
 			{
 				Utility.SaveJson(ConfigFile, this);
 			}
@@ -607,28 +607,28 @@ namespace UnrealGameSync
 			None		// Show no robomerge changes
 		};
 
-		FileReference FileName;
-		ConfigFile ConfigFile;
+		FileReference _fileName;
+		ConfigFile _configFile;
 
 		// General settings
 		public UserSettingsVersion Version = UserSettingsVersion.Latest;
-		public bool bBuildAfterSync;
-		public bool bRunAfterSync;
-		public bool bOpenSolutionAfterSync;
-		public bool bShowLogWindow;
-		public bool bAutoResolveConflicts;
-		public bool bShowUnreviewedChanges;
-		public bool bShowAutomatedChanges;
+		public bool BuildAfterSync;
+		public bool RunAfterSync;
+		public bool OpenSolutionAfterSync;
+		public bool ShowLogWindow;
+		public bool AutoResolveConflicts;
+		public bool ShowUnreviewedChanges;
+		public bool ShowAutomatedChanges;
 		public RobomergeShowChangesOption ShowRobomerge;
-		public bool bAnnotateRobmergeChanges;
-		public bool bShowLocalTimes;
-		public bool bKeepInTray;
+		public bool AnnotateRobmergeChanges;
+		public bool ShowLocalTimes;
+		public bool KeepInTray;
 		public Guid[] EnabledTools;
 		public int FilterIndex;
 		public UserSelectedProjectSettings? LastProject;
 		public List<UserSelectedProjectSettings> OpenProjects;
 		public List<UserSelectedProjectSettings> RecentProjects;
-		public string SyncTypeID;
+		public string SyncTypeId;
 		public BuildConfig CompiledEditorBuildConfig; // NB: This assumes not using precompiled editor. See CurrentBuildConfig.
 		public TabLabels TabLabels;
 
@@ -639,19 +639,19 @@ namespace UnrealGameSync
 		public Dictionary<string, string> ProviderToRefreshTokens = new Dictionary<string, string>();
 
 		// Window settings
-		public bool bWindowVisible;
+		public bool WindowVisible;
 		public string WindowState;
 		public Rectangle? WindowBounds;
 		
 		// Schedule settings
-		public bool bScheduleEnabled;
+		public bool ScheduleEnabled;
 		public TimeSpan ScheduleTime;
 		public bool ScheduleAnyOpenProject;
 		public List<UserSelectedProjectSettings> ScheduleProjects;
 
 		// Run configuration
 		public List<Tuple<string, bool>> EditorArguments = new List<Tuple<string,bool>>();
-		public bool bEditorArgumentsPrompt;
+		public bool EditorArgumentsPrompt;
 
 		// Notification settings
 		public List<string> NotifyProjects;
@@ -660,395 +660,395 @@ namespace UnrealGameSync
 		public int NotifyUnresolvedMinutes;
 
 		// Project settings
-		Dictionary<DirectoryReference, UserWorkspaceState> WorkspaceDirToState = new Dictionary<DirectoryReference, UserWorkspaceState>();
-		Dictionary<DirectoryReference, UserWorkspaceSettings> WorkspaceDirToSettings = new Dictionary<DirectoryReference, UserWorkspaceSettings>();
-		Dictionary<FileReference, UserProjectSettings> ProjectKeyToSettings = new Dictionary<FileReference, UserProjectSettings>();
+		Dictionary<DirectoryReference, UserWorkspaceState> _workspaceDirToState = new Dictionary<DirectoryReference, UserWorkspaceState>();
+		Dictionary<DirectoryReference, UserWorkspaceSettings> _workspaceDirToSettings = new Dictionary<DirectoryReference, UserWorkspaceSettings>();
+		Dictionary<FileReference, UserProjectSettings> _projectKeyToSettings = new Dictionary<FileReference, UserProjectSettings>();
 
 		// Perforce settings
 		public PerforceSyncOptions SyncOptions = new PerforceSyncOptions();
 
-		private List<UserSelectedProjectSettings> ReadProjectList(string SettingName, string LegacySettingName)
+		private List<UserSelectedProjectSettings> ReadProjectList(string settingName, string legacySettingName)
 		{
-			List<UserSelectedProjectSettings> Projects = new List<UserSelectedProjectSettings>();
+			List<UserSelectedProjectSettings> projects = new List<UserSelectedProjectSettings>();
 
-			string[]? ProjectStrings = ConfigFile.GetValues(SettingName, null);
-			if(ProjectStrings != null)
+			string[]? projectStrings = _configFile.GetValues(settingName, null);
+			if(projectStrings != null)
 			{
-				foreach(string ProjectString in ProjectStrings)
+				foreach(string projectString in projectStrings)
 				{
-					UserSelectedProjectSettings? Project;
-					if(UserSelectedProjectSettings.TryParseConfigEntry(ProjectString, out Project))
+					UserSelectedProjectSettings? project;
+					if(UserSelectedProjectSettings.TryParseConfigEntry(projectString, out project))
 					{
-						Projects.Add(Project);
+						projects.Add(project);
 					}
 				}
 			}
-			else if(LegacySettingName != null)
+			else if(legacySettingName != null)
 			{
-				string[]? LegacyProjectStrings = ConfigFile.GetValues(LegacySettingName, null);
-				if(LegacyProjectStrings != null)
+				string[]? legacyProjectStrings = _configFile.GetValues(legacySettingName, null);
+				if(legacyProjectStrings != null)
 				{
-					foreach(string LegacyProjectString in LegacyProjectStrings)
+					foreach(string legacyProjectString in legacyProjectStrings)
 					{
-						if(!String.IsNullOrWhiteSpace(LegacyProjectString))
+						if(!String.IsNullOrWhiteSpace(legacyProjectString))
 						{
-							Projects.Add(new UserSelectedProjectSettings(null, null, UserSelectedProjectType.Local, null, LegacyProjectString));
+							projects.Add(new UserSelectedProjectSettings(null, null, UserSelectedProjectType.Local, null, legacyProjectString));
 						}
 					}
 				}
 			}
 
-			return Projects;
+			return projects;
 		}
 
-		public static UserSettings Create(DirectoryReference SettingsDir, ILogger Logger)
+		public static UserSettings Create(DirectoryReference settingsDir, ILogger logger)
 		{
-			return Create(FileReference.Combine(SettingsDir, "UnrealGameSyncV2.ini"), FileReference.Combine(SettingsDir, "UnrealGameSync.ini"), FileReference.Combine(SettingsDir, "Global.json"), Logger);
+			return Create(FileReference.Combine(settingsDir, "UnrealGameSyncV2.ini"), FileReference.Combine(settingsDir, "UnrealGameSync.ini"), FileReference.Combine(settingsDir, "Global.json"), logger);
 		}
 
-		public static UserSettings Create(FileReference FileName, FileReference LegacyFileName, FileReference CoreFileName, ILogger Logger)
+		public static UserSettings Create(FileReference fileName, FileReference legacyFileName, FileReference coreFileName, ILogger logger)
 		{
-			ConfigFile ConfigFile = new ConfigFile();
-			if (FileReference.Exists(FileName))
+			ConfigFile configFile = new ConfigFile();
+			if (FileReference.Exists(fileName))
 			{
-				ConfigFile.TryLoad(FileName, Logger);
+				configFile.TryLoad(fileName, logger);
 			}
 			else
 			{
-				ConfigFile.TryLoad(LegacyFileName, Logger);
+				configFile.TryLoad(legacyFileName, logger);
 			}
 
-			GlobalSettings? CoreSettingsData = null;
-			if (FileReference.Exists(CoreFileName))
+			GlobalSettings? coreSettingsData = null;
+			if (FileReference.Exists(coreFileName))
 			{
 				try
 				{
-					CoreSettingsData = Utility.LoadJson<GlobalSettings>(CoreFileName);
+					coreSettingsData = Utility.LoadJson<GlobalSettings>(coreFileName);
 				}
-				catch (Exception Ex)
+				catch (Exception ex)
 				{
-					Logger.LogError(Ex, "Error while reading {File}.", CoreFileName);
+					logger.LogError(ex, "Error while reading {File}.", coreFileName);
 				}
 			}
 
-			if (CoreSettingsData == null)
+			if (coreSettingsData == null)
 			{
-				CoreSettingsData = new GlobalSettings();
-				CoreSettingsData.Filter.View = ConfigFile.GetValues("General.SyncFilter", new string[0]).ToList();
-				CoreSettingsData.Filter.SetCategories(GetCategorySettings(ConfigFile.FindSection("General"), "SyncIncludedCategories", "SyncExcludedCategories"));
-				CoreSettingsData.Filter.AllProjects = ConfigFile.GetValue("General.SyncAllProjects", false);
-				CoreSettingsData.Filter.AllProjectsInSln = ConfigFile.GetValue("General.IncludeAllProjectsInSolution", false);
+				coreSettingsData = new GlobalSettings();
+				coreSettingsData.Filter.View = configFile.GetValues("General.SyncFilter", new string[0]).ToList();
+				coreSettingsData.Filter.SetCategories(GetCategorySettings(configFile.FindSection("General"), "SyncIncludedCategories", "SyncExcludedCategories"));
+				coreSettingsData.Filter.AllProjects = configFile.GetValue("General.SyncAllProjects", false);
+				coreSettingsData.Filter.AllProjectsInSln = configFile.GetValue("General.IncludeAllProjectsInSolution", false);
 			}
 
-			return new UserSettings(FileName, ConfigFile, CoreFileName, CoreSettingsData);
+			return new UserSettings(fileName, configFile, coreFileName, coreSettingsData);
 		}
 
-		public UserSettings(FileReference InFileName, ConfigFile InConfigFile, FileReference InCoreFileName, GlobalSettings InCoreSettingsData)
-			: base(InCoreFileName, InCoreSettingsData)
+		public UserSettings(FileReference inFileName, ConfigFile inConfigFile, FileReference inCoreFileName, GlobalSettings inCoreSettingsData)
+			: base(inCoreFileName, inCoreSettingsData)
 		{
-			this.FileName = InFileName;
-			this.ConfigFile = InConfigFile;
+			this._fileName = inFileName;
+			this._configFile = inConfigFile;
 
 			// General settings
-			Version = (UserSettingsVersion)ConfigFile.GetValue("General.Version", (int)UserSettingsVersion.Initial);
-			bBuildAfterSync = (ConfigFile.GetValue("General.BuildAfterSync", "1") != "0");
-			bRunAfterSync = (ConfigFile.GetValue("General.RunAfterSync", "1") != "0");
-			bool bSyncPrecompiledEditor = (ConfigFile.GetValue("General.SyncPrecompiledEditor", "0") != "0");
-			bOpenSolutionAfterSync = (ConfigFile.GetValue("General.OpenSolutionAfterSync", "0") != "0");
-			bShowLogWindow = (ConfigFile.GetValue("General.ShowLogWindow", false));
-			bAutoResolveConflicts = (ConfigFile.GetValue("General.AutoResolveConflicts", "1") != "0");
-			bShowUnreviewedChanges = ConfigFile.GetValue("General.ShowUnreviewed", true);
-			bShowAutomatedChanges = ConfigFile.GetValue("General.ShowAutomated", false);
+			Version = (UserSettingsVersion)_configFile.GetValue("General.Version", (int)UserSettingsVersion.Initial);
+			BuildAfterSync = (_configFile.GetValue("General.BuildAfterSync", "1") != "0");
+			RunAfterSync = (_configFile.GetValue("General.RunAfterSync", "1") != "0");
+			bool syncPrecompiledEditor = (_configFile.GetValue("General.SyncPrecompiledEditor", "0") != "0");
+			OpenSolutionAfterSync = (_configFile.GetValue("General.OpenSolutionAfterSync", "0") != "0");
+			ShowLogWindow = (_configFile.GetValue("General.ShowLogWindow", false));
+			AutoResolveConflicts = (_configFile.GetValue("General.AutoResolveConflicts", "1") != "0");
+			ShowUnreviewedChanges = _configFile.GetValue("General.ShowUnreviewed", true);
+			ShowAutomatedChanges = _configFile.GetValue("General.ShowAutomated", false);
 
 			// safely parse the filter enum
 			ShowRobomerge = RobomergeShowChangesOption.All;
-			Enum.TryParse(ConfigFile.GetValue("General.RobomergeFilter", ""), out ShowRobomerge);
+			Enum.TryParse(_configFile.GetValue("General.RobomergeFilter", ""), out ShowRobomerge);
 
-			bAnnotateRobmergeChanges = ConfigFile.GetValue("General.AnnotateRobomerge", true);
-			bShowLocalTimes = ConfigFile.GetValue("General.ShowLocalTimes", false);
-			bKeepInTray = ConfigFile.GetValue("General.KeepInTray", true);
+			AnnotateRobmergeChanges = _configFile.GetValue("General.AnnotateRobomerge", true);
+			ShowLocalTimes = _configFile.GetValue("General.ShowLocalTimes", false);
+			KeepInTray = _configFile.GetValue("General.KeepInTray", true);
 
-			List<Guid> EnabledTools = ConfigFile.GetGuidValues("General.EnabledTools", new Guid[0]).ToList();
-			if (ConfigFile.GetValue("General.EnableP4VExtensions", false))
+			List<Guid> enabledTools = _configFile.GetGuidValues("General.EnabledTools", new Guid[0]).ToList();
+			if (_configFile.GetValue("General.EnableP4VExtensions", false))
 			{
-				EnabledTools.Add(new Guid("963850A0-BF63-4E0E-B903-1C5954C7DCF8"));
+				enabledTools.Add(new Guid("963850A0-BF63-4E0E-B903-1C5954C7DCF8"));
 			}
-			if (ConfigFile.GetValue("General.EnableUshell", false))
+			if (_configFile.GetValue("General.EnableUshell", false))
 			{
-				EnabledTools.Add(new Guid("922EED87-E732-464C-92DC-5A8F7ED955E2"));
+				enabledTools.Add(new Guid("922EED87-E732-464C-92DC-5A8F7ED955E2"));
 			}
-			this.EnabledTools = EnabledTools.ToArray();
+			this.EnabledTools = enabledTools.ToArray();
 
-			int.TryParse(ConfigFile.GetValue("General.FilterIndex", "0"), out FilterIndex);
+			int.TryParse(_configFile.GetValue("General.FilterIndex", "0"), out FilterIndex);
 
-			string? LastProjectString = ConfigFile.GetValue("General.LastProject", null);
-			if(LastProjectString != null)
+			string? lastProjectString = _configFile.GetValue("General.LastProject", null);
+			if(lastProjectString != null)
 			{
-				UserSelectedProjectSettings.TryParseConfigEntry(LastProjectString, out LastProject);
+				UserSelectedProjectSettings.TryParseConfigEntry(lastProjectString, out LastProject);
 			}
 			else
 			{
-				string? LastProjectFileName = ConfigFile.GetValue("General.LastProjectFileName", null);
-				if(LastProjectFileName != null)
+				string? lastProjectFileName = _configFile.GetValue("General.LastProjectFileName", null);
+				if(lastProjectFileName != null)
 				{
-					LastProject = new UserSelectedProjectSettings(null, null, UserSelectedProjectType.Local, null, LastProjectFileName);
+					LastProject = new UserSelectedProjectSettings(null, null, UserSelectedProjectType.Local, null, lastProjectFileName);
 				}
 			}
 
 			OpenProjects = ReadProjectList("General.OpenProjects", "General.OpenProjectFileNames");
 			RecentProjects = ReadProjectList("General.RecentProjects", "General.OtherProjectFileNames");
-			SyncTypeID = ConfigFile.GetValue("General.SyncTypeID", "");
-			string? OldSyncTye = ConfigFile.GetValue("General.SyncType", null);
-			if (OldSyncTye != null)
+			SyncTypeId = _configFile.GetValue("General.SyncTypeID", "");
+			string? oldSyncTye = _configFile.GetValue("General.SyncType", null);
+			if (oldSyncTye != null)
 			{
-				if (OldSyncTye == "Any")
+				if (oldSyncTye == "Any")
 				{
-					SyncTypeID = LatestChangeType.LatestChange().Name;
+					SyncTypeId = LatestChangeType.LatestChange().Name;
 				}
-				else if (OldSyncTye == "Good")
+				else if (oldSyncTye == "Good")
 				{
-					SyncTypeID = LatestChangeType.LatestGoodChange().Name;
+					SyncTypeId = LatestChangeType.LatestGoodChange().Name;
 				}
-				else if (OldSyncTye == "Starred")
+				else if (oldSyncTye == "Starred")
 				{
-					SyncTypeID = LatestChangeType.LatestStarredChange().Name;
+					SyncTypeId = LatestChangeType.LatestStarredChange().Name;
 				}
 			}
 
 			// Build configuration
-			string CompiledEditorBuildConfigName = ConfigFile.GetValue("General.BuildConfig", "");
-			if(!Enum.TryParse(CompiledEditorBuildConfigName, true, out CompiledEditorBuildConfig))
+			string compiledEditorBuildConfigName = _configFile.GetValue("General.BuildConfig", "");
+			if(!Enum.TryParse(compiledEditorBuildConfigName, true, out CompiledEditorBuildConfig))
 			{
 				CompiledEditorBuildConfig = BuildConfig.DebugGame;
 			}
 
 			// Tab names
-			string TabLabelsValue = ConfigFile.GetValue("General.TabLabels", "");
-			if(!Enum.TryParse(TabLabelsValue, true, out TabLabels))
+			string tabLabelsValue = _configFile.GetValue("General.TabLabels", "");
+			if(!Enum.TryParse(tabLabelsValue, true, out TabLabels))
 			{
 				TabLabels = TabLabels.Stream;
 			}
 
 			// Editor arguments
-			string[] Arguments = ConfigFile.GetValues("General.EditorArguments", new string[]{ "0:-log", "0:-fastload" });
+			string[] arguments = _configFile.GetValues("General.EditorArguments", new string[]{ "0:-log", "0:-fastload" });
 			if (Version < UserSettingsVersion.XgeShaderCompilation)
 			{
-				Arguments = Enumerable.Concat(Arguments, new string[] { "0:-noxgeshadercompile" }).ToArray();
+				arguments = Enumerable.Concat(arguments, new string[] { "0:-noxgeshadercompile" }).ToArray();
 			}
-			foreach(string Argument in Arguments)
+			foreach(string argument in arguments)
 			{
-				if(Argument.StartsWith("0:"))
+				if(argument.StartsWith("0:"))
 				{
-					EditorArguments.Add(new Tuple<string,bool>(Argument.Substring(2), false));
+					EditorArguments.Add(new Tuple<string,bool>(argument.Substring(2), false));
 				}
-				else if(Argument.StartsWith("1:"))
+				else if(argument.StartsWith("1:"))
 				{
-					EditorArguments.Add(new Tuple<string,bool>(Argument.Substring(2), true));
+					EditorArguments.Add(new Tuple<string,bool>(argument.Substring(2), true));
 				}
 				else
 				{
-					EditorArguments.Add(new Tuple<string,bool>(Argument, true));
+					EditorArguments.Add(new Tuple<string,bool>(argument, true));
 				}
 			}
-			bEditorArgumentsPrompt = ConfigFile.GetValue("General.EditorArgumentsPrompt", false);
+			EditorArgumentsPrompt = _configFile.GetValue("General.EditorArgumentsPrompt", false);
 
 			// Precompiled binaries
-			string[] ArchiveValues = ConfigFile.GetValues("PrecompiledBinaries.Archives", new string[0]);
-			foreach (string ArchiveValue in ArchiveValues)
+			string[] archiveValues = _configFile.GetValues("PrecompiledBinaries.Archives", new string[0]);
+			foreach (string archiveValue in archiveValues)
 			{
-				ArchiveSettings? Settings;
-				if (ArchiveSettings.TryParseConfigEntry(ArchiveValue, out Settings))
+				ArchiveSettings? settings;
+				if (ArchiveSettings.TryParseConfigEntry(archiveValue, out settings))
 				{
-					Archives.Add(Settings);
+					Archives.Add(settings);
 				}
 			}
 
-			if (bSyncPrecompiledEditor)
+			if (syncPrecompiledEditor)
 			{
 				Archives.Add(new ArchiveSettings(true, "Editor", new string[0]));
 			}
 
 			// OIDC Settings
-			string[] Tokens = ConfigFile.GetValues("OIDCProviders.Tokens", new string[0]);
-			foreach (string TokenValue in Tokens)
+			string[] tokens = _configFile.GetValues("OIDCProviders.Tokens", new string[0]);
+			foreach (string tokenValue in tokens)
 			{
-				ConfigObject O = new ConfigObject(TokenValue);
-				string? Provider = O.GetValue("Provider");
-				string? Token = O.GetValue("Token");
-				if (Provider != null && Token != null)
+				ConfigObject o = new ConfigObject(tokenValue);
+				string? provider = o.GetValue("Provider");
+				string? token = o.GetValue("Token");
+				if (provider != null && token != null)
 				{
-					ProviderToRefreshTokens.TryAdd(Provider, Token);
+					ProviderToRefreshTokens.TryAdd(provider, token);
 				}
 			}
 
 			// Window settings
-			bWindowVisible = ConfigFile.GetValue("Window.Visible", true);
-			WindowState = ConfigFile.GetValue("Window.State", "");
-			WindowBounds = ParseRectangleValue(ConfigFile.GetValue("Window.Bounds", ""));
+			WindowVisible = _configFile.GetValue("Window.Visible", true);
+			WindowState = _configFile.GetValue("Window.State", "");
+			WindowBounds = ParseRectangleValue(_configFile.GetValue("Window.Bounds", ""));
 
 			// Schedule settings
-			bScheduleEnabled = ConfigFile.GetValue("Schedule.Enabled", false);
-			if(!TimeSpan.TryParse(ConfigFile.GetValue("Schedule.Time", ""), out ScheduleTime))
+			ScheduleEnabled = _configFile.GetValue("Schedule.Enabled", false);
+			if(!TimeSpan.TryParse(_configFile.GetValue("Schedule.Time", ""), out ScheduleTime))
 			{
 				ScheduleTime = new TimeSpan(6, 0, 0);
 			}
 
-			ScheduleAnyOpenProject = ConfigFile.GetValue("Schedule.AnyOpenProject", true);
+			ScheduleAnyOpenProject = _configFile.GetValue("Schedule.AnyOpenProject", true);
 			ScheduleProjects = ReadProjectList("Schedule.Projects", "Schedule.ProjectFileNames");
 
 			// Notification settings
-			NotifyProjects = ConfigFile.GetValues("Notifications.NotifyProjects", new string[0]).ToList();
-			NotifyUnassignedMinutes = ConfigFile.GetValue("Notifications.NotifyUnassignedMinutes", -1);
-			NotifyUnacknowledgedMinutes = ConfigFile.GetValue("Notifications.NotifyUnacknowledgedMinutes", -1);
-			NotifyUnresolvedMinutes = ConfigFile.GetValue("Notifications.NotifyUnresolvedMinutes", -1);
+			NotifyProjects = _configFile.GetValues("Notifications.NotifyProjects", new string[0]).ToList();
+			NotifyUnassignedMinutes = _configFile.GetValue("Notifications.NotifyUnassignedMinutes", -1);
+			NotifyUnacknowledgedMinutes = _configFile.GetValue("Notifications.NotifyUnacknowledgedMinutes", -1);
+			NotifyUnresolvedMinutes = _configFile.GetValue("Notifications.NotifyUnresolvedMinutes", -1);
 
 			// Perforce settings
-			if(!int.TryParse(ConfigFile.GetValue("Perforce.NumRetries", "0"), out SyncOptions.NumRetries))
+			if(!int.TryParse(_configFile.GetValue("Perforce.NumRetries", "0"), out SyncOptions.NumRetries))
 			{
 				SyncOptions.NumRetries = 0;
 			}
 
-			int NumThreads;
-			if(int.TryParse(ConfigFile.GetValue("Perforce.NumThreads", "0"), out NumThreads) && NumThreads > 0)
+			int numThreads;
+			if(int.TryParse(_configFile.GetValue("Perforce.NumThreads", "0"), out numThreads) && numThreads > 0)
 			{
-				if(Version >= UserSettingsVersion.DefaultNumberOfThreads || NumThreads > 1)
+				if(Version >= UserSettingsVersion.DefaultNumberOfThreads || numThreads > 1)
 				{
-					SyncOptions.NumThreads = NumThreads;
+					SyncOptions.NumThreads = numThreads;
 				}
 			}
 
-			SyncOptions.TcpBufferSize = ConfigFile.GetValue("Perforce.TcpBufferSize", PerforceSyncOptions.DefaultTcpBufferSize);
-			SyncOptions.FileBufferSize = ConfigFile.GetValue("Perforce.FileBufferSize", PerforceSyncOptions.DefaultFileBufferSize);
-			SyncOptions.MaxCommandsPerBatch = ConfigFile.GetValue("Perforce.MaxCommandsPerBatch", PerforceSyncOptions.DefaultMaxCommandsPerBatch);
-			SyncOptions.MaxSizePerBatch = ConfigFile.GetValue("Perforce.MaxSizePerBatch", PerforceSyncOptions.DefaultMaxSizePerBatch);
-			SyncOptions.NumSyncErrorRetries = ConfigFile.GetValue("Perforce.NumSyncErrorRetries", PerforceSyncOptions.DefaultNumSyncErrorRetries);
+			SyncOptions.TcpBufferSize = _configFile.GetValue("Perforce.TcpBufferSize", PerforceSyncOptions.DefaultTcpBufferSize);
+			SyncOptions.FileBufferSize = _configFile.GetValue("Perforce.FileBufferSize", PerforceSyncOptions.DefaultFileBufferSize);
+			SyncOptions.MaxCommandsPerBatch = _configFile.GetValue("Perforce.MaxCommandsPerBatch", PerforceSyncOptions.DefaultMaxCommandsPerBatch);
+			SyncOptions.MaxSizePerBatch = _configFile.GetValue("Perforce.MaxSizePerBatch", PerforceSyncOptions.DefaultMaxSizePerBatch);
+			SyncOptions.NumSyncErrorRetries = _configFile.GetValue("Perforce.NumSyncErrorRetries", PerforceSyncOptions.DefaultNumSyncErrorRetries);
 		}
 
-		static Dictionary<Guid, bool> GetCategorySettings(ConfigSection Section, string IncludedKey, string ExcludedKey)
+		static Dictionary<Guid, bool> GetCategorySettings(ConfigSection section, string includedKey, string excludedKey)
 		{
-			Dictionary<Guid, bool> Result = new Dictionary<Guid, bool>();
-			if (Section != null)
+			Dictionary<Guid, bool> result = new Dictionary<Guid, bool>();
+			if (section != null)
 			{
-				foreach (Guid UniqueId in Section.GetValues(IncludedKey, new Guid[0]))
+				foreach (Guid uniqueId in section.GetValues(includedKey, new Guid[0]))
 				{
-					Result[UniqueId] = true;
+					result[uniqueId] = true;
 				}
-				foreach (Guid UniqueId in Section.GetValues(ExcludedKey, new Guid[0]))
+				foreach (Guid uniqueId in section.GetValues(excludedKey, new Guid[0]))
 				{
-					Result[UniqueId] = false;
+					result[uniqueId] = false;
 				}
 			}
-			return Result;
+			return result;
 		}
 
-		static void SetCategorySettings(ConfigSection Section, string IncludedKey, string ExcludedKey, Dictionary<Guid, bool> Categories)
+		static void SetCategorySettings(ConfigSection section, string includedKey, string excludedKey, Dictionary<Guid, bool> categories)
 		{
-			Guid[] IncludedCategories = Categories.Where(x => x.Value).Select(x => x.Key).ToArray();
-			if (IncludedCategories.Length > 0)
+			Guid[] includedCategories = categories.Where(x => x.Value).Select(x => x.Key).ToArray();
+			if (includedCategories.Length > 0)
 			{
-				Section.SetValues(IncludedKey, IncludedCategories);
+				section.SetValues(includedKey, includedCategories);
 			}
 
-			Guid[] ExcludedCategories = Categories.Where(x => !x.Value).Select(x => x.Key).ToArray();
-			if (ExcludedCategories.Length > 0)
+			Guid[] excludedCategories = categories.Where(x => !x.Value).Select(x => x.Key).ToArray();
+			if (excludedCategories.Length > 0)
 			{
-				Section.SetValues(ExcludedKey, ExcludedCategories);
+				section.SetValues(excludedKey, excludedCategories);
 			}
 		}
 
-		static Rectangle? ParseRectangleValue(string Text)
+		static Rectangle? ParseRectangleValue(string text)
 		{
-			ConfigObject Object = new ConfigObject(Text);
+			ConfigObject obj = new ConfigObject(text);
 
-			int X = Object.GetValue("X", -1);
-			int Y = Object.GetValue("Y", -1);
-			int W = Object.GetValue("W", -1);
-			int H = Object.GetValue("H", -1);
+			int x = obj.GetValue("X", -1);
+			int y = obj.GetValue("Y", -1);
+			int w = obj.GetValue("W", -1);
+			int h = obj.GetValue("H", -1);
 
-			if(X == -1 || Y == -1 || W == -1 || H == -1)
+			if(x == -1 || y == -1 || w == -1 || h == -1)
 			{
 				return null;
 			}
 			else
 			{
-				return new Rectangle(X, Y, W, H);
+				return new Rectangle(x, y, w, h);
 			}
 		}
 
-		static string FormatRectangleValue(Rectangle Value)
+		static string FormatRectangleValue(Rectangle value)
 		{
-			ConfigObject Object = new ConfigObject();
+			ConfigObject obj = new ConfigObject();
 
-			Object.SetValue("X", Value.X);
-			Object.SetValue("Y", Value.Y);
-			Object.SetValue("W", Value.Width);
-			Object.SetValue("H", Value.Height);
+			obj.SetValue("X", value.X);
+			obj.SetValue("Y", value.Y);
+			obj.SetValue("W", value.Width);
+			obj.SetValue("H", value.Height);
 
-			return Object.ToString();
+			return obj.ToString();
 		}
 
-		public static DirectoryReference GetConfigDir(DirectoryReference WorkspaceDir)
+		public static DirectoryReference GetConfigDir(DirectoryReference workspaceDir)
 		{
-			DirectoryReference ConfigDir = DirectoryReference.Combine(WorkspaceDir, ".ugs");
-			return ConfigDir;
+			DirectoryReference configDir = DirectoryReference.Combine(workspaceDir, ".ugs");
+			return configDir;
 		}
 
-		public static void CreateConfigDir(DirectoryReference ConfigDir)
+		public static void CreateConfigDir(DirectoryReference configDir)
 		{
-			DirectoryInfo ConfigDirInfo = ConfigDir.ToDirectoryInfo();
-			if (!ConfigDirInfo.Exists)
+			DirectoryInfo configDirInfo = configDir.ToDirectoryInfo();
+			if (!configDirInfo.Exists)
 			{
-				ConfigDirInfo.Create();
+				configDirInfo.Create();
 				if (RuntimeInformation.IsOSPlatform(OSPlatform.Windows))
 				{
-					ConfigDirInfo.Attributes = FileAttributes.Directory | FileAttributes.Hidden;
+					configDirInfo.Attributes = FileAttributes.Directory | FileAttributes.Hidden;
 				}
 			}
 		}
 
-		protected override void ImportWorkspaceState(DirectoryReference RootDir, string ClientName, string BranchPath, UserWorkspaceState CurrentWorkspace)
+		protected override void ImportWorkspaceState(DirectoryReference rootDir, string clientName, string branchPath, UserWorkspaceState currentWorkspace)
 		{
 			// Read the workspace settings
-			ConfigSection WorkspaceSection = ConfigFile.FindSection(ClientName + BranchPath);
-			if(WorkspaceSection == null)
+			ConfigSection workspaceSection = _configFile.FindSection(clientName + branchPath);
+			if(workspaceSection == null)
 			{
-				string LegacyBranchAndClientKey = ClientName + BranchPath;
+				string legacyBranchAndClientKey = clientName + branchPath;
 
-				int SlashIdx = LegacyBranchAndClientKey.IndexOf('/');
-				if(SlashIdx != -1)
+				int slashIdx = legacyBranchAndClientKey.IndexOf('/');
+				if(slashIdx != -1)
 				{
-					LegacyBranchAndClientKey = LegacyBranchAndClientKey.Substring(0, SlashIdx) + "$" + LegacyBranchAndClientKey.Substring(SlashIdx + 1);
+					legacyBranchAndClientKey = legacyBranchAndClientKey.Substring(0, slashIdx) + "$" + legacyBranchAndClientKey.Substring(slashIdx + 1);
 				}
 
-				string? CurrentSync = ConfigFile.GetValue("Clients." + LegacyBranchAndClientKey, null);
-				if(CurrentSync != null)
+				string? currentSync = _configFile.GetValue("Clients." + legacyBranchAndClientKey, null);
+				if(currentSync != null)
 				{
-					int AtIdx = CurrentSync.LastIndexOf('@');
-					if(AtIdx != -1)
+					int atIdx = currentSync.LastIndexOf('@');
+					if(atIdx != -1)
 					{
-						int ChangeNumber;
-						if(int.TryParse(CurrentSync.Substring(AtIdx + 1), out ChangeNumber))
+						int changeNumber;
+						if(int.TryParse(currentSync.Substring(atIdx + 1), out changeNumber))
 						{
-							CurrentWorkspace.ProjectIdentifier = CurrentSync.Substring(0, AtIdx);
-							CurrentWorkspace.CurrentChangeNumber = ChangeNumber;
+							currentWorkspace.ProjectIdentifier = currentSync.Substring(0, atIdx);
+							currentWorkspace.CurrentChangeNumber = changeNumber;
 						}
 					}
 				}
 
-				string? LastUpdateResultText = ConfigFile.GetValue("Clients." + LegacyBranchAndClientKey + "$LastUpdate", null);
-				if(LastUpdateResultText != null)
+				string? lastUpdateResultText = _configFile.GetValue("Clients." + legacyBranchAndClientKey + "$LastUpdate", null);
+				if(lastUpdateResultText != null)
 				{
-					int ColonIdx = LastUpdateResultText.LastIndexOf(':');
-					if(ColonIdx != -1)
+					int colonIdx = lastUpdateResultText.LastIndexOf(':');
+					if(colonIdx != -1)
 					{
-						int ChangeNumber;
-						if(int.TryParse(LastUpdateResultText.Substring(0, ColonIdx), out ChangeNumber))
+						int changeNumber;
+						if(int.TryParse(lastUpdateResultText.Substring(0, colonIdx), out changeNumber))
 						{
-							WorkspaceUpdateResult Result;
-							if(Enum.TryParse(LastUpdateResultText.Substring(ColonIdx + 1), out Result))
+							WorkspaceUpdateResult result;
+							if(Enum.TryParse(lastUpdateResultText.Substring(colonIdx + 1), out result))
 							{
-								CurrentWorkspace.LastSyncChangeNumber = ChangeNumber;
-								CurrentWorkspace.LastSyncResult = Result;
+								currentWorkspace.LastSyncChangeNumber = changeNumber;
+								currentWorkspace.LastSyncResult = result;
 							}
 						}
 					}
@@ -1056,326 +1056,326 @@ namespace UnrealGameSync
 			}
 			else
 			{
-				CurrentWorkspace.ProjectIdentifier = WorkspaceSection.GetValue("CurrentProjectPath", "");
-				CurrentWorkspace.CurrentChangeNumber = WorkspaceSection.GetValue("CurrentChangeNumber", -1);
-				CurrentWorkspace.CurrentSyncFilterHash = WorkspaceSection.GetValue("CurrentSyncFilterHash", null);
-				foreach(string AdditionalChangeNumberString in WorkspaceSection.GetValues("AdditionalChangeNumbers", new string[0]))
+				currentWorkspace.ProjectIdentifier = workspaceSection.GetValue("CurrentProjectPath", "");
+				currentWorkspace.CurrentChangeNumber = workspaceSection.GetValue("CurrentChangeNumber", -1);
+				currentWorkspace.CurrentSyncFilterHash = workspaceSection.GetValue("CurrentSyncFilterHash", null);
+				foreach(string additionalChangeNumberString in workspaceSection.GetValues("AdditionalChangeNumbers", new string[0]))
 				{
-					int AdditionalChangeNumber;
-					if(int.TryParse(AdditionalChangeNumberString, out AdditionalChangeNumber))
+					int additionalChangeNumber;
+					if(int.TryParse(additionalChangeNumberString, out additionalChangeNumber))
 					{
-						CurrentWorkspace.AdditionalChangeNumbers.Add(AdditionalChangeNumber);
+						currentWorkspace.AdditionalChangeNumbers.Add(additionalChangeNumber);
 					}
 				}
 
-				WorkspaceUpdateResult LastSyncResult;
-				Enum.TryParse(WorkspaceSection.GetValue("LastSyncResult", ""), out LastSyncResult);
-				CurrentWorkspace.LastSyncResult = LastSyncResult;
+				WorkspaceUpdateResult lastSyncResult;
+				Enum.TryParse(workspaceSection.GetValue("LastSyncResult", ""), out lastSyncResult);
+				currentWorkspace.LastSyncResult = lastSyncResult;
 
-				CurrentWorkspace.LastSyncResultMessage = UnescapeText(WorkspaceSection.GetValue("LastSyncResultMessage"));
-				CurrentWorkspace.LastSyncChangeNumber = WorkspaceSection.GetValue("LastSyncChangeNumber", -1);
+				currentWorkspace.LastSyncResultMessage = UnescapeText(workspaceSection.GetValue("LastSyncResultMessage"));
+				currentWorkspace.LastSyncChangeNumber = workspaceSection.GetValue("LastSyncChangeNumber", -1);
 
-				DateTime LastSyncTime;
-				if(DateTime.TryParse(WorkspaceSection.GetValue("LastSyncTime", ""), out LastSyncTime))
+				DateTime lastSyncTime;
+				if(DateTime.TryParse(workspaceSection.GetValue("LastSyncTime", ""), out lastSyncTime))
 				{
-					CurrentWorkspace.LastSyncTime = LastSyncTime;
+					currentWorkspace.LastSyncTime = lastSyncTime;
 				}
 
-				CurrentWorkspace.LastSyncDurationSeconds = WorkspaceSection.GetValue("LastSyncDuration", 0);
-				CurrentWorkspace.LastBuiltChangeNumber = WorkspaceSection.GetValue("LastBuiltChangeNumber", 0);
-				CurrentWorkspace.ExpandedArchiveTypes = WorkspaceSection.GetValues("ExpandedArchiveName", new string[0]);
+				currentWorkspace.LastSyncDurationSeconds = workspaceSection.GetValue("LastSyncDuration", 0);
+				currentWorkspace.LastBuiltChangeNumber = workspaceSection.GetValue("LastBuiltChangeNumber", 0);
+				currentWorkspace.ExpandedArchiveTypes = workspaceSection.GetValues("ExpandedArchiveName", new string[0]);
 
-				string[] BisectEntries = WorkspaceSection.GetValues("Bisect", new string[0]);
-				foreach(string BisectEntry in BisectEntries)
+				string[] bisectEntries = workspaceSection.GetValues("Bisect", new string[0]);
+				foreach(string bisectEntry in bisectEntries)
 				{
-					ConfigObject BisectEntryObject = new ConfigObject(BisectEntry);
+					ConfigObject bisectEntryObject = new ConfigObject(bisectEntry);
 
-					int ChangeNumber = BisectEntryObject.GetValue("Change", -1);
-					if(ChangeNumber != -1)
+					int changeNumber = bisectEntryObject.GetValue("Change", -1);
+					if(changeNumber != -1)
 					{
-						BisectState State;
-						if(Enum.TryParse(BisectEntryObject.GetValue("State", ""), out State))
+						BisectState state;
+						if(Enum.TryParse(bisectEntryObject.GetValue("State", ""), out state))
 						{
-							BisectEntry Entry = new BisectEntry();
-							Entry.Change = ChangeNumber;
-							Entry.State = State;
-							CurrentWorkspace.BisectChanges.Add(Entry);
+							BisectEntry entry = new BisectEntry();
+							entry.Change = changeNumber;
+							entry.State = state;
+							currentWorkspace.BisectChanges.Add(entry);
 						}
 					}
 				}
 			}
 		}
 
-		protected override void ImportWorkspaceSettings(DirectoryReference RootDir, string ClientName, string BranchPath, UserWorkspaceSettings CurrentWorkspace)
+		protected override void ImportWorkspaceSettings(DirectoryReference rootDir, string clientName, string branchPath, UserWorkspaceSettings currentWorkspace)
 		{
-			ConfigSection WorkspaceSection = ConfigFile.FindSection(ClientName + BranchPath);
-			if (WorkspaceSection != null)
+			ConfigSection workspaceSection = _configFile.FindSection(clientName + branchPath);
+			if (workspaceSection != null)
 			{
-				CurrentWorkspace.Filter.View = WorkspaceSection.GetValues("SyncFilter", new string[0]).ToList();
-				CurrentWorkspace.Filter.SetCategories(GetCategorySettings(WorkspaceSection, "SyncIncludedCategories", "SyncExcludedCategories"));
+				currentWorkspace.Filter.View = workspaceSection.GetValues("SyncFilter", new string[0]).ToList();
+				currentWorkspace.Filter.SetCategories(GetCategorySettings(workspaceSection, "SyncIncludedCategories", "SyncExcludedCategories"));
 
-				int SyncAllProjects = WorkspaceSection.GetValue("SyncAllProjects", -1);
-				CurrentWorkspace.Filter.AllProjects = (SyncAllProjects == 0) ? (bool?)false : (SyncAllProjects == 1) ? (bool?)true : (bool?)null;
+				int syncAllProjects = workspaceSection.GetValue("SyncAllProjects", -1);
+				currentWorkspace.Filter.AllProjects = (syncAllProjects == 0) ? (bool?)false : (syncAllProjects == 1) ? (bool?)true : (bool?)null;
 
-				int IncludeAllProjectsInSolution = WorkspaceSection.GetValue("IncludeAllProjectsInSolution", -1);
-				CurrentWorkspace.Filter.AllProjectsInSln = (IncludeAllProjectsInSolution == 0) ? (bool?)false : (IncludeAllProjectsInSolution == 1) ? (bool?)true : (bool?)null;
+				int includeAllProjectsInSolution = workspaceSection.GetValue("IncludeAllProjectsInSolution", -1);
+				currentWorkspace.Filter.AllProjectsInSln = (includeAllProjectsInSolution == 0) ? (bool?)false : (includeAllProjectsInSolution == 1) ? (bool?)true : (bool?)null;
 			}
 		}
 
-		protected override void ImportProjectSettings(ProjectInfo ProjectInfo, UserProjectSettings CurrentProject)
+		protected override void ImportProjectSettings(ProjectInfo projectInfo, UserProjectSettings currentProject)
 		{
-			string ClientProjectFileName = ProjectInfo.ClientFileName;
+			string clientProjectFileName = projectInfo.ClientFileName;
 
 			// Read the project settings
-			ConfigSection ProjectSection = ConfigFile.FindOrAddSection(ClientProjectFileName);
-			CurrentProject.BuildSteps.AddRange(ProjectSection.GetValues("BuildStep", new string[0]).Select(x => new ConfigObject(x)));
+			ConfigSection projectSection = _configFile.FindOrAddSection(clientProjectFileName);
+			currentProject.BuildSteps.AddRange(projectSection.GetValues("BuildStep", new string[0]).Select(x => new ConfigObject(x)));
 
-			FilterType FilterType;
-			if (!Enum.TryParse(ProjectSection.GetValue("FilterType", ""), true, out FilterType))
+			FilterType filterType;
+			if (!Enum.TryParse(projectSection.GetValue("FilterType", ""), true, out filterType))
 			{
-				FilterType = FilterType.None;
+				filterType = FilterType.None;
 			}
 
-			CurrentProject.FilterType = FilterType;
-			CurrentProject.FilterBadges.UnionWith(ProjectSection.GetValues("FilterBadges", new string[0]));
+			currentProject.FilterType = filterType;
+			currentProject.FilterBadges.UnionWith(projectSection.GetValues("FilterBadges", new string[0]));
 		}
 
-		public override bool Save(ILogger Logger)
+		public override bool Save(ILogger logger)
 		{
-			if (!base.Save(Logger))
+			if (!base.Save(logger))
 			{
 				return false;
 			}
 
 			// General settings
-			ConfigSection GeneralSection = ConfigFile.FindOrAddSection("General");
-			GeneralSection.Clear();
-			GeneralSection.SetValue("Version", (int)Version);
-			GeneralSection.SetValue("BuildAfterSync", bBuildAfterSync);
-			GeneralSection.SetValue("RunAfterSync", bRunAfterSync);
-			GeneralSection.SetValue("OpenSolutionAfterSync", bOpenSolutionAfterSync);
-			GeneralSection.SetValue("ShowLogWindow", bShowLogWindow);
-			GeneralSection.SetValue("AutoResolveConflicts", bAutoResolveConflicts);
-			GeneralSection.SetValue("ShowUnreviewed", bShowUnreviewedChanges);
-			GeneralSection.SetValue("ShowAutomated", bShowAutomatedChanges);
-			GeneralSection.SetValue("RobomergeFilter", ShowRobomerge.ToString());
-			GeneralSection.SetValue("AnnotateRobomerge", bAnnotateRobmergeChanges);
-			GeneralSection.SetValue("ShowLocalTimes", bShowLocalTimes);
+			ConfigSection generalSection = _configFile.FindOrAddSection("General");
+			generalSection.Clear();
+			generalSection.SetValue("Version", (int)Version);
+			generalSection.SetValue("BuildAfterSync", BuildAfterSync);
+			generalSection.SetValue("RunAfterSync", RunAfterSync);
+			generalSection.SetValue("OpenSolutionAfterSync", OpenSolutionAfterSync);
+			generalSection.SetValue("ShowLogWindow", ShowLogWindow);
+			generalSection.SetValue("AutoResolveConflicts", AutoResolveConflicts);
+			generalSection.SetValue("ShowUnreviewed", ShowUnreviewedChanges);
+			generalSection.SetValue("ShowAutomated", ShowAutomatedChanges);
+			generalSection.SetValue("RobomergeFilter", ShowRobomerge.ToString());
+			generalSection.SetValue("AnnotateRobomerge", AnnotateRobmergeChanges);
+			generalSection.SetValue("ShowLocalTimes", ShowLocalTimes);
 			if(LastProject != null)
 			{
-				GeneralSection.SetValue("LastProject", LastProject.ToConfigEntry());
+				generalSection.SetValue("LastProject", LastProject.ToConfigEntry());
 			}
-			GeneralSection.SetValues("OpenProjects", OpenProjects.Select(x => x.ToConfigEntry()).ToArray());
-			GeneralSection.SetValue("KeepInTray", bKeepInTray);
-			GeneralSection.SetValues("EnabledTools", EnabledTools);
-			GeneralSection.SetValue("FilterIndex", FilterIndex);
-			GeneralSection.SetValues("RecentProjects", RecentProjects.Select(x => x.ToConfigEntry()).ToArray());
-			GeneralSection.SetValue("SyncTypeID", SyncTypeID);
+			generalSection.SetValues("OpenProjects", OpenProjects.Select(x => x.ToConfigEntry()).ToArray());
+			generalSection.SetValue("KeepInTray", KeepInTray);
+			generalSection.SetValues("EnabledTools", EnabledTools);
+			generalSection.SetValue("FilterIndex", FilterIndex);
+			generalSection.SetValues("RecentProjects", RecentProjects.Select(x => x.ToConfigEntry()).ToArray());
+			generalSection.SetValue("SyncTypeID", SyncTypeId);
 
 			// Build configuration
-			GeneralSection.SetValue("BuildConfig", CompiledEditorBuildConfig.ToString());
+			generalSection.SetValue("BuildConfig", CompiledEditorBuildConfig.ToString());
 
 			// Tab labels
-			GeneralSection.SetValue("TabLabels", TabLabels.ToString());
+			generalSection.SetValue("TabLabels", TabLabels.ToString());
 
 			// Editor arguments
-			List<string> EditorArgumentList = new List<string>();
-			foreach(Tuple<string, bool> EditorArgument in EditorArguments)
+			List<string> editorArgumentList = new List<string>();
+			foreach(Tuple<string, bool> editorArgument in EditorArguments)
 			{
-				EditorArgumentList.Add(String.Format("{0}:{1}", EditorArgument.Item2? 1 : 0, EditorArgument.Item1));
+				editorArgumentList.Add(String.Format("{0}:{1}", editorArgument.Item2? 1 : 0, editorArgument.Item1));
 			}
-			GeneralSection.SetValues("EditorArguments", EditorArgumentList.ToArray());
-			GeneralSection.SetValue("EditorArgumentsPrompt", bEditorArgumentsPrompt);
+			generalSection.SetValues("EditorArguments", editorArgumentList.ToArray());
+			generalSection.SetValue("EditorArgumentsPrompt", EditorArgumentsPrompt);
 
 			// Schedule settings
-			ConfigSection ScheduleSection = ConfigFile.FindOrAddSection("Schedule");
-			ScheduleSection.Clear();
-			ScheduleSection.SetValue("Enabled", bScheduleEnabled);
-			ScheduleSection.SetValue("Time", ScheduleTime.ToString());
-			ScheduleSection.SetValue("AnyOpenProject", ScheduleAnyOpenProject);
-			ScheduleSection.SetValues("Projects", ScheduleProjects.Select(x => x.ToConfigEntry()).ToArray());
+			ConfigSection scheduleSection = _configFile.FindOrAddSection("Schedule");
+			scheduleSection.Clear();
+			scheduleSection.SetValue("Enabled", ScheduleEnabled);
+			scheduleSection.SetValue("Time", ScheduleTime.ToString());
+			scheduleSection.SetValue("AnyOpenProject", ScheduleAnyOpenProject);
+			scheduleSection.SetValues("Projects", ScheduleProjects.Select(x => x.ToConfigEntry()).ToArray());
 
 			// Precompiled binaries
-			ConfigSection ArchivesSection = ConfigFile.FindOrAddSection("PrecompiledBinaries");
-			ArchivesSection.SetValues("Archives", Archives.Select(x => x.ToConfigEntry()).ToArray());
+			ConfigSection archivesSection = _configFile.FindOrAddSection("PrecompiledBinaries");
+			archivesSection.SetValues("Archives", Archives.Select(x => x.ToConfigEntry()).ToArray());
 
 			// OIDC Settings
-			ConfigSection OIDCSection = ConfigFile.FindOrAddSection("OIDCProviders");
-			List<ConfigObject> TokenObjects = new List<ConfigObject>();
-			foreach (KeyValuePair<string, string> Pair in ProviderToRefreshTokens)
+			ConfigSection oidcSection = _configFile.FindOrAddSection("OIDCProviders");
+			List<ConfigObject> tokenObjects = new List<ConfigObject>();
+			foreach (KeyValuePair<string, string> pair in ProviderToRefreshTokens)
 			{
-				ConfigObject TokenEntryObject = new ConfigObject();
-				TokenEntryObject.SetValue("Provider", Pair.Key);
-				TokenEntryObject.SetValue("Token", Pair.Value);
-				TokenObjects.Add(TokenEntryObject);
+				ConfigObject tokenEntryObject = new ConfigObject();
+				tokenEntryObject.SetValue("Provider", pair.Key);
+				tokenEntryObject.SetValue("Token", pair.Value);
+				tokenObjects.Add(tokenEntryObject);
 			}
-			OIDCSection.SetValues("Tokens", TokenObjects.Select(x => x.ToString()).ToArray());
+			oidcSection.SetValues("Tokens", tokenObjects.Select(x => x.ToString()).ToArray());
 
 
 			// Window settings
-			ConfigSection WindowSection = ConfigFile.FindOrAddSection("Window");
-			WindowSection.Clear();
-			WindowSection.SetValue("Visible", bWindowVisible);
-			WindowSection.SetValue("State", WindowState);
+			ConfigSection windowSection = _configFile.FindOrAddSection("Window");
+			windowSection.Clear();
+			windowSection.SetValue("Visible", WindowVisible);
+			windowSection.SetValue("State", WindowState);
 			if(WindowBounds != null)
 			{
-				WindowSection.SetValue("Bounds", FormatRectangleValue(WindowBounds.Value));
+				windowSection.SetValue("Bounds", FormatRectangleValue(WindowBounds.Value));
 			}
 
 			// Notification settings
-			ConfigSection NotificationSection = ConfigFile.FindOrAddSection("Notifications");
-			NotificationSection.Clear();
+			ConfigSection notificationSection = _configFile.FindOrAddSection("Notifications");
+			notificationSection.Clear();
 			if (NotifyProjects.Count > 0)
 			{
-				NotificationSection.SetValues("NotifyProjects", NotifyProjects.ToArray());
+				notificationSection.SetValues("NotifyProjects", NotifyProjects.ToArray());
 			}
 			if (NotifyUnassignedMinutes != -1)
 			{
-				NotificationSection.SetValue("NotifyUnassignedMinutes", NotifyUnassignedMinutes);
+				notificationSection.SetValue("NotifyUnassignedMinutes", NotifyUnassignedMinutes);
 			}
 			if (NotifyUnacknowledgedMinutes != -1)
 			{
-				NotificationSection.SetValue("NotifyUnacknowledgedMinutes", NotifyUnacknowledgedMinutes);
+				notificationSection.SetValue("NotifyUnacknowledgedMinutes", NotifyUnacknowledgedMinutes);
 			}
 			if (NotifyUnresolvedMinutes != -1)
 			{
-				NotificationSection.SetValue("NotifyUnresolvedMinutes", NotifyUnresolvedMinutes);
+				notificationSection.SetValue("NotifyUnresolvedMinutes", NotifyUnresolvedMinutes);
 			}
 
 			// Perforce settings
-			ConfigSection PerforceSection = ConfigFile.FindOrAddSection("Perforce");
-			PerforceSection.Clear();
+			ConfigSection perforceSection = _configFile.FindOrAddSection("Perforce");
+			perforceSection.Clear();
 			if(SyncOptions.NumRetries > 0 && SyncOptions.NumRetries != PerforceSyncOptions.DefaultNumRetries)
 			{
-				PerforceSection.SetValue("NumRetries", SyncOptions.NumRetries);
+				perforceSection.SetValue("NumRetries", SyncOptions.NumRetries);
 			}
 			if(SyncOptions.NumThreads > 0 && SyncOptions.NumThreads != PerforceSyncOptions.DefaultNumThreads)
 			{
-				PerforceSection.SetValue("NumThreads", SyncOptions.NumThreads);
+				perforceSection.SetValue("NumThreads", SyncOptions.NumThreads);
 			}
 			if(SyncOptions.TcpBufferSize > 0 && SyncOptions.TcpBufferSize != PerforceSyncOptions.DefaultTcpBufferSize)
 			{
-				PerforceSection.SetValue("TcpBufferSize", SyncOptions.TcpBufferSize);
+				perforceSection.SetValue("TcpBufferSize", SyncOptions.TcpBufferSize);
 			}
 			if (SyncOptions.FileBufferSize > 0 && SyncOptions.FileBufferSize != PerforceSyncOptions.DefaultFileBufferSize)
 			{
-				PerforceSection.SetValue("FileBufferSize", SyncOptions.FileBufferSize);
+				perforceSection.SetValue("FileBufferSize", SyncOptions.FileBufferSize);
 			}
 			if (SyncOptions.MaxCommandsPerBatch > 0 && SyncOptions.MaxCommandsPerBatch != PerforceSyncOptions.DefaultMaxCommandsPerBatch)
 			{
-				PerforceSection.SetValue("MaxCommandsPerBatch", SyncOptions.MaxCommandsPerBatch);
+				perforceSection.SetValue("MaxCommandsPerBatch", SyncOptions.MaxCommandsPerBatch);
 			}
 			if (SyncOptions.MaxSizePerBatch > 0 && SyncOptions.MaxSizePerBatch != PerforceSyncOptions.DefaultMaxSizePerBatch)
 			{
-				PerforceSection.SetValue("MaxSizePerBatch", SyncOptions.MaxSizePerBatch);
+				perforceSection.SetValue("MaxSizePerBatch", SyncOptions.MaxSizePerBatch);
 			}
 			if (SyncOptions.NumSyncErrorRetries > 0 && SyncOptions.NumSyncErrorRetries != PerforceSyncOptions.DefaultNumSyncErrorRetries)
 			{
-				PerforceSection.SetValue("NumSyncErrorRetries", SyncOptions.NumSyncErrorRetries);
+				perforceSection.SetValue("NumSyncErrorRetries", SyncOptions.NumSyncErrorRetries);
 			}
 
 			// Save the file
 			try
 			{
-				ConfigFile.Save(FileName);
+				_configFile.Save(_fileName);
 				return true;
 			}
-			catch (Exception Ex)
+			catch (Exception ex)
 			{
-				Logger.LogError(Ex, "Unable to save config file {FileName}: {Message}", FileName, Ex.Message);
+				logger.LogError(ex, "Unable to save config file {FileName}: {Message}", _fileName, ex.Message);
 				return false;
 			}
 		}
 
-		[return: NotNullIfNotNull("Text")]
-		static string? EscapeText(string? Text)
+		[return: NotNullIfNotNull("text")]
+		static string? EscapeText(string? text)
 		{
-			if(Text == null)
+			if(text == null)
 			{
 				return null;
 			}
 
-			StringBuilder Result = new StringBuilder();
-			for(int Idx = 0; Idx < Text.Length; Idx++)
+			StringBuilder result = new StringBuilder();
+			for(int idx = 0; idx < text.Length; idx++)
 			{
-				switch(Text[Idx])
+				switch(text[idx])
 				{
 					case '\\':
-						Result.Append("\\\\");
+						result.Append("\\\\");
 						break;
 					case '\t':
-						Result.Append("\\t");
+						result.Append("\\t");
 						break;
 					case '\r':
-						Result.Append("\\r");
+						result.Append("\\r");
 						break;
 					case '\n':
-						Result.Append("\\n");
+						result.Append("\\n");
 						break;
 					case '\'':
-						Result.Append("\\\'");
+						result.Append("\\\'");
 						break;
 					case '\"':
-						Result.Append("\\\"");
+						result.Append("\\\"");
 						break;
 					default:
-						Result.Append(Text[Idx]);
+						result.Append(text[idx]);
 						break;
 				}
 			}
-			return Result.ToString();
+			return result.ToString();
 		}
 
-		[return: NotNullIfNotNull("Text")]
-		static string? UnescapeText(string? Text)
+		[return: NotNullIfNotNull("text")]
+		static string? UnescapeText(string? text)
 		{
-			if(Text == null)
+			if(text == null)
 			{
 				return null;
 			}
 
-			StringBuilder Result = new StringBuilder();
-			for(int Idx = 0; Idx < Text.Length; Idx++)
+			StringBuilder result = new StringBuilder();
+			for(int idx = 0; idx < text.Length; idx++)
 			{
-				if(Text[Idx] == '\\' && Idx + 1 < Text.Length)
+				if(text[idx] == '\\' && idx + 1 < text.Length)
 				{
-					switch(Text[++Idx])
+					switch(text[++idx])
 					{
 						case 't':
-							Result.Append('\t');
+							result.Append('\t');
 							break;
 						case 'r':
-							Result.Append('\r');
+							result.Append('\r');
 							break;
 						case 'n':
-							Result.Append('\n');
+							result.Append('\n');
 							break;
 						case '\'':
-							Result.Append('\'');
+							result.Append('\'');
 							break;
 						case '\"':
-							Result.Append('\"');
+							result.Append('\"');
 							break;
 						default:
-							Result.Append(Text[Idx]);
+							result.Append(text[idx]);
 							break;
 					}
 				}
 				else
 				{
-					Result.Append(Text[Idx]);
+					result.Append(text[idx]);
 				}
 			}
-			return Result.ToString();
+			return result.ToString();
 		}
 
 		public IEnumerable<FileReference> GetCachedFilePaths()
 		{
-			List<FileReference> Files = new List<FileReference>();
-			Files.AddRange(WorkspaceDirToSettings.Values.Select(x => x.ConfigFile));
-			Files.AddRange(WorkspaceDirToState.Values.Select(x => x.ConfigFile));
-			Files.AddRange(ProjectKeyToSettings.Keys);
-			return Files;
+			List<FileReference> files = new List<FileReference>();
+			files.AddRange(_workspaceDirToSettings.Values.Select(x => x.ConfigFile));
+			files.AddRange(_workspaceDirToState.Values.Select(x => x.ConfigFile));
+			files.AddRange(_projectKeyToSettings.Keys);
+			return files;
 		}
 	}
 }
