UGS: Reformat to standard C# style conventions.

diff --git a/Engine/Source/Programs/UnrealGameSync/UnrealGameSyncShared/Utility.cs b/Engine/Source/Programs/UnrealGameSync/UnrealGameSyncShared/Utility.cs
--- a/Engine/Source/Programs/UnrealGameSync/UnrealGameSyncShared/Utility.cs
+++ b/Engine/Source/Programs/UnrealGameSync/UnrealGameSyncShared/Utility.cs
@@ -24,9 +24,9 @@ namespace UnrealGameSync
 	{
 		public int Code { get; }
 
-		public UserErrorException(string Message, int Code = 1) : base(Message)
+		public UserErrorException(string message, int code = 1) : base(message)
 		{
-			this.Code = Code;
+			this.Code = code;
 		}
 	}
 
@@ -57,26 +57,26 @@ namespace UnrealGameSync
 		};
 
 		public string Description;
-		public bool bContainsCode;
-		public bool bContainsContent;
+		public bool ContainsCode;
+		public bool ContainsContent;
 
-		public PerforceChangeDetails(DescribeRecord DescribeRecord)
+		public PerforceChangeDetails(DescribeRecord describeRecord)
 		{
-			Description = DescribeRecord.Description;
+			Description = describeRecord.Description;
 
 			// Check whether the files are code or content
-			foreach (DescribeFileRecord File in DescribeRecord.Files)
+			foreach (DescribeFileRecord file in describeRecord.Files)
 			{
-				if (CodeExtensions.Any(Extension => File.DepotFile.EndsWith(Extension, StringComparison.OrdinalIgnoreCase)))
+				if (CodeExtensions.Any(extension => file.DepotFile.EndsWith(extension, StringComparison.OrdinalIgnoreCase)))
 				{
-					bContainsCode = true;
+					ContainsCode = true;
 				}
 				else
 				{
-					bContainsContent = true;
+					ContainsContent = true;
 				}
 
-				if (bContainsCode && bContainsContent)
+				if (ContainsCode && ContainsContent)
 				{
 					break;
 				}
@@ -88,123 +88,123 @@ namespace UnrealGameSync
 	{
 		static JsonSerializerOptions GetDefaultJsonSerializerOptions()
 		{
-			JsonSerializerOptions Options = new JsonSerializerOptions();
-			Options.AllowTrailingCommas = true;
-			Options.ReadCommentHandling = JsonCommentHandling.Skip;
-			Options.PropertyNameCaseInsensitive = true;
-			Options.Converters.Add(new JsonStringEnumConverter());
-			return Options;
+			JsonSerializerOptions options = new JsonSerializerOptions();
+			options.AllowTrailingCommas = true;
+			options.ReadCommentHandling = JsonCommentHandling.Skip;
+			options.PropertyNameCaseInsensitive = true;
+			options.Converters.Add(new JsonStringEnumConverter());
+			return options;
 		}
 
 		public static JsonSerializerOptions DefaultJsonSerializerOptions { get; } = GetDefaultJsonSerializerOptions();
 
-		public static bool TryLoadJson<T>(FileReference File, [NotNullWhen(true)] out T? Object) where T : class
+		public static bool TryLoadJson<T>(FileReference file, [NotNullWhen(true)] out T? obj) where T : class
 		{
-			if (!FileReference.Exists(File))
+			if (!FileReference.Exists(file))
 			{
-				Object = null;
+				obj = null;
 				return false;
 			}
 
 			try
 			{
-				Object = LoadJson<T>(File);
+				obj = LoadJson<T>(file);
 				return true;
 			}
 			catch
 			{
-				Object = null;
+				obj = null;
 				return false;
 			}
 		}
 
-		public static T LoadJson<T>(FileReference File)
+		public static T LoadJson<T>(FileReference file)
 		{
-			byte[] Data = FileReference.ReadAllBytes(File);
-			return JsonSerializer.Deserialize<T>(Data, DefaultJsonSerializerOptions);
+			byte[] data = FileReference.ReadAllBytes(file);
+			return JsonSerializer.Deserialize<T>(data, DefaultJsonSerializerOptions);
 		}
 
-		public static void SaveJson<T>(FileReference File, T Object)
+		public static void SaveJson<T>(FileReference file, T obj)
 		{
-			JsonSerializerOptions Options = new JsonSerializerOptions { IgnoreNullValues = true, WriteIndented = true };
-			Options.Converters.Add(new JsonStringEnumConverter());
+			JsonSerializerOptions options = new JsonSerializerOptions { IgnoreNullValues = true, WriteIndented = true };
+			options.Converters.Add(new JsonStringEnumConverter());
 
-			byte[] Buffer;
-			using (MemoryStream Stream = new MemoryStream())
+			byte[] buffer;
+			using (MemoryStream stream = new MemoryStream())
 			{
-				using (Utf8JsonWriter Writer = new Utf8JsonWriter(Stream, new JsonWriterOptions { Indented = true }))
+				using (Utf8JsonWriter writer = new Utf8JsonWriter(stream, new JsonWriterOptions { Indented = true }))
 				{
-					JsonSerializer.Serialize(Writer, Object, Options);
+					JsonSerializer.Serialize(writer, obj, options);
 				}
-				Buffer = Stream.ToArray();
+				buffer = stream.ToArray();
 			}
 
-			FileReference.WriteAllBytes(File, Buffer);
+			FileReference.WriteAllBytes(file, buffer);
 		}
 
-		public static string GetPathWithCorrectCase(FileInfo Info)
+		public static string GetPathWithCorrectCase(FileInfo info)
 		{
-			DirectoryInfo ParentInfo = Info.Directory;
-			if(Info.Exists)
+			DirectoryInfo parentInfo = info.Directory;
+			if(info.Exists)
 			{
-				return Path.Combine(GetPathWithCorrectCase(ParentInfo), ParentInfo.GetFiles(Info.Name)[0].Name); 
+				return Path.Combine(GetPathWithCorrectCase(parentInfo), parentInfo.GetFiles(info.Name)[0].Name); 
 			}
 			else
 			{
-				return Path.Combine(GetPathWithCorrectCase(ParentInfo), Info.Name);
+				return Path.Combine(GetPathWithCorrectCase(parentInfo), info.Name);
 			}
 		}
 
-		public static string GetPathWithCorrectCase(DirectoryInfo Info)
+		public static string GetPathWithCorrectCase(DirectoryInfo info)
 		{
-			DirectoryInfo ParentInfo = Info.Parent;
-			if(ParentInfo == null)
+			DirectoryInfo parentInfo = info.Parent;
+			if(parentInfo == null)
 			{
-				return Info.FullName.ToUpperInvariant();
+				return info.FullName.ToUpperInvariant();
 			}
-			else if(Info.Exists)
+			else if(info.Exists)
 			{
-				return Path.Combine(GetPathWithCorrectCase(ParentInfo), ParentInfo.GetDirectories(Info.Name)[0].Name);
+				return Path.Combine(GetPathWithCorrectCase(parentInfo), parentInfo.GetDirectories(info.Name)[0].Name);
 			}
 			else
 			{
-				return Path.Combine(GetPathWithCorrectCase(ParentInfo), Info.Name);
+				return Path.Combine(GetPathWithCorrectCase(parentInfo), info.Name);
 			}
 		}
 
-		public static void ForceDeleteFile(string FileName)
+		public static void ForceDeleteFile(string fileName)
 		{
-			if(File.Exists(FileName))
+			if(File.Exists(fileName))
 			{
-				File.SetAttributes(FileName, File.GetAttributes(FileName) & ~FileAttributes.ReadOnly);
-				File.Delete(FileName);
+				File.SetAttributes(fileName, File.GetAttributes(fileName) & ~FileAttributes.ReadOnly);
+				File.Delete(fileName);
 			}
 		}
 
-		public static bool SpawnProcess(string FileName, string CommandLine)
+		public static bool SpawnProcess(string fileName, string commandLine)
 		{
-			using(Process ChildProcess = new Process())
+			using(Process childProcess = new Process())
 			{
-				ChildProcess.StartInfo.FileName = FileName;
-				ChildProcess.StartInfo.Arguments = String.IsNullOrEmpty(CommandLine) ? "" : CommandLine;
-				ChildProcess.StartInfo.UseShellExecute = false;
-				return ChildProcess.Start();
+				childProcess.StartInfo.FileName = fileName;
+				childProcess.StartInfo.Arguments = String.IsNullOrEmpty(commandLine) ? "" : commandLine;
+				childProcess.StartInfo.UseShellExecute = false;
+				return childProcess.Start();
 			}
 		}
 
-		public static bool SpawnHiddenProcess(string FileName, string CommandLine)
+		public static bool SpawnHiddenProcess(string fileName, string commandLine)
 		{
-			using(Process ChildProcess = new Process())
+			using(Process childProcess = new Process())
 			{
-				ChildProcess.StartInfo.FileName = FileName;
-				ChildProcess.StartInfo.Arguments = String.IsNullOrEmpty(CommandLine) ? "" : CommandLine;
-				ChildProcess.StartInfo.UseShellExecute = false;
-				ChildProcess.StartInfo.RedirectStandardOutput = true;
-				ChildProcess.StartInfo.RedirectStandardError = true;
-				ChildProcess.StartInfo.CreateNoWindow = true;
+				childProcess.StartInfo.FileName = fileName;
+				childProcess.StartInfo.Arguments = String.IsNullOrEmpty(commandLine) ? "" : commandLine;
+				childProcess.StartInfo.UseShellExecute = false;
+				childProcess.StartInfo.RedirectStandardOutput = true;
+				childProcess.StartInfo.RedirectStandardError = true;
+				childProcess.StartInfo.CreateNoWindow = true;
 				try
 				{
-					return ChildProcess.Start();
+					return childProcess.Start();
 				}
 				catch
 				{
@@ -213,30 +213,30 @@ namespace UnrealGameSync
 			}
 		}
 
-		public static async Task<int> ExecuteProcessAsync(string FileName, string? WorkingDir, string CommandLine, Action<string> OutputLine, CancellationToken CancellationToken)
+		public static async Task<int> ExecuteProcessAsync(string fileName, string? workingDir, string commandLine, Action<string> outputLine, CancellationToken cancellationToken)
 		{
-			using (ManagedProcess NewProcess = new ManagedProcess(null, FileName, CommandLine, WorkingDir, null, null, ProcessPriorityClass.Normal))
+			using (ManagedProcess newProcess = new ManagedProcess(null, fileName, commandLine, workingDir, null, null, ProcessPriorityClass.Normal))
 			{
 				for (; ; )
 				{
-					string? Line = await NewProcess.ReadLineAsync(CancellationToken);
-					if (Line == null)
+					string? line = await newProcess.ReadLineAsync(cancellationToken);
+					if (line == null)
 					{
-						NewProcess.WaitForExit();
-						return NewProcess.ExitCode;
+						newProcess.WaitForExit();
+						return newProcess.ExitCode;
 					}
-					OutputLine(Line);
+					outputLine(line);
 				}
 			}
 		}
 
-		public static bool SafeIsFileUnderDirectory(string FileName, string DirectoryName)
+		public static bool SafeIsFileUnderDirectory(string fileName, string directoryName)
 		{
 			try
 			{
-				string FullDirectoryName = Path.GetFullPath(DirectoryName).TrimEnd(Path.DirectorySeparatorChar) + Path.DirectorySeparatorChar;
-				string FullFileName = Path.GetFullPath(FileName);
-				return FullFileName.StartsWith(FullDirectoryName, StringComparison.InvariantCultureIgnoreCase);
+				string fullDirectoryName = Path.GetFullPath(directoryName).TrimEnd(Path.DirectorySeparatorChar) + Path.DirectorySeparatorChar;
+				string fullFileName = Path.GetFullPath(fileName);
+				return fullFileName.StartsWith(fullDirectoryName, StringComparison.InvariantCultureIgnoreCase);
 			}
 			catch(Exception)
 			{
@@ -248,58 +248,58 @@ namespace UnrealGameSync
 		/// Expands variables in $(VarName) format in the given string. Variables are retrieved from the given dictionary, or through the environment of the current process.
 		/// Any unknown variables are ignored.
 		/// </summary>
-		/// <param name="InputString">String to search for variable names</param>
-		/// <param name="AdditionalVariables">Lookup of variable names to values</param>
+		/// <param name="inputString">String to search for variable names</param>
+		/// <param name="additionalVariables">Lookup of variable names to values</param>
 		/// <returns>String with all variables replaced</returns>
-		public static string ExpandVariables(string InputString, Dictionary<string, string>? AdditionalVariables = null)
+		public static string ExpandVariables(string inputString, Dictionary<string, string>? additionalVariables = null)
 		{
-			string Result = InputString;
-			for (int Idx = Result.IndexOf("$("); Idx != -1; Idx = Result.IndexOf("$(", Idx))
+			string result = inputString;
+			for (int idx = result.IndexOf("$("); idx != -1; idx = result.IndexOf("$(", idx))
 			{
 				// Find the end of the variable name
-				int EndIdx = Result.IndexOf(')', Idx + 2);
-				if (EndIdx == -1)
+				int endIdx = result.IndexOf(')', idx + 2);
+				if (endIdx == -1)
 				{
 					break;
 				}
 
 				// Extract the variable name from the string
-				string Name = Result.Substring(Idx + 2, EndIdx - (Idx + 2));
+				string name = result.Substring(idx + 2, endIdx - (idx + 2));
 
 				// Strip the format from the name
-				string? Format = null;
-				int FormatIdx = Name.IndexOf(':');
-				if(FormatIdx != -1)
+				string? format = null;
+				int formatIdx = name.IndexOf(':');
+				if(formatIdx != -1)
 				{ 
-					Format = Name.Substring(FormatIdx + 1);
-					Name = Name.Substring(0, FormatIdx);
+					format = name.Substring(formatIdx + 1);
+					name = name.Substring(0, formatIdx);
 				}
 
 				// Find the value for it, either from the dictionary or the environment block
-				string? Value;
-				if (AdditionalVariables == null || !AdditionalVariables.TryGetValue(Name, out Value))
+				string? value;
+				if (additionalVariables == null || !additionalVariables.TryGetValue(name, out value))
 				{
-					Value = Environment.GetEnvironmentVariable(Name);
-					if (Value == null)
+					value = Environment.GetEnvironmentVariable(name);
+					if (value == null)
 					{
-						Idx = EndIdx + 1;
+						idx = endIdx + 1;
 						continue;
 					}
 				}
 
 				// Encode the variable if necessary
-				if(Format != null)
+				if(format != null)
 				{
-					if(String.Equals(Format, "URI", StringComparison.InvariantCultureIgnoreCase))
+					if(String.Equals(format, "URI", StringComparison.InvariantCultureIgnoreCase))
 					{
-						Value = Uri.EscapeDataString(Value);
+						value = Uri.EscapeDataString(value);
 					}
 				}
 
 				// Replace the variable, or skip past it
-				Result = Result.Substring(0, Idx) + Value + Result.Substring(EndIdx + 1);
+				result = result.Substring(0, idx) + value + result.Substring(endIdx + 1);
 			}
-			return Result;
+			return result;
 		}
 
 		class ProjectJson
@@ -312,17 +312,17 @@ namespace UnrealGameSync
 		/// </summary>
 		/// <param name="FileName">Path to the project file</param>
 		/// <returns>True if the given filename is an enterprise project</returns>
-		public static bool IsEnterpriseProjectFromText(string Text)
+		public static bool IsEnterpriseProjectFromText(string text)
 		{
 			try
 			{
-				JsonSerializerOptions Options = new JsonSerializerOptions();
-				Options.PropertyNameCaseInsensitive = true;
-				Options.Converters.Add(new JsonStringEnumConverter());
+				JsonSerializerOptions options = new JsonSerializerOptions();
+				options.PropertyNameCaseInsensitive = true;
+				options.Converters.Add(new JsonStringEnumConverter());
 
-				ProjectJson Project = JsonSerializer.Deserialize<ProjectJson>(Text, Options);
+				ProjectJson project = JsonSerializer.Deserialize<ProjectJson>(text, options);
 
-				return Project.Enterprise;
+				return project.Enterprise;
 			}
 			catch
 			{
@@ -332,110 +332,110 @@ namespace UnrealGameSync
 
 		/******/
 
-		private static void AddLocalConfigPaths_WithSubFolders(DirectoryInfo BaseDir, string FileName, List<FileInfo> Files)
+		private static void AddLocalConfigPaths_WithSubFolders(DirectoryInfo baseDir, string fileName, List<FileInfo> files)
 		{
-			if(BaseDir.Exists)
+			if(baseDir.Exists)
 			{
-				FileInfo BaseFileInfo = new FileInfo(Path.Combine(BaseDir.FullName, FileName));
-				if(BaseFileInfo.Exists)
+				FileInfo baseFileInfo = new FileInfo(Path.Combine(baseDir.FullName, fileName));
+				if(baseFileInfo.Exists)
 				{
-					Files.Add(BaseFileInfo);
+					files.Add(baseFileInfo);
 				}
 
-				foreach (DirectoryInfo SubDirInfo in BaseDir.EnumerateDirectories())
+				foreach (DirectoryInfo subDirInfo in baseDir.EnumerateDirectories())
 				{
-					FileInfo SubFile = new FileInfo(Path.Combine(SubDirInfo.FullName, FileName));
-					if (SubFile.Exists)
+					FileInfo subFile = new FileInfo(Path.Combine(subDirInfo.FullName, fileName));
+					if (subFile.Exists)
 					{
-						Files.Add(SubFile);
+						files.Add(subFile);
 					}
 				}
 			}
 		}
 
-		private static void AddLocalConfigPaths_WithExtensionDirs(DirectoryInfo BaseDir, string RelativePath, string FileName, List<FileInfo> Files)
+		private static void AddLocalConfigPaths_WithExtensionDirs(DirectoryInfo baseDir, string relativePath, string fileName, List<FileInfo> files)
 		{
-			if (BaseDir.Exists)
+			if (baseDir.Exists)
 			{
-				AddLocalConfigPaths_WithSubFolders(new DirectoryInfo(Path.Combine(BaseDir.FullName, RelativePath)), FileName, Files);
+				AddLocalConfigPaths_WithSubFolders(new DirectoryInfo(Path.Combine(baseDir.FullName, relativePath)), fileName, files);
 
-				DirectoryInfo PlatformExtensionsDir = new DirectoryInfo(Path.Combine(BaseDir.FullName, "Platforms"));
-				if (PlatformExtensionsDir.Exists)
+				DirectoryInfo platformExtensionsDir = new DirectoryInfo(Path.Combine(baseDir.FullName, "Platforms"));
+				if (platformExtensionsDir.Exists)
 				{
-					foreach (DirectoryInfo PlatformExtensionDir in PlatformExtensionsDir.EnumerateDirectories())
+					foreach (DirectoryInfo platformExtensionDir in platformExtensionsDir.EnumerateDirectories())
 					{
-						AddLocalConfigPaths_WithSubFolders(new DirectoryInfo(Path.Combine(PlatformExtensionDir.FullName, RelativePath)), FileName, Files);
+						AddLocalConfigPaths_WithSubFolders(new DirectoryInfo(Path.Combine(platformExtensionDir.FullName, relativePath)), fileName, files);
 					}
 				}
 
-				DirectoryInfo RestrictedBaseDir = new DirectoryInfo(Path.Combine(BaseDir.FullName, "Restricted"));
-				if (RestrictedBaseDir.Exists)
+				DirectoryInfo restrictedBaseDir = new DirectoryInfo(Path.Combine(baseDir.FullName, "Restricted"));
+				if (restrictedBaseDir.Exists)
 				{
-					foreach (DirectoryInfo RestrictedDir in RestrictedBaseDir.EnumerateDirectories())
+					foreach (DirectoryInfo restrictedDir in restrictedBaseDir.EnumerateDirectories())
 					{
-						AddLocalConfigPaths_WithSubFolders(new DirectoryInfo(Path.Combine(RestrictedDir.FullName, RelativePath)), FileName, Files);
+						AddLocalConfigPaths_WithSubFolders(new DirectoryInfo(Path.Combine(restrictedDir.FullName, relativePath)), fileName, files);
 					}
 				}
 			}
 		}
 
-		public static List<FileInfo> GetLocalConfigPaths(DirectoryInfo EngineDir, FileInfo ProjectFile)
+		public static List<FileInfo> GetLocalConfigPaths(DirectoryInfo engineDir, FileInfo projectFile)
 		{
-			List<FileInfo> SearchPaths = new List<FileInfo>();
-			AddLocalConfigPaths_WithExtensionDirs(EngineDir, "Programs/UnrealGameSync", "UnrealGameSync.ini", SearchPaths);
+			List<FileInfo> searchPaths = new List<FileInfo>();
+			AddLocalConfigPaths_WithExtensionDirs(engineDir, "Programs/UnrealGameSync", "UnrealGameSync.ini", searchPaths);
 
-			if (ProjectFile.Name.EndsWith(".uproject", StringComparison.OrdinalIgnoreCase))
+			if (projectFile.Name.EndsWith(".uproject", StringComparison.OrdinalIgnoreCase))
 			{
-				AddLocalConfigPaths_WithExtensionDirs(ProjectFile.Directory, "Build", "UnrealGameSync.ini", SearchPaths);
+				AddLocalConfigPaths_WithExtensionDirs(projectFile.Directory, "Build", "UnrealGameSync.ini", searchPaths);
 			}
 			else
 			{
-				AddLocalConfigPaths_WithExtensionDirs(EngineDir, "Programs/UnrealGameSync", "DefaultEngine.ini", SearchPaths);
+				AddLocalConfigPaths_WithExtensionDirs(engineDir, "Programs/UnrealGameSync", "DefaultEngine.ini", searchPaths);
 			}
-			return SearchPaths;
+			return searchPaths;
 		}
 
 		/******/
 
-		private static void AddDepotConfigPaths_PlatformFolders(string BasePath, string FileName, List<string> SearchPaths)
+		private static void AddDepotConfigPaths_PlatformFolders(string basePath, string fileName, List<string> searchPaths)
 		{
-			SearchPaths.Add(String.Format("{0}/{1}", BasePath, FileName));
-			SearchPaths.Add(String.Format("{0}/*/{1}", BasePath, FileName));
+			searchPaths.Add(String.Format("{0}/{1}", basePath, fileName));
+			searchPaths.Add(String.Format("{0}/*/{1}", basePath, fileName));
 		}
 
-		private static void AddDepotConfigPaths_PlatformExtensions(string BasePath, string RelativePath, string FileName, List<string> SearchPaths)
+		private static void AddDepotConfigPaths_PlatformExtensions(string basePath, string relativePath, string fileName, List<string> searchPaths)
 		{
-			AddDepotConfigPaths_PlatformFolders(BasePath + RelativePath, FileName, SearchPaths);
-			AddDepotConfigPaths_PlatformFolders(BasePath + "/Platforms/*" + RelativePath, FileName, SearchPaths);
-			AddDepotConfigPaths_PlatformFolders(BasePath + "/Restricted/*" + RelativePath, FileName, SearchPaths);
+			AddDepotConfigPaths_PlatformFolders(basePath + relativePath, fileName, searchPaths);
+			AddDepotConfigPaths_PlatformFolders(basePath + "/Platforms/*" + relativePath, fileName, searchPaths);
+			AddDepotConfigPaths_PlatformFolders(basePath + "/Restricted/*" + relativePath, fileName, searchPaths);
 		}
 
-		public static List<string> GetDepotConfigPaths(string EnginePath, string ProjectPath)
+		public static List<string> GetDepotConfigPaths(string enginePath, string projectPath)
 		{
-			List<string> SearchPaths = new List<string>();
-			AddDepotConfigPaths_PlatformExtensions(EnginePath, "/Programs/UnrealGameSync", "UnrealGameSync.ini", SearchPaths);
+			List<string> searchPaths = new List<string>();
+			AddDepotConfigPaths_PlatformExtensions(enginePath, "/Programs/UnrealGameSync", "UnrealGameSync.ini", searchPaths);
 
-			if (ProjectPath.EndsWith(".uproject", StringComparison.OrdinalIgnoreCase))
+			if (projectPath.EndsWith(".uproject", StringComparison.OrdinalIgnoreCase))
 			{
-				AddDepotConfigPaths_PlatformExtensions(ProjectPath.Substring(0, ProjectPath.LastIndexOf('/')), "/Build", "UnrealGameSync.ini", SearchPaths);
+				AddDepotConfigPaths_PlatformExtensions(projectPath.Substring(0, projectPath.LastIndexOf('/')), "/Build", "UnrealGameSync.ini", searchPaths);
 			}
 			else
 			{
-				AddDepotConfigPaths_PlatformExtensions(EnginePath, "/Programs/UnrealGameSync", "DefaultEngine.ini", SearchPaths);
+				AddDepotConfigPaths_PlatformExtensions(enginePath, "/Programs/UnrealGameSync", "DefaultEngine.ini", searchPaths);
 			}
-			return SearchPaths;
+			return searchPaths;
 		}
 
 		/******/
 
-		public static async Task<string[]?> TryPrintFileUsingCacheAsync(IPerforceConnection Perforce, string DepotPath, DirectoryReference CacheFolder, string? Digest, ILogger Logger, CancellationToken CancellationToken)
+		public static async Task<string[]?> TryPrintFileUsingCacheAsync(IPerforceConnection perforce, string depotPath, DirectoryReference cacheFolder, string? digest, ILogger logger, CancellationToken cancellationToken)
 		{
-			if(Digest == null)
+			if(digest == null)
 			{
-				PerforceResponse<PrintRecord<string[]>> Response = await Perforce.TryPrintLinesAsync(DepotPath, CancellationToken);
-				if (Response.Succeeded)
+				PerforceResponse<PrintRecord<string[]>> response = await perforce.TryPrintLinesAsync(depotPath, cancellationToken);
+				if (response.Succeeded)
 				{
-					return Response.Data.Contents;
+					return response.Data.Contents;
 				}
 				else
 				{
@@ -443,69 +443,69 @@ namespace UnrealGameSync
 				}
 			}
 
-			FileReference CacheFile = FileReference.Combine(CacheFolder, Digest);
-			if(FileReference.Exists(CacheFile))
+			FileReference cacheFile = FileReference.Combine(cacheFolder, digest);
+			if(FileReference.Exists(cacheFile))
 			{
-				Logger.LogDebug("Reading cached copy of {DepotFile} from {LocalFile}", DepotPath, CacheFile);
-				string[] Lines = FileReference.ReadAllLines(CacheFile);
+				logger.LogDebug("Reading cached copy of {DepotFile} from {LocalFile}", depotPath, cacheFile);
+				string[] lines = FileReference.ReadAllLines(cacheFile);
 				try
 				{
-					FileReference.SetLastWriteTimeUtc(CacheFile, DateTime.UtcNow);
+					FileReference.SetLastWriteTimeUtc(cacheFile, DateTime.UtcNow);
 				}
-				catch(Exception Ex)
+				catch(Exception ex)
 				{
-					Logger.LogWarning(Ex, "Exception touching cache file {LocalFile}", CacheFile);
+					logger.LogWarning(ex, "Exception touching cache file {LocalFile}", cacheFile);
 				}
-				return Lines;
+				return lines;
 			}
 			else
 			{
-				DirectoryReference.CreateDirectory(CacheFolder);
+				DirectoryReference.CreateDirectory(cacheFolder);
 
-				FileReference TempFile = new FileReference(String.Format("{0}.{1}.temp", CacheFile.FullName, Guid.NewGuid()));
-				PerforceResponse<PrintRecord> Response = await Perforce.TryPrintAsync(TempFile.FullName, DepotPath, CancellationToken);
-				if (!Response.Succeeded)
+				FileReference tempFile = new FileReference(String.Format("{0}.{1}.temp", cacheFile.FullName, Guid.NewGuid()));
+				PerforceResponse<PrintRecord> response = await perforce.TryPrintAsync(tempFile.FullName, depotPath, cancellationToken);
+				if (!response.Succeeded)
 				{
 					return null;
 				}
 				else
 				{
-					string[] Lines = await FileReference.ReadAllLinesAsync(TempFile);
+					string[] lines = await FileReference.ReadAllLinesAsync(tempFile);
 					try
 					{
-						FileReference.SetAttributes(TempFile, FileAttributes.Normal);
-						FileReference.SetLastWriteTimeUtc(TempFile, DateTime.UtcNow);
-						FileReference.Move(TempFile, CacheFile);
+						FileReference.SetAttributes(tempFile, FileAttributes.Normal);
+						FileReference.SetLastWriteTimeUtc(tempFile, DateTime.UtcNow);
+						FileReference.Move(tempFile, cacheFile);
 					}
 					catch
 					{
 						try
 						{
-							FileReference.Delete(TempFile);
+							FileReference.Delete(tempFile);
 						}
 						catch
 						{
 						}
 					}
-					return Lines;
+					return lines;
 				}
 			}
 		}
 
-		public static void ClearPrintCache(DirectoryReference CacheFolder)
+		public static void ClearPrintCache(DirectoryReference cacheFolder)
 		{
-			DirectoryInfo CacheDir = CacheFolder.ToDirectoryInfo();
-			if(CacheDir.Exists)
+			DirectoryInfo cacheDir = cacheFolder.ToDirectoryInfo();
+			if(cacheDir.Exists)
 			{
-				DateTime DeleteTime = DateTime.UtcNow - TimeSpan.FromDays(5.0);
-				foreach(FileInfo CacheFile in CacheDir.EnumerateFiles())
+				DateTime deleteTime = DateTime.UtcNow - TimeSpan.FromDays(5.0);
+				foreach(FileInfo cacheFile in cacheDir.EnumerateFiles())
 				{
-					if(CacheFile.LastWriteTimeUtc < DeleteTime || CacheFile.Name.EndsWith(".temp", StringComparison.OrdinalIgnoreCase))
+					if(cacheFile.LastWriteTimeUtc < deleteTime || cacheFile.Name.EndsWith(".temp", StringComparison.OrdinalIgnoreCase))
 					{
 						try
 						{
-							CacheFile.Attributes = FileAttributes.Normal;
-							CacheFile.Delete();
+							cacheFile.Attributes = FileAttributes.Normal;
+							cacheFile.Delete();
 						}
 						catch
 						{
@@ -515,102 +515,102 @@ namespace UnrealGameSync
 			}
 		}
 
-		public static Color Blend(Color First, Color Second, float T)
+		public static Color Blend(Color first, Color second, float T)
 		{
-			return Color.FromArgb((int)(First.R + (Second.R - First.R) * T), (int)(First.G + (Second.G - First.G) * T), (int)(First.B + (Second.B - First.B) * T));
+			return Color.FromArgb((int)(first.R + (second.R - first.R) * T), (int)(first.G + (second.G - first.G) * T), (int)(first.B + (second.B - first.B) * T));
 		}
 
-		public static PerforceSettings OverridePerforceSettings(IPerforceSettings DefaultConnection, string? ServerAndPort, string? UserName)
+		public static PerforceSettings OverridePerforceSettings(IPerforceSettings defaultConnection, string? serverAndPort, string? userName)
 		{
-			PerforceSettings NewSettings = new PerforceSettings(DefaultConnection);
-			if(!String.IsNullOrWhiteSpace(ServerAndPort))
+			PerforceSettings newSettings = new PerforceSettings(defaultConnection);
+			if(!String.IsNullOrWhiteSpace(serverAndPort))
 			{
-				NewSettings.ServerAndPort = ServerAndPort;
+				newSettings.ServerAndPort = serverAndPort;
 			}
-			if (!String.IsNullOrWhiteSpace(UserName))
+			if (!String.IsNullOrWhiteSpace(userName))
 			{
-				NewSettings.UserName = UserName;
+				newSettings.UserName = userName;
 			}
-			return NewSettings;
+			return newSettings;
 		}
 
-		public static string FormatRecentDateTime(DateTime Date)
+		public static string FormatRecentDateTime(DateTime date)
 		{
-			DateTime Now = DateTime.Now;
+			DateTime now = DateTime.Now;
 
-			DateTime Midnight = new DateTime(Now.Year, Now.Month, Now.Day);
-			DateTime MidnightTonight = Midnight + TimeSpan.FromDays(1.0);
+			DateTime midnight = new DateTime(now.Year, now.Month, now.Day);
+			DateTime midnightTonight = midnight + TimeSpan.FromDays(1.0);
 
-			if(Date > MidnightTonight)
+			if(date > midnightTonight)
 			{
-				return String.Format("{0} at {1}", Date.ToLongDateString(), Date.ToShortTimeString());
+				return String.Format("{0} at {1}", date.ToLongDateString(), date.ToShortTimeString());
 			}
-			else if(Date >= Midnight)
+			else if(date >= midnight)
 			{
-				return String.Format("today at {0}", Date.ToShortTimeString());
+				return String.Format("today at {0}", date.ToShortTimeString());
 			}
-			else if(Date >= Midnight - TimeSpan.FromDays(1.0))
+			else if(date >= midnight - TimeSpan.FromDays(1.0))
 			{
-				return String.Format("yesterday at {0}", Date.ToShortTimeString());
+				return String.Format("yesterday at {0}", date.ToShortTimeString());
 			}
-			else if(Date >= Midnight - TimeSpan.FromDays(5.0))
+			else if(date >= midnight - TimeSpan.FromDays(5.0))
 			{
-				return String.Format("{0:dddd} at {1}", Date, Date.ToShortTimeString());
+				return String.Format("{0:dddd} at {1}", date, date.ToShortTimeString());
 			}
 			else
 			{
-				return String.Format("{0} at {1}", Date.ToLongDateString(), Date.ToShortTimeString());
+				return String.Format("{0} at {1}", date.ToLongDateString(), date.ToShortTimeString());
 			}
 		}
 
-		public static string FormatDurationMinutes(TimeSpan Duration)
+		public static string FormatDurationMinutes(TimeSpan duration)
 		{
-			return FormatDurationMinutes((int)(Duration.TotalMinutes + 1));
+			return FormatDurationMinutes((int)(duration.TotalMinutes + 1));
 		}
 
-		public static string FormatDurationMinutes(int TotalMinutes)
+		public static string FormatDurationMinutes(int totalMinutes)
 		{
-			if(TotalMinutes > 24 * 60)
+			if(totalMinutes > 24 * 60)
 			{
-				return String.Format("{0}d {1}h", TotalMinutes / (24 * 60), (TotalMinutes / 60) % 24);
+				return String.Format("{0}d {1}h", totalMinutes / (24 * 60), (totalMinutes / 60) % 24);
 			}
-			else if(TotalMinutes > 60)
+			else if(totalMinutes > 60)
 			{
-				return String.Format("{0}h {1}m", TotalMinutes / 60, TotalMinutes % 60);
+				return String.Format("{0}h {1}m", totalMinutes / 60, totalMinutes % 60);
 			}
 			else
 			{
-				return String.Format("{0}m", TotalMinutes);
+				return String.Format("{0}m", totalMinutes);
 			}
 		}
 
-		public static string FormatUserName(string UserName)
+		public static string FormatUserName(string userName)
 		{
-			StringBuilder NormalUserName = new StringBuilder();
-			for(int Idx = 0; Idx < UserName.Length; Idx++)
+			StringBuilder normalUserName = new StringBuilder();
+			for(int idx = 0; idx < userName.Length; idx++)
 			{
-				if(Idx == 0 || UserName[Idx - 1] == '.')
+				if(idx == 0 || userName[idx - 1] == '.')
 				{
-					NormalUserName.Append(Char.ToUpper(UserName[Idx]));
+					normalUserName.Append(Char.ToUpper(userName[idx]));
 				}
-				else if(UserName[Idx] == '.')
+				else if(userName[idx] == '.')
 				{
-					NormalUserName.Append(' ');
+					normalUserName.Append(' ');
 				}
 				else
 				{
-					NormalUserName.Append(Char.ToLower(UserName[Idx]));
+					normalUserName.Append(Char.ToLower(userName[idx]));
 				}
 			}
-			return NormalUserName.ToString();
+			return normalUserName.ToString();
 		}
 
-		public static void OpenUrl(string Url)
+		public static void OpenUrl(string url)
 		{
-			ProcessStartInfo StartInfo = new ProcessStartInfo();
-			StartInfo.FileName = Url;
-			StartInfo.UseShellExecute = true;
-			using Process _ = Process.Start(StartInfo);
+			ProcessStartInfo startInfo = new ProcessStartInfo();
+			startInfo.FileName = url;
+			startInfo.UseShellExecute = true;
+			using Process _ = Process.Start(startInfo);
 		}
 	}
 }
