UGS: Reformat to standard C# style conventions.

diff --git a/Engine/Source/Programs/UnrealGameSync/UnrealGameSyncShared/WorkspaceUpdate.cs b/Engine/Source/Programs/UnrealGameSync/UnrealGameSyncShared/WorkspaceUpdate.cs
--- a/Engine/Source/Programs/UnrealGameSync/UnrealGameSyncShared/WorkspaceUpdate.cs
+++ b/Engine/Source/Programs/UnrealGameSync/UnrealGameSyncShared/WorkspaceUpdate.cs
@@ -96,8 +96,8 @@ namespace UnrealGameSync
 		string BasePath { get; }
 		string? Target { get; }
 		bool Exists();
-		bool TryGetArchiveKeyForChangeNumber(int ChangeNumber, [NotNullWhen(true)] out string? ArchiveKey);
-		Task<bool> DownloadArchive(IPerforceConnection Perforce, string ArchiveKey, DirectoryReference LocalRootPath, FileReference ManifestFileName, ILogger Logger, ProgressValue Progress, CancellationToken CancellationToken);
+		bool TryGetArchiveKeyForChangeNumber(int changeNumber, [NotNullWhen(true)] out string? archiveKey);
+		Task<bool> DownloadArchive(IPerforceConnection perforce, string archiveKey, DirectoryReference localRootPath, FileReference manifestFileName, ILogger logger, ProgressValue progress, CancellationToken cancellationToken);
 	}
 
 	public class WorkspaceUpdateContext
@@ -120,63 +120,63 @@ namespace UnrealGameSync
 		public ConfigFile? ProjectConfigFile;
 		public IReadOnlyList<string>? ProjectStreamFilter;
 
-		public WorkspaceUpdateContext(int InChangeNumber, WorkspaceUpdateOptions InOptions, BuildConfig InEditorConfig, string[]? InSyncFilter, List<ConfigObject> InUserBuildSteps, HashSet<Guid>? InCustomBuildSteps)
+		public WorkspaceUpdateContext(int inChangeNumber, WorkspaceUpdateOptions inOptions, BuildConfig inEditorConfig, string[]? inSyncFilter, List<ConfigObject> inUserBuildSteps, HashSet<Guid>? inCustomBuildSteps)
 		{
-			ChangeNumber = InChangeNumber;
-			Options = InOptions;
-			EditorConfig = InEditorConfig;
-			SyncFilter = InSyncFilter;
-			UserBuildStepObjects = InUserBuildSteps;
-			CustomBuildSteps = InCustomBuildSteps;
+			ChangeNumber = inChangeNumber;
+			Options = inOptions;
+			EditorConfig = inEditorConfig;
+			SyncFilter = inSyncFilter;
+			UserBuildStepObjects = inUserBuildSteps;
+			CustomBuildSteps = inCustomBuildSteps;
 		}
 	}
 
 	public class WorkspaceSyncCategory
 	{
 		public Guid UniqueId;
-		public bool bEnable;
+		public bool Enable;
 		public string Name;
 		public string[] Paths;
-		public bool bHidden;
+		public bool Hidden;
 		public Guid[] Requires;
 
-		public WorkspaceSyncCategory(Guid UniqueId) : this(UniqueId, "Unnamed")
+		public WorkspaceSyncCategory(Guid uniqueId) : this(uniqueId, "Unnamed")
 		{
 		}
 
-		public WorkspaceSyncCategory(Guid UniqueId, string Name, params string[] Paths)
+		public WorkspaceSyncCategory(Guid uniqueId, string name, params string[] paths)
 		{
-			this.UniqueId = UniqueId;
-			this.bEnable = true;
-			this.Name = Name;
-			this.Paths = Paths;
+			this.UniqueId = uniqueId;
+			this.Enable = true;
+			this.Name = name;
+			this.Paths = paths;
 			this.Requires = new Guid[0];
 		}
 
-		public static Dictionary<Guid, bool> GetDefault(IEnumerable<WorkspaceSyncCategory> Categories)
+		public static Dictionary<Guid, bool> GetDefault(IEnumerable<WorkspaceSyncCategory> categories)
 		{
-			return Categories.ToDictionary(x => x.UniqueId, x => x.bEnable);
+			return categories.ToDictionary(x => x.UniqueId, x => x.Enable);
 		}
 
-		public static Dictionary<Guid, bool> GetDelta(Dictionary<Guid, bool> Source, Dictionary<Guid, bool> Target)
+		public static Dictionary<Guid, bool> GetDelta(Dictionary<Guid, bool> source, Dictionary<Guid, bool> target)
 		{
-			Dictionary<Guid, bool> Changes = new Dictionary<Guid, bool>();
-			foreach (KeyValuePair<Guid, bool> Pair in Target)
+			Dictionary<Guid, bool> changes = new Dictionary<Guid, bool>();
+			foreach (KeyValuePair<Guid, bool> pair in target)
 			{
-				bool bValue;
-				if (!Source.TryGetValue(Pair.Key, out bValue) || bValue != Pair.Value)
+				bool value;
+				if (!source.TryGetValue(pair.Key, out value) || value != pair.Value)
 				{
-					Changes[Pair.Key] = Pair.Value;
+					changes[pair.Key] = pair.Value;
 				}
 			}
-			return Changes;
+			return changes;
 		}
 
-		public static void ApplyDelta(Dictionary<Guid, bool> Categories, Dictionary<Guid, bool> Delta)
+		public static void ApplyDelta(Dictionary<Guid, bool> categories, Dictionary<Guid, bool> delta)
 		{
-			foreach(KeyValuePair<Guid, bool> Pair in Delta)
+			foreach(KeyValuePair<Guid, bool> pair in delta)
 			{
-				Categories[Pair.Key] = Pair.Value;
+				categories[pair.Key] = pair.Value;
 			}
 		}
 
@@ -197,7 +197,7 @@ namespace UnrealGameSync
 		public string? StreamName { get; } // name of the current stream
 
 		public string ProjectIdentifier { get; } // stream path to project
-		public bool bIsEnterpriseProject { get; } // whether it's an enterprise project
+		public bool IsEnterpriseProject { get; } // whether it's an enterprise project
 
 		// derived properties
 		public FileReference LocalFileName => new FileReference(LocalRootPath.FullName + ProjectPath);
@@ -209,127 +209,127 @@ namespace UnrealGameSync
 		public DirectoryReference DataFolder => GetDataFolder(LocalRootPath);
 		public DirectoryReference CacheFolder => GetCacheFolder(LocalRootPath);
 
-		public static DirectoryReference GetDataFolder(DirectoryReference WorkspaceDir) => DirectoryReference.Combine(WorkspaceDir, ".ugs");
-		public static DirectoryReference GetCacheFolder(DirectoryReference WorkspaceDir) => DirectoryReference.Combine(WorkspaceDir, ".ugs", "cache");
+		public static DirectoryReference GetDataFolder(DirectoryReference workspaceDir) => DirectoryReference.Combine(workspaceDir, ".ugs");
+		public static DirectoryReference GetCacheFolder(DirectoryReference workspaceDir) => DirectoryReference.Combine(workspaceDir, ".ugs", "cache");
 
-		public ProjectInfo(DirectoryReference LocalRootPath, string ClientName, string BranchPath, string ProjectPath, string? StreamName, string ProjectIdentifier, bool bIsEnterpriseProject)
+		public ProjectInfo(DirectoryReference localRootPath, string clientName, string branchPath, string projectPath, string? streamName, string projectIdentifier, bool isEnterpriseProject)
 		{
-			ValidateBranchPath(BranchPath);
-			ValidateProjectPath(ProjectPath);
-
-			this.LocalRootPath = LocalRootPath;
-			this.ClientName = ClientName;
-			this.BranchPath = BranchPath;
-			this.ProjectPath = ProjectPath;
-			this.StreamName = StreamName;
-			this.ProjectIdentifier = ProjectIdentifier;
-			this.bIsEnterpriseProject = bIsEnterpriseProject;
+			ValidateBranchPath(branchPath);
+			ValidateProjectPath(projectPath);
+
+			this.LocalRootPath = localRootPath;
+			this.ClientName = clientName;
+			this.BranchPath = branchPath;
+			this.ProjectPath = projectPath;
+			this.StreamName = streamName;
+			this.ProjectIdentifier = projectIdentifier;
+			this.IsEnterpriseProject = isEnterpriseProject;
 		}
 
-		public static async Task<ProjectInfo> CreateAsync(IPerforceConnection PerforceClient, UserWorkspaceSettings Settings, CancellationToken CancellationToken)
+		public static async Task<ProjectInfo> CreateAsync(IPerforceConnection perforceClient, UserWorkspaceSettings settings, CancellationToken cancellationToken)
 		{
-			string? StreamName = await PerforceClient.GetCurrentStreamAsync(CancellationToken);
+			string? streamName = await perforceClient.GetCurrentStreamAsync(cancellationToken);
 
 			// Get a unique name for the project that's selected. For regular branches, this can be the depot path. For streams, we want to include the stream name to encode imports.
-			string NewSelectedProjectIdentifier;
-			if (StreamName != null)
+			string newSelectedProjectIdentifier;
+			if (streamName != null)
 			{
-				string ExpectedPrefix = String.Format("//{0}/", PerforceClient.Settings.ClientName);
-				if (!Settings.ClientProjectPath.StartsWith(ExpectedPrefix, StringComparison.InvariantCultureIgnoreCase))
+				string expectedPrefix = String.Format("//{0}/", perforceClient.Settings.ClientName);
+				if (!settings.ClientProjectPath.StartsWith(expectedPrefix, StringComparison.InvariantCultureIgnoreCase))
 				{
-					throw new UserErrorException($"Unexpected client path; expected '{Settings.ClientProjectPath}' to begin with '{ExpectedPrefix}'");
+					throw new UserErrorException($"Unexpected client path; expected '{settings.ClientProjectPath}' to begin with '{expectedPrefix}'");
 				}
-				string? StreamPrefix = await TryGetStreamPrefixAsync(PerforceClient, StreamName, CancellationToken);
-				if (StreamPrefix == null)
+				string? streamPrefix = await TryGetStreamPrefixAsync(perforceClient, streamName, cancellationToken);
+				if (streamPrefix == null)
 				{
 					throw new UserErrorException("Unable to get stream prefix");
 				}
-				NewSelectedProjectIdentifier = String.Format("{0}/{1}", StreamPrefix, Settings.ClientProjectPath.Substring(ExpectedPrefix.Length));
+				newSelectedProjectIdentifier = String.Format("{0}/{1}", streamPrefix, settings.ClientProjectPath.Substring(expectedPrefix.Length));
 			}
 			else
 			{
-				List<PerforceResponse<WhereRecord>> Records = await PerforceClient.TryWhereAsync(Settings.ClientProjectPath, CancellationToken).Where(x => !x.Succeeded || !x.Data.Unmap).ToListAsync(CancellationToken);
-				if (!Records.Succeeded() || Records.Count != 1)
+				List<PerforceResponse<WhereRecord>> records = await perforceClient.TryWhereAsync(settings.ClientProjectPath, cancellationToken).Where(x => !x.Succeeded || !x.Data.Unmap).ToListAsync(cancellationToken);
+				if (!records.Succeeded() || records.Count != 1)
 				{
-					throw new UserErrorException($"Couldn't get depot path for {Settings.ClientProjectPath}");
+					throw new UserErrorException($"Couldn't get depot path for {settings.ClientProjectPath}");
 				}
 
-				NewSelectedProjectIdentifier = Records[0].Data.DepotFile;
+				newSelectedProjectIdentifier = records[0].Data.DepotFile;
 
-				Match Match = Regex.Match(NewSelectedProjectIdentifier, "//([^/]+)/");
-				if (Match.Success)
+				Match match = Regex.Match(newSelectedProjectIdentifier, "//([^/]+)/");
+				if (match.Success)
 				{
-					DepotRecord Depot = await PerforceClient.GetDepotAsync(Match.Groups[1].Value, CancellationToken);
-					if (Depot.Type == "stream")
+					DepotRecord depot = await perforceClient.GetDepotAsync(match.Groups[1].Value, cancellationToken);
+					if (depot.Type == "stream")
 					{
-						throw new UserErrorException($"Cannot use a legacy client ({PerforceClient.Settings.ClientName}) with a stream depot ({Depot.Depot}).");
+						throw new UserErrorException($"Cannot use a legacy client ({perforceClient.Settings.ClientName}) with a stream depot ({depot.Depot}).");
 					}
 				}
 			}
 
 			// Figure out if it's an enterprise project. Use the synced version if we have it.
-			bool bIsEnterpriseProject = false;
-			if (Settings.ClientProjectPath.EndsWith(".uproject", StringComparison.InvariantCultureIgnoreCase))
+			bool isEnterpriseProject = false;
+			if (settings.ClientProjectPath.EndsWith(".uproject", StringComparison.InvariantCultureIgnoreCase))
 			{
-				string Text;
-				if (FileReference.Exists(Settings.LocalProjectPath))
+				string text;
+				if (FileReference.Exists(settings.LocalProjectPath))
 				{
-					Text = FileReference.ReadAllText(Settings.LocalProjectPath);
+					text = FileReference.ReadAllText(settings.LocalProjectPath);
 				}
 				else
 				{
-					PerforceResponse<PrintRecord<string[]>> ProjectLines = await PerforceClient.TryPrintLinesAsync(Settings.ClientProjectPath, CancellationToken);
-					if (!ProjectLines.Succeeded)
+					PerforceResponse<PrintRecord<string[]>> projectLines = await perforceClient.TryPrintLinesAsync(settings.ClientProjectPath, cancellationToken);
+					if (!projectLines.Succeeded)
 					{
-						throw new UserErrorException($"Unable to get contents of {Settings.ClientProjectPath}");
+						throw new UserErrorException($"Unable to get contents of {settings.ClientProjectPath}");
 					}
-					Text = String.Join("\n", ProjectLines.Data.Contents!);
+					text = String.Join("\n", projectLines.Data.Contents!);
 				}
-				bIsEnterpriseProject = Utility.IsEnterpriseProjectFromText(Text);
+				isEnterpriseProject = Utility.IsEnterpriseProjectFromText(text);
 			}
 
-			return new ProjectInfo(Settings.RootDir, Settings.ClientName, Settings.BranchPath, Settings.ProjectPath, StreamName, NewSelectedProjectIdentifier, bIsEnterpriseProject);
+			return new ProjectInfo(settings.RootDir, settings.ClientName, settings.BranchPath, settings.ProjectPath, streamName, newSelectedProjectIdentifier, isEnterpriseProject);
 		}
 
-		static async Task<string?> TryGetStreamPrefixAsync(IPerforceConnection Perforce, string StreamName, CancellationToken CancellationToken)
+		static async Task<string?> TryGetStreamPrefixAsync(IPerforceConnection perforce, string streamName, CancellationToken cancellationToken)
 		{
-			string? CurrentStreamName = StreamName;
-			while (!String.IsNullOrEmpty(CurrentStreamName))
+			string? currentStreamName = streamName;
+			while (!String.IsNullOrEmpty(currentStreamName))
 			{
-				PerforceResponse<StreamRecord> Response = await Perforce.TryGetStreamAsync(CurrentStreamName, false, CancellationToken);
-				if (!Response.Succeeded)
+				PerforceResponse<StreamRecord> response = await perforce.TryGetStreamAsync(currentStreamName, false, cancellationToken);
+				if (!response.Succeeded)
 				{
 					return null;
 				}
 
-				StreamRecord StreamSpec = Response.Data;
-				if (StreamSpec.Type != "virtual")
+				StreamRecord streamSpec = response.Data;
+				if (streamSpec.Type != "virtual")
 				{
-					return CurrentStreamName;
+					return currentStreamName;
 				}
 
-				CurrentStreamName = StreamSpec.Parent;
+				currentStreamName = streamSpec.Parent;
 			}
 			return null;
 		}
 
-		public static void ValidateBranchPath(string BranchPath)
+		public static void ValidateBranchPath(string branchPath)
 		{
-			if (BranchPath.Length > 0 && (!BranchPath.StartsWith("/") || BranchPath.EndsWith("/")))
+			if (branchPath.Length > 0 && (!branchPath.StartsWith("/") || branchPath.EndsWith("/")))
 			{
-				throw new ArgumentException("Branch path must start with a slash, and not end with a slash", nameof(BranchPath));
+				throw new ArgumentException("Branch path must start with a slash, and not end with a slash", nameof(branchPath));
 			}
 		}
 
-		public static void ValidateProjectPath(string ProjectPath)
+		public static void ValidateProjectPath(string projectPath)
 		{
-			if (!ProjectPath.StartsWith("/"))
+			if (!projectPath.StartsWith("/"))
 			{
-				throw new ArgumentException("Project path must start with a slash", nameof(ProjectPath));
+				throw new ArgumentException("Project path must start with a slash", nameof(projectPath));
 			}
-			if (!ProjectPath.EndsWith(".uproject", StringComparison.OrdinalIgnoreCase) && !ProjectPath.EndsWith(".uprojectdirs", StringComparison.OrdinalIgnoreCase))
+			if (!projectPath.EndsWith(".uproject", StringComparison.OrdinalIgnoreCase) && !projectPath.EndsWith(".uprojectdirs", StringComparison.OrdinalIgnoreCase))
 			{
-				throw new ArgumentException("Project path must be to a .uproject or .uprojectdirs file", nameof(ProjectPath));
+				throw new ArgumentException("Project path must be to a .uproject or .uprojectdirs file", nameof(projectPath));
 			}
 		}
 	}
@@ -343,21 +343,21 @@ namespace UnrealGameSync
 		static readonly string LocalVersionHeaderFileName = VersionHeaderFileName.Replace('/', Path.DirectorySeparatorChar);
 		static readonly string LocalObjectVersionFileName = ObjectVersionFileName.Replace('/', Path.DirectorySeparatorChar);
 
-		static SemaphoreSlim UpdateSemaphore = new SemaphoreSlim(1);
+		static SemaphoreSlim _updateSemaphore = new SemaphoreSlim(1);
 
 		class RecordCounter : IDisposable
 		{
-			ProgressValue Progress;
-			string Message;
-			int Count;
-			Stopwatch Timer = Stopwatch.StartNew();
+			ProgressValue _progress;
+			string _message;
+			int _count;
+			Stopwatch _timer = Stopwatch.StartNew();
 
-			public RecordCounter(ProgressValue Progress, string Message)
+			public RecordCounter(ProgressValue progress, string message)
 			{
-				this.Progress = Progress;
-				this.Message = Message;
+				this._progress = progress;
+				this._message = message;
 
-				Progress.Set(Message);
+				progress.Set(message);
 			}
 
 			public void Dispose()
@@ -367,8 +367,8 @@ namespace UnrealGameSync
 
 			public void Increment()
 			{
-				Count++;
-				if (Timer.ElapsedMilliseconds > 250)
+				_count++;
+				if (_timer.ElapsedMilliseconds > 250)
 				{
 					UpdateMessage();
 				}
@@ -376,8 +376,8 @@ namespace UnrealGameSync
 
 			public void UpdateMessage()
 			{
-				Progress.Set(String.Format("{0} ({1:N0})", Message, Count));
-				Timer.Restart();
+				_progress.Set(String.Format("{0} ({1:N0})", _message, _count));
+				_timer.Restart();
 			}
 		}
 
@@ -387,116 +387,116 @@ namespace UnrealGameSync
 			public long MaxSizePerList { get; }
 			public Queue<List<string>> Batches { get; }
 
-			List<string>? Commands;
-			List<string>? DeleteCommands;
-			long Size;
+			List<string>? _commands;
+			List<string>? _deleteCommands;
+			long _size;
 
-			public SyncBatchBuilder(int? MaxCommandsPerList, long? MaxSizePerList)
+			public SyncBatchBuilder(int? maxCommandsPerList, long? maxSizePerList)
 			{
-				this.MaxCommandsPerList = MaxCommandsPerList ?? PerforceSyncOptions.DefaultMaxCommandsPerBatch;
-				this.MaxSizePerList = MaxSizePerList ?? PerforceSyncOptions.DefaultMaxSizePerBatch;
+				this.MaxCommandsPerList = maxCommandsPerList ?? PerforceSyncOptions.DefaultMaxCommandsPerBatch;
+				this.MaxSizePerList = maxSizePerList ?? PerforceSyncOptions.DefaultMaxSizePerBatch;
 				this.Batches = new Queue<List<string>>();
 			}
 
-			public void Add(string NewCommand, long NewSize)
+			public void Add(string newCommand, long newSize)
 			{
-				if (NewSize == 0)
+				if (newSize == 0)
 				{
-					if (DeleteCommands == null || DeleteCommands.Count >= MaxCommandsPerList)
+					if (_deleteCommands == null || _deleteCommands.Count >= MaxCommandsPerList)
 					{
-						DeleteCommands = new List<string>();
-						Batches.Enqueue(DeleteCommands);
+						_deleteCommands = new List<string>();
+						Batches.Enqueue(_deleteCommands);
 					}
 
-					DeleteCommands.Add(NewCommand);
+					_deleteCommands.Add(newCommand);
 				}
 				else
 				{
-					if (Commands == null || Commands.Count >= MaxCommandsPerList || Size + NewSize >= MaxSizePerList)
+					if (_commands == null || _commands.Count >= MaxCommandsPerList || _size + newSize >= MaxSizePerList)
 					{
-						Commands = new List<string>();
-						Batches.Enqueue(Commands);
-						Size = 0;
+						_commands = new List<string>();
+						Batches.Enqueue(_commands);
+						_size = 0;
 					}
 
-					Commands.Add(NewCommand);
-					Size += NewSize;
+					_commands.Add(newCommand);
+					_size += newSize;
 				}
 			}
 		}
 
 		class SyncTree
 		{
-			public bool bCanUseWildcard;
+			public bool CanUseWildcard;
 			public int TotalIncludedFiles;
 			public long TotalSize;
 			public int TotalExcludedFiles;
 			public Dictionary<string, long> IncludedFiles = new Dictionary<string, long>();
 			public Dictionary<string, SyncTree> NameToSubTree = new Dictionary<string, SyncTree>(StringComparer.OrdinalIgnoreCase);
 
-			public SyncTree(bool bCanUseWildcard)
+			public SyncTree(bool canUseWildcard)
 			{
-				this.bCanUseWildcard = bCanUseWildcard;
+				this.CanUseWildcard = canUseWildcard;
 			}
 
-			public SyncTree FindOrAddSubTree(string Name)
+			public SyncTree FindOrAddSubTree(string name)
 			{
-				SyncTree? Result;
-				if (!NameToSubTree.TryGetValue(Name, out Result))
+				SyncTree? result;
+				if (!NameToSubTree.TryGetValue(name, out result))
 				{
-					Result = new SyncTree(bCanUseWildcard);
-					NameToSubTree.Add(Name, Result);
+					result = new SyncTree(CanUseWildcard);
+					NameToSubTree.Add(name, result);
 				}
-				return Result;
+				return result;
 			}
 
-			public void IncludeFile(string Path, long Size, ILogger Logger) => IncludeFile(Path, Path, Size, Logger);
+			public void IncludeFile(string path, long size, ILogger logger) => IncludeFile(path, path, size, logger);
 
-			private void IncludeFile(string FullPath, string Path, long Size, ILogger Logger)
+			private void IncludeFile(string fullPath, string path, long size, ILogger logger)
 			{
-				int Idx = Path.IndexOf('/');
-				if (Idx == -1)
+				int idx = path.IndexOf('/');
+				if (idx == -1)
 				{
-					if (!IncludedFiles.ContainsKey(Path))
+					if (!IncludedFiles.ContainsKey(path))
 					{
-						IncludedFiles.Add(Path, Size);
+						IncludedFiles.Add(path, size);
 					}
 				}
 				else
 				{
-					SyncTree SubTree = FindOrAddSubTree(Path.Substring(0, Idx));
-					SubTree.IncludeFile(FullPath, Path.Substring(Idx + 1), Size, Logger);
+					SyncTree subTree = FindOrAddSubTree(path.Substring(0, idx));
+					subTree.IncludeFile(fullPath, path.Substring(idx + 1), size, logger);
 				}
 				TotalIncludedFiles++;
-				TotalSize += Size;
+				TotalSize += size;
 			}
 
-			public void ExcludeFile(string Path)
+			public void ExcludeFile(string path)
 			{
-				int Idx = Path.IndexOf('/');
-				if (Idx != -1)
+				int idx = path.IndexOf('/');
+				if (idx != -1)
 				{
-					SyncTree SubTree = FindOrAddSubTree(Path.Substring(0, Idx));
-					SubTree.ExcludeFile(Path.Substring(Idx + 1));
+					SyncTree subTree = FindOrAddSubTree(path.Substring(0, idx));
+					subTree.ExcludeFile(path.Substring(idx + 1));
 				}
 				TotalExcludedFiles++;
 			}
 
-			public void GetOptimizedSyncCommands(string Prefix, int ChangeNumber, SyncBatchBuilder Builder)
+			public void GetOptimizedSyncCommands(string prefix, int changeNumber, SyncBatchBuilder builder)
 			{
-				if (bCanUseWildcard && TotalExcludedFiles == 0 && TotalSize < Builder.MaxSizePerList)
+				if (CanUseWildcard && TotalExcludedFiles == 0 && TotalSize < builder.MaxSizePerList)
 				{
-					Builder.Add(String.Format("{0}/...@{1}", Prefix, ChangeNumber), TotalSize);
+					builder.Add(String.Format("{0}/...@{1}", prefix, changeNumber), TotalSize);
 				}
 				else
 				{
-					foreach (KeyValuePair<string, long> File in IncludedFiles)
+					foreach (KeyValuePair<string, long> file in IncludedFiles)
 					{
-						Builder.Add(String.Format("{0}/{1}@{2}", Prefix, File.Key, ChangeNumber), File.Value);
+						builder.Add(String.Format("{0}/{1}@{2}", prefix, file.Key, changeNumber), file.Value);
 					}
-					foreach (KeyValuePair<string, SyncTree> Pair in NameToSubTree)
+					foreach (KeyValuePair<string, SyncTree> pair in NameToSubTree)
 					{
-						Pair.Value.GetOptimizedSyncCommands(String.Format("{0}/{1}", Prefix, PerforceUtils.EscapePath(Pair.Key)), ChangeNumber, Builder);
+						pair.Value.GetOptimizedSyncCommands(String.Format("{0}/{1}", prefix, PerforceUtils.EscapePath(pair.Key)), changeNumber, builder);
 					}
 				}
 			}
@@ -505,9 +505,9 @@ namespace UnrealGameSync
 		public WorkspaceUpdateContext Context { get; }
 		public ProgressValue Progress { get; } = new ProgressValue();
 
-		public WorkspaceUpdate(WorkspaceUpdateContext Context)
+		public WorkspaceUpdate(WorkspaceUpdateContext context)
 		{
-			this.Context = Context;
+			this.Context = context;
 		}
 
 		class SyncFile
@@ -516,33 +516,33 @@ namespace UnrealGameSync
 			public string RelativePath;
 			public long Size;
 
-			public SyncFile(string DepotFile, string RelativePath, long Size)
+			public SyncFile(string depotFile, string relativePath, long size)
 			{
-				this.DepotFile = DepotFile;
-				this.RelativePath = RelativePath;
-				this.Size = Size;
+				this.DepotFile = depotFile;
+				this.RelativePath = relativePath;
+				this.Size = size;
 			}
 		};
 
 		class SemaphoreScope : IDisposable
 		{
-			SemaphoreSlim Semaphore;
+			SemaphoreSlim _semaphore;
 			public bool HasLock { get; private set; }
 
-			public SemaphoreScope(SemaphoreSlim Semaphore)
+			public SemaphoreScope(SemaphoreSlim semaphore)
 			{
-				this.Semaphore = Semaphore;
+				this._semaphore = semaphore;
 			}
 
 			public bool TryAcquire()
 			{
-				HasLock = HasLock || Semaphore.Wait(0);
+				HasLock = HasLock || _semaphore.Wait(0);
 				return HasLock;
 			}
 
-			public async Task AcquireAsync(CancellationToken CancellationToken)
+			public async Task AcquireAsync(CancellationToken cancellationToken)
 			{
-				await Semaphore.WaitAsync(CancellationToken);
+				await _semaphore.WaitAsync(cancellationToken);
 				HasLock = true;
 			}
 
@@ -550,7 +550,7 @@ namespace UnrealGameSync
 			{
 				if (HasLock)
 				{
-					Semaphore.Release();
+					_semaphore.Release();
 					HasLock = false;
 				}
 			}
@@ -560,116 +560,116 @@ namespace UnrealGameSync
 
 		public static string ShellScriptExt { get; }= RuntimeInformation.IsOSPlatform(OSPlatform.Windows)? "bat" : "sh";
 
-		public Task<int> ExecuteShellCommandAsync(string CommandLine, string? WorkingDir, Action<string> ProcessOutput, CancellationToken CancellationToken)
+		public Task<int> ExecuteShellCommandAsync(string commandLine, string? workingDir, Action<string> processOutput, CancellationToken cancellationToken)
 		{
 			if(RuntimeInformation.IsOSPlatform(OSPlatform.Windows))
 			{
-				string CmdExe = Path.Combine(Environment.GetFolderPath(Environment.SpecialFolder.System), "cmd.exe");
-				return Utility.ExecuteProcessAsync(CmdExe, null, $"/C \"{CommandLine}\"", ProcessOutput, CancellationToken);
+				string cmdExe = Path.Combine(Environment.GetFolderPath(Environment.SpecialFolder.System), "cmd.exe");
+				return Utility.ExecuteProcessAsync(cmdExe, null, $"/C \"{commandLine}\"", processOutput, cancellationToken);
 			}
 			else
 			{
-				string ShellExe = "/bin/sh";
-				return Utility.ExecuteProcessAsync(ShellExe, null, $"{CommandLine}", ProcessOutput, CancellationToken);
+				string shellExe = "/bin/sh";
+				return Utility.ExecuteProcessAsync(shellExe, null, $"{commandLine}", processOutput, cancellationToken);
 			}
 		}
 
-		public async Task<(WorkspaceUpdateResult, string)> ExecuteAsync(IPerforceSettings PerforceSettings, ProjectInfo Project, UserWorkspaceState State, ILogger Logger, CancellationToken CancellationToken)
+		public async Task<(WorkspaceUpdateResult, string)> ExecuteAsync(IPerforceSettings perforceSettings, ProjectInfo project, UserWorkspaceState state, ILogger logger, CancellationToken cancellationToken)
 		{
-			using IPerforceConnection Perforce = await PerforceConnection.CreateAsync(PerforceSettings, Logger);
+			using IPerforceConnection perforce = await PerforceConnection.CreateAsync(perforceSettings, logger);
 
-			List<Tuple<string, TimeSpan>> Times = new List<Tuple<string, TimeSpan>>();
+			List<Tuple<string, TimeSpan>> times = new List<Tuple<string, TimeSpan>>();
 
-			int NumFilesSynced = 0;
+			int numFilesSynced = 0;
 			if (Context.Options.HasFlag(WorkspaceUpdateOptions.Sync) || Context.Options.HasFlag(WorkspaceUpdateOptions.SyncSingleChange))
 			{
-				using (TelemetryStopwatch SyncTelemetryStopwatch = new TelemetryStopwatch("Workspace_Sync", Project.TelemetryProjectIdentifier))
+				using (TelemetryStopwatch syncTelemetryStopwatch = new TelemetryStopwatch("Workspace_Sync", project.TelemetryProjectIdentifier))
 				{
-					Logger.LogInformation("Syncing to {Change}...", Context.ChangeNumber);
+					logger.LogInformation("Syncing to {Change}...", Context.ChangeNumber);
 
 					// Make sure we're logged in
-					PerforceResponse<LoginRecord> LoginResponse = await Perforce.TryGetLoginStateAsync(CancellationToken);
-					if (!LoginResponse.Succeeded)
+					PerforceResponse<LoginRecord> loginResponse = await perforce.TryGetLoginStateAsync(cancellationToken);
+					if (!loginResponse.Succeeded)
 					{
 						return (WorkspaceUpdateResult.FailedToSyncLoginExpired, "User is not logged in.");
 					}
 
 					// Figure out which paths to sync
-					List<string> RelativeSyncPaths = GetRelativeSyncPaths(Project, (Context.Options & WorkspaceUpdateOptions.SyncAllProjects) != 0, Context.SyncFilter);
-					List<string> SyncPaths = new List<string>(RelativeSyncPaths.Select(x => Project.ClientRootPath + x));
+					List<string> relativeSyncPaths = GetRelativeSyncPaths(project, (Context.Options & WorkspaceUpdateOptions.SyncAllProjects) != 0, Context.SyncFilter);
+					List<string> syncPaths = new List<string>(relativeSyncPaths.Select(x => project.ClientRootPath + x));
 
 					// Get the user's sync filter
-					FileFilter UserFilter = new FileFilter(FileFilterType.Include);
+					FileFilter userFilter = new FileFilter(FileFilterType.Include);
 					if (Context.SyncFilter != null)
 					{
-						UserFilter.AddRules(Context.SyncFilter.Select(x => x.Trim()).Where(x => x.Length > 0 && !x.StartsWith(";") && !x.StartsWith("#")));
+						userFilter.AddRules(Context.SyncFilter.Select(x => x.Trim()).Where(x => x.Length > 0 && !x.StartsWith(";") && !x.StartsWith("#")));
 					}
 
 					// Check if the new sync filter matches the previous one. If not, we'll enumerate all files in the workspace and make sure there's nothing extra there.
-					string? NextSyncFilterHash = null;
-					using (SHA1Managed SHA = new SHA1Managed())
+					string? nextSyncFilterHash = null;
+					using (SHA1Managed sha = new SHA1Managed())
 					{
-						StringBuilder CombinedFilter = new StringBuilder();
-						foreach (string RelativeSyncPath in RelativeSyncPaths)
+						StringBuilder combinedFilter = new StringBuilder();
+						foreach (string relativeSyncPath in relativeSyncPaths)
 						{
-							CombinedFilter.AppendFormat("{0}\n", RelativeSyncPath);
+							combinedFilter.AppendFormat("{0}\n", relativeSyncPath);
 						}
 						if (Context.SyncFilter != null)
 						{
-							CombinedFilter.Append("--FROM--\n");
-							CombinedFilter.Append(String.Join("\n", Context.SyncFilter));
+							combinedFilter.Append("--FROM--\n");
+							combinedFilter.Append(String.Join("\n", Context.SyncFilter));
 						}
-						NextSyncFilterHash = BitConverter.ToString(SHA.ComputeHash(Encoding.UTF8.GetBytes(CombinedFilter.ToString()))).Replace("-", "");
+						nextSyncFilterHash = BitConverter.ToString(sha.ComputeHash(Encoding.UTF8.GetBytes(combinedFilter.ToString()))).Replace("-", "");
 					}
 
 					// If the hash differs, enumerate everything in the workspace to find what needs to be removed
-					if (NextSyncFilterHash != State.CurrentSyncFilterHash || (Context.Options & WorkspaceUpdateOptions.Refilter) != 0)
+					if (nextSyncFilterHash != state.CurrentSyncFilterHash || (Context.Options & WorkspaceUpdateOptions.Refilter) != 0)
 					{
-						using (TelemetryStopwatch FilterStopwatch = new TelemetryStopwatch("Workspace_Sync_FilterChanged", Project.TelemetryProjectIdentifier))
+						using (TelemetryStopwatch filterStopwatch = new TelemetryStopwatch("Workspace_Sync_FilterChanged", project.TelemetryProjectIdentifier))
 						{
-							Logger.LogInformation("Filter has changed ({PrevHash} -> {NextHash}); finding files in workspace that need to be removed.", (String.IsNullOrEmpty(State.CurrentSyncFilterHash)) ? "None" : State.CurrentSyncFilterHash, NextSyncFilterHash);
+							logger.LogInformation("Filter has changed ({PrevHash} -> {NextHash}); finding files in workspace that need to be removed.", (String.IsNullOrEmpty(state.CurrentSyncFilterHash)) ? "None" : state.CurrentSyncFilterHash, nextSyncFilterHash);
 
 							// Find all the files that are in this workspace
-							List<HaveRecord>? HaveFiles = Context.HaveFiles;
-							if (HaveFiles == null)
+							List<HaveRecord>? haveFiles = Context.HaveFiles;
+							if (haveFiles == null)
 							{
-								HaveFiles = new List<HaveRecord>();
-								using (RecordCounter HaveCounter = new RecordCounter(Progress, "Sync filter changed; checking workspace..."))
+								haveFiles = new List<HaveRecord>();
+								using (RecordCounter haveCounter = new RecordCounter(Progress, "Sync filter changed; checking workspace..."))
 								{
-									await foreach (PerforceResponse<HaveRecord> Record in Perforce.TryHaveAsync(FileSpecList.Any, CancellationToken))
+									await foreach (PerforceResponse<HaveRecord> record in perforce.TryHaveAsync(FileSpecList.Any, cancellationToken))
 									{
-										if (Record.Succeeded)
+										if (record.Succeeded)
 										{
-											HaveFiles.Add(Record.Data);
-											HaveCounter.Increment();
+											haveFiles.Add(record.Data);
+											haveCounter.Increment();
 										}
 										else
 										{
-											return (WorkspaceUpdateResult.FailedToSync, $"Unable to query files ({Record}).");
+											return (WorkspaceUpdateResult.FailedToSync, $"Unable to query files ({record}).");
 										}
 									}
 								}
-								Context.HaveFiles = HaveFiles;
+								Context.HaveFiles = haveFiles;
 							}
 
 							// Build a filter for the current sync paths
-							FileFilter SyncPathsFilter = new FileFilter(FileFilterType.Exclude);
-							foreach (string RelativeSyncPath in RelativeSyncPaths)
+							FileFilter syncPathsFilter = new FileFilter(FileFilterType.Exclude);
+							foreach (string relativeSyncPath in relativeSyncPaths)
 							{
-								SyncPathsFilter.Include(RelativeSyncPath);
+								syncPathsFilter.Include(relativeSyncPath);
 							}
 
 							// Remove all the files that are not included by the filter
-							List<string> RemoveDepotPaths = new List<string>();
-							foreach (HaveRecord HaveFile in HaveFiles)
+							List<string> removeDepotPaths = new List<string>();
+							foreach (HaveRecord haveFile in haveFiles)
 							{
 								try
 								{
-									FileReference FullPath = new FileReference(HaveFile.Path);
-									if (MatchFilter(Project, FullPath, SyncPathsFilter) && !MatchFilter(Project, FullPath, UserFilter))
+									FileReference fullPath = new FileReference(haveFile.Path);
+									if (MatchFilter(project, fullPath, syncPathsFilter) && !MatchFilter(project, fullPath, userFilter))
 									{
-										Logger.LogInformation("  {DepotFile}", HaveFile.DepotFile);
-										RemoveDepotPaths.Add(HaveFile.DepotFile);
+										logger.LogInformation("  {DepotFile}", haveFile.DepotFile);
+										removeDepotPaths.Add(haveFile.DepotFile);
 									}
 								}
 								catch (PathTooLongException)
@@ -679,262 +679,262 @@ namespace UnrealGameSync
 							}
 
 							// Check if there are any paths outside the regular sync paths
-							if (RemoveDepotPaths.Count > 0 && (Context.Options & WorkspaceUpdateOptions.RemoveFilteredFiles) == 0)
+							if (removeDepotPaths.Count > 0 && (Context.Options & WorkspaceUpdateOptions.RemoveFilteredFiles) == 0)
 							{
-								bool bDeleteListMatches = true;
+								bool deleteListMatches = true;
 
-								Dictionary<string, bool> NewDeleteFiles = new Dictionary<string, bool>(StringComparer.OrdinalIgnoreCase);
-								foreach (string RemoveDepotPath in RemoveDepotPaths)
+								Dictionary<string, bool> newDeleteFiles = new Dictionary<string, bool>(StringComparer.OrdinalIgnoreCase);
+								foreach (string removeDepotPath in removeDepotPaths)
 								{
-									bool bDelete;
-									if (!Context.DeleteFiles.TryGetValue(RemoveDepotPath, out bDelete))
+									bool delete;
+									if (!Context.DeleteFiles.TryGetValue(removeDepotPath, out delete))
 									{
-										bDeleteListMatches = false;
-										bDelete = true;
+										deleteListMatches = false;
+										delete = true;
 									}
-									NewDeleteFiles[RemoveDepotPath] = bDelete;
+									newDeleteFiles[removeDepotPath] = delete;
 								}
-								Context.DeleteFiles = NewDeleteFiles;
+								Context.DeleteFiles = newDeleteFiles;
 
-								if (!bDeleteListMatches)
+								if (!deleteListMatches)
 								{
-									return (WorkspaceUpdateResult.FilesToDelete, $"Cancelled after finding {NewDeleteFiles.Count} files excluded by filter");
+									return (WorkspaceUpdateResult.FilesToDelete, $"Cancelled after finding {newDeleteFiles.Count} files excluded by filter");
 								}
 
-								RemoveDepotPaths.RemoveAll(x => !Context.DeleteFiles[x]);
+								removeDepotPaths.RemoveAll(x => !Context.DeleteFiles[x]);
 							}
 
 							// Actually delete any files that we don't want
-							if (RemoveDepotPaths.Count > 0)
+							if (removeDepotPaths.Count > 0)
 							{
 								// Clear the current sync filter hash. If the sync is canceled, we'll be in an indeterminate state, and we should always clean next time round.
-								State.CurrentSyncFilterHash = "INVALID";
-								State.Save(Logger);
+								state.CurrentSyncFilterHash = "INVALID";
+								state.Save(logger);
 
 								// Find all the depot paths that will be synced
-								HashSet<string> RemainingDepotPathsToRemove = new HashSet<string>(StringComparer.OrdinalIgnoreCase);
-								RemainingDepotPathsToRemove.UnionWith(RemoveDepotPaths);
+								HashSet<string> remainingDepotPathsToRemove = new HashSet<string>(StringComparer.OrdinalIgnoreCase);
+								remainingDepotPathsToRemove.UnionWith(removeDepotPaths);
 
 								// Build the list of revisions to sync
-								List<string> RevisionsToRemove = new List<string>();
-								RevisionsToRemove.AddRange(RemoveDepotPaths.Select(x => String.Format("{0}#0", x)));
+								List<string> revisionsToRemove = new List<string>();
+								revisionsToRemove.AddRange(removeDepotPaths.Select(x => String.Format("{0}#0", x)));
 
-								(WorkspaceUpdateResult, string) RemoveResult = await SyncFileRevisions(Perforce, "Removing files...", Context, RevisionsToRemove, RemainingDepotPathsToRemove, Progress, Logger, CancellationToken);
-								if (RemoveResult.Item1 != WorkspaceUpdateResult.Success)
+								(WorkspaceUpdateResult, string) removeResult = await SyncFileRevisions(perforce, "Removing files...", Context, revisionsToRemove, remainingDepotPathsToRemove, Progress, logger, cancellationToken);
+								if (removeResult.Item1 != WorkspaceUpdateResult.Success)
 								{
-									return RemoveResult;
+									return removeResult;
 								}
 							}
 
 							// Update the sync filter hash. We've removed any files we need to at this point.
-							State.CurrentSyncFilterHash = NextSyncFilterHash;
-							State.Save(Logger);
+							state.CurrentSyncFilterHash = nextSyncFilterHash;
+							state.Save(logger);
 						}
 					}
 
 					// Create a filter for all the files we don't want
-					FileFilter Filter = new FileFilter(UserFilter);
-					Filter.Exclude(BuildVersionFileName);
+					FileFilter filter = new FileFilter(userFilter);
+					filter.Exclude(BuildVersionFileName);
 					if (Context.Options.HasFlag(WorkspaceUpdateOptions.ContentOnly))
 					{
-						Filter.Exclude("*.usf");
-						Filter.Exclude("*.ush");
+						filter.Exclude("*.usf");
+						filter.Exclude("*.ush");
 					}
 
 					// Create a tree to store the sync path
-					SyncTree SyncTree = new SyncTree(false);
+					SyncTree syncTree = new SyncTree(false);
 					if (!Context.Options.HasFlag(WorkspaceUpdateOptions.SyncSingleChange))
 					{
-						foreach (string RelativeSyncPath in RelativeSyncPaths)
+						foreach (string relativeSyncPath in relativeSyncPaths)
 						{
-							const string WildcardSuffix = "/...";
-							if (RelativeSyncPath.EndsWith(WildcardSuffix, StringComparison.Ordinal))
+							const string wildcardSuffix = "/...";
+							if (relativeSyncPath.EndsWith(wildcardSuffix, StringComparison.Ordinal))
 							{
-								SyncTree Leaf = SyncTree;
+								SyncTree leaf = syncTree;
 
-								string[] Fragments = RelativeSyncPath.Split('/');
-								for (int Idx = 1; Idx < Fragments.Length - 1; Idx++)
+								string[] fragments = relativeSyncPath.Split('/');
+								for (int idx = 1; idx < fragments.Length - 1; idx++)
 								{
-									Leaf = Leaf.FindOrAddSubTree(Fragments[Idx]);
+									leaf = leaf.FindOrAddSubTree(fragments[idx]);
 								}
 
-								Leaf.bCanUseWildcard = true;
+								leaf.CanUseWildcard = true;
 							}
 						}
 					}
 
 					// Find all the server changes, and anything that's opened for edit locally. We need to sync files we have open to schedule a resolve.
-					SyncBatchBuilder BatchBuilder = new SyncBatchBuilder(Context.PerforceSyncOptions?.MaxCommandsPerBatch, Context.PerforceSyncOptions?.MaxSizePerBatch);
-					List<string> SyncDepotPaths = new List<string>();
-					using (RecordCounter Counter = new RecordCounter(Progress, "Filtering files..."))
+					SyncBatchBuilder batchBuilder = new SyncBatchBuilder(Context.PerforceSyncOptions?.MaxCommandsPerBatch, Context.PerforceSyncOptions?.MaxSizePerBatch);
+					List<string> syncDepotPaths = new List<string>();
+					using (RecordCounter counter = new RecordCounter(Progress, "Filtering files..."))
 					{
 						// Track the total new bytes that will be required on disk when syncing. Add an extra 100MB for padding.
-						long RequiredFreeSpace = 100 * 1024 * 1024;
+						long requiredFreeSpace = 100 * 1024 * 1024;
 
-						foreach (string SyncPath in SyncPaths)
+						foreach (string syncPath in syncPaths)
 						{
-							string SyncFilter = Context.Options.HasFlag(WorkspaceUpdateOptions.Sync) ? $"{Context.ChangeNumber}" : $"={Context.ChangeNumber}";
+							string syncFilter = Context.Options.HasFlag(WorkspaceUpdateOptions.Sync) ? $"{Context.ChangeNumber}" : $"={Context.ChangeNumber}";
 
-							List<SyncFile> SyncFiles = new List<SyncFile>();
-							await foreach (PerforceResponse<SyncRecord> Response in Perforce.TrySyncAsync(SyncOptions.PreviewOnly, -1, $"{SyncPath}@{SyncFilter}", CancellationToken))
+							List<SyncFile> syncFiles = new List<SyncFile>();
+							await foreach (PerforceResponse<SyncRecord> response in perforce.TrySyncAsync(SyncOptions.PreviewOnly, -1, $"{syncPath}@{syncFilter}", cancellationToken))
 							{
-								if (!Response.Succeeded)
+								if (!response.Succeeded)
 								{
-									return (WorkspaceUpdateResult.FailedToSync, $"Couldn't enumerate changes matching {SyncPath}.");
+									return (WorkspaceUpdateResult.FailedToSync, $"Couldn't enumerate changes matching {syncPath}.");
 								}
-								if (Response.Info != null)
+								if (response.Info != null)
 								{
-									Logger.LogInformation("Note: {0}", Response.Info.Data);
+									logger.LogInformation("Note: {0}", response.Info.Data);
 									continue;
 								}
 
-								SyncRecord Record = Response.Data;
+								SyncRecord record = response.Data;
 
 								// Get the full local path
-								string RelativePath;
+								string relativePath;
 								try
 								{
-									FileReference SyncFile = new FileReference(Record.Path.ToString());
-									RelativePath = PerforceUtils.GetClientRelativePath(Project.LocalRootPath, SyncFile);
+									FileReference syncFile = new FileReference(record.Path.ToString());
+									relativePath = PerforceUtils.GetClientRelativePath(project.LocalRootPath, syncFile);
 								}
 								catch (PathTooLongException)
 								{
-									Logger.LogInformation("The local path for {Path} exceeds the maximum allowed by Windows. Re-sync your workspace to a directory with a shorter name, or delete the file from the server.", Record.Path);
+									logger.LogInformation("The local path for {Path} exceeds the maximum allowed by Windows. Re-sync your workspace to a directory with a shorter name, or delete the file from the server.", record.Path);
 									return (WorkspaceUpdateResult.FailedToSync, "File exceeds maximum path length allowed by Windows.");
 								}
 
 								// Create the sync record
-								long SyncSize = (Record.Action == SyncAction.Deleted) ? 0 : Record.FileSize;
-								SyncFiles.Add(new SyncFile(Record.DepotFile.ToString(), RelativePath, SyncSize));
-								Counter.Increment();
+								long syncSize = (record.Action == SyncAction.Deleted) ? 0 : record.FileSize;
+								syncFiles.Add(new SyncFile(record.DepotFile.ToString(), relativePath, syncSize));
+								counter.Increment();
 							}
 
 							// Also sync the currently open files
-							await foreach (PerforceResponse<OpenedRecord> Response in Perforce.TryOpenedAsync(OpenedOptions.None, FileSpecList.Any, CancellationToken))
+							await foreach (PerforceResponse<OpenedRecord> response in perforce.TryOpenedAsync(OpenedOptions.None, FileSpecList.Any, cancellationToken))
 							{
-								if (!Response.Succeeded)
+								if (!response.Succeeded)
 								{
-									return (WorkspaceUpdateResult.FailedToSync, $"Couldn't enumerate changes matching {SyncPath}.");
+									return (WorkspaceUpdateResult.FailedToSync, $"Couldn't enumerate changes matching {syncPath}.");
 								}
 
-								OpenedRecord Record = Response.Data;
-								if (Record.Action != FileAction.Add || Record.Action != FileAction.Branch || Record.Action != FileAction.MoveAdd)
+								OpenedRecord record = response.Data;
+								if (record.Action != FileAction.Add || record.Action != FileAction.Branch || record.Action != FileAction.MoveAdd)
 								{
-									string RelativePath = PerforceUtils.GetClientRelativePath(Record.ClientFile);
-									SyncFiles.Add(new SyncFile(Record.DepotFile, RelativePath, 0));
+									string relativePath = PerforceUtils.GetClientRelativePath(record.ClientFile);
+									syncFiles.Add(new SyncFile(record.DepotFile, relativePath, 0));
 								}
 							}
 
 							// Enumerate all the files to be synced. NOTE: depotPath is escaped, whereas clientPath is not.
-							foreach (SyncFile SyncRecord in SyncFiles)
+							foreach (SyncFile syncRecord in syncFiles)
 							{
-								if (Filter.Matches(SyncRecord.RelativePath))
+								if (filter.Matches(syncRecord.RelativePath))
 								{
-									SyncTree.IncludeFile(PerforceUtils.EscapePath(SyncRecord.RelativePath), SyncRecord.Size, Logger);
-									SyncDepotPaths.Add(SyncRecord.DepotFile);
-									RequiredFreeSpace += SyncRecord.Size;
+									syncTree.IncludeFile(PerforceUtils.EscapePath(syncRecord.RelativePath), syncRecord.Size, logger);
+									syncDepotPaths.Add(syncRecord.DepotFile);
+									requiredFreeSpace += syncRecord.Size;
 
 									// If the file exists the required free space can be reduced as those bytes will be replaced.
-									FileInfo LocalFileInfo = FileReference.Combine(Project.LocalRootPath, SyncRecord.RelativePath).ToFileInfo();
-									if (LocalFileInfo.Exists)
+									FileInfo localFileInfo = FileReference.Combine(project.LocalRootPath, syncRecord.RelativePath).ToFileInfo();
+									if (localFileInfo.Exists)
 									{
-										RequiredFreeSpace -= LocalFileInfo.Length;
+										requiredFreeSpace -= localFileInfo.Length;
 									}
 								}
 								else
 								{
-									SyncTree.ExcludeFile(PerforceUtils.EscapePath(SyncRecord.RelativePath));
+									syncTree.ExcludeFile(PerforceUtils.EscapePath(syncRecord.RelativePath));
 								}
 							}
 						}
 
 						try
 						{
-							DirectoryInfo LocalRootInfo = Project.LocalRootPath.ToDirectoryInfo();
-							DriveInfo Drive = new DriveInfo(LocalRootInfo.Root.FullName);
+							DirectoryInfo localRootInfo = project.LocalRootPath.ToDirectoryInfo();
+							DriveInfo drive = new DriveInfo(localRootInfo.Root.FullName);
 
-							if (Drive.AvailableFreeSpace < RequiredFreeSpace)
+							if (drive.AvailableFreeSpace < requiredFreeSpace)
 							{
-								Logger.LogInformation("Syncing requires {RequiredSpace} which exceeds the {AvailableSpace} available free space on {Drive}.", StringUtils.FormatBytesString(RequiredFreeSpace), StringUtils.FormatBytesString(Drive.AvailableFreeSpace), Drive.Name);
+								logger.LogInformation("Syncing requires {RequiredSpace} which exceeds the {AvailableSpace} available free space on {Drive}.", StringUtils.FormatBytesString(requiredFreeSpace), StringUtils.FormatBytesString(drive.AvailableFreeSpace), drive.Name);
 								return (WorkspaceUpdateResult.FailedToSync, "Not enough available free space.");
 							}
 						}
 						catch (SystemException)
 						{
-							Logger.LogInformation("Unable to check available free space for {RootPath}.", Project.LocalRootPath);
+							logger.LogInformation("Unable to check available free space for {RootPath}.", project.LocalRootPath);
 						}
 					}
-					SyncTree.GetOptimizedSyncCommands(Project.ClientRootPath, Context.ChangeNumber, BatchBuilder);
+					syncTree.GetOptimizedSyncCommands(project.ClientRootPath, Context.ChangeNumber, batchBuilder);
 
 					// Clear the current sync changelist, in case we cancel
 					if (!Context.Options.HasFlag(WorkspaceUpdateOptions.SyncSingleChange))
 					{
-						State.CurrentChangeNumber = -1;
-						State.AdditionalChangeNumbers.Clear();
-						State.Save(Logger);
+						state.CurrentChangeNumber = -1;
+						state.AdditionalChangeNumbers.Clear();
+						state.Save(logger);
 					}
 
 					// Find all the depot paths that will be synced
-					HashSet<string> RemainingDepotPaths = new HashSet<string>(StringComparer.OrdinalIgnoreCase);
-					RemainingDepotPaths.UnionWith(SyncDepotPaths);
+					HashSet<string> remainingDepotPaths = new HashSet<string>(StringComparer.OrdinalIgnoreCase);
+					remainingDepotPaths.UnionWith(syncDepotPaths);
 
-					using (TelemetryStopwatch TransferStopwatch = new TelemetryStopwatch("Workspace_Sync_TransferFiles", Project.TelemetryProjectIdentifier))
+					using (TelemetryStopwatch transferStopwatch = new TelemetryStopwatch("Workspace_Sync_TransferFiles", project.TelemetryProjectIdentifier))
 					{
-						TransferStopwatch.AddData(new { MachineName = Environment.MachineName, DomainName = Environment.UserDomainName, ServerAndPort = Perforce.Settings.ServerAndPort, UserName = Perforce.Settings.UserName, IncludedFiles = SyncTree.TotalIncludedFiles, ExcludedFiles = SyncTree.TotalExcludedFiles, Size = SyncTree.TotalSize, NumThreads = Context.PerforceSyncOptions?.NumThreads ?? PerforceSyncOptions.DefaultNumThreads });
+						transferStopwatch.AddData(new { MachineName = Environment.MachineName, DomainName = Environment.UserDomainName, ServerAndPort = perforce.Settings.ServerAndPort, UserName = perforce.Settings.UserName, IncludedFiles = syncTree.TotalIncludedFiles, ExcludedFiles = syncTree.TotalExcludedFiles, Size = syncTree.TotalSize, NumThreads = Context.PerforceSyncOptions?.NumThreads ?? PerforceSyncOptions.DefaultNumThreads });
 
-						(WorkspaceUpdateResult, string) SyncResult = await SyncFileRevisions(Perforce, "Syncing files...", Context, BatchBuilder.Batches, RemainingDepotPaths, Progress, Logger, CancellationToken);
-						if (SyncResult.Item1 != WorkspaceUpdateResult.Success)
+						(WorkspaceUpdateResult, string) syncResult = await SyncFileRevisions(perforce, "Syncing files...", Context, batchBuilder.Batches, remainingDepotPaths, Progress, logger, cancellationToken);
+						if (syncResult.Item1 != WorkspaceUpdateResult.Success)
 						{
-							TransferStopwatch.AddData(new { SyncResult = SyncResult.Item2.ToString(), CompletedFilesFiles = SyncDepotPaths.Count - RemainingDepotPaths.Count });
-							return SyncResult;
+							transferStopwatch.AddData(new { SyncResult = syncResult.Item2.ToString(), CompletedFilesFiles = syncDepotPaths.Count - remainingDepotPaths.Count });
+							return syncResult;
 						}
 
-						TransferStopwatch.Stop("Ok");
-						TransferStopwatch.AddData(new { TransferRate = SyncTree.TotalSize / Math.Max(TransferStopwatch.Elapsed.TotalSeconds, 0.0001f) });
+						transferStopwatch.Stop("Ok");
+						transferStopwatch.AddData(new { TransferRate = syncTree.TotalSize / Math.Max(transferStopwatch.Elapsed.TotalSeconds, 0.0001f) });
 					}
 
-					int VersionChangeNumber = -1;
+					int versionChangeNumber = -1;
 					if (Context.Options.HasFlag(WorkspaceUpdateOptions.Sync) && !Context.Options.HasFlag(WorkspaceUpdateOptions.UpdateFilter))
 					{
 						// Read the new config file
-						Context.ProjectConfigFile = await ReadProjectConfigFile(Project.LocalRootPath, Project.LocalFileName, Logger);
-						Context.ProjectStreamFilter = await ReadProjectStreamFilter(Perforce, Context.ProjectConfigFile, Logger, CancellationToken);
+						Context.ProjectConfigFile = await ReadProjectConfigFile(project.LocalRootPath, project.LocalFileName, logger);
+						Context.ProjectStreamFilter = await ReadProjectStreamFilter(perforce, Context.ProjectConfigFile, logger, cancellationToken);
 
 						// Get the branch name
-						string? BranchOrStreamName = await Perforce.GetCurrentStreamAsync(CancellationToken);
-						if (BranchOrStreamName != null)
+						string? branchOrStreamName = await perforce.GetCurrentStreamAsync(cancellationToken);
+						if (branchOrStreamName != null)
 						{
 							// If it's a virtual stream, take the concrete parent stream instead
 							for (; ; )
 							{
-								StreamRecord StreamSpec = await Perforce.GetStreamAsync(BranchOrStreamName, false, CancellationToken);
-								if (StreamSpec.Type != "virtual" || StreamSpec.Parent == "none" || StreamSpec.Parent == null)
+								StreamRecord streamSpec = await perforce.GetStreamAsync(branchOrStreamName, false, cancellationToken);
+								if (streamSpec.Type != "virtual" || streamSpec.Parent == "none" || streamSpec.Parent == null)
 								{
 									break;
 								}
-								BranchOrStreamName = StreamSpec.Parent;
+								branchOrStreamName = streamSpec.Parent;
 							}
 						}
 						else
 						{
 							// Otherwise use the depot path for GenerateProjectFiles.bat in the root of the workspace
-							List<WhereRecord> Files = await Perforce.WhereAsync(Project.ClientRootPath + "/GenerateProjectFiles.bat").Where(x => !x.Unmap).ToListAsync(CancellationToken);
-							if (Files.Count != 1)
+							List<WhereRecord> files = await perforce.WhereAsync(project.ClientRootPath + "/GenerateProjectFiles.bat").Where(x => !x.Unmap).ToListAsync(cancellationToken);
+							if (files.Count != 1)
 							{
-								return (WorkspaceUpdateResult.FailedToSync, $"Couldn't determine branch name for {Project.ClientFileName}.");
+								return (WorkspaceUpdateResult.FailedToSync, $"Couldn't determine branch name for {project.ClientFileName}.");
 							}
-							BranchOrStreamName = PerforceUtils.GetClientOrDepotDirectoryName(Files[0].DepotFile);
+							branchOrStreamName = PerforceUtils.GetClientOrDepotDirectoryName(files[0].DepotFile);
 						}
 
 						// Find the last code change before this changelist. For consistency in versioning between local builds and precompiled binaries, we need to use the last submitted code changelist as our version number.
-						string[] CodeFilter = new string[] { ".cs", ".h", ".cpp", ".usf", ".ush", ".uproject", ".uplugin" }.SelectMany(x => SyncPaths.Select(y => String.Format("{0}{1}@<={2}", y, x, Context.ChangeNumber))).ToArray();
+						string[] codeFilter = new string[] { ".cs", ".h", ".cpp", ".usf", ".ush", ".uproject", ".uplugin" }.SelectMany(x => syncPaths.Select(y => String.Format("{0}{1}@<={2}", y, x, Context.ChangeNumber))).ToArray();
 
-						PerforceResponseList<ChangesRecord> CodeChanges = await Perforce.TryGetChangesAsync(ChangesOptions.None, 1, ChangeStatus.Submitted, CodeFilter, CancellationToken);
-						if (!CodeChanges.Succeeded)
+						PerforceResponseList<ChangesRecord> codeChanges = await perforce.TryGetChangesAsync(ChangesOptions.None, 1, ChangeStatus.Submitted, codeFilter, cancellationToken);
+						if (!codeChanges.Succeeded)
 						{
 							return (WorkspaceUpdateResult.FailedToSync, $"Couldn't determine last code changelist before CL {Context.ChangeNumber}.");
 						}
-						if (CodeChanges.Count == 0)
+						if (codeChanges.Count == 0)
 						{
 							return (WorkspaceUpdateResult.FailedToSync, $"Could not find any code changes before CL {Context.ChangeNumber}.");
 						}
@@ -942,60 +942,60 @@ namespace UnrealGameSync
 						// Get the last code change
 						if (Context.ProjectConfigFile.GetValue("Options.VersionToLastCodeChange", true))
 						{
-							VersionChangeNumber = CodeChanges.Max(x => x.Data.Number);
+							versionChangeNumber = codeChanges.Max(x => x.Data.Number);
 						}
 						else
 						{
-							VersionChangeNumber = Context.ChangeNumber;
+							versionChangeNumber = Context.ChangeNumber;
 						}
 
 						// Update the version files
 						if (Context.ProjectConfigFile.GetValue("Options.UseFastModularVersioningV2", false))
 						{
-							bool bIsLicenseeVersion = await IsLicenseeVersion(Perforce, Project, Logger, CancellationToken);
-							if (!await UpdateVersionFile(Perforce, Project.ClientRootPath + BuildVersionFileName, Context.ChangeNumber, Text => UpdateBuildVersion(Text, Context.ChangeNumber, VersionChangeNumber, BranchOrStreamName, bIsLicenseeVersion), Logger, CancellationToken))
+							bool isLicenseeVersion = await IsLicenseeVersion(perforce, project, logger, cancellationToken);
+							if (!await UpdateVersionFile(perforce, project.ClientRootPath + BuildVersionFileName, Context.ChangeNumber, text => UpdateBuildVersion(text, Context.ChangeNumber, versionChangeNumber, branchOrStreamName, isLicenseeVersion), logger, cancellationToken))
 							{
 								return (WorkspaceUpdateResult.FailedToSync, $"Failed to update {BuildVersionFileName}.");
 							}
 						}
 						else if (Context.ProjectConfigFile.GetValue("Options.UseFastModularVersioning", false))
 						{
-							bool bIsLicenseeVersion = await IsLicenseeVersion(Perforce, Project, Logger, CancellationToken);
-							if (!await UpdateVersionFile(Perforce, Project.ClientRootPath + BuildVersionFileName, Context.ChangeNumber, Text => UpdateBuildVersion(Text, Context.ChangeNumber, VersionChangeNumber, BranchOrStreamName, bIsLicenseeVersion), Logger, CancellationToken))
+							bool isLicenseeVersion = await IsLicenseeVersion(perforce, project, logger, cancellationToken);
+							if (!await UpdateVersionFile(perforce, project.ClientRootPath + BuildVersionFileName, Context.ChangeNumber, text => UpdateBuildVersion(text, Context.ChangeNumber, versionChangeNumber, branchOrStreamName, isLicenseeVersion), logger, cancellationToken))
 							{
 								return (WorkspaceUpdateResult.FailedToSync, $"Failed to update {BuildVersionFileName}");
 							}
 
-							Dictionary<string, string> VersionHeaderStrings = new Dictionary<string, string>();
-							VersionHeaderStrings["#define ENGINE_IS_PROMOTED_BUILD"] = " (0)";
-							VersionHeaderStrings["#define BUILT_FROM_CHANGELIST"] = " 0";
-							VersionHeaderStrings["#define BRANCH_NAME"] = " \"" + BranchOrStreamName.Replace('/', '+') + "\"";
-							if (!await UpdateVersionFile(Perforce, Project.ClientRootPath + VersionHeaderFileName, VersionHeaderStrings, Context.ChangeNumber, Logger, CancellationToken))
+							Dictionary<string, string> versionHeaderStrings = new Dictionary<string, string>();
+							versionHeaderStrings["#define ENGINE_IS_PROMOTED_BUILD"] = " (0)";
+							versionHeaderStrings["#define BUILT_FROM_CHANGELIST"] = " 0";
+							versionHeaderStrings["#define BRANCH_NAME"] = " \"" + branchOrStreamName.Replace('/', '+') + "\"";
+							if (!await UpdateVersionFile(perforce, project.ClientRootPath + VersionHeaderFileName, versionHeaderStrings, Context.ChangeNumber, logger, cancellationToken))
 							{
 								return (WorkspaceUpdateResult.FailedToSync, $"Failed to update {VersionHeaderFileName}.");
 							}
-							if (!await UpdateVersionFile(Perforce, Project.ClientRootPath + ObjectVersionFileName, new Dictionary<string, string>(), Context.ChangeNumber, Logger, CancellationToken))
+							if (!await UpdateVersionFile(perforce, project.ClientRootPath + ObjectVersionFileName, new Dictionary<string, string>(), Context.ChangeNumber, logger, cancellationToken))
 							{
 								return (WorkspaceUpdateResult.FailedToSync, $"Failed to update {ObjectVersionFileName}.");
 							}
 						}
 						else
 						{
-							if (!await UpdateVersionFile(Perforce, Project.ClientRootPath + BuildVersionFileName, new Dictionary<string, string>(), Context.ChangeNumber, Logger, CancellationToken))
+							if (!await UpdateVersionFile(perforce, project.ClientRootPath + BuildVersionFileName, new Dictionary<string, string>(), Context.ChangeNumber, logger, cancellationToken))
 							{
 								return (WorkspaceUpdateResult.FailedToSync, $"Failed to update {BuildVersionFileName}");
 							}
 
-							Dictionary<string, string> VersionStrings = new Dictionary<string, string>();
-							VersionStrings["#define ENGINE_VERSION"] = " " + VersionChangeNumber.ToString();
-							VersionStrings["#define ENGINE_IS_PROMOTED_BUILD"] = " (0)";
-							VersionStrings["#define BUILT_FROM_CHANGELIST"] = " " + VersionChangeNumber.ToString();
-							VersionStrings["#define BRANCH_NAME"] = " \"" + BranchOrStreamName.Replace('/', '+') + "\"";
-							if (!await UpdateVersionFile(Perforce, Project.ClientRootPath + VersionHeaderFileName, VersionStrings, Context.ChangeNumber, Logger, CancellationToken))
+							Dictionary<string, string> versionStrings = new Dictionary<string, string>();
+							versionStrings["#define ENGINE_VERSION"] = " " + versionChangeNumber.ToString();
+							versionStrings["#define ENGINE_IS_PROMOTED_BUILD"] = " (0)";
+							versionStrings["#define BUILT_FROM_CHANGELIST"] = " " + versionChangeNumber.ToString();
+							versionStrings["#define BRANCH_NAME"] = " \"" + branchOrStreamName.Replace('/', '+') + "\"";
+							if (!await UpdateVersionFile(perforce, project.ClientRootPath + VersionHeaderFileName, versionStrings, Context.ChangeNumber, logger, cancellationToken))
 							{
 								return (WorkspaceUpdateResult.FailedToSync, $"Failed to update {VersionHeaderFileName}");
 							}
-							if (!await UpdateVersionFile(Perforce, Project.ClientRootPath + ObjectVersionFileName, VersionStrings, Context.ChangeNumber, Logger, CancellationToken))
+							if (!await UpdateVersionFile(perforce, project.ClientRootPath + ObjectVersionFileName, versionStrings, Context.ChangeNumber, logger, cancellationToken))
 							{
 								return (WorkspaceUpdateResult.FailedToSync, $"Failed to update {ObjectVersionFileName}");
 							}
@@ -1003,24 +1003,24 @@ namespace UnrealGameSync
 					}
 
 					// Check if there are any files which need resolving
-					List<FStatRecord> UnresolvedFiles = await FindUnresolvedFiles(Perforce, SyncPaths, CancellationToken);
-					if (UnresolvedFiles.Count > 0 && Context.Options.HasFlag(WorkspaceUpdateOptions.AutoResolveChanges))
+					List<FStatRecord> unresolvedFiles = await FindUnresolvedFiles(perforce, syncPaths, cancellationToken);
+					if (unresolvedFiles.Count > 0 && Context.Options.HasFlag(WorkspaceUpdateOptions.AutoResolveChanges))
 					{
-						foreach (FStatRecord UnresolvedFile in UnresolvedFiles)
+						foreach (FStatRecord unresolvedFile in unresolvedFiles)
 						{
-							if (UnresolvedFile.DepotFile != null)
+							if (unresolvedFile.DepotFile != null)
 							{
-								await Perforce.ResolveAsync(-1, ResolveOptions.Automatic, UnresolvedFile.DepotFile, CancellationToken);
+								await perforce.ResolveAsync(-1, ResolveOptions.Automatic, unresolvedFile.DepotFile, cancellationToken);
 							}
 						}
-						UnresolvedFiles = await FindUnresolvedFiles(Perforce, SyncPaths, CancellationToken);
+						unresolvedFiles = await FindUnresolvedFiles(perforce, syncPaths, cancellationToken);
 					}
-					if (UnresolvedFiles.Count > 0)
+					if (unresolvedFiles.Count > 0)
 					{
-						Logger.LogInformation("{NumFiles} files need resolving:", UnresolvedFiles.Count);
-						foreach (FStatRecord UnresolvedFile in UnresolvedFiles)
+						logger.LogInformation("{NumFiles} files need resolving:", unresolvedFiles.Count);
+						foreach (FStatRecord unresolvedFile in unresolvedFiles)
 						{
-							Logger.LogInformation("  {ClientFile}", UnresolvedFile.ClientFile);
+							logger.LogInformation("  {ClientFile}", unresolvedFile.ClientFile);
 						}
 						return (WorkspaceUpdateResult.FilesToResolve, "Files need resolving.");
 					}
@@ -1028,31 +1028,31 @@ namespace UnrealGameSync
 					// Continue processing sync-only actions
 					if (Context.Options.HasFlag(WorkspaceUpdateOptions.Sync) && !Context.Options.HasFlag(WorkspaceUpdateOptions.UpdateFilter))
 					{
-						Context.ProjectConfigFile ??= await ReadProjectConfigFile(Project.LocalRootPath, Project.LocalFileName, Logger);
+						Context.ProjectConfigFile ??= await ReadProjectConfigFile(project.LocalRootPath, project.LocalFileName, logger);
 
 						// Execute any project specific post-sync steps
-						string[]? PostSyncSteps = Context.ProjectConfigFile.GetValues("Sync.Step", null);
-						if (PostSyncSteps != null)
+						string[]? postSyncSteps = Context.ProjectConfigFile.GetValues("Sync.Step", null);
+						if (postSyncSteps != null)
 						{
-							Logger.LogInformation("");
-							Logger.LogInformation("Executing post-sync steps...");
+							logger.LogInformation("");
+							logger.LogInformation("Executing post-sync steps...");
 
-							Dictionary<string, string> PostSyncVariables = ConfigUtils.GetWorkspaceVariables(Project, Context.ChangeNumber, VersionChangeNumber, null, Context.ProjectConfigFile);
-							foreach (string PostSyncStep in PostSyncSteps.Select(x => x.Trim()))
+							Dictionary<string, string> postSyncVariables = ConfigUtils.GetWorkspaceVariables(project, Context.ChangeNumber, versionChangeNumber, null, Context.ProjectConfigFile);
+							foreach (string postSyncStep in postSyncSteps.Select(x => x.Trim()))
 							{
-								ConfigObject PostSyncStepObject = new ConfigObject(PostSyncStep);
+								ConfigObject postSyncStepObject = new ConfigObject(postSyncStep);
 
-								string ToolFileName = Utility.ExpandVariables(PostSyncStepObject.GetValue("FileName", ""), PostSyncVariables);
-								if (ToolFileName != null)
+								string toolFileName = Utility.ExpandVariables(postSyncStepObject.GetValue("FileName", ""), postSyncVariables);
+								if (toolFileName != null)
 								{
-									string ToolArguments = Utility.ExpandVariables(PostSyncStepObject.GetValue("Arguments", ""), PostSyncVariables);
+									string toolArguments = Utility.ExpandVariables(postSyncStepObject.GetValue("Arguments", ""), postSyncVariables);
 
-									Logger.LogInformation("post-sync> Running {FileName} {Arguments}", ToolFileName, ToolArguments);
+									logger.LogInformation("post-sync> Running {FileName} {Arguments}", toolFileName, toolArguments);
 
-									int ResultFromTool = await Utility.ExecuteProcessAsync(ToolFileName, null, ToolArguments, Line => ProcessOutput(Line, "post-sync> ", Progress, Logger), CancellationToken);
-									if (ResultFromTool != 0)
+									int resultFromTool = await Utility.ExecuteProcessAsync(toolFileName, null, toolArguments, line => ProcessOutput(line, "post-sync> ", Progress, logger), cancellationToken);
+									if (resultFromTool != 0)
 									{
-										return (WorkspaceUpdateResult.FailedToSync, $"Post-sync step terminated with exit code {ResultFromTool}.");
+										return (WorkspaceUpdateResult.FailedToSync, $"Post-sync step terminated with exit code {resultFromTool}.");
 									}
 								}
 							}
@@ -1062,170 +1062,170 @@ namespace UnrealGameSync
 					// Update the current state
 					if (Context.Options.HasFlag(WorkspaceUpdateOptions.SyncSingleChange))
 					{
-						State.AdditionalChangeNumbers.Add(Context.ChangeNumber);
+						state.AdditionalChangeNumbers.Add(Context.ChangeNumber);
 					}
 					else
 					{
-						State.CurrentChangeNumber = Context.ChangeNumber;
-						State.CurrentCodeChangeNumber = VersionChangeNumber;
+						state.CurrentChangeNumber = Context.ChangeNumber;
+						state.CurrentCodeChangeNumber = versionChangeNumber;
 					}
-					State.Save(Logger);
+					state.Save(logger);
 
 					// Update the timing info
-					Times.Add(new Tuple<string, TimeSpan>("Sync", SyncTelemetryStopwatch.Stop("Success")));
+					times.Add(new Tuple<string, TimeSpan>("Sync", syncTelemetryStopwatch.Stop("Success")));
 
 					// Save the number of files synced
-					NumFilesSynced = SyncDepotPaths.Count;
-					Logger.LogInformation("");
+					numFilesSynced = syncDepotPaths.Count;
+					logger.LogInformation("");
 				}
 			}
 
 			// Extract an archive from the depot path
 			if (Context.Options.HasFlag(WorkspaceUpdateOptions.SyncArchives))
 			{
-				using (TelemetryStopwatch Stopwatch = new TelemetryStopwatch("Workspace_SyncArchives", Project.TelemetryProjectIdentifier))
+				using (TelemetryStopwatch stopwatch = new TelemetryStopwatch("Workspace_SyncArchives", project.TelemetryProjectIdentifier))
 				{
 					// Create the directory for extracted archive manifests
-					DirectoryReference ManifestDirectoryName;
-					if (Project.LocalFileName.HasExtension(".uproject"))
+					DirectoryReference manifestDirectoryName;
+					if (project.LocalFileName.HasExtension(".uproject"))
 					{
-						ManifestDirectoryName = DirectoryReference.Combine(Project.LocalFileName.Directory, "Saved", "UnrealGameSync");
+						manifestDirectoryName = DirectoryReference.Combine(project.LocalFileName.Directory, "Saved", "UnrealGameSync");
 					}
 					else
 					{
-						ManifestDirectoryName = DirectoryReference.Combine(Project.LocalFileName.Directory, "Engine", "Saved", "UnrealGameSync");
+						manifestDirectoryName = DirectoryReference.Combine(project.LocalFileName.Directory, "Engine", "Saved", "UnrealGameSync");
 					}
-					DirectoryReference.CreateDirectory(ManifestDirectoryName);
+					DirectoryReference.CreateDirectory(manifestDirectoryName);
 
 					// Sync and extract (or just remove) the given archives
-					foreach (KeyValuePair<string, Tuple<IArchiveInfo, string>?> ArchiveTypeAndArchive in Context.ArchiveTypeToArchive)
+					foreach (KeyValuePair<string, Tuple<IArchiveInfo, string>?> archiveTypeAndArchive in Context.ArchiveTypeToArchive)
 					{
-						string ArchiveType = ArchiveTypeAndArchive.Key;
+						string archiveType = archiveTypeAndArchive.Key;
 
 						// Remove any existing binaries
-						FileReference ManifestFileName = FileReference.Combine(ManifestDirectoryName, String.Format("{0}.zipmanifest", ArchiveType));
-						if (FileReference.Exists(ManifestFileName))
+						FileReference manifestFileName = FileReference.Combine(manifestDirectoryName, String.Format("{0}.zipmanifest", archiveType));
+						if (FileReference.Exists(manifestFileName))
 						{
-							Logger.LogInformation("Removing {ArchiveType} binaries...", ArchiveType);
-							Progress.Set(String.Format("Removing {0} binaries...", ArchiveType), 0.0f);
-							ArchiveUtils.RemoveExtractedFiles(Project.LocalRootPath, ManifestFileName, Progress, Logger);
-							FileReference.Delete(ManifestFileName);
-							Logger.LogInformation("");
+							logger.LogInformation("Removing {ArchiveType} binaries...", archiveType);
+							Progress.Set(String.Format("Removing {0} binaries...", archiveType), 0.0f);
+							ArchiveUtils.RemoveExtractedFiles(project.LocalRootPath, manifestFileName, Progress, logger);
+							FileReference.Delete(manifestFileName);
+							logger.LogInformation("");
 						}
 
 						// If we have a new depot path, sync it down and extract it
-						if (ArchiveTypeAndArchive.Value != null)
+						if (archiveTypeAndArchive.Value != null)
 						{
-							IArchiveInfo ArchiveInfo = ArchiveTypeAndArchive.Value.Item1;
-							string ArchiveKey = ArchiveTypeAndArchive.Value.Item2;
+							IArchiveInfo archiveInfo = archiveTypeAndArchive.Value.Item1;
+							string archiveKey = archiveTypeAndArchive.Value.Item2;
 
-							Logger.LogInformation("Syncing {ArchiveType} binaries...", ArchiveType.ToLowerInvariant());
-							Progress.Set(String.Format("Syncing {0} binaries...", ArchiveType.ToLowerInvariant()), 0.0f);
-							if (!ArchiveInfo.DownloadArchive(Perforce, ArchiveKey, Project.LocalRootPath, ManifestFileName, Logger, Progress, CancellationToken.None).Result)
+							logger.LogInformation("Syncing {ArchiveType} binaries...", archiveType.ToLowerInvariant());
+							Progress.Set(String.Format("Syncing {0} binaries...", archiveType.ToLowerInvariant()), 0.0f);
+							if (!archiveInfo.DownloadArchive(perforce, archiveKey, project.LocalRootPath, manifestFileName, logger, Progress, CancellationToken.None).Result)
 							{
-								return (WorkspaceUpdateResult.FailedToSync, $"Couldn't read {ArchiveKey}");
+								return (WorkspaceUpdateResult.FailedToSync, $"Couldn't read {archiveKey}");
 							}
 						}
 					}
 
 					// Update the state
-					State.ExpandedArchiveTypes = Context.ArchiveTypeToArchive.Where(x => x.Value != null).Select(x => x.Key).ToArray();
-					State.Save(Logger);
+					state.ExpandedArchiveTypes = Context.ArchiveTypeToArchive.Where(x => x.Value != null).Select(x => x.Key).ToArray();
+					state.Save(logger);
 
 					// Add the finish time
-					Times.Add(new Tuple<string, TimeSpan>("Archive", Stopwatch.Stop("Success")));
+					times.Add(new Tuple<string, TimeSpan>("Archive", stopwatch.Stop("Success")));
 				}
 			}
 
 			// Take the lock before doing anything else. Building and generating project files can only be done on one workspace at a time.
-			using SemaphoreScope Scope = new SemaphoreScope(UpdateSemaphore);
+			using SemaphoreScope scope = new SemaphoreScope(_updateSemaphore);
 			if (Context.Options.HasFlag(WorkspaceUpdateOptions.GenerateProjectFiles) || Context.Options.HasFlag(WorkspaceUpdateOptions.Build))
 			{
-				if (!Scope.TryAcquire())
+				if (!scope.TryAcquire())
 				{
-					Logger.LogInformation("Waiting for other workspaces to finish...");
-					await Scope.AcquireAsync(CancellationToken);
+					logger.LogInformation("Waiting for other workspaces to finish...");
+					await scope.AcquireAsync(cancellationToken);
 				}
 			}
 
 			// Generate project files in the workspace
 			if (Context.Options.HasFlag(WorkspaceUpdateOptions.GenerateProjectFiles))
 			{
-				using (TelemetryStopwatch Stopwatch = new TelemetryStopwatch("Workspace_GenerateProjectFiles", Project.TelemetryProjectIdentifier))
+				using (TelemetryStopwatch stopwatch = new TelemetryStopwatch("Workspace_GenerateProjectFiles", project.TelemetryProjectIdentifier))
 				{
 					Progress.Set("Generating project files...", 0.0f);
 
-					StringBuilder CommandLine = new StringBuilder();
-					CommandLine.AppendFormat("\"{0}\"", FileReference.Combine(Project.LocalRootPath, $"GenerateProjectFiles.{ShellScriptExt}"));
+					StringBuilder commandLine = new StringBuilder();
+					commandLine.AppendFormat("\"{0}\"", FileReference.Combine(project.LocalRootPath, $"GenerateProjectFiles.{ShellScriptExt}"));
 					if ((Context.Options & WorkspaceUpdateOptions.SyncAllProjects) == 0 && (Context.Options & WorkspaceUpdateOptions.IncludeAllProjectsInSolution) == 0)
 					{
-						if (Project.LocalFileName.HasExtension(".uproject"))
+						if (project.LocalFileName.HasExtension(".uproject"))
 						{
-							CommandLine.AppendFormat(" \"{0}\"", Project.LocalFileName);
+							commandLine.AppendFormat(" \"{0}\"", project.LocalFileName);
 						}
 					}
-					CommandLine.Append(" -progress");
-					Logger.LogInformation("Generating project files...");
-					Logger.LogInformation("gpf> Running {Arguments}", CommandLine);
+					commandLine.Append(" -progress");
+					logger.LogInformation("Generating project files...");
+					logger.LogInformation("gpf> Running {Arguments}", commandLine);
 
-					int GenerateProjectFilesResult = await ExecuteShellCommandAsync(CommandLine.ToString(), null, Line => ProcessOutput(Line, "gpf> ", Progress, Logger), CancellationToken);
-					if (GenerateProjectFilesResult != 0)
+					int generateProjectFilesResult = await ExecuteShellCommandAsync(commandLine.ToString(), null, line => ProcessOutput(line, "gpf> ", Progress, logger), cancellationToken);
+					if (generateProjectFilesResult != 0)
 					{
-						return (WorkspaceUpdateResult.FailedToCompile, $"Failed to generate project files (exit code {GenerateProjectFilesResult}).");
+						return (WorkspaceUpdateResult.FailedToCompile, $"Failed to generate project files (exit code {generateProjectFilesResult}).");
 					}
 
-					Logger.LogInformation("");
-					Times.Add(new Tuple<string, TimeSpan>("Prj gen", Stopwatch.Stop("Success")));
+					logger.LogInformation("");
+					times.Add(new Tuple<string, TimeSpan>("Prj gen", stopwatch.Stop("Success")));
 				}
 			}
 
 			// Build everything using MegaXGE
 			if (Context.Options.HasFlag(WorkspaceUpdateOptions.Build))
 			{
-				Context.ProjectConfigFile ??= await ReadProjectConfigFile(Project.LocalRootPath, Project.LocalFileName, Logger);
+				Context.ProjectConfigFile ??= await ReadProjectConfigFile(project.LocalRootPath, project.LocalFileName, logger);
 
 				// Figure out the new editor target
-				TargetReceipt DefaultEditorReceipt = ConfigUtils.CreateDefaultEditorReceipt(Project, Context.ProjectConfigFile, Context.EditorConfig);
+				TargetReceipt defaultEditorReceipt = ConfigUtils.CreateDefaultEditorReceipt(project, Context.ProjectConfigFile, Context.EditorConfig);
 
-				FileReference EditorTargetFile = ConfigUtils.GetEditorTargetFile(Project, Context.ProjectConfigFile);
-				string EditorTargetName = EditorTargetFile.GetFileNameWithoutAnyExtensions();
-				FileReference EditorReceiptFile = ConfigUtils.GetReceiptFile(Project, EditorTargetFile, Context.EditorConfig.ToString());
+				FileReference editorTargetFile = ConfigUtils.GetEditorTargetFile(project, Context.ProjectConfigFile);
+				string editorTargetName = editorTargetFile.GetFileNameWithoutAnyExtensions();
+				FileReference editorReceiptFile = ConfigUtils.GetReceiptFile(project, editorTargetFile, Context.EditorConfig.ToString());
 
 				// Get the build steps
-				bool UsingPrecompiledEditor = Context.ArchiveTypeToArchive.TryGetValue(IArchiveInfo.EditorArchiveType, out Tuple<IArchiveInfo, string>? ArchiveInfo) && ArchiveInfo != null;
-				Dictionary<Guid, ConfigObject> BuildStepObjects = ConfigUtils.GetDefaultBuildStepObjects(Project, EditorTargetName, Context.EditorConfig, Context.ProjectConfigFile, UsingPrecompiledEditor);
-				BuildStep.MergeBuildStepObjects(BuildStepObjects, Context.ProjectConfigFile.GetValues("Build.Step", new string[0]).Select(x => new ConfigObject(x)));
-				BuildStep.MergeBuildStepObjects(BuildStepObjects, Context.UserBuildStepObjects);
+				bool usingPrecompiledEditor = Context.ArchiveTypeToArchive.TryGetValue(IArchiveInfo.EditorArchiveType, out Tuple<IArchiveInfo, string>? archiveInfo) && archiveInfo != null;
+				Dictionary<Guid, ConfigObject> buildStepObjects = ConfigUtils.GetDefaultBuildStepObjects(project, editorTargetName, Context.EditorConfig, Context.ProjectConfigFile, usingPrecompiledEditor);
+				BuildStep.MergeBuildStepObjects(buildStepObjects, Context.ProjectConfigFile.GetValues("Build.Step", new string[0]).Select(x => new ConfigObject(x)));
+				BuildStep.MergeBuildStepObjects(buildStepObjects, Context.UserBuildStepObjects);
 
 				// Construct build steps from them
-				List<BuildStep> BuildSteps = BuildStepObjects.Values.Select(x => new BuildStep(x)).OrderBy(x => (x.OrderIndex == -1) ? 10000 : x.OrderIndex).ToList();
+				List<BuildStep> buildSteps = buildStepObjects.Values.Select(x => new BuildStep(x)).OrderBy(x => (x.OrderIndex == -1) ? 10000 : x.OrderIndex).ToList();
 				if (Context.CustomBuildSteps != null && Context.CustomBuildSteps.Count > 0)
 				{
-					BuildSteps.RemoveAll(x => !Context.CustomBuildSteps.Contains(x.UniqueId));
+					buildSteps.RemoveAll(x => !Context.CustomBuildSteps.Contains(x.UniqueId));
 				}
 				else if (Context.Options.HasFlag(WorkspaceUpdateOptions.ScheduledBuild))
 				{
-					BuildSteps.RemoveAll(x => !x.bScheduledSync);
+					buildSteps.RemoveAll(x => !x.ScheduledSync);
 				}
 				else
 				{
-					BuildSteps.RemoveAll(x => !x.bNormalSync);
+					buildSteps.RemoveAll(x => !x.NormalSync);
 				}
 
 				// Check if the last successful build was before a change that we need to force a clean for
-				bool bForceClean = false;
-				if (State.LastBuiltChangeNumber != 0)
+				bool forceClean = false;
+				if (state.LastBuiltChangeNumber != 0)
 				{
-					foreach (string CleanBuildChange in Context.ProjectConfigFile.GetValues("ForceClean.Changelist", new string[0]))
+					foreach (string cleanBuildChange in Context.ProjectConfigFile.GetValues("ForceClean.Changelist", new string[0]))
 					{
-						int ChangeNumber;
-						if (int.TryParse(CleanBuildChange, out ChangeNumber))
+						int changeNumber;
+						if (int.TryParse(cleanBuildChange, out changeNumber))
 						{
-							if ((State.LastBuiltChangeNumber >= ChangeNumber) != (State.CurrentChangeNumber >= ChangeNumber))
+							if ((state.LastBuiltChangeNumber >= changeNumber) != (state.CurrentChangeNumber >= changeNumber))
 							{
-								Logger.LogInformation("Forcing clean build due to changelist {Change}.", ChangeNumber);
-								Logger.LogInformation("");
-								bForceClean = true;
+								logger.LogInformation("Forcing clean build due to changelist {Change}.", changeNumber);
+								logger.LogInformation("");
+								forceClean = true;
 								break;
 							}
 						}
@@ -1233,271 +1233,271 @@ namespace UnrealGameSync
 				}
 
 				// Execute them all
-				using (TelemetryStopwatch Stopwatch = new TelemetryStopwatch("Workspace_Build", Project.TelemetryProjectIdentifier))
+				using (TelemetryStopwatch stopwatch = new TelemetryStopwatch("Workspace_Build", project.TelemetryProjectIdentifier))
 				{
 					Progress.Set("Starting build...", 0.0f);
 
 					// Execute all the steps
-					float MaxProgressFraction = 0.0f;
-					foreach (BuildStep Step in BuildSteps)
+					float maxProgressFraction = 0.0f;
+					foreach (BuildStep step in buildSteps)
 					{
-						MaxProgressFraction += (float)Step.EstimatedDuration / (float)Math.Max(BuildSteps.Sum(x => x.EstimatedDuration), 1);
+						maxProgressFraction += (float)step.EstimatedDuration / (float)Math.Max(buildSteps.Sum(x => x.EstimatedDuration), 1);
 
-						Progress.Set(Step.StatusText ?? "Executing build step");
-						Progress.Push(MaxProgressFraction);
+						Progress.Set(step.StatusText ?? "Executing build step");
+						Progress.Push(maxProgressFraction);
 
-						Logger.LogInformation("{Status}", Step.StatusText);
+						logger.LogInformation("{Status}", step.StatusText);
 
-						DirectoryReference BatchFilesDir = DirectoryReference.Combine(Project.LocalRootPath, "Engine", "Build", "BatchFiles");
+						DirectoryReference batchFilesDir = DirectoryReference.Combine(project.LocalRootPath, "Engine", "Build", "BatchFiles");
 						if(RuntimeInformation.IsOSPlatform(OSPlatform.OSX))
 						{
-							BatchFilesDir = DirectoryReference.Combine(BatchFilesDir, "Mac");
+							batchFilesDir = DirectoryReference.Combine(batchFilesDir, "Mac");
 						}
 						else if(RuntimeInformation.IsOSPlatform(OSPlatform.Linux))
 						{
-							BatchFilesDir = DirectoryReference.Combine(BatchFilesDir, "Linux");
+							batchFilesDir = DirectoryReference.Combine(batchFilesDir, "Linux");
 						}
 
-						if (Step.IsValid())
+						if (step.IsValid())
 						{
 							// Get the build variables for this step
-							TargetReceipt? EditorReceipt;
-							if (!ConfigUtils.TryReadEditorReceipt(Project, EditorReceiptFile, out EditorReceipt))
+							TargetReceipt? editorReceipt;
+							if (!ConfigUtils.TryReadEditorReceipt(project, editorReceiptFile, out editorReceipt))
 							{
-								EditorReceipt = DefaultEditorReceipt;
+								editorReceipt = defaultEditorReceipt;
 							}
-							Dictionary<string, string> Variables = ConfigUtils.GetWorkspaceVariables(Project, State.CurrentChangeNumber, State.CurrentCodeChangeNumber, EditorReceipt, Context.ProjectConfigFile, Context.AdditionalVariables);
+							Dictionary<string, string> variables = ConfigUtils.GetWorkspaceVariables(project, state.CurrentChangeNumber, state.CurrentCodeChangeNumber, editorReceipt, Context.ProjectConfigFile, Context.AdditionalVariables);
 
 							// Handle all the different types of steps
-							switch (Step.Type)
+							switch (step.Type)
 							{
 								case BuildStepType.Compile:
-									using (TelemetryStopwatch StepStopwatch = new TelemetryStopwatch("Workspace_Execute_Compile", Project.TelemetryProjectIdentifier))
+									using (TelemetryStopwatch stepStopwatch = new TelemetryStopwatch("Workspace_Execute_Compile", project.TelemetryProjectIdentifier))
 									{
-										StepStopwatch.AddData(new { Target = Step.Target });
+										stepStopwatch.AddData(new { Target = step.Target });
 
-										FileReference BuildBat = FileReference.Combine(BatchFilesDir, $"Build.{ShellScriptExt}");
+										FileReference buildBat = FileReference.Combine(batchFilesDir, $"Build.{ShellScriptExt}");
 
-										string CommandLine = $"\"{BuildBat}\" {Step.Target} {Step.Platform} {Step.Configuration} {Utility.ExpandVariables(Step.Arguments ?? "", Variables)} -NoHotReloadFromIDE";
-										if (Context.Options.HasFlag(WorkspaceUpdateOptions.Clean) || bForceClean)
+										string commandLine = $"\"{buildBat}\" {step.Target} {step.Platform} {step.Configuration} {Utility.ExpandVariables(step.Arguments ?? "", variables)} -NoHotReloadFromIDE";
+										if (Context.Options.HasFlag(WorkspaceUpdateOptions.Clean) || forceClean)
 										{
-											Logger.LogInformation("ubt> Running {Arguments}", CommandLine + " -clean");
-											await ExecuteShellCommandAsync(CommandLine + " -clean", null, Line => ProcessOutput(Line, "ubt> ", Progress, Logger), CancellationToken);
+											logger.LogInformation("ubt> Running {Arguments}", commandLine + " -clean");
+											await ExecuteShellCommandAsync(commandLine + " -clean", null, line => ProcessOutput(line, "ubt> ", Progress, logger), cancellationToken);
 										}
 
-										Logger.LogInformation("ubt> Running {FileName} {Arguments}", BuildBat, CommandLine + " -progress");
+										logger.LogInformation("ubt> Running {FileName} {Arguments}", buildBat, commandLine + " -progress");
 
-										int ResultFromBuild = await ExecuteShellCommandAsync(CommandLine + " -progress", null, Line => ProcessOutput(Line, "ubt> ", Progress, Logger), CancellationToken);
-										if (ResultFromBuild != 0)
+										int resultFromBuild = await ExecuteShellCommandAsync(commandLine + " -progress", null, line => ProcessOutput(line, "ubt> ", Progress, logger), cancellationToken);
+										if (resultFromBuild != 0)
 										{
-											StepStopwatch.Stop("Failed");
+											stepStopwatch.Stop("Failed");
 
-											WorkspaceUpdateResult Result;
-											if (await HasModifiedSourceFiles(Perforce, Project, CancellationToken) || Context.UserBuildStepObjects.Count > 0)
+											WorkspaceUpdateResult result;
+											if (await HasModifiedSourceFiles(perforce, project, cancellationToken) || Context.UserBuildStepObjects.Count > 0)
 											{
-												Result = WorkspaceUpdateResult.FailedToCompile;
+												result = WorkspaceUpdateResult.FailedToCompile;
 											}
 											else
 											{
-												Result = WorkspaceUpdateResult.FailedToCompileWithCleanWorkspace;
+												result = WorkspaceUpdateResult.FailedToCompileWithCleanWorkspace;
 											}
 
-											return (Result, $"Failed to compile {Step.Target}");
+											return (result, $"Failed to compile {step.Target}");
 										}
 
-										StepStopwatch.Stop("Success");
+										stepStopwatch.Stop("Success");
 									}
 									break;
 								case BuildStepType.Cook:
-									using (TelemetryStopwatch StepStopwatch = new TelemetryStopwatch("Workspace_Execute_Cook", Project.TelemetryProjectIdentifier))
+									using (TelemetryStopwatch stepStopwatch = new TelemetryStopwatch("Workspace_Execute_Cook", project.TelemetryProjectIdentifier))
 									{
-										StepStopwatch.AddData(new { Project = Path.GetFileNameWithoutExtension(Step.FileName) });
+										stepStopwatch.AddData(new { Project = Path.GetFileNameWithoutExtension(step.FileName) });
 
-										FileReference LocalRunUAT = FileReference.Combine(BatchFilesDir, $"RunUAT.{ShellScriptExt}");
-										string Arguments = String.Format("\"{0}\" -profile=\"{1}\"", LocalRunUAT, FileReference.Combine(Project.LocalRootPath, Step.FileName ?? "unknown"));
-										Logger.LogInformation("uat> Running {FileName} {Argument}", LocalRunUAT, Arguments);
+										FileReference localRunUat = FileReference.Combine(batchFilesDir, $"RunUAT.{ShellScriptExt}");
+										string arguments = String.Format("\"{0}\" -profile=\"{1}\"", localRunUat, FileReference.Combine(project.LocalRootPath, step.FileName ?? "unknown"));
+										logger.LogInformation("uat> Running {FileName} {Argument}", localRunUat, arguments);
 
-										int ResultFromUAT = await ExecuteShellCommandAsync(Arguments, null, Line => ProcessOutput(Line, "uat> ", Progress, Logger), CancellationToken);
-										if (ResultFromUAT != 0)
+										int resultFromUat = await ExecuteShellCommandAsync(arguments, null, line => ProcessOutput(line, "uat> ", Progress, logger), cancellationToken);
+										if (resultFromUat != 0)
 										{
-											StepStopwatch.Stop("Failed");
-											return (WorkspaceUpdateResult.FailedToCompile, $"Cook failed. ({ResultFromUAT})");
+											stepStopwatch.Stop("Failed");
+											return (WorkspaceUpdateResult.FailedToCompile, $"Cook failed. ({resultFromUat})");
 										}
 
-										StepStopwatch.Stop("Success");
+										stepStopwatch.Stop("Success");
 									}
 									break;
 								case BuildStepType.Other:
-									using (TelemetryStopwatch StepStopwatch = new TelemetryStopwatch("Workspace_Execute_Custom", Project.TelemetryProjectIdentifier))
+									using (TelemetryStopwatch stepStopwatch = new TelemetryStopwatch("Workspace_Execute_Custom", project.TelemetryProjectIdentifier))
 									{
-										StepStopwatch.AddData(new { FileName = Path.GetFileNameWithoutExtension(Step.FileName) });
+										stepStopwatch.AddData(new { FileName = Path.GetFileNameWithoutExtension(step.FileName) });
 
-										FileReference ToolFileName = FileReference.Combine(Project.LocalRootPath, Utility.ExpandVariables(Step.FileName ?? "unknown", Variables));
-										string ToolWorkingDir = String.IsNullOrWhiteSpace(Step.WorkingDir) ? ToolFileName.Directory.FullName : Utility.ExpandVariables(Step.WorkingDir, Variables);
-										string ToolArguments = Utility.ExpandVariables(Step.Arguments ?? "", Variables);
-										Logger.LogInformation("tool> Running {0} {1}", ToolFileName, ToolArguments);
+										FileReference toolFileName = FileReference.Combine(project.LocalRootPath, Utility.ExpandVariables(step.FileName ?? "unknown", variables));
+										string toolWorkingDir = String.IsNullOrWhiteSpace(step.WorkingDir) ? toolFileName.Directory.FullName : Utility.ExpandVariables(step.WorkingDir, variables);
+										string toolArguments = Utility.ExpandVariables(step.Arguments ?? "", variables);
+										logger.LogInformation("tool> Running {0} {1}", toolFileName, toolArguments);
 
-										if (Step.bUseLogWindow)
+										if (step.UseLogWindow)
 										{
-											int ResultFromTool = await Utility.ExecuteProcessAsync(ToolFileName.FullName, ToolWorkingDir, ToolArguments, Line => ProcessOutput(Line, "tool> ", Progress, Logger), CancellationToken);
-											if (ResultFromTool != 0)
+											int resultFromTool = await Utility.ExecuteProcessAsync(toolFileName.FullName, toolWorkingDir, toolArguments, line => ProcessOutput(line, "tool> ", Progress, logger), cancellationToken);
+											if (resultFromTool != 0)
 											{
-												StepStopwatch.Stop("Failed");
-												return (WorkspaceUpdateResult.FailedToCompile, $"Tool terminated with exit code {ResultFromTool}.");
+												stepStopwatch.Stop("Failed");
+												return (WorkspaceUpdateResult.FailedToCompile, $"Tool terminated with exit code {resultFromTool}.");
 											}
 										}
 										else
 										{
-											ProcessStartInfo StartInfo = new ProcessStartInfo(ToolFileName.FullName, ToolArguments);
-											StartInfo.WorkingDirectory = ToolWorkingDir;
-											using (Process.Start(StartInfo))
+											ProcessStartInfo startInfo = new ProcessStartInfo(toolFileName.FullName, toolArguments);
+											startInfo.WorkingDirectory = toolWorkingDir;
+											using (Process.Start(startInfo))
 											{
 											}
 										}
 
-										StepStopwatch.Stop("Success");
+										stepStopwatch.Stop("Success");
 									}
 									break;
 							}
 						}
 
-						Logger.LogInformation("");
+						logger.LogInformation("");
 						Progress.Pop();
 					}
 
-					Times.Add(new Tuple<string, TimeSpan>("Build", Stopwatch.Stop("Success")));
+					times.Add(new Tuple<string, TimeSpan>("Build", stopwatch.Stop("Success")));
 				}
 
 				// Update the last successful build change number
 				if (Context.CustomBuildSteps == null || Context.CustomBuildSteps.Count == 0)
 				{
-					State.LastBuiltChangeNumber = State.CurrentChangeNumber;
-					State.Save(Logger);
+					state.LastBuiltChangeNumber = state.CurrentChangeNumber;
+					state.Save(logger);
 				}
 			}
 
 			// Write out all the timing information
-			Logger.LogInformation("Total time: {TotalTime}", FormatTime(Times.Sum(x => (long)(x.Item2.TotalMilliseconds / 1000))));
-			foreach (Tuple<string, TimeSpan> Time in Times)
+			logger.LogInformation("Total time: {TotalTime}", FormatTime(times.Sum(x => (long)(x.Item2.TotalMilliseconds / 1000))));
+			foreach (Tuple<string, TimeSpan> time in times)
 			{
-				Logger.LogInformation("   {Name,-8}: {TimeSeconds}", Time.Item1, FormatTime((long)(Time.Item2.TotalMilliseconds / 1000)));
+				logger.LogInformation("   {Name,-8}: {TimeSeconds}", time.Item1, FormatTime((long)(time.Item2.TotalMilliseconds / 1000)));
 			}
-			if (NumFilesSynced > 0)
+			if (numFilesSynced > 0)
 			{
-				Logger.LogInformation("{NumFiles} files synced.", NumFilesSynced);
+				logger.LogInformation("{NumFiles} files synced.", numFilesSynced);
 			}
 
-			DateTime FinishTime = DateTime.Now;
-			Logger.LogInformation("");
-			Logger.LogInformation("UPDATE SUCCEEDED ({FinishDate} {FinishTime})", FinishTime.ToShortDateString(), FinishTime.ToShortTimeString());
+			DateTime finishTime = DateTime.Now;
+			logger.LogInformation("");
+			logger.LogInformation("UPDATE SUCCEEDED ({FinishDate} {FinishTime})", finishTime.ToShortDateString(), finishTime.ToShortTimeString());
 
 			return (WorkspaceUpdateResult.Success, "Update succeeded");
 		}
 
-		static void ProcessOutput(string Line, string Prefix, ProgressValue Progress, ILogger Logger)
+		static void ProcessOutput(string line, string prefix, ProgressValue progress, ILogger logger)
 		{
-			string? ParsedLine = ProgressTextWriter.ParseLine(Line, Progress);
-			if (ParsedLine != null)
+			string? parsedLine = ProgressTextWriter.ParseLine(line, progress);
+			if (parsedLine != null)
 			{
-				Logger.LogInformation("{Prefix}{Message}", Prefix, ParsedLine);
+				logger.LogInformation("{Prefix}{Message}", prefix, parsedLine);
 			}
 		}
 
-		static async Task<bool> IsLicenseeVersion(IPerforceConnection Perforce, ProjectInfo Project, ILogger Logger, CancellationToken CancellationToken)
+		static async Task<bool> IsLicenseeVersion(IPerforceConnection perforce, ProjectInfo project, ILogger logger, CancellationToken cancellationToken)
 		{
-			string[] Files =
+			string[] files =
 			{
-				Project.ClientRootPath + "/Engine/Build/NotForLicensees/EpicInternal.txt",
-				Project.ClientRootPath + "/Engine/Restricted/NotForLicensees/Build/EpicInternal.txt"
+				project.ClientRootPath + "/Engine/Build/NotForLicensees/EpicInternal.txt",
+				project.ClientRootPath + "/Engine/Restricted/NotForLicensees/Build/EpicInternal.txt"
 			};
 
-			List<FStatRecord> Records = await Perforce.FStatAsync(Files, CancellationToken).ToListAsync(CancellationToken);
-			return !Records.Any(x => x.IsMapped);
+			List<FStatRecord> records = await perforce.FStatAsync(files, cancellationToken).ToListAsync(cancellationToken);
+			return !records.Any(x => x.IsMapped);
 		}
 
-		public static List<string> GetSyncPaths(ProjectInfo Project, bool bSyncAllProjects, string[] SyncFilter)
+		public static List<string> GetSyncPaths(ProjectInfo project, bool syncAllProjects, string[] syncFilter)
 		{
-			List<string> SyncPaths = GetRelativeSyncPaths(Project, bSyncAllProjects, SyncFilter);
-			return SyncPaths.Select(x => Project.ClientRootPath + x).ToList();
+			List<string> syncPaths = GetRelativeSyncPaths(project, syncAllProjects, syncFilter);
+			return syncPaths.Select(x => project.ClientRootPath + x).ToList();
 		}
 
-		public static List<string> GetRelativeSyncPaths(ProjectInfo Project, bool bSyncAllProjects, string[]? SyncFilter)
+		public static List<string> GetRelativeSyncPaths(ProjectInfo project, bool syncAllProjects, string[]? syncFilter)
 		{
-			List<string> SyncPaths = new List<string>();
+			List<string> syncPaths = new List<string>();
 
 			// Check the client path is formatted correctly
-			if (!Project.ClientFileName.StartsWith(Project.ClientRootPath + "/"))
+			if (!project.ClientFileName.StartsWith(project.ClientRootPath + "/"))
 			{
-				throw new Exception(String.Format("Expected '{0}' to start with '{1}'", Project.ClientFileName, Project.ClientRootPath));
+				throw new Exception(String.Format("Expected '{0}' to start with '{1}'", project.ClientFileName, project.ClientRootPath));
 			}
 
 			// Add the default project paths
-			int LastSlashIdx = Project.ClientFileName.LastIndexOf('/');
-			if (bSyncAllProjects || !Project.ClientFileName.EndsWith(".uproject", StringComparison.OrdinalIgnoreCase) || LastSlashIdx <= Project.ClientRootPath.Length)
+			int lastSlashIdx = project.ClientFileName.LastIndexOf('/');
+			if (syncAllProjects || !project.ClientFileName.EndsWith(".uproject", StringComparison.OrdinalIgnoreCase) || lastSlashIdx <= project.ClientRootPath.Length)
 			{
-				SyncPaths.Add("/...");
+				syncPaths.Add("/...");
 			}
 			else
 			{
-				SyncPaths.Add("/*");
-				SyncPaths.Add("/Engine/...");
-				if (Project.bIsEnterpriseProject)
+				syncPaths.Add("/*");
+				syncPaths.Add("/Engine/...");
+				if (project.IsEnterpriseProject)
 				{
-					SyncPaths.Add("/Enterprise/...");
+					syncPaths.Add("/Enterprise/...");
 				}
-				SyncPaths.Add(Project.ClientFileName.Substring(Project.ClientRootPath.Length, LastSlashIdx - Project.ClientRootPath.Length) + "/...");
+				syncPaths.Add(project.ClientFileName.Substring(project.ClientRootPath.Length, lastSlashIdx - project.ClientRootPath.Length) + "/...");
 			}
 
 			// Apply the sync filter to that list. We only want inclusive rules in the output list, but we can manually apply exclusions to previous entries.
-			if (SyncFilter != null)
+			if (syncFilter != null)
 			{
-				foreach (string SyncPath in SyncFilter)
+				foreach (string syncPath in syncFilter)
 				{
-					string TrimSyncPath = SyncPath.Trim();
-					if (TrimSyncPath.StartsWith("/"))
+					string trimSyncPath = syncPath.Trim();
+					if (trimSyncPath.StartsWith("/"))
 					{
-						SyncPaths.Add(TrimSyncPath);
+						syncPaths.Add(trimSyncPath);
 					}
-					else if (TrimSyncPath.StartsWith("-/") && TrimSyncPath.EndsWith("..."))
+					else if (trimSyncPath.StartsWith("-/") && trimSyncPath.EndsWith("..."))
 					{
-						SyncPaths.RemoveAll(x => x.StartsWith(TrimSyncPath.Substring(1, TrimSyncPath.Length - 4)));
+						syncPaths.RemoveAll(x => x.StartsWith(trimSyncPath.Substring(1, trimSyncPath.Length - 4)));
 					}
 				}
 			}
 
 			// Sort the remaining paths by length, and remove any paths which are included twice
-			SyncPaths = SyncPaths.OrderBy(x => x.Length).ToList();
-			for (int Idx = 0; Idx < SyncPaths.Count; Idx++)
+			syncPaths = syncPaths.OrderBy(x => x.Length).ToList();
+			for (int idx = 0; idx < syncPaths.Count; idx++)
 			{
-				string SyncPath = SyncPaths[Idx];
-				if (SyncPath.EndsWith("..."))
+				string syncPath = syncPaths[idx];
+				if (syncPath.EndsWith("..."))
 				{
-					string SyncPathPrefix = SyncPath.Substring(0, SyncPath.Length - 3);
-					for (int OtherIdx = SyncPaths.Count - 1; OtherIdx > Idx; OtherIdx--)
+					string syncPathPrefix = syncPath.Substring(0, syncPath.Length - 3);
+					for (int otherIdx = syncPaths.Count - 1; otherIdx > idx; otherIdx--)
 					{
-						if (SyncPaths[OtherIdx].StartsWith(SyncPathPrefix))
+						if (syncPaths[otherIdx].StartsWith(syncPathPrefix))
 						{
-							SyncPaths.RemoveAt(OtherIdx);
+							syncPaths.RemoveAt(otherIdx);
 						}
 					}
 				}
 			}
 
-			return SyncPaths;
+			return syncPaths;
 		}
 
-		public static bool MatchFilter(ProjectInfo Project, FileReference FileName, FileFilter Filter)
+		public static bool MatchFilter(ProjectInfo project, FileReference fileName, FileFilter filter)
 		{
-			bool bMatch = true;
-			if (FileName.IsUnderDirectory(Project.LocalRootPath))
+			bool match = true;
+			if (fileName.IsUnderDirectory(project.LocalRootPath))
 			{
-				string RelativePath = FileName.MakeRelativeTo(Project.LocalRootPath);
-				if (!Filter.Matches(RelativePath))
+				string relativePath = fileName.MakeRelativeTo(project.LocalRootPath);
+				if (!filter.Matches(relativePath))
 				{
-					bMatch = false;
+					match = false;
 				}
 			}
-			return bMatch;
+			return match;
 		}
 
 		class SyncState
@@ -1508,119 +1508,119 @@ namespace UnrealGameSync
 			public string StatusMessage;
 			public WorkspaceUpdateResult Result = WorkspaceUpdateResult.Success;
 
-			public SyncState(HashSet<string> RemainingDepotPaths, Queue<List<string>> SyncCommandLists)
+			public SyncState(HashSet<string> remainingDepotPaths, Queue<List<string>> syncCommandLists)
 			{
-				this.TotalDepotPaths = RemainingDepotPaths.Count;
-				this.RemainingDepotPaths = RemainingDepotPaths;
-				this.SyncCommandLists = SyncCommandLists;
+				this.TotalDepotPaths = remainingDepotPaths.Count;
+				this.RemainingDepotPaths = remainingDepotPaths;
+				this.SyncCommandLists = syncCommandLists;
 				this.StatusMessage = "Succeeded.";
 			}
 		}
 
-		static Task<(WorkspaceUpdateResult, string)> SyncFileRevisions(IPerforceConnection Perforce, string Prefix, WorkspaceUpdateContext Context, List<string> SyncCommands, HashSet<string> RemainingDepotPaths, ProgressValue Progress, ILogger Logger, CancellationToken CancellationToken)
+		static Task<(WorkspaceUpdateResult, string)> SyncFileRevisions(IPerforceConnection perforce, string prefix, WorkspaceUpdateContext context, List<string> syncCommands, HashSet<string> remainingDepotPaths, ProgressValue progress, ILogger logger, CancellationToken cancellationToken)
 		{
-			Queue<List<string>> SyncCommandLists = new Queue<List<string>>();
-			SyncCommandLists.Enqueue(SyncCommands);
-			return SyncFileRevisions(Perforce, Prefix, Context, SyncCommandLists, RemainingDepotPaths, Progress, Logger, CancellationToken);
+			Queue<List<string>> syncCommandLists = new Queue<List<string>>();
+			syncCommandLists.Enqueue(syncCommands);
+			return SyncFileRevisions(perforce, prefix, context, syncCommandLists, remainingDepotPaths, progress, logger, cancellationToken);
 		}
 
-		static async Task<(WorkspaceUpdateResult, string)> SyncFileRevisions(IPerforceConnection Perforce, string Prefix, WorkspaceUpdateContext Context, Queue<List<string>> SyncCommandLists, HashSet<string> RemainingDepotPaths, ProgressValue Progress, ILogger Logger, CancellationToken CancellationToken)
+		static async Task<(WorkspaceUpdateResult, string)> SyncFileRevisions(IPerforceConnection perforce, string prefix, WorkspaceUpdateContext context, Queue<List<string>> syncCommandLists, HashSet<string> remainingDepotPaths, ProgressValue progress, ILogger logger, CancellationToken cancellationToken)
 		{
 			// Figure out the number of additional background threads we want to run with. We can run worker on the current thread.
-			int NumThreads = Context.PerforceSyncOptions?.NumThreads ?? PerforceSyncOptions.DefaultNumThreads;
-			int NumExtraThreads = Math.Max(Math.Min(SyncCommandLists.Count, NumThreads) - 1, 0);
+			int numThreads = context.PerforceSyncOptions?.NumThreads ?? PerforceSyncOptions.DefaultNumThreads;
+			int numExtraThreads = Math.Max(Math.Min(syncCommandLists.Count, numThreads) - 1, 0);
 
-			List<IPerforceConnection> ChildConnections = new List<IPerforceConnection>();
-			List<Task> ChildTasks = new List<Task>(NumExtraThreads);
+			List<IPerforceConnection> childConnections = new List<IPerforceConnection>();
+			List<Task> childTasks = new List<Task>(numExtraThreads);
 			try
 			{
 				// Create the state object shared by all the worker threads
-				SyncState State = new SyncState(RemainingDepotPaths, SyncCommandLists);
+				SyncState state = new SyncState(remainingDepotPaths, syncCommandLists);
 
 				// Wrapper writer around the log class to prevent multiple threads writing to it at once
-				ILogger LogWrapper = Logger;
+				ILogger logWrapper = logger;
 
 				// Initialize Sync Progress
-				UpdateSyncState(Prefix, State, Progress);
+				UpdateSyncState(prefix, state, progress);
 
 				// Delegate for updating the sync state after a file has been synced
-				Action<SyncRecord, ILogger> SyncOutput = (Record, LocalLog) => { UpdateSyncState(Prefix, Record, State, Progress, LocalLog); };
+				Action<SyncRecord, ILogger> syncOutput = (record, localLog) => { UpdateSyncState(prefix, record, state, progress, localLog); };
 
 				// Create all the child threads
-				for (int ThreadIdx = 0; ThreadIdx < NumExtraThreads; ThreadIdx++)
+				for (int threadIdx = 0; threadIdx < numExtraThreads; threadIdx++)
 				{
-					int ThreadNumber = ThreadIdx + 2;
+					int threadNumber = threadIdx + 2;
 
 					// Create connection
-					IPerforceConnection ChildConnection = await PerforceConnection.CreateAsync(Perforce.Settings, Perforce.Logger);
-					ChildConnections.Add(ChildConnection);
+					IPerforceConnection childConnection = await PerforceConnection.CreateAsync(perforce.Settings, perforce.Logger);
+					childConnections.Add(childConnection);
 
-					Task ChildTask = Task.Run(() => StaticSyncWorker(ThreadNumber, ChildConnection, Context, State, SyncOutput, LogWrapper, CancellationToken));
-					ChildTasks.Add(ChildTask);
+					Task childTask = Task.Run(() => StaticSyncWorker(threadNumber, childConnection, context, state, syncOutput, logWrapper, cancellationToken));
+					childTasks.Add(childTask);
 				}
 
 				// Run one worker on the current thread
-				await StaticSyncWorker(1, Perforce, Context, State, SyncOutput, LogWrapper, CancellationToken);
+				await StaticSyncWorker(1, perforce, context, state, syncOutput, logWrapper, cancellationToken);
 
 				// Allow the tasks to throw
-				foreach (Task ChildTask in ChildTasks)
+				foreach (Task childTask in childTasks)
 				{
-					await ChildTask;
+					await childTask;
 				}
 
 				// Return the result that was set on the state object
-				return (State.Result, State.StatusMessage);
+				return (state.Result, state.StatusMessage);
 			}
 			finally
 			{
-				foreach (Task ChildTask in ChildTasks)
+				foreach (Task childTask in childTasks)
 				{
-					await ChildTask.ContinueWith(_ => { }); // Swallow exceptions until we've disposed the connections
+					await childTask.ContinueWith(_ => { }); // Swallow exceptions until we've disposed the connections
 				}
 
-				foreach (IPerforceConnection ChildConnection in ChildConnections)
+				foreach (IPerforceConnection childConnection in childConnections)
 				{
-					ChildConnection.Dispose();
+					childConnection.Dispose();
 				}
 			}
 		}
 
-		static void UpdateSyncState(string Prefix, SyncState State, ProgressValue Progress)
+		static void UpdateSyncState(string prefix, SyncState state, ProgressValue progress)
 		{
-			lock (State)
+			lock (state)
 			{
-				string Message = String.Format("{0} ({1}/{2})", Prefix, State.TotalDepotPaths - State.RemainingDepotPaths.Count, State.TotalDepotPaths);
-				float Fraction = Math.Min((float)(State.TotalDepotPaths - State.RemainingDepotPaths.Count) / (float)State.TotalDepotPaths, 1.0f);
-				Progress.Set(Message, Fraction);
+				string message = String.Format("{0} ({1}/{2})", prefix, state.TotalDepotPaths - state.RemainingDepotPaths.Count, state.TotalDepotPaths);
+				float fraction = Math.Min((float)(state.TotalDepotPaths - state.RemainingDepotPaths.Count) / (float)state.TotalDepotPaths, 1.0f);
+				progress.Set(message, fraction);
 			}
 		}
 
-		static void UpdateSyncState(string Prefix, SyncRecord Record, SyncState State, ProgressValue Progress, ILogger Logger)
+		static void UpdateSyncState(string prefix, SyncRecord record, SyncState state, ProgressValue progress, ILogger logger)
 		{
-			lock (State)
+			lock (state)
 			{
-				State.RemainingDepotPaths.Remove(Record.DepotFile.ToString());
+				state.RemainingDepotPaths.Remove(record.DepotFile.ToString());
 
-				string Message = String.Format("{0} ({1}/{2})", Prefix, State.TotalDepotPaths - State.RemainingDepotPaths.Count, State.TotalDepotPaths);
-				float Fraction = Math.Min((float)(State.TotalDepotPaths - State.RemainingDepotPaths.Count) / (float)State.TotalDepotPaths, 1.0f);
-				Progress.Set(Message, Fraction);
+				string message = String.Format("{0} ({1}/{2})", prefix, state.TotalDepotPaths - state.RemainingDepotPaths.Count, state.TotalDepotPaths);
+				float fraction = Math.Min((float)(state.TotalDepotPaths - state.RemainingDepotPaths.Count) / (float)state.TotalDepotPaths, 1.0f);
+				progress.Set(message, fraction);
 
-				Logger.LogInformation("p4>   {Action} {Path}", Record.Action, Record.Path);
+				logger.LogInformation("p4>   {Action} {Path}", record.Action, record.Path);
 			}
 		}
 
-		static async Task StaticSyncWorker(int ThreadNumber, IPerforceConnection Perforce, WorkspaceUpdateContext Context, SyncState State, Action<SyncRecord, ILogger> SyncOutput, ILogger GlobalLog, CancellationToken CancellationToken)
+		static async Task StaticSyncWorker(int threadNumber, IPerforceConnection perforce, WorkspaceUpdateContext context, SyncState state, Action<SyncRecord, ILogger> syncOutput, ILogger globalLog, CancellationToken cancellationToken)
 		{
-			PrefixedTextWriter ThreadLog = new PrefixedTextWriter(String.Format("{0}:", ThreadNumber), GlobalLog);
+			PrefixedTextWriter threadLog = new PrefixedTextWriter(String.Format("{0}:", threadNumber), globalLog);
 			for (; ; )
 			{
 				// Remove the next batch that needs to be synced
-				List<string> SyncCommands;
-				lock (State)
+				List<string> syncCommands;
+				lock (state)
 				{
-					if (State.Result == WorkspaceUpdateResult.Success && State.SyncCommandLists.Count > 0)
+					if (state.Result == WorkspaceUpdateResult.Success && state.SyncCommandLists.Count > 0)
 					{
-						SyncCommands = State.SyncCommandLists.Dequeue();
+						syncCommands = state.SyncCommandLists.Dequeue();
 					}
 					else
 					{
@@ -1628,31 +1628,31 @@ namespace UnrealGameSync
 					}
 				}
 
-				WorkspaceUpdateResult Result = WorkspaceUpdateResult.FailedToSync;
-				string StatusMessage = "";
+				WorkspaceUpdateResult result = WorkspaceUpdateResult.FailedToSync;
+				string statusMessage = "";
 
-				int Retries = (null != Context.PerforceSyncOptions) ? Context.PerforceSyncOptions!.NumSyncErrorRetries : 0;
+				int retries = (null != context.PerforceSyncOptions) ? context.PerforceSyncOptions!.NumSyncErrorRetries : 0;
 
-				while (Retries >= 0 && WorkspaceUpdateResult.FailedToSync == Result)
+				while (retries >= 0 && WorkspaceUpdateResult.FailedToSync == result)
 				{
 					// Sync the files
-					(Result, StatusMessage) = await StaticSyncFileRevisions(Perforce, Context, SyncCommands, Record => SyncOutput(Record, ThreadLog), CancellationToken);
+					(result, statusMessage) = await StaticSyncFileRevisions(perforce, context, syncCommands, record => syncOutput(record, threadLog), cancellationToken);
 
-					if (WorkspaceUpdateResult.FailedToSync == Result && --Retries >= 0)
+					if (WorkspaceUpdateResult.FailedToSync == result && --retries >= 0)
 					{
-						ThreadLog.LogWarning("Sync Errors occurred.  Retrying: Remaining retries " + Retries.ToString());
+						threadLog.LogWarning("Sync Errors occurred.  Retrying: Remaining retries " + retries.ToString());
 					}
 				}
 
 				// If it failed, try to set it on the state if nothing else has failed first
-				if (Result != WorkspaceUpdateResult.Success)
+				if (result != WorkspaceUpdateResult.Success)
 				{
-					lock (State)
+					lock (state)
 					{
-						if (State.Result == WorkspaceUpdateResult.Success)
+						if (state.Result == WorkspaceUpdateResult.Success)
 						{
-							State.Result = Result;
-							State.StatusMessage = StatusMessage;
+							state.Result = result;
+							state.StatusMessage = statusMessage;
 						}
 					}
 					break;
@@ -1660,68 +1660,68 @@ namespace UnrealGameSync
 			}
 		}
 
-		static async Task<(WorkspaceUpdateResult, string)> StaticSyncFileRevisions(IPerforceConnection Perforce, WorkspaceUpdateContext Context, List<string> SyncCommands, Action<SyncRecord> SyncOutput, CancellationToken CancellationToken)
+		static async Task<(WorkspaceUpdateResult, string)> StaticSyncFileRevisions(IPerforceConnection perforce, WorkspaceUpdateContext context, List<string> syncCommands, Action<SyncRecord> syncOutput, CancellationToken cancellationToken)
 		{
 			// Sync them all
-			List<PerforceResponse<SyncRecord>> Responses = await Perforce.TrySyncAsync(SyncOptions.None, -1, SyncCommands, CancellationToken).ToListAsync(CancellationToken);
+			List<PerforceResponse<SyncRecord>> responses = await perforce.TrySyncAsync(SyncOptions.None, -1, syncCommands, cancellationToken).ToListAsync(cancellationToken);
 
-			List<string> TamperedFiles = new List<string>();
-			foreach (PerforceResponse<SyncRecord> Response in Responses)
+			List<string> tamperedFiles = new List<string>();
+			foreach (PerforceResponse<SyncRecord> response in responses)
 			{
-				const string NoClobberPrefix = "Can't clobber writable file ";
-				if (Response.Info != null)
+				const string noClobberPrefix = "Can't clobber writable file ";
+				if (response.Info != null)
 				{
 					continue;
 				}
-				else if (Response.Succeeded)
+				else if (response.Succeeded)
 				{
-					SyncOutput(Response.Data);
+					syncOutput(response.Data);
 				}
-				else if (Response.Error != null && Response.Error.Generic == PerforceGenericCode.Client && Response.Error.Data.StartsWith(NoClobberPrefix, StringComparison.OrdinalIgnoreCase))
+				else if (response.Error != null && response.Error.Generic == PerforceGenericCode.Client && response.Error.Data.StartsWith(noClobberPrefix, StringComparison.OrdinalIgnoreCase))
 				{
-					TamperedFiles.Add(Response.Error.Data.Substring(NoClobberPrefix.Length).Trim());
+					tamperedFiles.Add(response.Error.Data.Substring(noClobberPrefix.Length).Trim());
 				}
 				else
 				{
 					return (WorkspaceUpdateResult.FailedToSync, "Aborted sync due to errors.. Currently retries on sync error is set at " +
-						((null != Context.PerforceSyncOptions) ? Context.PerforceSyncOptions!.NumSyncErrorRetries : 0).ToString() +
+						((null != context.PerforceSyncOptions) ? context.PerforceSyncOptions!.NumSyncErrorRetries : 0).ToString() +
 						" in Options->Application Settings...->Advanced.  You might want to set it higher if you are on a bad connection.");
 				}
 			}
 
 			// If any files need to be clobbered, defer to the main thread to figure out which ones
-			if (TamperedFiles.Count > 0)
+			if (tamperedFiles.Count > 0)
 			{
-				if ((Context.Options & WorkspaceUpdateOptions.Clobber) == 0)
+				if ((context.Options & WorkspaceUpdateOptions.Clobber) == 0)
 				{
-					int NumNewFilesToClobber = 0;
-					foreach (string TamperedFile in TamperedFiles)
+					int numNewFilesToClobber = 0;
+					foreach (string tamperedFile in tamperedFiles)
 					{
-						if (!Context.ClobberFiles.ContainsKey(TamperedFile))
+						if (!context.ClobberFiles.ContainsKey(tamperedFile))
 						{
-							Context.ClobberFiles[TamperedFile] = true;
-							if (TamperedFile.EndsWith(LocalObjectVersionFileName, StringComparison.OrdinalIgnoreCase) || TamperedFile.EndsWith(LocalVersionHeaderFileName, StringComparison.OrdinalIgnoreCase))
+							context.ClobberFiles[tamperedFile] = true;
+							if (tamperedFile.EndsWith(LocalObjectVersionFileName, StringComparison.OrdinalIgnoreCase) || tamperedFile.EndsWith(LocalVersionHeaderFileName, StringComparison.OrdinalIgnoreCase))
 							{
 								// Hack for UseFastModularVersioningV2; we don't need to update these files any more.
 								continue;
 							}
-							NumNewFilesToClobber++;
+							numNewFilesToClobber++;
 						}
 					}
-					if (NumNewFilesToClobber > 0)
+					if (numNewFilesToClobber > 0)
 					{
-						return (WorkspaceUpdateResult.FilesToClobber, $"Cancelled sync after checking files to clobber ({NumNewFilesToClobber} new files).");
+						return (WorkspaceUpdateResult.FilesToClobber, $"Cancelled sync after checking files to clobber ({numNewFilesToClobber} new files).");
 					}
 				}
-				foreach (string TamperedFile in TamperedFiles)
+				foreach (string tamperedFile in tamperedFiles)
 				{
-					bool bShouldClobber = (Context.Options & WorkspaceUpdateOptions.Clobber) != 0 || Context.ClobberFiles[TamperedFile];
-					if (bShouldClobber)
+					bool shouldClobber = (context.Options & WorkspaceUpdateOptions.Clobber) != 0 || context.ClobberFiles[tamperedFile];
+					if (shouldClobber)
 					{
-						List<PerforceResponse<SyncRecord>> Response = await Perforce.TrySyncAsync(SyncOptions.Force, -1, TamperedFile, CancellationToken).ToListAsync();
-						if (!Response.Succeeded())
+						List<PerforceResponse<SyncRecord>> response = await perforce.TrySyncAsync(SyncOptions.Force, -1, tamperedFile, cancellationToken).ToListAsync();
+						if (!response.Succeeded())
 						{
-							return (WorkspaceUpdateResult.FailedToSync, $"Couldn't sync {TamperedFile}.");
+							return (WorkspaceUpdateResult.FailedToSync, $"Couldn't sync {tamperedFile}.");
 						}
 					}
 				}
@@ -1731,190 +1731,190 @@ namespace UnrealGameSync
 			return (WorkspaceUpdateResult.Success, "Succeeded.");
 		}
 
-		public static async Task<ConfigFile> ReadProjectConfigFile(DirectoryReference LocalRootPath, FileReference SelectedLocalFileName, ILogger Logger)
+		public static async Task<ConfigFile> ReadProjectConfigFile(DirectoryReference localRootPath, FileReference selectedLocalFileName, ILogger logger)
 		{
 			// Find the valid config file paths
-			DirectoryInfo EngineDir = DirectoryReference.Combine(LocalRootPath, "Engine").ToDirectoryInfo();
-			List<FileInfo> LocalConfigFiles = Utility.GetLocalConfigPaths(EngineDir, SelectedLocalFileName.ToFileInfo());
+			DirectoryInfo engineDir = DirectoryReference.Combine(localRootPath, "Engine").ToDirectoryInfo();
+			List<FileInfo> localConfigFiles = Utility.GetLocalConfigPaths(engineDir, selectedLocalFileName.ToFileInfo());
 
 			// Read them in
-			ConfigFile ProjectConfig = new ConfigFile();
-			foreach (FileInfo LocalConfigFile in LocalConfigFiles)
+			ConfigFile projectConfig = new ConfigFile();
+			foreach (FileInfo localConfigFile in localConfigFiles)
 			{
 				try
 				{
-					string[] Lines = await File.ReadAllLinesAsync(LocalConfigFile.FullName);
-					ProjectConfig.Parse(Lines);
-					Logger.LogDebug("Read config file from {FileName}", LocalConfigFile.FullName);
+					string[] lines = await File.ReadAllLinesAsync(localConfigFile.FullName);
+					projectConfig.Parse(lines);
+					logger.LogDebug("Read config file from {FileName}", localConfigFile.FullName);
 				}
-				catch (Exception Ex)
+				catch (Exception ex)
 				{
-					Logger.LogWarning(Ex, "Failed to read config file from {FileName}", LocalConfigFile.FullName);
+					logger.LogWarning(ex, "Failed to read config file from {FileName}", localConfigFile.FullName);
 				}
 			}
-			return ProjectConfig;
+			return projectConfig;
 		}
 
-		public static async Task<IReadOnlyList<string>?> ReadProjectStreamFilter(IPerforceConnection Perforce, ConfigFile ProjectConfigFile, ILogger Logger, CancellationToken CancellationToken)
+		public static async Task<IReadOnlyList<string>?> ReadProjectStreamFilter(IPerforceConnection perforce, ConfigFile projectConfigFile, ILogger logger, CancellationToken cancellationToken)
 		{
-			string? StreamListDepotPath = ProjectConfigFile.GetValue("Options.QuickSelectStreamList", null);
-			if (StreamListDepotPath == null)
+			string? streamListDepotPath = projectConfigFile.GetValue("Options.QuickSelectStreamList", null);
+			if (streamListDepotPath == null)
 			{
 				return null;
 			}
 
-			PerforceResponse<PrintRecord<string[]>> Response = await Perforce.TryPrintLinesAsync(StreamListDepotPath, CancellationToken);
-			if (!Response.Succeeded)
+			PerforceResponse<PrintRecord<string[]>> response = await perforce.TryPrintLinesAsync(streamListDepotPath, cancellationToken);
+			if (!response.Succeeded)
 			{
 				return null;
 			}
 
-			return Response.Data.Contents.Select(x => x.Trim()).Where(x => x.Length > 0).ToList().AsReadOnly();
+			return response.Data.Contents.Select(x => x.Trim()).Where(x => x.Length > 0).ToList().AsReadOnly();
 		}
 
-		static string FormatTime(long Seconds)
+		static string FormatTime(long seconds)
 		{
-			if (Seconds >= 60)
+			if (seconds >= 60)
 			{
-				return String.Format("{0,3}m {1:00}s", Seconds / 60, Seconds % 60);
+				return String.Format("{0,3}m {1:00}s", seconds / 60, seconds % 60);
 			}
 			else
 			{
-				return String.Format("     {0,2}s", Seconds);
+				return String.Format("     {0,2}s", seconds);
 			}
 		}
 
-		static async Task<bool> HasModifiedSourceFiles(IPerforceConnection Perforce, ProjectInfo Project, CancellationToken CancellationToken)
+		static async Task<bool> HasModifiedSourceFiles(IPerforceConnection perforce, ProjectInfo project, CancellationToken cancellationToken)
 		{
-			List<OpenedRecord> OpenFiles = await Perforce.OpenedAsync(OpenedOptions.None, Project.ClientRootPath + "/...", CancellationToken).ToListAsync(CancellationToken);
-			if (OpenFiles.Any(x => x.DepotFile.IndexOf("/Source/", StringComparison.OrdinalIgnoreCase) != -1))
+			List<OpenedRecord> openFiles = await perforce.OpenedAsync(OpenedOptions.None, project.ClientRootPath + "/...", cancellationToken).ToListAsync(cancellationToken);
+			if (openFiles.Any(x => x.DepotFile.IndexOf("/Source/", StringComparison.OrdinalIgnoreCase) != -1))
 			{
 				return true;
 			}
 			return false;
 		}
 
-		static async Task<List<FStatRecord>> FindUnresolvedFiles(IPerforceConnection Perforce, IEnumerable<string> SyncPaths, CancellationToken CancellationToken)
+		static async Task<List<FStatRecord>> FindUnresolvedFiles(IPerforceConnection perforce, IEnumerable<string> syncPaths, CancellationToken cancellationToken)
 		{
-			List<FStatRecord> UnresolvedFiles = new List<FStatRecord>();
-			foreach (string SyncPath in SyncPaths)
+			List<FStatRecord> unresolvedFiles = new List<FStatRecord>();
+			foreach (string syncPath in syncPaths)
 			{
-				List<FStatRecord> Records = await Perforce.FStatAsync(FStatOptions.OnlyUnresolved, SyncPath, CancellationToken).ToListAsync(CancellationToken);
-				UnresolvedFiles.AddRange(Records);
+				List<FStatRecord> records = await perforce.FStatAsync(FStatOptions.OnlyUnresolved, syncPath, cancellationToken).ToListAsync(cancellationToken);
+				unresolvedFiles.AddRange(records);
 			}
-			return UnresolvedFiles;
+			return unresolvedFiles;
 		}
 
-		static Task<bool> UpdateVersionFile(IPerforceConnection Perforce, string ClientPath, Dictionary<string, string> VersionStrings, int ChangeNumber, ILogger Logger, CancellationToken CancellationToken)
+		static Task<bool> UpdateVersionFile(IPerforceConnection perforce, string clientPath, Dictionary<string, string> versionStrings, int changeNumber, ILogger logger, CancellationToken cancellationToken)
 		{
-			return UpdateVersionFile(Perforce, ClientPath, ChangeNumber, Text => UpdateVersionStrings(Text, VersionStrings), Logger, CancellationToken);
+			return UpdateVersionFile(perforce, clientPath, changeNumber, text => UpdateVersionStrings(text, versionStrings), logger, cancellationToken);
 		}
 
-		static async Task<bool> UpdateVersionFile(IPerforceConnection Perforce, string ClientPath, int ChangeNumber, Func<string, string> Update, ILogger Logger, CancellationToken CancellationToken)
+		static async Task<bool> UpdateVersionFile(IPerforceConnection perforce, string clientPath, int changeNumber, Func<string, string> update, ILogger logger, CancellationToken cancellationToken)
 		{
-			List<PerforceResponse<FStatRecord>> Records = await Perforce.TryFStatAsync(FStatOptions.None, ClientPath, CancellationToken).ToListAsync(CancellationToken);
-			if (!Records.Succeeded())
+			List<PerforceResponse<FStatRecord>> records = await perforce.TryFStatAsync(FStatOptions.None, clientPath, cancellationToken).ToListAsync(cancellationToken);
+			if (!records.Succeeded())
 			{
-				Logger.LogInformation("Failed to query records for {ClientPath}", ClientPath);
+				logger.LogInformation("Failed to query records for {ClientPath}", clientPath);
 				return false;
 			}
-			if (Records.Count > 1)
+			if (records.Count > 1)
 			{
 				// Attempt to remove any existing file which is synced
-				await Perforce.SyncAsync(SyncOptions.Force, -1, $"{ClientPath}#0", CancellationToken).ToListAsync(CancellationToken);
+				await perforce.SyncAsync(SyncOptions.Force, -1, $"{clientPath}#0", cancellationToken).ToListAsync(cancellationToken);
 
 				// Try to get the mapped files again
-				Records = await Perforce.TryFStatAsync(FStatOptions.None, ClientPath, CancellationToken).ToListAsync(CancellationToken);
-				if (!Records.Succeeded())
+				records = await perforce.TryFStatAsync(FStatOptions.None, clientPath, cancellationToken).ToListAsync(cancellationToken);
+				if (!records.Succeeded())
 				{
-					Logger.LogInformation("Failed to query records for {ClientPath}", ClientPath);
+					logger.LogInformation("Failed to query records for {ClientPath}", clientPath);
 					return false;
 				}
 			}
-			if (Records.Count == 0)
+			if (records.Count == 0)
 			{
-				Logger.LogInformation("Ignoring {ClientPath}; not found on server.", ClientPath);
+				logger.LogInformation("Ignoring {ClientPath}; not found on server.", clientPath);
 				return true;
 			}
 
-			FStatRecord Record = Records[0].Data;
-			string? LocalPath = Record.ClientFile; // Actually a filesystem path
-			if (LocalPath == null)
+			FStatRecord record = records[0].Data;
+			string? localPath = record.ClientFile; // Actually a filesystem path
+			if (localPath == null)
 			{
-				Logger.LogInformation("Version file is not mapped to workspace ({ClientFile})", ClientPath);
+				logger.LogInformation("Version file is not mapped to workspace ({ClientFile})", clientPath);
 				return false;
 			}
-			string? DepotPath = Record.DepotFile;
-			if (DepotPath == null)
+			string? depotPath = record.DepotFile;
+			if (depotPath == null)
 			{
-				Logger.LogInformation("Version file does not exist in depot ({ClientFile})", ClientPath);
+				logger.LogInformation("Version file does not exist in depot ({ClientFile})", clientPath);
 				return false;
 			}
 
-			PerforceResponse<PrintRecord<string[]>> Response = await Perforce.TryPrintLinesAsync($"{DepotPath}@{ChangeNumber}", CancellationToken);
-			if (!Response.Succeeded)
+			PerforceResponse<PrintRecord<string[]>> response = await perforce.TryPrintLinesAsync($"{depotPath}@{changeNumber}", cancellationToken);
+			if (!response.Succeeded)
 			{
-				Logger.LogInformation("Couldn't get default contents of {DepotPath}", DepotPath);
+				logger.LogInformation("Couldn't get default contents of {DepotPath}", depotPath);
 				return false;
 			}
 
-			string[]? Contents = Response.Data.Contents;
-			if (Contents == null)
+			string[]? contents = response.Data.Contents;
+			if (contents == null)
 			{
-				Logger.LogInformation("No data returned for {DepotPath}", DepotPath);
+				logger.LogInformation("No data returned for {DepotPath}", depotPath);
 				return false;
 			}
 
-			string Text = String.Join("\n", Contents);
-			Text = Update(Text);
-			return await WriteVersionFile(Perforce, LocalPath, DepotPath, Text, Logger, CancellationToken);
+			string text = String.Join("\n", contents);
+			text = update(text);
+			return await WriteVersionFile(perforce, localPath, depotPath, text, logger, cancellationToken);
 		}
 
-		static string UpdateVersionStrings(string Text, Dictionary<string, string> VersionStrings)
+		static string UpdateVersionStrings(string text, Dictionary<string, string> versionStrings)
 		{
-			StringWriter Writer = new StringWriter();
-			foreach (string Line in Text.Split('\n'))
+			StringWriter writer = new StringWriter();
+			foreach (string line in text.Split('\n'))
 			{
-				string NewLine = Line;
-				foreach (KeyValuePair<string, string> VersionString in VersionStrings)
+				string newLine = line;
+				foreach (KeyValuePair<string, string> versionString in versionStrings)
 				{
-					if (UpdateVersionLine(ref NewLine, VersionString.Key, VersionString.Value))
+					if (UpdateVersionLine(ref newLine, versionString.Key, versionString.Value))
 					{
 						break;
 					}
 				}
-				Writer.WriteLine(NewLine);
+				writer.WriteLine(newLine);
 			}
-			return Writer.ToString();
+			return writer.ToString();
 		}
 
-		static string UpdateBuildVersion(string Text, int Changelist, int CodeChangelist, string BranchOrStreamName, bool bIsLicenseeVersion)
+		static string UpdateBuildVersion(string text, int changelist, int codeChangelist, string branchOrStreamName, bool isLicenseeVersion)
 		{
-			Dictionary<string, object> Object = JsonSerializer.Deserialize<Dictionary<string, object>>(Text, Utility.DefaultJsonSerializerOptions);
+			Dictionary<string, object> obj = JsonSerializer.Deserialize<Dictionary<string, object>>(text, Utility.DefaultJsonSerializerOptions);
 
-			int PrevCompatibleChangelist = 0;
-			if (Object.TryGetValue("CompatibleChangelist", out object? PrevCompatibleChangelistObj))
+			int prevCompatibleChangelist = 0;
+			if (obj.TryGetValue("CompatibleChangelist", out object? prevCompatibleChangelistObj))
 			{
-				int.TryParse(PrevCompatibleChangelistObj?.ToString(), out PrevCompatibleChangelist);
+				int.TryParse(prevCompatibleChangelistObj?.ToString(), out prevCompatibleChangelist);
 			}
 
-			int PrevIsLicenseeVersion = 0;
-			if (Object.TryGetValue("IsLicenseeVersion", out object? PrevIsLicenseeVersionObj))
+			int prevIsLicenseeVersion = 0;
+			if (obj.TryGetValue("IsLicenseeVersion", out object? prevIsLicenseeVersionObj))
 			{
-				int.TryParse(PrevIsLicenseeVersionObj?.ToString(), out PrevIsLicenseeVersion);
+				int.TryParse(prevIsLicenseeVersionObj?.ToString(), out prevIsLicenseeVersion);
 			}
 
-			Object["Changelist"] = Changelist;
-			if (PrevCompatibleChangelist == 0 || (PrevIsLicenseeVersion != 0) != bIsLicenseeVersion)
+			obj["Changelist"] = changelist;
+			if (prevCompatibleChangelist == 0 || (prevIsLicenseeVersion != 0) != isLicenseeVersion)
 			{
 				// Don't overwrite the compatible changelist if we're in a hotfix release
-				Object["CompatibleChangelist"] = CodeChangelist;
+				obj["CompatibleChangelist"] = codeChangelist;
 			}
-			Object["BranchName"] = BranchOrStreamName.Replace('/', '+');
-			Object["IsPromotedBuild"] = 0;
-			Object["IsLicenseeVersion"] = bIsLicenseeVersion ? 1 : 0;
+			obj["BranchName"] = branchOrStreamName.Replace('/', '+');
+			obj["IsPromotedBuild"] = 0;
+			obj["IsLicenseeVersion"] = isLicenseeVersion ? 1 : 0;
 
-			return JsonSerializer.Serialize(Object, new JsonSerializerOptions
+			return JsonSerializer.Serialize(obj, new JsonSerializerOptions
 			{
 				WriteIndented = true,
 				// do not escape +
@@ -1922,80 +1922,80 @@ namespace UnrealGameSync
 			});
 		}
 
-		static async Task<bool> WriteVersionFile(IPerforceConnection Perforce, string LocalPath, string DepotPath, string NewText, ILogger Logger, CancellationToken CancellationToken)
+		static async Task<bool> WriteVersionFile(IPerforceConnection perforce, string localPath, string depotPath, string newText, ILogger logger, CancellationToken cancellationToken)
 		{
 			try
 			{
-				if (File.Exists(LocalPath) && File.ReadAllText(LocalPath) == NewText)
+				if (File.Exists(localPath) && File.ReadAllText(localPath) == newText)
 				{
-					Logger.LogInformation("Ignored {FileName}; contents haven't changed", LocalPath);
+					logger.LogInformation("Ignored {FileName}; contents haven't changed", localPath);
 				}
 				else
 				{
-					Directory.CreateDirectory(Path.GetDirectoryName(LocalPath));
-					Utility.ForceDeleteFile(LocalPath);
-					if (DepotPath != null)
+					Directory.CreateDirectory(Path.GetDirectoryName(localPath));
+					Utility.ForceDeleteFile(localPath);
+					if (depotPath != null)
 					{
-						await Perforce.SyncAsync(DepotPath + "#0", CancellationToken).ToListAsync(CancellationToken);
+						await perforce.SyncAsync(depotPath + "#0", cancellationToken).ToListAsync(cancellationToken);
 					}
-					File.WriteAllText(LocalPath, NewText);
-					Logger.LogInformation("Written {FileName}", LocalPath);
+					File.WriteAllText(localPath, newText);
+					logger.LogInformation("Written {FileName}", localPath);
 				}
 				return true;
 			}
-			catch (Exception Ex)
+			catch (Exception ex)
 			{
-				Logger.LogError(Ex, "Failed to write to {FileName}.", LocalPath);
+				logger.LogError(ex, "Failed to write to {FileName}.", localPath);
 				return false;
 			}
 		}
 
-		static bool UpdateVersionLine(ref string Line, string Prefix, string Suffix)
+		static bool UpdateVersionLine(ref string line, string prefix, string suffix)
 		{
-			int LineIdx = 0;
-			int PrefixIdx = 0;
+			int lineIdx = 0;
+			int prefixIdx = 0;
 			for (; ; )
 			{
-				string? PrefixToken = ReadToken(Prefix, ref PrefixIdx);
-				if (PrefixToken == null)
+				string? prefixToken = ReadToken(prefix, ref prefixIdx);
+				if (prefixToken == null)
 				{
 					break;
 				}
 
-				string? LineToken = ReadToken(Line, ref LineIdx);
-				if (LineToken == null || LineToken != PrefixToken)
+				string? lineToken = ReadToken(line, ref lineIdx);
+				if (lineToken == null || lineToken != prefixToken)
 				{
 					return false;
 				}
 			}
-			Line = Line.Substring(0, LineIdx) + Suffix;
+			line = line.Substring(0, lineIdx) + suffix;
 			return true;
 		}
 
-		static string? ReadToken(string Line, ref int LineIdx)
+		static string? ReadToken(string line, ref int lineIdx)
 		{
-			for (; ; LineIdx++)
+			for (; ; lineIdx++)
 			{
-				if (LineIdx == Line.Length)
+				if (lineIdx == line.Length)
 				{
 					return null;
 				}
-				else if (!Char.IsWhiteSpace(Line[LineIdx]))
+				else if (!Char.IsWhiteSpace(line[lineIdx]))
 				{
 					break;
 				}
 			}
 
-			int StartIdx = LineIdx++;
-			if (Char.IsLetterOrDigit(Line[StartIdx]) || Line[StartIdx] == '_')
+			int startIdx = lineIdx++;
+			if (Char.IsLetterOrDigit(line[startIdx]) || line[startIdx] == '_')
 			{
-				while (LineIdx < Line.Length && (Char.IsLetterOrDigit(Line[LineIdx]) || Line[LineIdx] == '_'))
+				while (lineIdx < line.Length && (Char.IsLetterOrDigit(line[lineIdx]) || line[lineIdx] == '_'))
 				{
-					LineIdx++;
+					lineIdx++;
 				}
 			}
 
-			return Line.Substring(StartIdx, LineIdx - StartIdx);
+			return line.Substring(startIdx, lineIdx - startIdx);
 		}
 
 		public Tuple<string, float> CurrentProgress
