UnrealGameSync - Started using EpicGames.OIDC for token management (deleting the old duplicated code we had in UGS). This is not enabled by default unless you have a ini that configured which IdP to use.

diff --git a/Engine/Source/Programs/UnrealGameSync/UnrealGameSync/Controls/OIDCControl.cs b/Engine/Source/Programs/UnrealGameSync/UnrealGameSync/Controls/OIDCControl.cs
--- a/Engine/Source/Programs/UnrealGameSync/UnrealGameSync/Controls/OIDCControl.cs
+++ b/Engine/Source/Programs/UnrealGameSync/UnrealGameSync/Controls/OIDCControl.cs
@@ -4,6 +4,7 @@ using System;
 using System.Drawing;
 using System.Drawing.Drawing2D;
 using System.Windows.Forms;
+using EpicGames.OIDC;
 
 namespace UnrealGameSync.Controls
 {
diff --git a/Engine/Source/Programs/UnrealGameSync/UnrealGameSync/Controls/WorkspaceControl.cs b/Engine/Source/Programs/UnrealGameSync/UnrealGameSync/Controls/WorkspaceControl.cs
--- a/Engine/Source/Programs/UnrealGameSync/UnrealGameSync/Controls/WorkspaceControl.cs
+++ b/Engine/Source/Programs/UnrealGameSync/UnrealGameSync/Controls/WorkspaceControl.cs
@@ -21,6 +21,7 @@ using System.Threading.Tasks;
 using Microsoft.Extensions.Logging;
 using Microsoft.Extensions.DependencyInjection;
 using System.Diagnostics.CodeAnalysis;
+using EpicGames.OIDC;
 
 namespace UnrealGameSync
 {
diff --git a/Engine/Source/Programs/UnrealGameSync/UnrealGameSync/Forms/MainWindow.cs b/Engine/Source/Programs/UnrealGameSync/UnrealGameSync/Forms/MainWindow.cs
--- a/Engine/Source/Programs/UnrealGameSync/UnrealGameSync/Forms/MainWindow.cs
+++ b/Engine/Source/Programs/UnrealGameSync/UnrealGameSync/Forms/MainWindow.cs
@@ -20,6 +20,7 @@ using System.Threading.Tasks;
 using Microsoft.Extensions.Logging;
 using Microsoft.Extensions.DependencyInjection;
 using System.Diagnostics.CodeAnalysis;
+using EpicGames.OIDC;
 
 namespace UnrealGameSync
 {
diff --git a/Engine/Source/Programs/UnrealGameSync/UnrealGameSync/Forms/OIDCLoginWindow.cs b/Engine/Source/Programs/UnrealGameSync/UnrealGameSync/Forms/OIDCLoginWindow.cs
--- a/Engine/Source/Programs/UnrealGameSync/UnrealGameSync/Forms/OIDCLoginWindow.cs
+++ b/Engine/Source/Programs/UnrealGameSync/UnrealGameSync/Forms/OIDCLoginWindow.cs
@@ -3,6 +3,7 @@
 using System;
 using System.Drawing;
 using System.Windows.Forms;
+using EpicGames.OIDC;
 using UnrealGameSync.Controls;
 
 namespace UnrealGameSync.Forms
@@ -14,9 +15,9 @@ namespace UnrealGameSync.Forms
 			InitializeComponent();
 
 			int lastYPosition = 10;
-			foreach (OidcTokenManager.ProviderInfo provider in oidcManager.Providers.Values)
+			foreach ((string providerIdentifier, ProviderInfo provider) in oidcManager.Providers)
 			{
-				OidcControl serviceControl = new OidcControl(oidcManager, provider.Identifier, provider.DisplayName)
+				OidcControl serviceControl = new OidcControl(oidcManager, providerIdentifier, provider.DisplayName)
 				{
 					Anchor = AnchorStyles.Top | AnchorStyles.Left | AnchorStyles.Right,
 					Location = new Point(10, lastYPosition)
diff --git a/Engine/Source/Programs/UnrealGameSync/UnrealGameSync/Jupiter/JupiterMonitor.cs b/Engine/Source/Programs/UnrealGameSync/UnrealGameSync/Jupiter/JupiterMonitor.cs
--- a/Engine/Source/Programs/UnrealGameSync/UnrealGameSync/Jupiter/JupiterMonitor.cs
+++ b/Engine/Source/Programs/UnrealGameSync/UnrealGameSync/Jupiter/JupiterMonitor.cs
@@ -11,6 +11,7 @@ using System.Threading.Tasks;
 using IdentityModel.Client;
 using EpicGames.Core;
 using EpicGames.Jupiter;
+using EpicGames.OIDC;
 using EpicGames.Perforce;
 using Microsoft.Extensions.Logging;
 
@@ -115,13 +116,13 @@ namespace UnrealGameSync
 
 		private async Task<IReadOnlyList<IArchiveInfo>> GetAvailableArchives()
 		{
-			string token = await _tokenManager.GetAccessToken(_providerIdentifier);
+			OidcTokenInfo token = await _tokenManager.GetAccessToken(_providerIdentifier);
 
 			List<JupiterArchiveInfo> newArchives = new List<JupiterArchiveInfo>();
 			using (HttpClient client = new HttpClient())
 			{
 				client.BaseAddress = _jupiterUrl;
-				client.SetBearerToken(token);
+				client.SetBearerToken(token.AccessToken);
 
 				string responseBody = await client.GetStringAsync($"/api/v1/c/tree-root/{_jupiterNamespace}");
 				TreeRootListResponse response = JsonSerializer.Deserialize<TreeRootListResponse>(responseBody, new JsonSerializerOptions {PropertyNameCaseInsensitive = true});
diff --git a/Engine/Source/Programs/UnrealGameSync/UnrealGameSync/OIDCTokenManager.cs b/Engine/Source/Programs/UnrealGameSync/UnrealGameSync/OIDCTokenManager.cs
deleted file mode 100644
diff --git a/Engine/Source/Programs/UnrealGameSync/UnrealGameSync/ProgramApplicationContext.cs b/Engine/Source/Programs/UnrealGameSync/UnrealGameSync/ProgramApplicationContext.cs
--- a/Engine/Source/Programs/UnrealGameSync/UnrealGameSync/ProgramApplicationContext.cs
+++ b/Engine/Source/Programs/UnrealGameSync/UnrealGameSync/ProgramApplicationContext.cs
@@ -15,6 +15,7 @@ using System.Text;
 using System.Threading;
 using System.Threading.Tasks;
 using System.Windows.Forms;
+using EpicGames.OIDC;
 using UnrealGameSync.Forms;
 
 #nullable enable
@@ -56,6 +57,7 @@ namespace UnrealGameSync
 		ModalTaskWindow? _startupWindow;
 		MainWindow? _mainWindowInstance;
 
+		ITokenStore? _tokenStore;
 		OidcTokenManager? _oidcTokenManager;
 
 		public ProgramApplicationContext(IPerforceSettings defaultPerforceSettings, UpdateMonitor updateMonitor, string? apiUrl, DirectoryReference dataFolder, EventWaitHandle activateEvent, bool restoreState, string? updateSpawn, string? projectFileName, bool preview, IServiceProvider serviceProvider, string? uri)
@@ -256,12 +258,15 @@ namespace UnrealGameSync
 			// Clear out the cache folder
 			Utility.ClearPrintCache(_cacheFolder);
 
-			_oidcTokenManager = OidcTokenManager.CreateFromConfigFile(_settings, startupTasks.Where(x => x.Item2.Succeeded).Select(x => x.Item2.Result).ToList());
+			_tokenStore = TokenStoreFactory.CreateTokenStore();
+			_oidcTokenManager = CreateTokenManagerFromConfigFile(_settings, _tokenStore, startupTasks.Where(x => x.Item2.Succeeded).Select(x => x.Item2.Result).ToList());
 			// Verify that none of the projects we are opening needs a OIDC login, if they do prompt for the login
 			if (_oidcTokenManager?.HasUnfinishedLogin() ?? false)
 			{
 				OidcLoginWindow loginDialog = new OidcLoginWindow(_oidcTokenManager);
 				loginDialog.ShowDialog();
+
+				_tokenStore.Save();
 			}
 
 			// Get the application path
@@ -288,6 +293,71 @@ namespace UnrealGameSync
 			_mainWindowInstance.FormClosed += MainWindowInstance_FormClosed;
 		}
 
+		static OidcTokenManager? CreateTokenManagerFromConfigFile(UserSettings settings, ITokenStore tokenStore, List<OpenProjectInfo> configFiles)
+		{
+			// join the provider configuration from all projects
+			Dictionary<string, ProviderInfo> Providers = new Dictionary<string, ProviderInfo>();
+			foreach (OpenProjectInfo DetectProjectSettingsTask in configFiles)
+			{
+				if (DetectProjectSettingsTask == null)
+				{
+					continue;
+				}
+
+				ConfigFile ConfigFile = DetectProjectSettingsTask.LatestProjectConfigFile;
+				if (ConfigFile == null)
+				{
+					continue;
+				}
+
+				ConfigSection ProviderSection = ConfigFile.FindSection("OIDCProvider");
+				if (ProviderSection == null)
+				{
+					continue;
+				}
+
+				string[] ProviderValues = ProviderSection.GetValues("Provider", Array.Empty<string>());
+				foreach (ConfigObject Provider in ProviderValues.Select(s => new ConfigObject(s)).ToList())
+				{
+					string? Identifier = Provider.GetValue("Identifier");
+					string? ServerUri = Provider.GetValue("ServerUri");
+					string? ClientId = Provider.GetValue("ClientId");
+					string? DisplayName = Provider.GetValue("DisplayName");
+					string? RedirectUri = Provider.GetValue("RedirectUri");
+					string? Scopes = Provider.GetValue("Scopes");
+
+					if (string.IsNullOrEmpty(Identifier))
+						continue;
+
+					if (string.IsNullOrEmpty(ClientId))
+						continue;
+
+					if (string.IsNullOrEmpty(RedirectUri))
+						continue;
+
+					if (string.IsNullOrEmpty(ServerUri))
+						continue;
+
+					ProviderInfo info = new ProviderInfo
+					{
+						ClientId = ClientId,
+						DisplayName = DisplayName ?? Identifier,
+						RedirectUri = new Uri(RedirectUri),
+						ServerUri = new Uri(ServerUri),
+					};
+					if (Scopes != null)
+						info.Scopes = Scopes;
+
+					// we might get a provider with the same identifier from another project, in which case we only keep the first one
+					Providers.TryAdd(Identifier, info);
+				}
+			}
+
+			if (Providers.Count == 0)
+				return null;
+			return new OidcTokenManager(Providers, tokenStore);
+		}
+
 		private void MainWindowInstance_FormClosed(object sender, FormClosedEventArgs e)
 		{
 			ExitThread();
diff --git a/Engine/Source/Programs/UnrealGameSync/UnrealGameSync/UnrealGameSync.csproj b/Engine/Source/Programs/UnrealGameSync/UnrealGameSync/UnrealGameSync.csproj
--- a/Engine/Source/Programs/UnrealGameSync/UnrealGameSync/UnrealGameSync.csproj
+++ b/Engine/Source/Programs/UnrealGameSync/UnrealGameSync/UnrealGameSync.csproj
@@ -22,7 +22,6 @@
   </PropertyGroup>
   <ItemGroup>
     <PackageReference Include="EnvDTE" Version="16.7.30328.74" />
-    <PackageReference Include="IdentityModel.OidcClient" Version="3.1.2" />
     <PackageReference Include="Sentry" Version="3.15.0" />
     <PackageReference Include="Serilog.Expressions" Version="3.2.1" />
     <PackageReference Include="Serilog.Extensions.Hosting" Version="4.2.0" />
@@ -32,6 +31,7 @@
   <ItemGroup>
     <ProjectReference Include="..\..\Shared\EpicGames.Core\EpicGames.Core.csproj" />
     <ProjectReference Include="..\..\Shared\EpicGames.Jupiter\EpicGames.Jupiter.csproj" />
+    <ProjectReference Include="..\..\Shared\EpicGames.OIDC\EpicGames.OIDC.csproj" />
     <ProjectReference Include="..\..\Shared\EpicGames.Perforce\EpicGames.Perforce.csproj" />
     <ProjectReference Include="..\UnrealGameSyncShared\UnrealGameSyncShared.csproj" />
   </ItemGroup>
