UGS: Fix static analysis warnings.

diff --git a/Engine/Source/Programs/UnrealGameSync/UnrealGameSync/Controls/WorkspaceControl.cs b/Engine/Source/Programs/UnrealGameSync/UnrealGameSync/Controls/WorkspaceControl.cs
--- a/Engine/Source/Programs/UnrealGameSync/UnrealGameSync/Controls/WorkspaceControl.cs
+++ b/Engine/Source/Programs/UnrealGameSync/UnrealGameSync/Controls/WorkspaceControl.cs
@@ -1727,13 +1727,15 @@ namespace UnrealGameSync
 				try
 				{
 					string jsonString = File.ReadAllText(uncontrolledChangelistPersistencyFilePath);
-					UncontrolledChangelistPersistency uclPersistency = JsonSerializer.Deserialize<UncontrolledChangelistPersistency>(jsonString, Utility.DefaultJsonSerializerOptions);
-
-					foreach (UncontrolledChangelist ucl in uclPersistency.Changelists)
+					UncontrolledChangelistPersistency? uclPersistency = JsonSerializer.Deserialize<UncontrolledChangelistPersistency>(jsonString, Utility.DefaultJsonSerializerOptions);
+					if (uclPersistency != null)
 					{
-						foreach (string uncontrolledFile in ucl.Files)
+						foreach (UncontrolledChangelist ucl in uclPersistency.Changelists)
 						{
-							outUncontrolledFiles.Add(uncontrolledFile);
+							foreach (string uncontrolledFile in ucl.Files)
+							{
+								outUncontrolledFiles.Add(uncontrolledFile);
+							}
 						}
 					}
 				}
diff --git a/Engine/Source/Programs/UnrealGameSync/UnrealGameSync/RESTApi.cs b/Engine/Source/Programs/UnrealGameSync/UnrealGameSync/RESTApi.cs
--- a/Engine/Source/Programs/UnrealGameSync/UnrealGameSync/RESTApi.cs
+++ b/Engine/Source/Programs/UnrealGameSync/UnrealGameSync/RESTApi.cs
@@ -76,7 +76,7 @@ namespace UnrealGameSync
 
 		public static async Task<T> GetAsync<T>(string url, CancellationToken cancellationToken)
 		{
-			return JsonSerializer.Deserialize<T>(await GetAsync(url, cancellationToken), Utility.DefaultJsonSerializerOptions);
+			return JsonSerializer.Deserialize<T>(await GetAsync(url, cancellationToken), Utility.DefaultJsonSerializerOptions)!;
 		}
 
 		public static Task<string> PutAsync<T>(string url, T obj, CancellationToken cancellationToken)
