UnrealGameSync - Added login command which can be used to make sure we have a offline token available to use by UGS or other tools. This will spawn a browser to complete the interactive login so while its useable from a command line you do need to have access to a interactive session to complete the login.

diff --git a/Engine/Source/Programs/UnrealGameSync/UnrealGameSyncCmd/Program.cs b/Engine/Source/Programs/UnrealGameSync/UnrealGameSyncCmd/Program.cs
--- a/Engine/Source/Programs/UnrealGameSync/UnrealGameSyncCmd/Program.cs
+++ b/Engine/Source/Programs/UnrealGameSync/UnrealGameSyncCmd/Program.cs
@@ -17,6 +17,8 @@ using System.Text;
 using System.Text.RegularExpressions;
 using System.Threading;
 using System.Threading.Tasks;
+using EpicGames.OIDC;
+using Microsoft.Extensions.Configuration;
 using UnrealGameSync;
 
 namespace UnrealGameSyncCmd
@@ -103,6 +105,10 @@ namespace UnrealGameSyncCmd
 				"ugs status [-update]",
 				"Shows the status of the currently synced branch."
 			),
+			new CommandInfo("login", typeof(LoginCommand),
+				"ugs login",
+				"Starts a interactive login flow against the configured Identity Provider"
+			),
 			new CommandInfo("version", typeof(VersionCommand),
 				"ugs version",
 				"Prints the current application version"
@@ -1104,6 +1110,33 @@ namespace UnrealGameSyncCmd
 			}
 		}
 
+		class LoginCommand : Command
+		{
+			public override async Task ExecuteAsync(CommandContext context)
+			{
+				ILogger logger = context.Logger;
+
+				// Get the positional argument indicating the file to look for
+				if (!context.Arguments.TryGetPositionalArgument(out string? providerIdentifier))
+				{
+					throw new UserErrorException("Missing provider identifier to login to.");
+				}
+				context.Arguments.CheckAllArgumentsUsed();
+
+				UserWorkspaceSettings settings = ReadRequiredUserWorkspaceSettings();
+				
+				// Find the valid config file paths
+				DirectoryInfo engineDir = DirectoryReference.Combine(settings.RootDir, "Engine").ToDirectoryInfo();
+				DirectoryInfo gameDir = new DirectoryInfo(settings.ProjectPath);
+				using ITokenStore tokenStore = TokenStoreFactory.CreateTokenStore();
+				IConfiguration providerConfiguration = ProviderConfigurationFactory.ReadConfiguration(engineDir, gameDir);
+				OidcTokenManager oidcTokenManager = OidcTokenManager.CreateTokenManager(providerConfiguration, tokenStore, new List<string>() {providerIdentifier});
+				OidcTokenInfo result = await oidcTokenManager.Login(providerIdentifier);
+
+				logger.LogInformation("Logged in to provider {ProviderIdentifier}", providerIdentifier);
+			}
+		}
+
 		class SwitchCommand : Command
 		{
 			public override async Task ExecuteAsync(CommandContext context)
diff --git a/Engine/Source/Programs/UnrealGameSync/UnrealGameSyncShared/UnrealGameSyncShared.csproj b/Engine/Source/Programs/UnrealGameSync/UnrealGameSyncShared/UnrealGameSyncShared.csproj
--- a/Engine/Source/Programs/UnrealGameSync/UnrealGameSyncShared/UnrealGameSyncShared.csproj
+++ b/Engine/Source/Programs/UnrealGameSync/UnrealGameSyncShared/UnrealGameSyncShared.csproj
@@ -8,6 +8,7 @@
 
   <ItemGroup>
     <ProjectReference Include="..\..\Shared\EpicGames.Core\EpicGames.Core.csproj" />
+    <ProjectReference Include="..\..\Shared\EpicGames.OIDC\EpicGames.OIDC.csproj" />
     <ProjectReference Include="..\..\Shared\EpicGames.Perforce\EpicGames.Perforce.csproj" />
   </ItemGroup>
 
