UnrealGameSync: Also support PrimaryProjectName.txt when finding the name of a project

#rnx

diff --git a/Engine/Source/Programs/UnrealGameSync/UnrealGameSync/Controls/WorkspaceControl.cs b/Engine/Source/Programs/UnrealGameSync/UnrealGameSync/Controls/WorkspaceControl.cs
--- a/Engine/Source/Programs/UnrealGameSync/UnrealGameSync/Controls/WorkspaceControl.cs
+++ b/Engine/Source/Programs/UnrealGameSync/UnrealGameSync/Controls/WorkspaceControl.cs
@@ -4879,26 +4879,31 @@ namespace UnrealGameSync
 
 		private void OpenSolution()
 		{
-			string masterProjectName = "UE4";
+			string primaryProjectName = "UE4";
 			if (FileReference.Exists(FileReference.Combine(BranchDirectoryName, "UE5.sln")))
 			{
-				masterProjectName = "UE5";
+				primaryProjectName = "UE5";
 			}
 
-			FileReference masterProjectNameFileName = FileReference.Combine(BranchDirectoryName, "Engine", "Intermediate", "ProjectFiles", "MasterProjectName.txt");
-			if (FileReference.Exists(masterProjectNameFileName))
+			FileReference primaryProjectNameFileName = FileReference.Combine(BranchDirectoryName, "Engine", "Intermediate", "ProjectFiles", "PrimaryProjectName.txt");
+			if (!FileReference.Exists(primaryProjectNameFileName))
+			{
+				// Old path, needs to be maintained while projects prior to UE5.1 are supported
+				primaryProjectNameFileName = FileReference.Combine(BranchDirectoryName, "Engine", "Intermediate", "ProjectFiles", "MasterProjectName.txt");
+			}
+			if (FileReference.Exists(primaryProjectNameFileName))
 			{
 				try
 				{
-					masterProjectName = FileReference.ReadAllText(masterProjectNameFileName).Trim();
+					primaryProjectName = FileReference.ReadAllText(primaryProjectNameFileName).Trim();
 				}
 				catch (Exception ex)
 				{
-					_logger.LogError(ex, "Unable to read '{File}'", masterProjectNameFileName);
+					_logger.LogError(ex, "Unable to read '{File}'", primaryProjectNameFileName);
 				}
 			}
 
-			FileReference solutionFileName = FileReference.Combine(BranchDirectoryName, masterProjectName + ".sln");
+			FileReference solutionFileName = FileReference.Combine(BranchDirectoryName, primaryProjectName + ".sln");
 			if (!FileReference.Exists(solutionFileName))
 			{
 				MessageBox.Show(String.Format("Couldn't find solution at {0}", solutionFileName));
