UGS: Fix some null reference warnings.

diff --git a/Engine/Source/Programs/UnrealGameSync/UnrealGameSyncShared/ConfigFile.cs b/Engine/Source/Programs/UnrealGameSync/UnrealGameSyncShared/ConfigFile.cs
--- a/Engine/Source/Programs/UnrealGameSync/UnrealGameSyncShared/ConfigFile.cs
+++ b/Engine/Source/Programs/UnrealGameSync/UnrealGameSyncShared/ConfigFile.cs
@@ -310,9 +310,9 @@ namespace UnrealGameSync
 			{
 				while (reader.Read() && reader.TokenType == JsonTokenType.PropertyName)
 				{
-					string name = reader.GetString();
+					string name = reader.GetString()!;
 					reader.Read();
-					string value = reader.GetString();
+					string value = reader.GetString()!;
 					obj.Pairs.Add(KeyValuePair.Create(name, value));
 				}
 			}
diff --git a/Engine/Source/Programs/UnrealGameSync/UnrealGameSyncShared/Utility.cs b/Engine/Source/Programs/UnrealGameSync/UnrealGameSyncShared/Utility.cs
--- a/Engine/Source/Programs/UnrealGameSync/UnrealGameSyncShared/Utility.cs
+++ b/Engine/Source/Programs/UnrealGameSync/UnrealGameSyncShared/Utility.cs
@@ -121,7 +121,7 @@ namespace UnrealGameSync
 		public static T LoadJson<T>(FileReference file)
 		{
 			byte[] data = FileReference.ReadAllBytes(file);
-			return JsonSerializer.Deserialize<T>(data, DefaultJsonSerializerOptions);
+			return JsonSerializer.Deserialize<T>(data, DefaultJsonSerializerOptions)!;
 		}
 
 		public static void SaveJson<T>(FileReference file, T obj)
@@ -320,7 +320,7 @@ namespace UnrealGameSync
 				options.PropertyNameCaseInsensitive = true;
 				options.Converters.Add(new JsonStringEnumConverter());
 
-				ProjectJson project = JsonSerializer.Deserialize<ProjectJson>(text, options);
+				ProjectJson project = JsonSerializer.Deserialize<ProjectJson>(text, options)!;
 
 				return project.Enterprise;
 			}
diff --git a/Engine/Source/Programs/UnrealGameSync/UnrealGameSyncShared/WorkspaceUpdate.cs b/Engine/Source/Programs/UnrealGameSync/UnrealGameSyncShared/WorkspaceUpdate.cs
--- a/Engine/Source/Programs/UnrealGameSync/UnrealGameSyncShared/WorkspaceUpdate.cs
+++ b/Engine/Source/Programs/UnrealGameSync/UnrealGameSyncShared/WorkspaceUpdate.cs
@@ -1890,7 +1890,7 @@ namespace UnrealGameSync
 
 		static string UpdateBuildVersion(string text, int changelist, int codeChangelist, string branchOrStreamName, bool isLicenseeVersion)
 		{
-			Dictionary<string, object> obj = JsonSerializer.Deserialize<Dictionary<string, object>>(text, Utility.DefaultJsonSerializerOptions);
+			Dictionary<string, object> obj = JsonSerializer.Deserialize<Dictionary<string, object>>(text, Utility.DefaultJsonSerializerOptions)!;
 
 			int prevCompatibleChangelist = 0;
 			if (obj.TryGetValue("CompatibleChangelist", out object? prevCompatibleChangelistObj))
