UGS: Fix missing notification when switching streams with open files.

diff --git a/Engine/Source/Programs/UnrealGameSync/UnrealGameSync/Controls/WorkspaceControl.cs b/Engine/Source/Programs/UnrealGameSync/UnrealGameSync/Controls/WorkspaceControl.cs
--- a/Engine/Source/Programs/UnrealGameSync/UnrealGameSync/Controls/WorkspaceControl.cs
+++ b/Engine/Source/Programs/UnrealGameSync/UnrealGameSync/Controls/WorkspaceControl.cs
@@ -3870,6 +3870,10 @@ namespace UnrealGameSync
 							StatusPanel.ResumeLayout();
 							break;
 						}
+						if (MessageBox.Show("You have files open for edit in this workspace. If you continue, you will not be able to submit them until you switch back.\n\nContinue switching streams?", "Files checked out", MessageBoxButtons.YesNo) != DialogResult.Yes)
+						{
+							break;
+						}
 					}
 				}
 			}
