UGS: Install the ugs:// protocol handler by default for the current user, so people can use Horde integration out-of-the-box. Previous setting reworded to apply to all users.

diff --git a/Engine/Source/Programs/UnrealGameSync/UnrealGameSync/Program.cs b/Engine/Source/Programs/UnrealGameSync/UnrealGameSync/Program.cs
--- a/Engine/Source/Programs/UnrealGameSync/UnrealGameSync/Program.cs
+++ b/Engine/Source/Programs/UnrealGameSync/UnrealGameSync/Program.cs
@@ -212,6 +212,8 @@ namespace UnrealGameSync
 
 							IPerforceSettings defaultSettings = new PerforceSettings(serverAndPort, userName) { PreferNativeClient = true };
 
+							ProtocolHandlerUtils.InstallQuiet(logger);
+
 							using (UpdateMonitor updateMonitor = new UpdateMonitor(defaultSettings, updatePath, serviceProvider))
 							{
 								using ProgramApplicationContext context = new ProgramApplicationContext(defaultSettings, updateMonitor, DeploymentSettings.ApiUrl, dataFolder, activateEvent, restoreState, updateSpawn, projectFileName, preview, serviceProvider, uri);
diff --git a/Engine/Source/Programs/UnrealGameSync/UnrealGameSync/UriHandlers/UriHandler.cs b/Engine/Source/Programs/UnrealGameSync/UnrealGameSync/UriHandlers/UriHandler.cs
--- a/Engine/Source/Programs/UnrealGameSync/UnrealGameSync/UriHandlers/UriHandler.cs
+++ b/Engine/Source/Programs/UnrealGameSync/UnrealGameSync/UriHandlers/UriHandler.cs
@@ -12,6 +12,7 @@ using System.Diagnostics;
 using System.Windows.Forms;
 using System.Configuration;
 using System.IO;
+using Microsoft.Extensions.Logging;
 
 namespace UnrealGameSync
 {
@@ -296,10 +297,10 @@ namespace UnrealGameSync
 
 		public static string CurrentProcessFilePath
 		{
-			get { return Path.GetFullPath(Process.GetCurrentProcess().MainModule.FileName); }
+			get { return Path.ChangeExtension(Path.GetFullPath(Process.GetCurrentProcess().MainModule.FileName), ".exe"); }
 		}
 
-		static List<RegistrySetting> GetRegistrySettings()
+		static List<RegistrySetting> GetGlobalRegistrySettings()
 		{
 			string applicationPath = CurrentProcessFilePath;
 
@@ -311,14 +312,30 @@ namespace UnrealGameSync
 			return keys;
 		}
 
+		static List<RegistrySetting> GetUserRegistrySettings()
+		{
+			string applicationPath = CurrentProcessFilePath;
+
+			List<RegistrySetting> keys = new List<RegistrySetting>();
+			keys.Add(new RegistrySetting(Registry.CurrentUser, "Software\\Classes\\UGS", null, "URL:UGS Protocol"));
+			keys.Add(new RegistrySetting(Registry.CurrentUser, "Software\\Classes\\UGS", "URL Protocol", ""));
+			keys.Add(new RegistrySetting(Registry.CurrentUser, "Software\\Classes\\UGS\\DefaultIcon", null, String.Format("\"{0}\",0", applicationPath)));
+			keys.Add(new RegistrySetting(Registry.CurrentUser, "Software\\Classes\\UGS\\shell\\open\\command", null, String.Format("\"{0}\" -uri=\"%1\"", applicationPath)));
+			return keys;
+		}
+
 		public static ProtocolHandlerState GetState()
+		{
+			return GetState(GetGlobalRegistrySettings());
+		}
+
+		private static ProtocolHandlerState GetState(IEnumerable<RegistrySetting> keys)
 		{
 			try
 			{
 				bool hasAny = false;
 				bool hasAll = true;
 
-				List<RegistrySetting> keys = GetRegistrySettings();
 				foreach (IGrouping<RegistryKey, RegistrySetting> rootKeyGroup in keys.GroupBy(x => x.RootKey))
 				{
 					foreach (IGrouping<string, RegistrySetting> keyNameGroup in rootKeyGroup.GroupBy(x => x.KeyName))
@@ -355,20 +372,7 @@ namespace UnrealGameSync
 		{
 			try
 			{
-				List<RegistrySetting> keys = GetRegistrySettings();
-				foreach (IGrouping<RegistryKey, RegistrySetting> rootKeyGroup in keys.GroupBy(x => x.RootKey))
-				{
-					foreach (IGrouping<string, RegistrySetting> keyNameGroup in rootKeyGroup.GroupBy(x => x.KeyName))
-					{
-						using (RegistryKey registryKey = rootKeyGroup.Key.CreateSubKey(keyNameGroup.Key))
-						{
-							foreach (RegistrySetting setting in keyNameGroup)
-							{
-								registryKey.SetValue(setting.ValueName, setting.Value);
-							}
-						}
-					}
-				}
+				ApplyRegistrySettings(GetGlobalRegistrySettings());
 			}
 			catch (Exception ex)
 			{
@@ -376,6 +380,19 @@ namespace UnrealGameSync
 			}
 		}
 
+		public static void InstallQuiet(ILogger logger)
+		{
+			try
+			{
+				ApplyRegistrySettings(GetUserRegistrySettings());
+				logger.LogInformation("Installed protocol handler to user hive");
+			}
+			catch (Exception ex)
+			{
+				logger.LogWarning(ex, "Unable to install protocol handler");
+			}
+		}
+
 		public static void Uninstall()
 		{
 			RunElevated(String.Format("{0} {1}", UriHandler.UninstallHandlerArg, UriHandler.ElevatedArg));
@@ -405,6 +422,23 @@ namespace UnrealGameSync
 				process.WaitForExit();
 			}
 		}
+
+		private static void ApplyRegistrySettings(IEnumerable<RegistrySetting> settings)
+		{
+			foreach (IGrouping<RegistryKey, RegistrySetting> rootKeyGroup in settings.GroupBy(x => x.RootKey))
+			{
+				foreach (IGrouping<string, RegistrySetting> keyNameGroup in rootKeyGroup.GroupBy(x => x.KeyName))
+				{
+					using (RegistryKey registryKey = rootKeyGroup.Key.CreateSubKey(keyNameGroup.Key))
+					{
+						foreach (RegistrySetting setting in keyNameGroup)
+						{
+							registryKey.SetValue(setting.ValueName, setting.Value);
+						}
+					}
+				}
+			}
+		}
 	}
 }
 
