UGS: Add a setting to force projects to look for the shared editor receipt rather than a project-specific one.

Also include more information in the "you need to build the editor" dialog explaining why UGS thinks it's out of date.

diff --git a/Engine/Source/Programs/UnrealGameSync/UnrealGameSync/Controls/WorkspaceControl.cs b/Engine/Source/Programs/UnrealGameSync/UnrealGameSync/Controls/WorkspaceControl.cs
--- a/Engine/Source/Programs/UnrealGameSync/UnrealGameSync/Controls/WorkspaceControl.cs
+++ b/Engine/Source/Programs/UnrealGameSync/UnrealGameSync/Controls/WorkspaceControl.cs
@@ -1002,7 +1002,7 @@ namespace UnrealGameSync
 					FileReference targetFile = ConfigUtils.GetEditorTargetFile(_workspace.Project, _workspace.ProjectConfigFile);
 					foreach (BuildConfig config in Enum.GetValues(typeof(BuildConfig)).OfType<BuildConfig>())
 					{
-						FileReference receiptFile = ConfigUtils.GetReceiptFile(_workspace.Project, targetFile, config.ToString());
+						FileReference receiptFile = ConfigUtils.GetReceiptFile(_workspace.Project, _workspace.ProjectConfigFile, targetFile, config.ToString());
 						if (FileReference.Exists(receiptFile))
 						{
 							try { FileReference.Delete(receiptFile); } catch (Exception) { }
@@ -3038,7 +3038,24 @@ namespace UnrealGameSync
 				BuildConfig editorBuildConfig = GetEditorBuildConfig();
 
 				FileReference receiptFile = ConfigUtils.GetEditorReceiptFile(_workspace.Project, _workspace.ProjectConfigFile, editorBuildConfig);
-				if (ConfigUtils.TryReadEditorReceipt(_workspace.Project, receiptFile, out TargetReceipt? receipt) && receipt.Launch != null && File.Exists(receipt.Launch))
+
+				if (!FileReference.Exists(receiptFile))
+				{
+					ShowEditorLaunchError($"{receiptFile} not found");
+				}
+				else if (!ConfigUtils.TryReadEditorReceipt(_workspace.Project, receiptFile, out TargetReceipt? receipt))
+				{
+					ShowEditorLaunchError($"Unable to read {receiptFile}");
+				}
+				else if (receipt.Launch == null)
+				{
+					ShowEditorLaunchError($"No launch target in {receiptFile}");
+				}
+				else if (!File.Exists(receipt.Launch))
+				{
+					ShowEditorLaunchError($"{receipt.Launch} not found");
+				}
+				else 
 				{
 					if (_settings.EditorArgumentsPrompt && !ModifyEditorArguments())
 					{
@@ -3067,17 +3084,18 @@ namespace UnrealGameSync
 						ShowErrorDialog("Unable to spawn {0} {1}", receipt.Launch, launchArguments.ToString());
 					}
 				}
-				else
-				{
-					if (MessageBox.Show("The editor needs to be built before you can run it. Build it now?", "Editor out of date", MessageBoxButtons.YesNo) == System.Windows.Forms.DialogResult.Yes)
-					{
-						_owner.ShowAndActivate();
+			}
+		}
 
-						WorkspaceUpdateOptions options = WorkspaceUpdateOptions.Build | WorkspaceUpdateOptions.RunAfterSync;
-						WorkspaceUpdateContext context = new WorkspaceUpdateContext(_workspace.CurrentChangeNumber, options, _settings.CompiledEditorBuildConfig, null, _projectSettings.BuildSteps, null);
-						StartWorkspaceUpdate(context, null);
-					}
-				}
+		private void ShowEditorLaunchError(string message)
+		{
+			if (MessageBox.Show($"The editor needs to be built before you can run it. Build it now?\r\n\r\n({message})", "Editor out of date", MessageBoxButtons.YesNo) == System.Windows.Forms.DialogResult.Yes)
+			{
+				_owner.ShowAndActivate();
+
+				WorkspaceUpdateOptions options = WorkspaceUpdateOptions.Build | WorkspaceUpdateOptions.RunAfterSync;
+				WorkspaceUpdateContext context = new WorkspaceUpdateContext(_workspace.CurrentChangeNumber, options, _settings.CompiledEditorBuildConfig, null, _projectSettings.BuildSteps, null);
+				StartWorkspaceUpdate(context, null);
 			}
 		}
 
diff --git a/Engine/Source/Programs/UnrealGameSync/UnrealGameSyncShared/ConfigUtils.cs b/Engine/Source/Programs/UnrealGameSync/UnrealGameSyncShared/ConfigUtils.cs
--- a/Engine/Source/Programs/UnrealGameSync/UnrealGameSyncShared/ConfigUtils.cs
+++ b/Engine/Source/Programs/UnrealGameSync/UnrealGameSyncShared/ConfigUtils.cs
@@ -164,7 +164,7 @@ namespace UnrealGameSync
 		public static FileReference GetEditorReceiptFile(ProjectInfo projectInfo, ConfigFile projectConfig, BuildConfig config)
 		{
 			FileReference targetFile = GetEditorTargetFile(projectInfo, projectConfig);
-			return GetReceiptFile(projectInfo, targetFile, config.ToString());
+			return GetReceiptFile(projectInfo, projectConfig, targetFile, config.ToString());
 		}
 
 		private static List<FileReference> FindTargets(DirectoryReference engineOrProjectDir)
@@ -243,12 +243,26 @@ namespace UnrealGameSync
 			return receipt;
 		}
 
-		public static FileReference GetReceiptFile(ProjectInfo projectInfo, FileReference targetFile, string configuration)
+		private static bool UseSharedEditorReceipt(ProjectInfo projectInfo, ConfigFile projectConfig)
+		{
+			string? setting;
+			if (TryGetProjectSetting(projectConfig, projectInfo.ProjectIdentifier, "UseSharedEditor", out setting))
+			{
+				bool value;
+				if (bool.TryParse(setting, out value))
+				{
+					return value;
+				}
+			}
+			return false;
+		}
+
+		public static FileReference GetReceiptFile(ProjectInfo projectInfo, ConfigFile projectConfig, FileReference targetFile, string configuration)
 		{
 			string targetName = targetFile.GetFileNameWithoutAnyExtensions();
 
 			DirectoryReference? projectDir = projectInfo.ProjectDir;
-			if (projectDir != null)
+			if (projectDir != null && (targetFile.IsUnderDirectory(projectDir) || !UseSharedEditorReceipt(projectInfo, projectConfig)))
 			{
 				return GetReceiptFile(projectDir, targetName, configuration);
 			}
diff --git a/Engine/Source/Programs/UnrealGameSync/UnrealGameSyncShared/WorkspaceUpdate.cs b/Engine/Source/Programs/UnrealGameSync/UnrealGameSyncShared/WorkspaceUpdate.cs
--- a/Engine/Source/Programs/UnrealGameSync/UnrealGameSyncShared/WorkspaceUpdate.cs
+++ b/Engine/Source/Programs/UnrealGameSync/UnrealGameSyncShared/WorkspaceUpdate.cs
@@ -1189,7 +1189,7 @@ namespace UnrealGameSync
 
 				FileReference editorTargetFile = ConfigUtils.GetEditorTargetFile(project, Context.ProjectConfigFile);
 				string editorTargetName = editorTargetFile.GetFileNameWithoutAnyExtensions();
-				FileReference editorReceiptFile = ConfigUtils.GetReceiptFile(project, editorTargetFile, Context.EditorConfig.ToString());
+				FileReference editorReceiptFile = ConfigUtils.GetReceiptFile(project, Context.ProjectConfigFile, editorTargetFile, Context.EditorConfig.ToString());
 
 				// Get the build steps
 				bool usingPrecompiledEditor = Context.ArchiveTypeToArchive.TryGetValue(IArchiveInfo.EditorArchiveType, out Tuple<IArchiveInfo, string>? archiveInfo) && archiveInfo != null;
