UGS: Add a better error for missing a file for post-sync steps.

diff --git a/Engine/Source/Programs/UnrealGameSync/UnrealGameSyncShared/WorkspaceUpdate.cs b/Engine/Source/Programs/UnrealGameSync/UnrealGameSyncShared/WorkspaceUpdate.cs
--- a/Engine/Source/Programs/UnrealGameSync/UnrealGameSyncShared/WorkspaceUpdate.cs
+++ b/Engine/Source/Programs/UnrealGameSync/UnrealGameSyncShared/WorkspaceUpdate.cs
@@ -1049,6 +1049,11 @@ namespace UnrealGameSync
 
 									logger.LogInformation("post-sync> Running {FileName} {Arguments}", toolFileName, toolArguments);
 
+									if (!File.Exists(toolFileName))
+									{
+										return (WorkspaceUpdateResult.FailedToSync, $"Unable to find {toolFileName}. You may have an incomplete sync; try cleaning your workspace.");
+									}
+
 									int resultFromTool = await Utility.ExecuteProcessAsync(toolFileName, null, toolArguments, line => ProcessOutput(line, "post-sync> ", Progress, logger), cancellationToken);
 									if (resultFromTool != 0)
 									{
@@ -1337,6 +1342,13 @@ namespace UnrealGameSync
 										string toolArguments = Utility.ExpandVariables(step.Arguments ?? "", variables);
 										logger.LogInformation("tool> Running {0} {1}", toolFileName, toolArguments);
 
+										if (!FileReference.Exists(toolFileName))
+										{
+											logger.LogInformation("Unable to find {FileName}. You may have an incomplete sync; try cleaning your workspace.", toolFileName);
+											stepStopwatch.Stop("Failed");
+											return (WorkspaceUpdateResult.FailedToCompile, $"Unable to run {toolFileName}");
+										}
+
 										if (step.UseLogWindow)
 										{
 											int resultFromTool = await Utility.ExecuteProcessAsync(toolFileName.FullName, toolWorkingDir, toolArguments, line => ProcessOutput(line, "tool> ", Progress, logger), cancellationToken);
