UGS: Fix crash opening UGS without any project.

diff --git a/Engine/Source/Programs/UnrealGameSync/UnrealGameSync/ProgramApplicationContext.cs b/Engine/Source/Programs/UnrealGameSync/UnrealGameSync/ProgramApplicationContext.cs
--- a/Engine/Source/Programs/UnrealGameSync/UnrealGameSync/ProgramApplicationContext.cs
+++ b/Engine/Source/Programs/UnrealGameSync/UnrealGameSync/ProgramApplicationContext.cs
@@ -258,8 +258,6 @@ namespace UnrealGameSync
 			// Clear out the cache folder
 			Utility.ClearPrintCache(_cacheFolder);
 
-			OpenProjectInfo firstProject = startupTasks.FirstOrDefault().Item2.Result;
-			
 			_tokenStore = TokenStoreFactory.CreateTokenStore();
 			List<(DirectoryInfo, DirectoryInfo?)> configurationLocations = new List<(DirectoryInfo, DirectoryInfo?)>();
 			List<string> allowedProviders = new List<string>();
