UGS: Add support for syncing binaries from CLI (with the "-binaries" argument).

diff --git a/Engine/Source/Programs/UnrealGameSync/UnrealGameSync/PerforceMonitor.cs b/Engine/Source/Programs/UnrealGameSync/UnrealGameSync/PerforceMonitor.cs
--- a/Engine/Source/Programs/UnrealGameSync/UnrealGameSync/PerforceMonitor.cs
+++ b/Engine/Source/Programs/UnrealGameSync/UnrealGameSync/PerforceMonitor.cs
@@ -25,104 +25,6 @@ namespace UnrealGameSync
 
 	class PerforceMonitor : IDisposable, IArchiveInfoSource
 	{
-		internal class PerforceArchiveInfo : IArchiveInfo
-		{
-			public string Name { get; }
-			public string Type { get; }
-			public string DepotPath { get; }
-			public string? Target { get; }
-			public PerforceMonitor Outer { get; }
-
-			public string BasePath
-			{
-				get { return DepotPath; }
-			}
-
-			// TODO: executable/configuration?
-			public SortedList<int, string> ChangeNumberToFileRevision = new SortedList<int, string>();
-
-			public PerforceArchiveInfo(string name, string type, string depotPath, string? target, PerforceMonitor outer)
-			{
-				this.Name = name;
-				this.Type = type;
-				this.DepotPath = depotPath;
-				this.Target = target;
-				this.Outer = outer;
-			}
-
-			public override bool Equals(object? other)
-			{
-				PerforceArchiveInfo? otherArchive = other as PerforceArchiveInfo;
-				return otherArchive != null && Name == otherArchive.Name && Type == otherArchive.Type && DepotPath == otherArchive.DepotPath && Target == otherArchive.Target && Enumerable.SequenceEqual(ChangeNumberToFileRevision, otherArchive.ChangeNumberToFileRevision);
-			}
-
-			public override int GetHashCode()
-			{
-				throw new NotSupportedException();
-			}
-
-			public bool Exists()
-			{
-				return ChangeNumberToFileRevision.Count > 0;
-			}
-
-			public static bool TryParseConfigEntry(string text, PerforceMonitor outer, [NotNullWhen(true)] out PerforceArchiveInfo? info)
-			{
-				ConfigObject obj = new ConfigObject(text);
-
-				string? name = obj.GetValue("Name", null);
-				if (name == null)
-				{
-					info = null;
-					return false;
-				}
-
-				string? depotPath = obj.GetValue("DepotPath", null);
-				if (depotPath == null)
-				{
-					info = null;
-					return false;
-				}
-
-				string? target = obj.GetValue("Target", null);
-
-				string type = obj.GetValue("Type", null) ?? name;
-
-				info = new PerforceArchiveInfo(name, type, depotPath, target, outer);
-				return true;
-			}
-
-			public bool TryGetArchiveKeyForChangeNumber(int changeNumber, [NotNullWhen(true)] out string? archiveKey)
-			{
-				return ChangeNumberToFileRevision.TryGetValue(changeNumber, out archiveKey);
-			}
-
-			public async Task<bool> DownloadArchive(IPerforceConnection perforce, string archiveKey, DirectoryReference localRootPath, FileReference manifestFileName, ILogger logger, ProgressValue progress, CancellationToken cancellationToken)
-			{
-				FileReference tempZipFileName = new FileReference(Path.GetTempFileName());
-				try
-				{
-					PrintRecord record = await perforce.PrintAsync(tempZipFileName.FullName, archiveKey, cancellationToken);
-					if (tempZipFileName.ToFileInfo().Length == 0)
-					{
-						return false;
-					}
-					ArchiveUtils.ExtractFiles(tempZipFileName, localRootPath, manifestFileName, progress, logger);
-				}
-				finally
-				{
-					FileReference.SetAttributes(tempZipFileName, FileAttributes.Normal);
-					FileReference.Delete(tempZipFileName);
-				}
-
-				return true;
-			}
-
-			public override string ToString()
-			{
-				return Name;
-			}
-		}
 
 		class PerforceChangeSorter : IComparer<ChangesRecord>
 		{
@@ -595,57 +497,7 @@ namespace UnrealGameSync
 
 		async Task<bool> UpdateArchivesAsync(IPerforceConnection perforce, CancellationToken cancellationToken)
 		{
-			List<PerforceArchiveInfo> newArchives = new List<PerforceArchiveInfo>();
-
-			// Find all the zipped binaries under this stream
-			ConfigSection projectConfigSection = LatestProjectConfigFile.FindSection(_selectedProjectIdentifier);
-			if (projectConfigSection != null)
-			{
-				// Legacy
-				string? legacyEditorArchivePath = projectConfigSection.GetValue("ZippedBinariesPath", null);
-				if (legacyEditorArchivePath != null)
-				{
-					newArchives.Add(new PerforceArchiveInfo("Editor", "Editor", legacyEditorArchivePath, null, this));
-				}
-
-				// New style
-				foreach (string archiveValue in projectConfigSection.GetValues("Archives", new string[0]))
-				{
-					PerforceArchiveInfo? archive;
-					if (PerforceArchiveInfo.TryParseConfigEntry(archiveValue, this, out archive))
-					{
-						newArchives.Add(archive);
-					}
-				}
-
-				// Make sure the zipped binaries path exists
-				foreach (PerforceArchiveInfo newArchive in newArchives)
-				{
-					PerforceResponseList<FileLogRecord> response = await perforce.TryFileLogAsync(30, FileLogOptions.FullDescriptions, newArchive.DepotPath, cancellationToken);
-					if (response.Succeeded)
-					{
- 						// Build a new list of zipped binaries
-						foreach (FileLogRecord file in response.Data)
-						{
-							foreach(RevisionRecord revision in file.Revisions)
-							{
-								if (revision.Action != FileAction.Purge)
-								{
-									string[] tokens = revision.Description.Split(' ');
-									if (tokens[0].StartsWith("[CL") && tokens[1].EndsWith("]"))
-									{
-										int originalChangeNumber;
-										if (int.TryParse(tokens[1].Substring(0, tokens[1].Length - 1), out originalChangeNumber) && !newArchive.ChangeNumberToFileRevision.ContainsKey(originalChangeNumber))
-										{
-											newArchive.ChangeNumberToFileRevision[originalChangeNumber] = $"{newArchive.DepotPath}#{revision.RevisionNumber}";
-										}
-									}
-								}
-							}
-						}
-					}
-				}
-			}
+			List<PerforceArchiveInfo> newArchives = await PerforceArchive.EnumerateAsync(perforce, LatestProjectConfigFile, _selectedProjectIdentifier, cancellationToken);
 
 			// Check if the information has changed
 			if (!Enumerable.SequenceEqual(_archives, newArchives))
diff --git a/Engine/Source/Programs/UnrealGameSync/UnrealGameSyncCmd/Program.cs b/Engine/Source/Programs/UnrealGameSync/UnrealGameSyncCmd/Program.cs
--- a/Engine/Source/Programs/UnrealGameSync/UnrealGameSyncCmd/Program.cs
+++ b/Engine/Source/Programs/UnrealGameSync/UnrealGameSyncCmd/Program.cs
@@ -86,7 +86,7 @@ namespace UnrealGameSyncCmd
 				"Displays or updates the workspace or global sync filter"
 			),
 			new CommandInfo("sync", typeof(SyncCommand),
-				"ugs sync [change|'latest'] [-build] [-remove] [-only]",
+				"ugs sync [change|'latest'] [-build] [-binaries] [-remove] [-only]",
 				"Syncs the current workspace to the given changelist, optionally removing all local state."
 			),
 			new CommandInfo("clients", typeof(ClientsCommand),
@@ -574,6 +574,9 @@ namespace UnrealGameSyncCmd
 				[CommandLine("-Build")]
 				public bool Build { get; set; }
 
+				[CommandLine("-Binaries")]
+				public bool Binaries { get; set; }
+
 				[CommandLine("-NoGPF", Value = "false")]
 				[CommandLine("-NoProjectFiles", Value = "false")]
 				public bool ProjectFiles { get; set; } = true;
@@ -585,6 +588,24 @@ namespace UnrealGameSyncCmd
 				public bool Refilter { get; set; }
 			}
 
+			async Task<bool> IsCodeChangeAsync(IPerforceConnection perforce, int change)
+			{
+				DescribeRecord describeRecord = await perforce.DescribeAsync(change);
+				return IsCodeChange(describeRecord);
+			}
+
+			bool IsCodeChange(DescribeRecord describeRecord)
+			{
+				foreach (DescribeFileRecord file in describeRecord.Files)
+				{
+					if (PerforceUtils.CodeExtensions.Any(extension => file.DepotFile.EndsWith(extension, StringComparison.OrdinalIgnoreCase)))
+					{
+						return true;
+					}
+				}
+				return false;
+			}
+
 			public override async Task ExecuteAsync(CommandContext context)
 			{
 				ILogger logger = context.Logger;
@@ -601,10 +622,17 @@ namespace UnrealGameSyncCmd
 
 				changeString ??= "latest";
 
+				ProjectInfo projectInfo = state.CreateProjectInfo();
+				UserProjectSettings projectSettings = context.UserSettings.FindOrAddProjectSettings(projectInfo, settings, logger);
+
+				ConfigFile projectConfig = await ConfigUtils.ReadProjectConfigFileAsync(perforceClient, projectInfo, logger, CancellationToken.None);
+
+				bool syncLatest = String.Equals(changeString, "latest", StringComparison.OrdinalIgnoreCase);
+
 				int change;
 				if (!int.TryParse(changeString, out change))
 				{
-					if (String.Equals(changeString, "latest", StringComparison.OrdinalIgnoreCase))
+					if (syncLatest)
 					{
 						List<ChangesRecord> changes = await perforceClient.GetChangesAsync(ChangesOptions.None, 1, ChangeStatus.Submitted, $"//{settings.ClientName}/...");
 						change = changes[0].Number;
@@ -634,13 +662,51 @@ namespace UnrealGameSyncCmd
 				}
 				options |= WorkspaceUpdateOptions.RemoveFilteredFiles;
 
-				ProjectInfo projectInfo = state.CreateProjectInfo();
-				UserProjectSettings projectSettings = context.UserSettings.FindOrAddProjectSettings(projectInfo, settings, logger);
-
-				ConfigFile projectConfig = await ConfigUtils.ReadProjectConfigFileAsync(perforceClient, projectInfo, logger, CancellationToken.None);
 				string[] syncFilter = ReadSyncFilter(settings, context.UserSettings, projectConfig);
 
 				WorkspaceUpdateContext updateContext = new WorkspaceUpdateContext(change, options, BuildConfig.Development, syncFilter, projectSettings.BuildSteps, null);
+				if (syncOptions.Binaries)
+				{
+					List<PerforceArchiveInfo> archives = await PerforceArchive.EnumerateAsync(perforceClient, projectConfig, state.ProjectIdentifier, CancellationToken.None);
+
+					PerforceArchiveInfo? editorArchiveInfo = archives.FirstOrDefault(x => x.Name == IArchiveInfo.EditorArchiveType);
+					if (editorArchiveInfo == null)
+					{
+						throw new UserErrorException("No editor archives found for project");
+					}
+
+					KeyValuePair<int, string> revision = editorArchiveInfo.ChangeNumberToFileRevision.LastOrDefault(x => x.Key <= change);
+					if (revision.Key == 0)
+					{
+						throw new UserErrorException($"No editor archives found for CL {change}");
+					}
+
+					if (revision.Key < change)
+					{
+						int lastChange = revision.Key;
+
+						List<ChangesRecord> changeRecords = await perforceClient.GetChangesAsync(ChangesOptions.None, 1, ChangeStatus.Submitted, $"//{settings.ClientName}/...@{revision.Key + 1},{change}");
+						foreach (ChangesRecord changeRecord in changeRecords.OrderBy(x => x.Number))
+						{
+							if (await IsCodeChangeAsync(perforceClient, changeRecord.Number))
+							{
+								if (syncLatest)
+								{
+									updateContext.ChangeNumber = lastChange;
+								}
+								else
+								{
+									throw new UserErrorException($"No editor binaries found for CL {change} (last archive at CL {revision.Key}, but CL {changeRecord.Number} is a code change)");
+								}
+								break;
+							}
+							change = changeRecord.Number;
+						}
+					}
+
+					updateContext.Options |= WorkspaceUpdateOptions.SyncArchives;
+					updateContext.ArchiveTypeToArchive[IArchiveInfo.EditorArchiveType] = Tuple.Create<IArchiveInfo, string>(editorArchiveInfo, revision.Value);
+				}
 
 				WorkspaceUpdate update = new WorkspaceUpdate(updateContext);
 				(WorkspaceUpdateResult result, string message) = await update.ExecuteAsync(perforceClient.Settings, projectInfo, state, context.Logger, CancellationToken.None);
diff --git a/Engine/Source/Programs/UnrealGameSync/UnrealGameSyncShared/ArchiveInfo.cs b/Engine/Source/Programs/UnrealGameSync/UnrealGameSyncShared/ArchiveInfo.cs
new file mode 100644
--- /dev/null
+++ b/Engine/Source/Programs/UnrealGameSync/UnrealGameSyncShared/ArchiveInfo.cs
@@ -0,0 +1,187 @@
+// Copyright Epic Games, Inc. All Rights Reserved.
+
+using EpicGames.Core;
+using EpicGames.Perforce;
+using Microsoft.Extensions.Logging;
+using System;
+using System.Collections.Generic;
+using System.Diagnostics.CodeAnalysis;
+using System.IO;
+using System.Linq;
+using System.Text;
+using System.Threading;
+using System.Threading.Tasks;
+
+namespace UnrealGameSync
+{
+	public interface IArchiveInfo
+	{
+		public const string EditorArchiveType = "Editor";
+
+		string Name { get; }
+		string Type { get; }
+		string BasePath { get; }
+		string? Target { get; }
+
+		bool Exists();
+		bool TryGetArchiveKeyForChangeNumber(int changeNumber, [NotNullWhen(true)] out string? archiveKey);
+		Task<bool> DownloadArchive(IPerforceConnection perforce, string archiveKey, DirectoryReference localRootPath, FileReference manifestFileName, ILogger logger, ProgressValue progress, CancellationToken cancellationToken);
+	}
+
+	public class PerforceArchiveInfo : IArchiveInfo
+	{
+		public string Name { get; }
+		public string Type { get; }
+		public string DepotPath { get; }
+		public string? Target { get; }
+
+		public string BasePath
+		{
+			get { return DepotPath; }
+		}
+
+		// TODO: executable/configuration?
+		public SortedList<int, string> ChangeNumberToFileRevision = new SortedList<int, string>();
+
+		public PerforceArchiveInfo(string name, string type, string depotPath, string? target)
+		{
+			this.Name = name;
+			this.Type = type;
+			this.DepotPath = depotPath;
+			this.Target = target;
+		}
+
+		public override bool Equals(object? other)
+		{
+			PerforceArchiveInfo? otherArchive = other as PerforceArchiveInfo;
+			return otherArchive != null && Name == otherArchive.Name && Type == otherArchive.Type && DepotPath == otherArchive.DepotPath && Target == otherArchive.Target && Enumerable.SequenceEqual(ChangeNumberToFileRevision, otherArchive.ChangeNumberToFileRevision);
+		}
+
+		public override int GetHashCode()
+		{
+			throw new NotSupportedException();
+		}
+
+		public bool Exists()
+		{
+			return ChangeNumberToFileRevision.Count > 0;
+		}
+
+		public static bool TryParseConfigEntry(string text, [NotNullWhen(true)] out PerforceArchiveInfo? info)
+		{
+			ConfigObject obj = new ConfigObject(text);
+
+			string? name = obj.GetValue("Name", null);
+			if (name == null)
+			{
+				info = null;
+				return false;
+			}
+
+			string? depotPath = obj.GetValue("DepotPath", null);
+			if (depotPath == null)
+			{
+				info = null;
+				return false;
+			}
+
+			string? target = obj.GetValue("Target", null);
+
+			string type = obj.GetValue("Type", null) ?? name;
+
+			info = new PerforceArchiveInfo(name, type, depotPath, target);
+			return true;
+		}
+
+		public bool TryGetArchiveKeyForChangeNumber(int changeNumber, [NotNullWhen(true)] out string? archiveKey)
+		{
+			return ChangeNumberToFileRevision.TryGetValue(changeNumber, out archiveKey);
+		}
+
+		public async Task<bool> DownloadArchive(IPerforceConnection perforce, string archiveKey, DirectoryReference localRootPath, FileReference manifestFileName, ILogger logger, ProgressValue progress, CancellationToken cancellationToken)
+		{
+			FileReference tempZipFileName = new FileReference(Path.GetTempFileName());
+			try
+			{
+				PrintRecord record = await perforce.PrintAsync(tempZipFileName.FullName, archiveKey, cancellationToken);
+				if (tempZipFileName.ToFileInfo().Length == 0)
+				{
+					return false;
+				}
+				ArchiveUtils.ExtractFiles(tempZipFileName, localRootPath, manifestFileName, progress, logger);
+			}
+			finally
+			{
+				FileReference.SetAttributes(tempZipFileName, FileAttributes.Normal);
+				FileReference.Delete(tempZipFileName);
+			}
+
+			return true;
+		}
+
+		public override string ToString()
+		{
+			return Name;
+		}
+	}
+
+	public static class PerforceArchive
+	{
+		public static async Task<List<PerforceArchiveInfo>> EnumerateAsync(IPerforceConnection perforce, ConfigFile latestProjectConfigFile, string projectIdentifier, CancellationToken cancellationToken)
+		{
+			List<PerforceArchiveInfo> newArchives = new List<PerforceArchiveInfo>();
+
+			// Find all the zipped binaries under this stream
+			ConfigSection projectConfigSection = latestProjectConfigFile.FindSection(projectIdentifier);
+			if (projectConfigSection != null)
+			{
+				// Legacy
+				string? legacyEditorArchivePath = projectConfigSection.GetValue("ZippedBinariesPath", null);
+				if (legacyEditorArchivePath != null)
+				{
+					newArchives.Add(new PerforceArchiveInfo("Editor", "Editor", legacyEditorArchivePath, null));
+				}
+
+				// New style
+				foreach (string archiveValue in projectConfigSection.GetValues("Archives", new string[0]))
+				{
+					PerforceArchiveInfo? archive;
+					if (PerforceArchiveInfo.TryParseConfigEntry(archiveValue, out archive))
+					{
+						newArchives.Add(archive);
+					}
+				}
+
+				// Make sure the zipped binaries path exists
+				foreach (PerforceArchiveInfo newArchive in newArchives)
+				{
+					PerforceResponseList<FileLogRecord> response = await perforce.TryFileLogAsync(30, FileLogOptions.FullDescriptions, newArchive.DepotPath, cancellationToken);
+					if (response.Succeeded)
+					{
+						// Build a new list of zipped binaries
+						foreach (FileLogRecord file in response.Data)
+						{
+							foreach (RevisionRecord revision in file.Revisions)
+							{
+								if (revision.Action != FileAction.Purge)
+								{
+									string[] tokens = revision.Description.Split(' ');
+									if (tokens[0].StartsWith("[CL") && tokens[1].EndsWith("]"))
+									{
+										int originalChangeNumber;
+										if (int.TryParse(tokens[1].Substring(0, tokens[1].Length - 1), out originalChangeNumber) && !newArchive.ChangeNumberToFileRevision.ContainsKey(originalChangeNumber))
+										{
+											newArchive.ChangeNumberToFileRevision[originalChangeNumber] = $"{newArchive.DepotPath}#{revision.RevisionNumber}";
+										}
+									}
+								}
+							}
+						}
+					}
+				}
+			}
+
+			return newArchives;
+		}
+	}
+}
diff --git a/Engine/Source/Programs/UnrealGameSync/UnrealGameSyncShared/WorkspaceUpdate.cs b/Engine/Source/Programs/UnrealGameSync/UnrealGameSyncShared/WorkspaceUpdate.cs
--- a/Engine/Source/Programs/UnrealGameSync/UnrealGameSyncShared/WorkspaceUpdate.cs
+++ b/Engine/Source/Programs/UnrealGameSync/UnrealGameSyncShared/WorkspaceUpdate.cs
@@ -87,19 +87,6 @@ namespace UnrealGameSync
 		}
 	}
 
-	public interface IArchiveInfo
-	{
-		public const string EditorArchiveType = "Editor";
-
-		string Name { get; }
-		string Type { get; }
-		string BasePath { get; }
-		string? Target { get; }
-		bool Exists();
-		bool TryGetArchiveKeyForChangeNumber(int changeNumber, [NotNullWhen(true)] out string? archiveKey);
-		Task<bool> DownloadArchive(IPerforceConnection perforce, string archiveKey, DirectoryReference localRootPath, FileReference manifestFileName, ILogger logger, ProgressValue progress, CancellationToken cancellationToken);
-	}
-
 	public class WorkspaceUpdateContext
 	{
 		public DateTime StartTime = DateTime.UtcNow;
