UGS: Fix clean workspace task not correctly identifying writable files.

diff --git a/Engine/Source/Programs/UnrealGameSync/UnrealGameSync/FindFoldersToCleanTask.cs b/Engine/Source/Programs/UnrealGameSync/UnrealGameSync/FindFoldersToCleanTask.cs
--- a/Engine/Source/Programs/UnrealGameSync/UnrealGameSync/FindFoldersToCleanTask.cs
+++ b/Engine/Source/Programs/UnrealGameSync/UnrealGameSync/FindFoldersToCleanTask.cs
@@ -157,7 +157,7 @@ namespace UnrealGameSync
 					{
 						localFolder.FilesToSync.Add(new FileInfo(Path.Combine(localFolder.Directory.FullName, filePair.Key)));
 					}
-					else if(IsFileTypeWritable(filePair.Value?.HeadType ?? filePair.Value?.Type) && (localFile.Attributes & FileAttributes.ReadOnly) == 0 && !openClientPaths.Contains(filePair.Value?.ClientFile ?? ""))
+					else if(!IsFileTypeWritable(filePair.Value?.HeadType ?? filePair.Value?.Type) && (localFile.Attributes & FileAttributes.ReadOnly) == 0 && !openClientPaths.Contains(filePair.Value?.ClientFile ?? ""))
 					{
 						localFolder.FilesToSync.Add(localFile);
 					}
@@ -196,7 +196,7 @@ namespace UnrealGameSync
 					return type.IndexOf('w', idx) != -1;
 				}
 			}
-			return true;
+			return false;
 		}
 
 		void RemoveEmptyFolders(FolderToClean folder)
