UGS: Fix some build issues after recent changes to EpicGames.Perforce.

diff --git a/Engine/Source/Programs/UnrealGameSync/UnrealGameSync/Controls/LogControl.cs b/Engine/Source/Programs/UnrealGameSync/UnrealGameSync/Controls/LogControl.cs
--- a/Engine/Source/Programs/UnrealGameSync/UnrealGameSync/Controls/LogControl.cs
+++ b/Engine/Source/Programs/UnrealGameSync/UnrealGameSync/Controls/LogControl.cs
@@ -925,7 +925,7 @@ namespace UnrealGameSync
 			_logControl = inLogControl;
 		}
 
-		public void Log<TState>(LogLevel logLevel, EventId eventId, TState state, Exception exception, Func<TState, Exception, string> formatter)
+		public void Log<TState>(LogLevel logLevel, EventId eventId, TState state, Exception? exception, Func<TState, Exception?, string> formatter)
 		{
 			string message = formatter(state, exception);
 			_logControl.AppendLine(message);
diff --git a/Engine/Source/Programs/UnrealGameSync/UnrealGameSync/ToolUpdateMonitor.cs b/Engine/Source/Programs/UnrealGameSync/UnrealGameSync/ToolUpdateMonitor.cs
--- a/Engine/Source/Programs/UnrealGameSync/UnrealGameSync/ToolUpdateMonitor.cs
+++ b/Engine/Source/Programs/UnrealGameSync/UnrealGameSync/ToolUpdateMonitor.cs
@@ -402,7 +402,7 @@ namespace UnrealGameSync
 			{
 				FileReference zipFile = FileReference.Combine(nextToolZipsDir, String.Format("{0}.{1}.zip", toolName, idx));
 
-				PerforceResponse<PrintRecord> response = await perforce.TryPrintAsync(zipFile.FullName, $"{records[idx].DepotFile}#{records[idx].HeadRevision}", cancellationToken);
+				PerforceResponseList<PrintRecord> response = await perforce.TryPrintAsync(zipFile.FullName, $"{records[idx].DepotFile}#{records[idx].HeadRevision}", cancellationToken);
 				if(!response.Succeeded || !FileReference.Exists(zipFile))
 				{
 					logger.LogError("Unable to print {0}", records[idx].DepotFile);
diff --git a/Engine/Source/Programs/UnrealGameSync/UnrealGameSyncShared/OutputAdapters.cs b/Engine/Source/Programs/UnrealGameSync/UnrealGameSyncShared/OutputAdapters.cs
--- a/Engine/Source/Programs/UnrealGameSync/UnrealGameSyncShared/OutputAdapters.cs
+++ b/Engine/Source/Programs/UnrealGameSync/UnrealGameSyncShared/OutputAdapters.cs
@@ -28,7 +28,7 @@ namespace UnrealGameSync
 
 		public bool IsEnabled(LogLevel logLevel) => _inner.IsEnabled(logLevel);
 
-		public void Log<TState>(LogLevel logLevel, EventId eventId, TState state, Exception exception, Func<TState, Exception, string> formatter)
+		public void Log<TState>(LogLevel logLevel, EventId eventId, TState state, Exception? exception, Func<TState, Exception?, string> formatter)
 		{
 			_inner.Log(logLevel, eventId, state, exception, (state, exception) => _prefix + formatter(state, exception));
 		}
diff --git a/Engine/Source/Programs/UnrealGameSync/UnrealGameSyncShared/Utility.cs b/Engine/Source/Programs/UnrealGameSync/UnrealGameSyncShared/Utility.cs
--- a/Engine/Source/Programs/UnrealGameSync/UnrealGameSyncShared/Utility.cs
+++ b/Engine/Source/Programs/UnrealGameSync/UnrealGameSyncShared/Utility.cs
@@ -439,7 +439,7 @@ namespace UnrealGameSync
 				DirectoryReference.CreateDirectory(cacheFolder);
 
 				FileReference tempFile = new FileReference(String.Format("{0}.{1}.temp", cacheFile.FullName, Guid.NewGuid()));
-				PerforceResponse<PrintRecord> response = await perforce.TryPrintAsync(tempFile.FullName, depotPath, cancellationToken);
+				PerforceResponseList<PrintRecord> response = await perforce.TryPrintAsync(tempFile.FullName, depotPath, cancellationToken);
 				if (!response.Succeeded)
 				{
 					return null;
