Fix to 'changes' cmd -- clamping linecount to address out-of-bounds exception.

diff --git a/Engine/Source/Programs/UnrealGameSync/UnrealGameSyncCmd/Program.cs b/Engine/Source/Programs/UnrealGameSync/UnrealGameSyncCmd/Program.cs
--- a/Engine/Source/Programs/UnrealGameSync/UnrealGameSyncCmd/Program.cs
+++ b/Engine/Source/Programs/UnrealGameSync/UnrealGameSyncCmd/Program.cs
@@ -863,6 +863,8 @@ namespace UnrealGameSyncCmd
 							lines.Add(String.Empty);
 						}
 
+						lineCount = Math.Min(lineCount, lines.Count);
+
 						logger.LogInformation("  {Change,-9} {Type,-8} {Author,-15} {Description}", describeRecord.Number, type, author, lines[0]);
 						for (int lineIndex = 1; lineIndex < lineCount; lineIndex++)
 						{
