UGS : added the changelist context menu to current changelist

diff --git a/Engine/Source/Programs/UnrealGameSync/UnrealGameSync/Controls/WorkspaceControl.cs b/Engine/Source/Programs/UnrealGameSync/UnrealGameSync/Controls/WorkspaceControl.cs
--- a/Engine/Source/Programs/UnrealGameSync/UnrealGameSync/Controls/WorkspaceControl.cs
+++ b/Engine/Source/Programs/UnrealGameSync/UnrealGameSync/Controls/WorkspaceControl.cs
@@ -288,7 +288,7 @@ namespace UnrealGameSync
 
 		// When an author filter is applied this will be non-empty and != AuthorFilterPlaceholderText
 		string _authorFilterText = "";
-		
+
 		// Placeholder text that is in the control and cleared when the user starts editing.
 		static string _authorFilterPlaceholderText = "<username>";
 
@@ -378,7 +378,7 @@ namespace UnrealGameSync
 			_perforceMonitor.OnStreamChange += StreamChanged;
 			_perforceMonitor.OnLoginExpired += LoginExpired;
 
-			ILogger eventLogger = _serviceProvider.GetRequiredService < ILogger<EventMonitor>>();
+			ILogger eventLogger = _serviceProvider.GetRequiredService<ILogger<EventMonitor>>();
 			_eventMonitor = new EventMonitor(_apiUrl, PerforceUtils.GetClientOrDepotDirectoryName(SelectedProjectIdentifier), openProjectInfo.PerforceSettings.UserName, _serviceProvider);
 			_eventMonitor.OnUpdatesReady += UpdateReviewsCallback;
 
@@ -1054,7 +1054,7 @@ namespace UnrealGameSync
 
 		void UpdateComplete(WorkspaceUpdateContext context, WorkspaceUpdateResult result, string resultMessage)
 		{
-			if(_isDisposing)
+			if (_isDisposing)
 			{
 				return;
 			}
@@ -1083,7 +1083,7 @@ namespace UnrealGameSync
 				_owner.UpdateProgress();
 
 				HashSet<string> uncontrolledFiles = new HashSet<string>();
-				
+
 				if (SelectedProject.LocalPath != null && SelectedProject.LocalPath.EndsWith(".uprojectdirs", StringComparison.InvariantCultureIgnoreCase))
 				{
 					List<string> projectRoots = GetProjectRoots(SelectedProject.LocalPath);
@@ -1099,7 +1099,7 @@ namespace UnrealGameSync
 				}
 
 				ClobberWindow window = new ClobberWindow(context.ClobberFiles, uncontrolledFiles);
-				
+
 				if (window.ShowDialog(this) == DialogResult.OK)
 				{
 					StartWorkspaceUpdate(context, _updateCallback);
@@ -1431,7 +1431,7 @@ namespace UnrealGameSync
 			{
 				return false;
 			}
-			
+
 			// If this is a robomerge change, check if any filters will cause it to be hidden
 			if (IsRobomergeChange(change))
 			{
@@ -3055,7 +3055,7 @@ namespace UnrealGameSync
 				{
 					ShowEditorLaunchError($"{receipt.Launch} not found");
 				}
-				else 
+				else
 				{
 					if (_settings.EditorArgumentsPrompt && !ModifyEditorArguments())
 					{
@@ -3231,7 +3231,7 @@ namespace UnrealGameSync
 						string? name = config.GetValue("Name", null);
 						string? pattern = config.GetValue("Pattern", null);
 
-						if(name != null && pattern != null && StreamName != null)
+						if (name != null && pattern != null && StreamName != null)
 						{
 							string? depotName;
 							if (PerforceUtils.TryGetDepotName(StreamName, out depotName))
@@ -3244,7 +3244,7 @@ namespace UnrealGameSync
 							customFilters[name] = x => x.Streams.Any(y => patternRegex.IsMatch(y));
 						}
 					}
-					catch(Exception ex)
+					catch (Exception ex)
 					{
 						_logger.LogError(ex, "Unable to parse config filter '{Filter}'", filter);
 					}
@@ -3283,19 +3283,19 @@ namespace UnrealGameSync
 
 		private void UpdateStatusPanel_CrossThread()
 		{
-			_mainThreadSynchronizationContext.Post((o) => 
-			{ 
+			_mainThreadSynchronizationContext.Post((o) =>
+			{
 				if (!IsDisposed && !_isDisposing)
-				{ 
+				{
 					UpdateBuildSteps();
-					UpdateStatusPanel(); 
-				} 
+					UpdateStatusPanel();
+				}
 			}, null);
 		}
 
 		private void UpdateStatusPanel()
 		{
-			if(_workspace == null)
+			if (_workspace == null)
 			{
 				return;
 			}
@@ -3388,6 +3388,50 @@ namespace UnrealGameSync
 							summaryLine.AddText(" at changelist ");
 						}
 						summaryLine.AddLink(String.Format("{0}.", _workspace.CurrentChangeNumber), FontStyle.Regular, () => { SelectChange(_workspace.CurrentChangeNumber); });
+						summaryLine.AddLink(" \u25BE", 0, (P, R) =>
+						{
+							// Try to find the CL record in our stored list in order to avoid requesting the info from p4:
+							foreach (ListViewItem? item in BuildList.Items)
+							{
+								if (item != null)
+								{
+									ChangesRecord summary = (ChangesRecord)item.Tag;
+									if (summary != null)
+									{
+										// No need to keep on iterating if we're past the required CL :
+										if (summary.Number < _workspace.CurrentChangeNumber)
+										{
+											break;
+										}
+										else if (summary.Number == _workspace.CurrentChangeNumber)
+										{
+											_contextMenuChange = summary;
+											ShowChangelistContextMenu(StatusPanel, showTimeRelatedItems: false, new Point(R.Left, R.Bottom), ToolStripDropDownDirection.BelowRight);
+											return;
+										}
+									}
+								}
+							}
+
+							// Couldn't find the CL in our BuildList, get the info from p4 :
+							ChangesRecord? foundRecord = null;
+							Func<IPerforceConnection, CancellationToken, Task<bool>> getChangesRecordFunc = async (Perforce, CancellationToken) =>
+							{
+								List<ChangesRecord> changes = await Perforce.GetChangesAsync(ChangesOptions.None, 1, ChangeStatus.Submitted, $"//{_perforceSettings.ClientName}/...@{_workspace.CurrentChangeNumber},{_workspace.CurrentChangeNumber}", CancellationToken);
+								if (changes.Count == 1)
+								{
+									foundRecord = changes[0];
+								}
+								return true;
+							};
+
+							ModalTask<bool>? getChangesTask = PerforceModalTask.Execute<bool>(this, "Getting CL info", "Please wait...", _perforceSettings, getChangesRecordFunc, _logger);
+							if ((getChangesTask != null) && getChangesTask.Succeeded && (foundRecord != null))
+							{
+								_contextMenuChange = foundRecord;
+								ShowChangelistContextMenu(StatusPanel, showTimeRelatedItems: false, new Point(R.Left, R.Bottom), ToolStripDropDownDirection.BelowRight);
+							}
+						});
 					}
 					else
 					{
@@ -3432,7 +3476,7 @@ namespace UnrealGameSync
 				}
 
 				List<ToolDefinition> tools = _owner.ToolUpdateMonitor.Tools;
-				foreach(ToolDefinition tool in tools)
+				foreach (ToolDefinition tool in tools)
 				{
 					if (tool.Enabled)
 					{
@@ -4108,98 +4152,106 @@ namespace UnrealGameSync
 					else
 					{
 						_contextMenuChange = (ChangesRecord)hitTest.Item.Tag;
+						// Show time related context menu items if the click was made on the time column :
+						bool isTimeColumn = (hitTest.Item.SubItems.IndexOf(hitTest.SubItem) == TimeColumn.Index);
+						ShowChangelistContextMenu(BuildList, showTimeRelatedItems: isTimeColumn, args.Location);
+					}
+				}
+			}
+		}
 
-						BuildListContextMenu_WithdrawReview.Visible = (_eventMonitor.GetReviewByCurrentUser(_contextMenuChange.Number) != null);
-						BuildListContextMenu_StartInvestigating.Visible = !_eventMonitor.IsUnderInvestigationByCurrentUser(_contextMenuChange.Number);
-						BuildListContextMenu_FinishInvestigating.Visible = _eventMonitor.IsUnderInvestigation(_contextMenuChange.Number);
-
-						string? commentText;
-						bool hasExistingComment = _eventMonitor.GetCommentByCurrentUser(_contextMenuChange.Number, out commentText);
-						BuildListContextMenu_LeaveComment.Visible = !hasExistingComment;
-						BuildListContextMenu_EditComment.Visible = hasExistingComment;
-
-						bool isBusy = _workspace.IsBusy();
-						bool isCurrentChange = (_contextMenuChange.Number == _workspace.CurrentChangeNumber);
-						BuildListContextMenu_Sync.Visible = !isBusy;
-						BuildListContextMenu_Sync.Font = new Font(SystemFonts.MenuFont, isCurrentChange ? FontStyle.Regular : FontStyle.Bold);
-						BuildListContextMenu_SyncContentOnly.Visible = !isBusy && ShouldSyncPrecompiledEditor;
-						BuildListContextMenu_SyncOnlyThisChange.Visible = !isBusy && !isCurrentChange && _contextMenuChange.Number > _workspace.CurrentChangeNumber && _workspace.CurrentChangeNumber != -1;
-						BuildListContextMenu_Build.Visible = !isBusy && isCurrentChange && !ShouldSyncPrecompiledEditor;
-						BuildListContextMenu_Rebuild.Visible = !isBusy && isCurrentChange && !ShouldSyncPrecompiledEditor;
-						BuildListContextMenu_GenerateProjectFiles.Visible = !isBusy && isCurrentChange;
-						BuildListContextMenu_LaunchEditor.Visible = !isBusy && _contextMenuChange.Number == _workspace.CurrentChangeNumber;
-						BuildListContextMenu_LaunchEditor.Font = new Font(SystemFonts.MenuFont, FontStyle.Bold);
-						BuildListContextMenu_OpenVisualStudio.Visible = !isBusy && isCurrentChange;
-						BuildListContextMenu_Cancel.Visible = isBusy;
-
-						BisectState state = _workspaceState.BisectChanges.FirstOrDefault(x => x.Change == _contextMenuChange.Number)?.State ?? BisectState.Include;
-						bool isBisectMode = IsBisectModeEnabled();
-						BuildListContextMenu_Bisect_Pass.Visible = isBisectMode && state != BisectState.Pass;
-						BuildListContextMenu_Bisect_Fail.Visible = isBisectMode && state != BisectState.Fail;
-						BuildListContextMenu_Bisect_Exclude.Visible = isBisectMode && state != BisectState.Exclude;
-						BuildListContextMenu_Bisect_Include.Visible = isBisectMode && state != BisectState.Include;
-						BuildListContextMenu_Bisect_Separator.Visible = isBisectMode;
-
-						BuildListContextMenu_MarkGood.Visible = !isBisectMode;
-						BuildListContextMenu_MarkBad.Visible = !isBisectMode;
-						BuildListContextMenu_WithdrawReview.Visible = !isBisectMode;
-
-						EventSummary? summary = _eventMonitor.GetSummaryForChange(_contextMenuChange.Number);
-						bool starred = (summary != null && summary.LastStarReview != null && summary.LastStarReview.Type == EventType.Starred);
-						BuildListContextMenu_AddStar.Visible = !starred;
-						BuildListContextMenu_RemoveStar.Visible = starred;
+		private void ShowChangelistContextMenu(Control parentControl, bool showTimeRelatedItems, Point location, ToolStripDropDownDirection direction = ToolStripDropDownDirection.BelowRight)
+		{
+			if (_contextMenuChange != null)
+			{
+				BuildListContextMenu_WithdrawReview.Visible = (_eventMonitor.GetReviewByCurrentUser(_contextMenuChange.Number) != null);
+				BuildListContextMenu_StartInvestigating.Visible = !_eventMonitor.IsUnderInvestigationByCurrentUser(_contextMenuChange.Number);
+				BuildListContextMenu_FinishInvestigating.Visible = _eventMonitor.IsUnderInvestigation(_contextMenuChange.Number);
 
-						bool isTimeColumn = (hitTest.Item.SubItems.IndexOf(hitTest.SubItem) == TimeColumn.Index);
-						BuildListContextMenu_TimeZoneSeparator.Visible = isTimeColumn;
-						BuildListContextMenu_ShowLocalTimes.Visible = isTimeColumn;
-						BuildListContextMenu_ShowLocalTimes.Checked = _settings.ShowLocalTimes;
-						BuildListContextMenu_ShowServerTimes.Visible = isTimeColumn;
-						BuildListContextMenu_ShowServerTimes.Checked = !_settings.ShowLocalTimes;
-
-						int customToolStart = BuildListContextMenu.Items.IndexOf(BuildListContextMenu_CustomTool_Start) + 1;
-						int customToolEnd = BuildListContextMenu.Items.IndexOf(BuildListContextMenu_CustomTool_End);
-						while (customToolEnd > customToolStart)
-						{
-							BuildListContextMenu.Items.RemoveAt(customToolEnd - 1);
-							customToolEnd--;
-						}
+				string? commentText;
+				bool hasExistingComment = _eventMonitor.GetCommentByCurrentUser(_contextMenuChange.Number, out commentText);
+				BuildListContextMenu_LeaveComment.Visible = !hasExistingComment;
+				BuildListContextMenu_EditComment.Visible = hasExistingComment;
+
+				bool isBusy = _workspace.IsBusy();
+				bool isCurrentChange = (_contextMenuChange.Number == _workspace.CurrentChangeNumber);
+				BuildListContextMenu_Sync.Visible = !isBusy;
+				BuildListContextMenu_Sync.Font = new Font(SystemFonts.MenuFont, isCurrentChange ? FontStyle.Regular : FontStyle.Bold);
+				BuildListContextMenu_SyncContentOnly.Visible = !isBusy && ShouldSyncPrecompiledEditor;
+				BuildListContextMenu_SyncOnlyThisChange.Visible = !isBusy && !isCurrentChange && _contextMenuChange.Number > _workspace.CurrentChangeNumber && _workspace.CurrentChangeNumber != -1;
+				BuildListContextMenu_Build.Visible = !isBusy && isCurrentChange && !ShouldSyncPrecompiledEditor;
+				BuildListContextMenu_Rebuild.Visible = !isBusy && isCurrentChange && !ShouldSyncPrecompiledEditor;
+				BuildListContextMenu_GenerateProjectFiles.Visible = !isBusy && isCurrentChange;
+				BuildListContextMenu_LaunchEditor.Visible = !isBusy && _contextMenuChange.Number == _workspace.CurrentChangeNumber;
+				BuildListContextMenu_LaunchEditor.Font = new Font(SystemFonts.MenuFont, FontStyle.Bold);
+				BuildListContextMenu_OpenVisualStudio.Visible = !isBusy && isCurrentChange;
+				BuildListContextMenu_Cancel.Visible = isBusy;
+
+				BisectState state = _workspaceState.BisectChanges.FirstOrDefault(x => x.Change == _contextMenuChange.Number)?.State ?? BisectState.Include;
+				bool isBisectMode = IsBisectModeEnabled();
+				BuildListContextMenu_Bisect_Pass.Visible = isBisectMode && state != BisectState.Pass;
+				BuildListContextMenu_Bisect_Fail.Visible = isBisectMode && state != BisectState.Fail;
+				BuildListContextMenu_Bisect_Exclude.Visible = isBisectMode && state != BisectState.Exclude;
+				BuildListContextMenu_Bisect_Include.Visible = isBisectMode && state != BisectState.Include;
+				BuildListContextMenu_Bisect_Separator.Visible = isBisectMode;
+
+				BuildListContextMenu_MarkGood.Visible = !isBisectMode;
+				BuildListContextMenu_MarkBad.Visible = !isBisectMode;
+				BuildListContextMenu_WithdrawReview.Visible = !isBisectMode;
 
-						ConfigFile? projectConfigFile = _perforceMonitor.LatestProjectConfigFile;
-						if (projectConfigFile != null)
-						{
-							Dictionary<string, string> variables = _workspace.GetVariables(GetEditorBuildConfig(), _contextMenuChange.Number, -1);
+				EventSummary? summary = _eventMonitor.GetSummaryForChange(_contextMenuChange.Number);
+				bool starred = (summary != null && summary.LastStarReview != null && summary.LastStarReview.Type == EventType.Starred);
+				BuildListContextMenu_AddStar.Visible = !starred;
+				BuildListContextMenu_RemoveStar.Visible = starred;
 
-							string[] changeContextMenuEntries = projectConfigFile.GetValues("Options.ContextMenu", new string[0]);
-							foreach (string changeContextMenuEntry in changeContextMenuEntries)
-							{
-								ConfigObject obj = new ConfigObject(changeContextMenuEntry);
+				BuildListContextMenu_TimeZoneSeparator.Visible = showTimeRelatedItems;
+				BuildListContextMenu_ShowLocalTimes.Visible = showTimeRelatedItems;
+				BuildListContextMenu_ShowLocalTimes.Checked = _settings.ShowLocalTimes;
+				BuildListContextMenu_ShowServerTimes.Visible = showTimeRelatedItems;
+				BuildListContextMenu_ShowServerTimes.Checked = !_settings.ShowLocalTimes;
 
-								string? label = obj.GetValue("Label");
-								string? execute = obj.GetValue("Execute");
-								string? arguments = obj.GetValue("Arguments");
+				int customToolStart = BuildListContextMenu.Items.IndexOf(BuildListContextMenu_CustomTool_Start) + 1;
+				int customToolEnd = BuildListContextMenu.Items.IndexOf(BuildListContextMenu_CustomTool_End);
+				while (customToolEnd > customToolStart)
+				{
+					BuildListContextMenu.Items.RemoveAt(customToolEnd - 1);
+					customToolEnd--;
+				}
 
-								if (label != null && execute != null)
-								{
-									label = Utility.ExpandVariables(label, variables);
-									execute = Utility.ExpandVariables(execute, variables);
-									arguments = Utility.ExpandVariables(arguments ?? "", variables);
+				ConfigFile? projectConfigFile = _perforceMonitor.LatestProjectConfigFile;
+				if (projectConfigFile != null)
+				{
+					Dictionary<string, string> variables = _workspace.GetVariables(GetEditorBuildConfig(), _contextMenuChange.Number, -1);
 
-									ToolStripMenuItem item = new ToolStripMenuItem(label, null, new EventHandler((o, a) => SafeProcessStart(execute, arguments)));
+					string[] changeContextMenuEntries = projectConfigFile.GetValues("Options.ContextMenu", new string[0]);
+					foreach (string changeContextMenuEntry in changeContextMenuEntries)
+					{
+						ConfigObject obj = new ConfigObject(changeContextMenuEntry);
 
-									BuildListContextMenu.Items.Insert(customToolEnd, item);
-									customToolEnd++;
-								}
-							}
-						}
+						string? label = obj.GetValue("Label");
+						string? execute = obj.GetValue("Execute");
+						string? arguments = obj.GetValue("Arguments");
+
+						if (label != null && execute != null)
+						{
+							label = Utility.ExpandVariables(label, variables);
+							execute = Utility.ExpandVariables(execute, variables);
+							arguments = Utility.ExpandVariables(arguments ?? "", variables);
 
-						BuildListContextMenu_CustomTool_End.Visible = (customToolEnd > customToolStart);
+							ToolStripMenuItem item = new ToolStripMenuItem(label, null, new EventHandler((o, a) => SafeProcessStart(execute, arguments)));
 
-						string? swarmUrl;
-						BuildListContextMenu_ViewInSwarm.Visible = projectConfigFile != null && TryGetProjectSetting(projectConfigFile, "SwarmURL", out swarmUrl);
-							
-						BuildListContextMenu.Show(BuildList, args.Location);
+							BuildListContextMenu.Items.Insert(customToolEnd, item);
+							customToolEnd++;
+						}
 					}
 				}
+
+				BuildListContextMenu_CustomTool_End.Visible = (customToolEnd > customToolStart);
+
+				string? swarmUrl;
+				BuildListContextMenu_ViewInSwarm.Visible = projectConfigFile != null && TryGetProjectSetting(projectConfigFile, "SwarmURL", out swarmUrl);
+
+				BuildListContextMenu.Show(parentControl, location, direction);
 			}
 		}
 
@@ -4525,7 +4577,7 @@ namespace UnrealGameSync
 
 		public void SyncChange(int changeNumber, bool syncOnly, WorkspaceUpdateCallback? callback)
 		{
-			if(_workspace != null)
+			if (_workspace != null)
 			{
 				_owner.ShowAndActivate();
 				SelectChange(changeNumber);
@@ -5061,7 +5113,7 @@ namespace UnrealGameSync
 						{
 							changeNumber = change.Number;
 						}
-						
+
 						return true;
 					}
 				}
@@ -5236,7 +5288,7 @@ namespace UnrealGameSync
 
 		private bool CanRunStep(BuildStep step)
 		{
-			if(step.ToolId != Guid.Empty)
+			if (step.ToolId != Guid.Empty)
 			{
 				string? toolName = _owner.ToolUpdateMonitor.GetToolName(step.ToolId);
 				if (toolName == null)
@@ -5411,7 +5463,7 @@ namespace UnrealGameSync
 			}
 
 			// Create the expanded task objects
-			return userBuildStepObjects.Values.Select(x => new BuildStep(x)).OrderBy(x => (x.OrderIndex == -1)? 10000 : x.OrderIndex).ToList();
+			return userBuildStepObjects.Values.Select(x => new BuildStep(x)).OrderBy(x => (x.OrderIndex == -1) ? 10000 : x.OrderIndex).ToList();
 		}
 
 		private void OptionsContextMenu_SyncPrecompiledBinaries_Click(object sender, EventArgs e)
@@ -5823,7 +5875,7 @@ namespace UnrealGameSync
 				changeNumberToBisectState[changeNumberToBisectState.Keys.Min()] = BisectState.Pass;
 				changeNumberToBisectState[changeNumberToBisectState.Keys.Max()] = BisectState.Fail;
 
-				_workspaceState.BisectChanges = changeNumberToBisectState.Select(x => new BisectEntry {  Change = x.Key, State = x.Value }).ToList();
+				_workspaceState.BisectChanges = changeNumberToBisectState.Select(x => new BisectEntry { Change = x.Key, State = x.Value }).ToList();
 				_workspaceState.Save(_logger);
 
 				UpdateBuildList();
@@ -6007,7 +6059,7 @@ namespace UnrealGameSync
 
 		private void FilterContextMenu_Robomerge_ShowAll_Click(object sender, EventArgs e)
 		{
-			_settings.ShowRobomerge = UserSettings.RobomergeShowChangesOption.All;	
+			_settings.ShowRobomerge = UserSettings.RobomergeShowChangesOption.All;
 			_settings.Save(_logger);
 			UpdateBuildListFilter();
 		}
@@ -6030,7 +6082,7 @@ namespace UnrealGameSync
 		{
 			_settings.AnnotateRobmergeChanges ^= true;
 			_settings.Save(_logger);
-			BuildList.Items.Clear();	// need to reload to update user names
+			BuildList.Items.Clear();    // need to reload to update user names
 
 			UpdateBuildListFilter();
 		}
