UGS: Speculative fix for out-of-bounds access reported in Sentry.

diff --git a/Engine/Source/Programs/UnrealGameSync/UnrealGameSync/Controls/WorkspaceControl.cs b/Engine/Source/Programs/UnrealGameSync/UnrealGameSync/Controls/WorkspaceControl.cs
--- a/Engine/Source/Programs/UnrealGameSync/UnrealGameSync/Controls/WorkspaceControl.cs
+++ b/Engine/Source/Programs/UnrealGameSync/UnrealGameSync/Controls/WorkspaceControl.cs
@@ -4829,20 +4829,26 @@ namespace UnrealGameSync
 						return;
 					}
 
-					Rectangle cisBounds = args.Item.SubItems[CISColumn.Index].Bounds;
-					if (cisBounds.Contains(clientPoint))
+					if (CISColumn.Index < args.Item.SubItems.Count)
 					{
-						return;
+						Rectangle cisBounds = args.Item.SubItems[CISColumn.Index].Bounds;
+						if (cisBounds.Contains(clientPoint))
+						{
+							return;
+						}
 					}
 
-					Rectangle descriptionBounds = args.Item.SubItems[DescriptionColumn.Index].Bounds;
-					if (descriptionBounds.Contains(clientPoint))
+					if (DescriptionColumn.Index < args.Item.SubItems.Count)
 					{
-						ChangeLayoutInfo layoutInfo = GetChangeLayoutInfo(change);
-						if (layoutInfo.DescriptionBadges.Count == 0 || clientPoint.X < GetBadgeListLocation(layoutInfo.DescriptionBadges, descriptionBounds, HorizontalAlign.Right, VerticalAlignment.Middle).X - 2)
+						Rectangle descriptionBounds = args.Item.SubItems[DescriptionColumn.Index].Bounds;
+						if (descriptionBounds.Contains(clientPoint))
 						{
-							BuildListToolTip.Show(change.Description, BuildList, new Point(descriptionBounds.Left, descriptionBounds.Bottom + 2));
-							return;
+							ChangeLayoutInfo layoutInfo = GetChangeLayoutInfo(change);
+							if (layoutInfo.DescriptionBadges.Count == 0 || clientPoint.X < GetBadgeListLocation(layoutInfo.DescriptionBadges, descriptionBounds, HorizontalAlign.Right, VerticalAlignment.Middle).X - 2)
+							{
+								BuildListToolTip.Show(change.Description, BuildList, new Point(descriptionBounds.Left, descriptionBounds.Bottom + 2));
+								return;
+							}
 						}
 					}
 
@@ -4872,7 +4878,7 @@ namespace UnrealGameSync
 						{
 							AppendUserList(summaryText, "\n", "Starred by {0}.", new EventData[] { summary.LastStarReview });
 						}
-						if (summaryText.Length > 0)
+						if (summaryText.Length > 0 && StatusColumn.Index < args.Item.SubItems.Count)
 						{
 							Rectangle summaryBounds = args.Item.SubItems[StatusColumn.Index].Bounds;
 							BuildListToolTip.Show(summaryText.ToString(), BuildList, new Point(summaryBounds.Left, summaryBounds.Bottom));
