UGS: Improve performance of ranged changelist queries. (PR #9584)

* Querying for changelists using [REDACTED_HANDLE],m syntax allows the Perforce server to execute queries using an optimized codepath. Attempt to query for code changes using this syntax by default (with an arbitrary lower bound for the range), and fallback to the legacy query syntax if we don't find any results.
* Use the -e argument when executing p4 changes commands to specify a minimum changelist number. This query syntax is more efficient than specifying a wildcard with revision range.
* Use PerforceUtils.CodeExtensions as the list of code extensions, to unify with UAT/Horde.

diff --git a/Engine/Source/Programs/UnrealGameSync/UnrealGameSync/PerforceMonitor.cs b/Engine/Source/Programs/UnrealGameSync/UnrealGameSync/PerforceMonitor.cs
--- a/Engine/Source/Programs/UnrealGameSync/UnrealGameSync/PerforceMonitor.cs
+++ b/Engine/Source/Programs/UnrealGameSync/UnrealGameSync/PerforceMonitor.cs
@@ -223,7 +223,7 @@ namespace UnrealGameSync
 					}
 				}
 
-				// Wait for another request, or scan for new builds after a timeout
+                // Wait for another request, or scan for new builds after a timeout
 				Task delayTask = Task.Delay(TimeSpan.FromMinutes(IsActive ? 2 : 10), cancellationToken);
 				await Task.WhenAny(nextRefreshTask, delayTask);
 			}
@@ -284,7 +284,7 @@ namespace UnrealGameSync
 			}
 			else
 			{
-				newChanges = await perforce.GetChangesAsync(ChangesOptions.IncludeTimes | ChangesOptions.LongOutput, -1, ChangeStatus.Submitted, depotPaths.Select(x => $"{x}@>{newestChangeNumber}").ToArray(), cancellationToken);
+				newChanges = await perforce.GetChangesAsync(ChangesOptions.IncludeTimes | ChangesOptions.LongOutput, null, newestChangeNumber + 1, -1, ChangeStatus.Submitted, null, depotPaths, cancellationToken);
 			}
 
 			// Remove anything we already have
diff --git a/Engine/Source/Programs/UnrealGameSync/UnrealGameSyncShared/WorkspaceUpdate.cs b/Engine/Source/Programs/UnrealGameSync/UnrealGameSyncShared/WorkspaceUpdate.cs
--- a/Engine/Source/Programs/UnrealGameSync/UnrealGameSyncShared/WorkspaceUpdate.cs
+++ b/Engine/Source/Programs/UnrealGameSync/UnrealGameSyncShared/WorkspaceUpdate.cs
@@ -913,14 +913,32 @@ namespace UnrealGameSync
 							branchOrStreamName = PerforceUtils.GetClientOrDepotDirectoryName(files[0].DepotFile);
 						}
 
-						// Find the last code change before this changelist. For consistency in versioning between local builds and precompiled binaries, we need to use the last submitted code changelist as our version number.
-						string[] codeFilter = new string[] { ".cs", ".h", ".cpp", ".usf", ".ush", ".uproject", ".uplugin" }.SelectMany(x => syncPaths.Select(y => String.Format("{0}{1}@<={2}", y, x, Context.ChangeNumber))).ToArray();
-
-						PerforceResponseList<ChangesRecord> codeChanges = await perforce.TryGetChangesAsync(ChangesOptions.None, 1, ChangeStatus.Submitted, codeFilter, cancellationToken);
-						if (!codeChanges.Succeeded)
-						{
-							return (WorkspaceUpdateResult.FailedToSync, $"Couldn't determine last code changelist before CL {Context.ChangeNumber}.");
-						}
+                        PerforceResponseList<ChangesRecord>? codeChanges = null;
+
+                        // First, check the most recent changes (this is a much cheaper query because of revcx optimization in the Perforce server.)
+                        int optimisticRangeStart = Context.ChangeNumber - 1000;
+                        if (optimisticRangeStart >= 1)
+                        {
+                            string[] codeFilter = PerforceUtils.CodeExtensions.SelectMany(x => syncPaths.Select(y => $"{y}{x}@{optimisticRangeStart},{Context.ChangeNumber}")).ToArray();
+							codeChanges = await perforce.TryGetChangesAsync(ChangesOptions.None, 1, ChangeStatus.Submitted, codeFilter, cancellationToken);
+
+							if (!codeChanges.Succeeded)
+                            {
+                                return (WorkspaceUpdateResult.FailedToSync, $"Couldn't determine last code changelist before CL {Context.ChangeNumber}.");
+                            }
+                        }
+
+                        // If no change found in recent changes, do the full and expensive check.
+                        if (codeChanges == null || codeChanges.Count == 0)
+                        {
+                            string[] codeFilter = PerforceUtils.CodeExtensions.SelectMany(x => syncPaths.Select(y => $"{y}{x}@<={Context.ChangeNumber}")).ToArray();
+                            codeChanges = await perforce.TryGetChangesAsync(ChangesOptions.None, 1, ChangeStatus.Submitted, codeFilter, cancellationToken);
+
+							if (!codeChanges.Succeeded)
+                            {
+                                return (WorkspaceUpdateResult.FailedToSync, $"Couldn't determine last code changelist before CL {Context.ChangeNumber}.");
+                            }
+                        }
 						if (codeChanges.Count == 0)
 						{
 							return (WorkspaceUpdateResult.FailedToSync, $"Could not find any code changes before CL {Context.ChangeNumber}.");
