UGS: Disable debug/trace level output for syncing.

diff --git a/Engine/Source/Programs/UnrealGameSync/UnrealGameSync/Controls/LogControl.cs b/Engine/Source/Programs/UnrealGameSync/UnrealGameSync/Controls/LogControl.cs
--- a/Engine/Source/Programs/UnrealGameSync/UnrealGameSync/Controls/LogControl.cs
+++ b/Engine/Source/Programs/UnrealGameSync/UnrealGameSync/Controls/LogControl.cs
@@ -913,6 +913,8 @@ namespace UnrealGameSync
 
 	public class LogControlTextWriter : ILogger
 	{
+		const LogLevel MinLogLevel = LogLevel.Information;
+
 		class NullDisposable : IDisposable
 		{
 			public void Dispose() { }
@@ -927,6 +929,11 @@ namespace UnrealGameSync
 
 		public void Log<TState>(LogLevel logLevel, EventId eventId, TState state, Exception? exception, Func<TState, Exception?, string> formatter)
 		{
+			if (logLevel < MinLogLevel)
+			{
+				return;
+			}
+
 			string message = formatter(state, exception);
 			_logControl.AppendLine(message);
 
