UGS: Add random variance to background refresh times, to avoid DDOSing Perforce server. Also increase refresh times to 5m (when active) or 30m (when running the background). (#9584)

diff --git a/Engine/Source/Programs/UnrealGameSync/UnrealGameSync/PerforceMonitor.cs b/Engine/Source/Programs/UnrealGameSync/UnrealGameSync/PerforceMonitor.cs
--- a/Engine/Source/Programs/UnrealGameSync/UnrealGameSync/PerforceMonitor.cs
+++ b/Engine/Source/Programs/UnrealGameSync/UnrealGameSync/PerforceMonitor.cs
@@ -224,7 +224,13 @@ namespace UnrealGameSync
 				}
 
                 // Wait for another request, or scan for new builds after a timeout
-				Task delayTask = Task.Delay(TimeSpan.FromMinutes(IsActive ? 2 : 10), cancellationToken);
+				// Add random deviation to refresh event, to try and combat many UGS clients getting in sync and DDoSing Perforce
+                TimeSpan baseDelay = TimeSpan.FromMinutes(IsActive ? 5 : 30);
+
+                Random random = new Random();
+                TimeSpan randomDeviation = TimeSpan.FromSeconds(random.Next(0, 60));
+
+                Task delayTask = Task.Delay(baseDelay + randomDeviation, cancellationToken);
 				await Task.WhenAny(nextRefreshTask, delayTask);
 			}
 		}
