UGS: Support for manually refreshing list of changes, by pressing Ctrl+R or F5.

diff --git a/Engine/Source/Programs/UnrealGameSync/UnrealGameSync/Controls/WorkspaceControl.cs b/Engine/Source/Programs/UnrealGameSync/UnrealGameSync/Controls/WorkspaceControl.cs
--- a/Engine/Source/Programs/UnrealGameSync/UnrealGameSync/Controls/WorkspaceControl.cs
+++ b/Engine/Source/Programs/UnrealGameSync/UnrealGameSync/Controls/WorkspaceControl.cs
@@ -5602,6 +5602,14 @@ namespace UnrealGameSync
 
 		private void BuildList_KeyDown(object sender, KeyEventArgs args)
 		{
+			if (args.KeyCode == Keys.F5 || (args.Control && args.KeyCode == Keys.R))
+			{
+				if (_perforceMonitor != null)
+				{
+					_perforceMonitor.Refresh();
+				}
+			}
+
 			if (args.Control && args.KeyCode == Keys.C && BuildList.SelectedItems.Count > 0)
 			{
 				int selectedChange = ((ChangesRecord)BuildList.SelectedItems[0].Tag).Number;
