UGS: Duplicate shared projects which UGS depends on, including files from their original locations, so we can upgrade shared projects to NET 6 without requiring all users to have it installed. (In-engine projects all use NET SDK distributed with the engine, whereas UGS does not).

diff --git a/Engine/Source/Programs/UnrealGameSync/Shared/EpicGames.Core/EpicGames.Core.csproj b/Engine/Source/Programs/UnrealGameSync/Shared/EpicGames.Core/EpicGames.Core.csproj
new file mode 100644
--- /dev/null
+++ b/Engine/Source/Programs/UnrealGameSync/Shared/EpicGames.Core/EpicGames.Core.csproj
@@ -0,0 +1,51 @@
+﻿<Project Sdk="Microsoft.NET.Sdk">
+  <PropertyGroup>
+    <TargetFramework>netcoreapp3.1</TargetFramework>
+    <Configuration Condition=" '$(Configuration)' == '' ">Development</Configuration>
+    <OutputType>Library</OutputType>
+    <Deterministic>true</Deterministic>
+    <RootNamespace>EpicGames.Core</RootNamespace>
+    <GenerateAssemblyInfo>false</GenerateAssemblyInfo>
+    <GenerateTargetFrameworkAttribute>false</GenerateTargetFrameworkAttribute>
+    <Configurations>Debug;Release;Development;Analyze</Configurations>
+    <Nullable>enable</Nullable>
+    <AllowUnsafeBlocks>true</AllowUnsafeBlocks>
+    <DebugType>pdbonly</DebugType>
+    <RunAnalyzersDuringBuild>false</RunAnalyzersDuringBuild>
+  </PropertyGroup>
+  <PropertyGroup Condition=" '$(Configuration)|$(Platform)' == 'Debug|AnyCPU' ">
+    <DocumentationFile></DocumentationFile>
+  </PropertyGroup>
+  <PropertyGroup Condition="'$(Configuration)|$(Platform)' == 'Development|AnyCPU'">
+    <DocumentationFile></DocumentationFile>
+    <Optimize>true</Optimize>
+  </PropertyGroup>
+  <PropertyGroup Condition="'$(Configuration)|$(Platform)' == 'Release|AnyCPU'">
+    <DocumentationFile></DocumentationFile>
+    <Optimize>true</Optimize>
+  </PropertyGroup>
+  <PropertyGroup Condition=" '$(Configuration)|$(Platform)' == 'Analyze|AnyCPU' ">
+    <RunAnalyzersDuringBuild>True</RunAnalyzersDuringBuild>
+    <TreatWarningsAsErrors>false</TreatWarningsAsErrors>
+    <GenerateDocumentationFile>True</GenerateDocumentationFile>
+    <DocumentationFile></DocumentationFile>
+  </PropertyGroup>
+  <ItemGroup>
+    <Compile Include="../../../Shared\MetaData.cs">
+      <Link>Properties\MetaData.cs</Link>
+    </Compile>
+  </ItemGroup>
+  <ItemGroup>
+    <Compile Include="../../../Shared/EpicGames.Core/**/*.cs">
+      <Link>%(RecursiveDir)/%(Filename)%(Extension)</Link>
+    </Compile>
+  </ItemGroup>
+  <ItemGroup>
+    <PackageReference Include="JetBrains.Annotations" Version="2022.1.0" />
+    <PackageReference Include="Microsoft.Extensions.Logging" Version="6.0.0" />
+    <PackageReference Include="Microsoft.Extensions.Logging.Abstractions" Version="6.0.1" />
+    <PackageReference Include="Microsoft.Extensions.ObjectPool" Version="6.0.8" />
+    <PackageReference Include="OpenTracing" Version="0.12.1" />
+    <PackageReference Include="System.Memory" Version="4.5.5" />
+  </ItemGroup>
+</Project>
\ No newline at end of file
diff --git a/Engine/Source/Programs/UnrealGameSync/Shared/EpicGames.Jupiter/EpicGames.Jupiter.csproj b/Engine/Source/Programs/UnrealGameSync/Shared/EpicGames.Jupiter/EpicGames.Jupiter.csproj
new file mode 100644
--- /dev/null
+++ b/Engine/Source/Programs/UnrealGameSync/Shared/EpicGames.Jupiter/EpicGames.Jupiter.csproj
@@ -0,0 +1,40 @@
+<Project Sdk="Microsoft.NET.Sdk">
+  <PropertyGroup>
+    <TargetFramework>netcoreapp3.1</TargetFramework>
+    <Configuration Condition=" '$(Configuration)' == '' ">Development</Configuration>
+    <OutputType>Library</OutputType>
+    <Deterministic>true</Deterministic>
+    <RootNamespace>EpicGames.Jupiter</RootNamespace>
+    <GenerateAssemblyInfo>false</GenerateAssemblyInfo>
+    <GenerateTargetFrameworkAttribute>false</GenerateTargetFrameworkAttribute>
+    <Configurations>Debug;Release;Development</Configurations>
+    <DebugType>pdbonly</DebugType>
+  </PropertyGroup>
+  <PropertyGroup Condition=" '$(Configuration)|$(Platform)' == 'Debug|AnyCPU' ">
+    <DocumentationFile></DocumentationFile>
+  </PropertyGroup>
+  <PropertyGroup Condition="'$(Configuration)|$(Platform)' == 'Development|AnyCPU'">
+    <DocumentationFile></DocumentationFile>
+    <Optimize>true</Optimize>
+  </PropertyGroup>
+  <PropertyGroup Condition="'$(Configuration)|$(Platform)' == 'Release|AnyCPU'">
+    <DocumentationFile></DocumentationFile>
+    <Optimize>true</Optimize>
+  </PropertyGroup>
+  <ItemGroup>
+    <Compile Include="../../../Shared/MetaData.cs">
+      <Link>Properties\MetaData.cs</Link>
+    </Compile>
+  </ItemGroup>
+  <ItemGroup>
+    <Compile Include="../../../Shared/EpicGames.Jupiter/**/*.cs">
+      <Link>%(RecursiveDir)/%(Filename)%(Extension)</Link>
+    </Compile>
+  </ItemGroup>
+  <ItemGroup>
+    <PackageReference Include="System.Memory" Version="4.5.5" />
+  </ItemGroup>
+  <ItemGroup>
+    <ProjectReference Include="..\EpicGames.Core\EpicGames.Core.csproj" />
+  </ItemGroup>
+</Project>
\ No newline at end of file
diff --git a/Engine/Source/Programs/UnrealGameSync/Shared/EpicGames.OIDC/EpicGames.OIDC.csproj b/Engine/Source/Programs/UnrealGameSync/Shared/EpicGames.OIDC/EpicGames.OIDC.csproj
new file mode 100644
--- /dev/null
+++ b/Engine/Source/Programs/UnrealGameSync/Shared/EpicGames.OIDC/EpicGames.OIDC.csproj
@@ -0,0 +1,40 @@
+<Project Sdk="Microsoft.NET.Sdk">
+
+  <PropertyGroup>
+    <OutputType>Library</OutputType>
+    <RootNamespace>EpicGames.OIDC</RootNamespace>
+    <TargetFramework>netcoreapp3.1</TargetFramework>
+    <Nullable>enable</Nullable>
+    <GenerateAssemblyInfo>false</GenerateAssemblyInfo>
+    <GenerateTargetFrameworkAttribute>false</GenerateTargetFrameworkAttribute>
+    <Configurations>Debug;Release;Development</Configurations>
+  </PropertyGroup>
+
+  <PropertyGroup Condition=" '$(Configuration)|$(Platform)' == 'Debug|AnyCPU' ">
+  </PropertyGroup>
+  <PropertyGroup Condition="'$(Configuration)|$(Platform)' == 'Development|AnyCPU'">
+    <Optimize>true</Optimize>
+  </PropertyGroup>
+  <PropertyGroup Condition="'$(Configuration)|$(Platform)' == 'Release|AnyCPU'">
+    <Optimize>true</Optimize>
+  </PropertyGroup>
+
+  <ItemGroup>
+    <Compile Include="../../../Shared/MetaData.cs">
+      <Link>Properties\MetaData.cs</Link>
+    </Compile>
+  </ItemGroup>
+  <ItemGroup>
+    <Compile Include="../../../Shared/EpicGames.OIDC/**/*.cs">
+      <Link>%(RecursiveDir)/%(Filename)%(Extension)</Link>
+    </Compile>
+  </ItemGroup>
+  
+  <ItemGroup>
+    <PackageReference Include="IdentityModel.OidcClient" Version="4.0.0" />
+    <PackageReference Include="Microsoft.Extensions.Configuration" Version="5.0.0" />
+    <PackageReference Include="Microsoft.Extensions.Configuration.Binder" Version="5.0.0" />
+    <PackageReference Include="Microsoft.Extensions.Configuration.Json" Version="5.0.0" />
+  </ItemGroup>
+
+</Project>
diff --git a/Engine/Source/Programs/UnrealGameSync/Shared/EpicGames.Perforce/EpicGames.Perforce.csproj b/Engine/Source/Programs/UnrealGameSync/Shared/EpicGames.Perforce/EpicGames.Perforce.csproj
new file mode 100644
--- /dev/null
+++ b/Engine/Source/Programs/UnrealGameSync/Shared/EpicGames.Perforce/EpicGames.Perforce.csproj
@@ -0,0 +1,71 @@
+<Project Sdk="Microsoft.NET.Sdk">
+
+  <PropertyGroup>
+    <TargetFramework>netcoreapp3.1</TargetFramework>
+    <Configuration Condition=" '$(Configuration)' == '' ">Development</Configuration>
+    <OutputType>Library</OutputType>
+    <RootNamespace>EpicGames.Perforce</RootNamespace>
+    <GenerateAssemblyInfo>false</GenerateAssemblyInfo>
+    <GenerateTargetFrameworkAttribute>false</GenerateTargetFrameworkAttribute>
+    <Configurations>Debug;Release;Development</Configurations>
+    <nullable>enable</nullable>
+    <DebugType>pdbonly</DebugType>
+  </PropertyGroup>
+  <PropertyGroup Condition=" '$(Configuration)|$(Platform)' == 'Debug|AnyCPU' ">
+    <DocumentationFile>EpicGames.Perforce.xml</DocumentationFile>
+  </PropertyGroup>
+  <PropertyGroup Condition="'$(Configuration)|$(Platform)' == 'Development|AnyCPU'">
+    <DocumentationFile></DocumentationFile>
+    <Optimize>true</Optimize>
+  </PropertyGroup>
+  <PropertyGroup Condition="'$(Configuration)|$(Platform)' == 'Release|AnyCPU'">
+    <DocumentationFile></DocumentationFile>
+    <Optimize>true</Optimize>
+  </PropertyGroup>
+  <ItemGroup>
+    <None Remove="EpicGames.Perforce.uatbuildrecord" />
+    <None Remove="EpicGames.Perforce.xml" />
+  </ItemGroup>
+  <ItemGroup>
+    <Compile Include="../../../Shared/MetaData.cs">
+      <Link>Properties\MetaData.cs</Link>
+    </Compile>
+  </ItemGroup>
+  <ItemGroup>
+    <Compile Include="../../../Shared/EpicGames.Perforce/**/*.cs">
+      <Link>%(RecursiveDir)/%(Filename)%(Extension)</Link>
+    </Compile>
+  </ItemGroup>
+
+  <ItemGroup>
+    <ProjectReference Include="..\EpicGames.Core\EpicGames.Core.csproj" />
+  </ItemGroup>
+
+  <PropertyGroup>
+    <WindowsNativeConfig>Release</WindowsNativeConfig>
+    <WindowsNativeConfig Condition=" '$(Configuration)|$(Platform)' == 'Debug|AnyCPU' ">Debug</WindowsNativeConfig>
+    <DebugType>full</DebugType>
+    <EnableNETAnalyzers>True</EnableNETAnalyzers>
+  </PropertyGroup>
+
+  <ItemGroup>
+    <None Include="..\..\..\..\..\Binaries\DotNET\EpicGames.Perforce.Native\win-x64\$(WindowsNativeConfig)\EpicGames.Perforce.Native.dll" Visible="false">
+      <CopyToOutputDirectory>PreserveNewest</CopyToOutputDirectory>
+    </None>
+    <None Include="..\..\..\..\..\Binaries\DotNET\EpicGames.Perforce.Native\win-x64\$(WindowsNativeConfig)\EpicGames.Perforce.Native.pdb" Visible="false" Condition="Exists('..\..\..\..\Binaries\DotNET\EpicGames.Perforce.Native\win-x64\$(WindowsNativeConfig)\EpicGames.Perforce.Native.pdb')">
+      <CopyToOutputDirectory>PreserveNewest</CopyToOutputDirectory>
+    </None>
+    <None Include="..\..\..\..\..\Binaries\DotNET\EpicGames.Perforce.Native\linux-x64\Release\EpicGames.Perforce.Native.so" Visible="false">
+      <CopyToOutputDirectory>PreserveNewest</CopyToOutputDirectory>
+    </None>
+    <None Include="..\..\..\..\..\Binaries\DotNET\EpicGames.Perforce.Native\mac-x64\Release\EpicGames.Perforce.Native.dylib" Visible="false">
+      <CopyToOutputDirectory>PreserveNewest</CopyToOutputDirectory>
+    </None>
+  </ItemGroup>
+
+  <ItemGroup>
+    <PackageReference Include="Microsoft.Win32.Registry" Version="5.0.0" />
+    <PackageReference Include="System.Linq.Async" Version="6.0.1" />
+  </ItemGroup>
+
+</Project>
diff --git a/Engine/Source/Programs/UnrealGameSync/UnrealGameSync/UnrealGameSync.csproj b/Engine/Source/Programs/UnrealGameSync/UnrealGameSync/UnrealGameSync.csproj
--- a/Engine/Source/Programs/UnrealGameSync/UnrealGameSync/UnrealGameSync.csproj
+++ b/Engine/Source/Programs/UnrealGameSync/UnrealGameSync/UnrealGameSync.csproj
@@ -29,10 +29,10 @@
     <PackageReference Include="Serilog.Sinks.File" Version="5.0.0" />
   </ItemGroup>
   <ItemGroup>
-    <ProjectReference Include="..\..\Shared\EpicGames.Core\EpicGames.Core.csproj" />
-    <ProjectReference Include="..\..\Shared\EpicGames.Jupiter\EpicGames.Jupiter.csproj" />
-    <ProjectReference Include="..\..\Shared\EpicGames.OIDC\EpicGames.OIDC.csproj" />
-    <ProjectReference Include="..\..\Shared\EpicGames.Perforce\EpicGames.Perforce.csproj" />
+    <ProjectReference Include="..\Shared\EpicGames.Core\EpicGames.Core.csproj" />
+    <ProjectReference Include="..\Shared\EpicGames.Jupiter\EpicGames.Jupiter.csproj" />
+    <ProjectReference Include="..\Shared\EpicGames.OIDC\EpicGames.OIDC.csproj" />
+    <ProjectReference Include="..\Shared\EpicGames.Perforce\EpicGames.Perforce.csproj" />
     <ProjectReference Include="..\UnrealGameSyncShared\UnrealGameSyncShared.csproj" />
   </ItemGroup>
   <ItemGroup>
diff --git a/Engine/Source/Programs/UnrealGameSync/UnrealGameSyncCmd/UnrealGameSyncCmd.csproj b/Engine/Source/Programs/UnrealGameSync/UnrealGameSyncCmd/UnrealGameSyncCmd.csproj
--- a/Engine/Source/Programs/UnrealGameSync/UnrealGameSyncCmd/UnrealGameSyncCmd.csproj
+++ b/Engine/Source/Programs/UnrealGameSync/UnrealGameSyncCmd/UnrealGameSyncCmd.csproj
@@ -15,7 +15,7 @@
   </ItemGroup>
 
   <ItemGroup>
-    <ProjectReference Include="..\..\Shared\EpicGames.Core\EpicGames.Core.csproj" />
+    <ProjectReference Include="..\Shared\EpicGames.Core\EpicGames.Core.csproj" />
     <ProjectReference Include="..\UnrealGameSyncShared\UnrealGameSyncShared.csproj" />
   </ItemGroup>
 
diff --git a/Engine/Source/Programs/UnrealGameSync/UnrealGameSyncLauncher/UnrealGameSyncLauncher.csproj b/Engine/Source/Programs/UnrealGameSync/UnrealGameSyncLauncher/UnrealGameSyncLauncher.csproj
--- a/Engine/Source/Programs/UnrealGameSync/UnrealGameSyncLauncher/UnrealGameSyncLauncher.csproj
+++ b/Engine/Source/Programs/UnrealGameSync/UnrealGameSyncLauncher/UnrealGameSyncLauncher.csproj
@@ -49,8 +49,8 @@
     </Compile>
   </ItemGroup>
   <ItemGroup>
-    <ProjectReference Include="..\..\Shared\EpicGames.Core\EpicGames.Core.csproj" />
-    <ProjectReference Include="..\..\Shared\EpicGames.Perforce\EpicGames.Perforce.csproj" />
+    <ProjectReference Include="..\Shared\EpicGames.Core\EpicGames.Core.csproj" />
+    <ProjectReference Include="..\Shared\EpicGames.Perforce\EpicGames.Perforce.csproj" />
   </ItemGroup>
   <ItemGroup>
     <None Update="Properties\Settings.settings">
diff --git a/Engine/Source/Programs/UnrealGameSync/UnrealGameSyncShared/UnrealGameSyncShared.csproj b/Engine/Source/Programs/UnrealGameSync/UnrealGameSyncShared/UnrealGameSyncShared.csproj
--- a/Engine/Source/Programs/UnrealGameSync/UnrealGameSyncShared/UnrealGameSyncShared.csproj
+++ b/Engine/Source/Programs/UnrealGameSync/UnrealGameSyncShared/UnrealGameSyncShared.csproj
@@ -7,9 +7,9 @@
   </PropertyGroup>
 
   <ItemGroup>
-    <ProjectReference Include="..\..\Shared\EpicGames.Core\EpicGames.Core.csproj" />
-    <ProjectReference Include="..\..\Shared\EpicGames.OIDC\EpicGames.OIDC.csproj" />
-    <ProjectReference Include="..\..\Shared\EpicGames.Perforce\EpicGames.Perforce.csproj" />
+    <ProjectReference Include="..\Shared\EpicGames.Core\EpicGames.Core.csproj" />
+    <ProjectReference Include="..\Shared\EpicGames.OIDC\EpicGames.OIDC.csproj" />
+    <ProjectReference Include="..\Shared\EpicGames.Perforce\EpicGames.Perforce.csproj" />
   </ItemGroup>
 
 </Project>
