UGS: Improve fast path for finding code changes. Find the last cpp change within an optimistic range first, and only query all extensions for changes after then.

diff --git a/Engine/Source/Programs/UnrealGameSync/UnrealGameSyncShared/WorkspaceUpdate.cs b/Engine/Source/Programs/UnrealGameSync/UnrealGameSyncShared/WorkspaceUpdate.cs
--- a/Engine/Source/Programs/UnrealGameSync/UnrealGameSyncShared/WorkspaceUpdate.cs
+++ b/Engine/Source/Programs/UnrealGameSync/UnrealGameSyncShared/WorkspaceUpdate.cs
@@ -91,6 +91,7 @@ namespace UnrealGameSync
 	{
 		public DateTime StartTime = DateTime.UtcNow;
 		public int ChangeNumber;
+		public int? CodeChangeNumber;
 		public WorkspaceUpdateOptions Options;
 		public BuildConfig EditorConfig;
 		public string[]? SyncFilter;
@@ -913,41 +914,49 @@ namespace UnrealGameSync
 							branchOrStreamName = PerforceUtils.GetClientOrDepotDirectoryName(files[0].DepotFile);
 						}
 
-                        PerforceResponseList<ChangesRecord>? codeChanges = null;
+						// Get the last code change
+						int codeChangeNumber = Context.CodeChangeNumber ?? 0;
+						if (codeChangeNumber == 0)
+						{
+							string range = $"<={Context.ChangeNumber}";
 
-                        // First, check the most recent changes (this is a much cheaper query because of revcx optimization in the Perforce server.)
-                        int optimisticRangeStart = Context.ChangeNumber - 1000;
-                        if (optimisticRangeStart >= 1)
-                        {
-                            string[] codeFilter = PerforceUtils.CodeExtensions.SelectMany(x => syncPaths.Select(y => $"{y}{x}@{optimisticRangeStart},{Context.ChangeNumber}")).ToArray();
-							codeChanges = await perforce.TryGetChangesAsync(ChangesOptions.None, 1, ChangeStatus.Submitted, codeFilter, cancellationToken);
+							// First, check the most recent changes for a limited subset of file types (this is a much cheaper query because of revcx optimization in the Perforce server.)
+							int optimisticRangeStart = Context.ChangeNumber - 1000;
+							if (optimisticRangeStart >= 1)
+							{
+								string[] optimisticFileTypes = { ".cpp" };
 
-							if (!codeChanges.Succeeded)
-                            {
-                                return (WorkspaceUpdateResult.FailedToSync, $"Couldn't determine last code changelist before CL {Context.ChangeNumber}.");
-                            }
-                        }
+								string[] optimisticCodeFilter = optimisticFileTypes.SelectMany(x => syncPaths.Where(x => x.EndsWith("...", StringComparison.Ordinal)).Select(y => $"{y}{x}@{optimisticRangeStart},{Context.ChangeNumber}")).ToArray();
+								PerforceResponseList<ChangesRecord> optimisticChanges = await perforce.TryGetChangesAsync(ChangesOptions.None, 1, ChangeStatus.Submitted, optimisticCodeFilter, cancellationToken);
+
+								if (optimisticChanges.Succeeded && optimisticChanges.Count > 0)
+								{
+									int maxChange = optimisticChanges.Max(x => x.Data.Number);
+									range = $"{maxChange},{Context.ChangeNumber}";
+								}
+							}
 
-                        // If no change found in recent changes, do the full and expensive check.
-                        if (codeChanges == null || codeChanges.Count == 0)
-                        {
-                            string[] codeFilter = PerforceUtils.CodeExtensions.SelectMany(x => syncPaths.Select(y => $"{y}{x}@<={Context.ChangeNumber}")).ToArray();
-                            codeChanges = await perforce.TryGetChangesAsync(ChangesOptions.None, 1, ChangeStatus.Submitted, codeFilter, cancellationToken);
+							// If no change found in recent changes, do the full and expensive check.
+							string[] codeFilter = PerforceUtils.CodeExtensions.SelectMany(x => syncPaths.Select(y => $"{y}{x}@{range}")).ToArray();
+							PerforceResponseList<ChangesRecord> codeChanges = await perforce.TryGetChangesAsync(ChangesOptions.None, 1, ChangeStatus.Submitted, codeFilter, cancellationToken);
 
 							if (!codeChanges.Succeeded)
-                            {
-                                return (WorkspaceUpdateResult.FailedToSync, $"Couldn't determine last code changelist before CL {Context.ChangeNumber}.");
-                            }
-                        }
-						if (codeChanges.Count == 0)
-						{
-							return (WorkspaceUpdateResult.FailedToSync, $"Could not find any code changes before CL {Context.ChangeNumber}.");
+							{
+								return (WorkspaceUpdateResult.FailedToSync, $"Couldn't determine last code changelist before CL {Context.ChangeNumber}.");
+							}
+							if (codeChanges.Count == 0)
+							{
+								return (WorkspaceUpdateResult.FailedToSync, $"Could not find any code changes before CL {Context.ChangeNumber}.");
+							}
+
+							// Get the last code change number
+							codeChangeNumber = codeChanges.Max(x => x.Data.Number);
 						}
 
-						// Get the last code change
+						// Set the version change
 						if (Context.ProjectConfigFile.GetValue("Options.VersionToLastCodeChange", true))
 						{
-							versionChangeNumber = codeChanges.Max(x => x.Data.Number);
+							versionChangeNumber = codeChangeNumber;
 						}
 						else
 						{
