UGS: Revert change to query submitted changes using -e parameter. Seems to run much slower on our p4 server (?).

diff --git a/Engine/Source/Programs/UnrealGameSync/UnrealGameSync/PerforceMonitor.cs b/Engine/Source/Programs/UnrealGameSync/UnrealGameSync/PerforceMonitor.cs
--- a/Engine/Source/Programs/UnrealGameSync/UnrealGameSync/PerforceMonitor.cs
+++ b/Engine/Source/Programs/UnrealGameSync/UnrealGameSync/PerforceMonitor.cs
@@ -290,7 +290,8 @@ namespace UnrealGameSync
 			}
 			else
 			{
-				newChanges = await perforce.GetChangesAsync(ChangesOptions.IncludeTimes | ChangesOptions.LongOutput, null, newestChangeNumber + 1, -1, ChangeStatus.Submitted, null, depotPaths, cancellationToken);
+				newChanges = await perforce.GetChangesAsync(ChangesOptions.IncludeTimes | ChangesOptions.LongOutput, -1, ChangeStatus.Submitted, depotPaths.Select(x => $"{x}@>{newestChangeNumber}").ToArray(), cancellationToken);
+//				newChanges = await perforce.GetChangesAsync(ChangesOptions.IncludeTimes | ChangesOptions.LongOutput, null, newestChangeNumber + 1, -1, ChangeStatus.Submitted, null, depotPaths, cancellationToken);
 			}
 
 			// Remove anything we already have
