Horde: Prevent fetching more than the maximum number of desired changes, even when doing incremental updates.

diff --git a/Engine/Source/Programs/UnrealGameSync/UnrealGameSync/PerforceMonitor.cs b/Engine/Source/Programs/UnrealGameSync/UnrealGameSync/PerforceMonitor.cs
--- a/Engine/Source/Programs/UnrealGameSync/UnrealGameSync/PerforceMonitor.cs
+++ b/Engine/Source/Programs/UnrealGameSync/UnrealGameSync/PerforceMonitor.cs
@@ -284,14 +284,14 @@ namespace UnrealGameSync
 
 			// Read any new changes
 			List<ChangesRecord> newChanges;
-			if(maxChanges > CurrentMaxChanges)
+			if(maxChanges > CurrentMaxChanges || newestChangeNumber == -1)
 			{
 				newChanges = await perforce.GetChangesAsync(ChangesOptions.IncludeTimes | ChangesOptions.LongOutput, maxChanges, ChangeStatus.Submitted, depotPaths, cancellationToken);
 			}
 			else
 			{
-				newChanges = await perforce.GetChangesAsync(ChangesOptions.IncludeTimes | ChangesOptions.LongOutput, -1, ChangeStatus.Submitted, depotPaths.Select(x => $"{x}@>{newestChangeNumber}").ToArray(), cancellationToken);
-//				newChanges = await perforce.GetChangesAsync(ChangesOptions.IncludeTimes | ChangesOptions.LongOutput, null, newestChangeNumber + 1, -1, ChangeStatus.Submitted, null, depotPaths, cancellationToken);
+//				newChanges = await perforce.GetChangesAsync(ChangesOptions.IncludeTimes | ChangesOptions.LongOutput, maxChanges, ChangeStatus.Submitted, depotPaths.Select(x => $"{x}@>{newestChangeNumber}").ToArray(), cancellationToken);
+				newChanges = await perforce.GetChangesAsync(ChangesOptions.IncludeTimes | ChangesOptions.LongOutput, clientName: null, minChangeNumber: newestChangeNumber + 1, maxChanges: maxChanges, status: ChangeStatus.Submitted, userName: null, fileSpecs: depotPaths, cancellationToken);
 			}
 
 			// Remove anything we already have
