UGS: Disable logic to query for changes up to the last PCB zip file. This can result in a lot of unnecessary traffic to the P4 server, and can cause the UI to lock up on inserting so many changes.

diff --git a/Engine/Source/Programs/UnrealGameSync/UnrealGameSync/PerforceMonitor.cs b/Engine/Source/Programs/UnrealGameSync/UnrealGameSync/PerforceMonitor.cs
--- a/Engine/Source/Programs/UnrealGameSync/UnrealGameSync/PerforceMonitor.cs
+++ b/Engine/Source/Programs/UnrealGameSync/UnrealGameSync/PerforceMonitor.cs
@@ -25,7 +25,6 @@ namespace UnrealGameSync
 
 	class PerforceMonitor : IDisposable, IArchiveInfoSource
 	{
-
 		class PerforceChangeSorter : IComparer<ChangesRecord>
 		{
 			public int Compare(ChangesRecord? summaryA, ChangesRecord? summaryB)
@@ -304,6 +303,8 @@ namespace UnrealGameSync
 				newestChangeNumber = Math.Min(newestChangeNumber, newChanges.First().Number);
 			}
 
+			// The code below is correct, but can cause a lot of load on the Perforce server when we query a large number of changes because PCBs are far behind.
+#if false
 			// If we are using zipped binaries, make sure we have every change since the last zip containing them. This is necessary for ensuring that content changes show as
 			// syncable in the workspace view if there have been a large number of content changes since the last code change.
 			int minZippedChangeNumber = -1;
@@ -323,6 +324,7 @@ namespace UnrealGameSync
 				List<ChangesRecord> zipChanges = await perforce.GetChangesAsync(ChangesOptions.None, -1, ChangeStatus.Submitted, filteredPaths, cancellationToken);
 				newChanges.AddRange(zipChanges);
 			}
+#endif
 
 			// Fixup any ROBOMERGE authors
 			const string roboMergePrefix = "#ROBOMERGE-AUTHOR:";
