UGS: Allow overriding the filter used to determine if a file contains code.

The CodeFilter property is read from the [Perforce] section in the config file, and is appended to the default list of filters used to find code changes (having a cpp/h extension and so on). Syntax is in standard Perforce syntax, with a '-' prefix to exclude matching files. For example, the following entry..

  [Perforce]
  CodeFilter=-.../Experimental/...
  CodeFilter=-*.h

Excludes all *.h and files under an /Experimental/ folder anywhere in the tree from being treated as code changes.

diff --git a/Engine/Source/Programs/UnrealGameSync/UnrealGameSync/PerforceMonitor.cs b/Engine/Source/Programs/UnrealGameSync/UnrealGameSync/PerforceMonitor.cs
--- a/Engine/Source/Programs/UnrealGameSync/UnrealGameSync/PerforceMonitor.cs
+++ b/Engine/Source/Programs/UnrealGameSync/UnrealGameSync/PerforceMonitor.cs
@@ -52,6 +52,7 @@ namespace UnrealGameSync
 		DirectoryReference _cacheFolder;
 		List<KeyValuePair<FileReference, DateTime>> _localConfigFiles;
 		IAsyncDisposer _asyncDisposeTasks;
+		string[] prevCodeRules = Array.Empty<string>();
 
 		SynchronizationContext _synchronizationContext;
 		public event Action? OnUpdate;
@@ -394,8 +395,48 @@ namespace UnrealGameSync
 			return true;
 		}
 
+		Func<string, bool>? CreateCodeFilterFunc()
+		{
+			ConfigSection projectConfigSection = LatestProjectConfigFile.FindSection("Perforce");
+
+			string[] rules = projectConfigSection.GetValues("CodeFilter", new string[0]);
+			if (rules.Length == 0)
+			{
+				return null;
+			}
+
+			FileFilter filter = new FileFilter(PerforceUtils.CodeExtensions.Select(x => $"*.{x}"));
+			foreach (string rule in rules)
+			{
+				filter.AddRule(rule);
+			}
+
+			return filter.Matches;
+		}
+
 		public async Task<bool> UpdateChangeTypesAsync(IPerforceConnection perforce, CancellationToken cancellationToken)
 		{
+			// Get the filter for code changes
+			ConfigSection projectConfigSection = LatestProjectConfigFile.FindSection("Perforce");
+
+			string[] codeRules = projectConfigSection.GetValues("CodeFilter", new string[0]);
+			if (!Enumerable.SequenceEqual(codeRules, prevCodeRules))
+			{
+				_changeDetails.Clear();
+				prevCodeRules = codeRules;
+			}
+
+			Func<string, bool>? isCodeFile = null;
+			if (codeRules.Length > 0)
+			{
+				FileFilter filter = new FileFilter(PerforceUtils.CodeExtensions.Select(x => $"*{x}"));
+				foreach (string codeRule in codeRules)
+				{
+					filter.AddRule(codeRule);
+				}
+				isCodeFile = filter.Matches;
+			}
+
 			// Find the changes we need to query
 			List<int> queryChangeNumbers = new List<int>();
 			lock(this)
@@ -433,7 +474,7 @@ namespace UnrealGameSync
 						DescribeRecord describeRecord = describeRecordLoop;
 						int queryChangeNumber = describeRecord.Number;
 
-						PerforceChangeDetails details = new PerforceChangeDetails(describeRecord);
+						PerforceChangeDetails details = new PerforceChangeDetails(describeRecord, isCodeFile);
 
 						// Content only changes must be flagged accurately, because code changes invalidate precompiled binaries. Increase the number of files fetched until we can classify it correctly.
 						while (describeRecord.Files.Count >= maxFiles && !details.ContainsCode)
@@ -448,7 +489,7 @@ namespace UnrealGameSync
 							}
 
 							describeRecord = newDescribeRecords[0];
-							details = new PerforceChangeDetails(describeRecord);
+							details = new PerforceChangeDetails(describeRecord, isCodeFile);
 						}
 
 						lock (this)
diff --git a/Engine/Source/Programs/UnrealGameSync/UnrealGameSyncShared/Utility.cs b/Engine/Source/Programs/UnrealGameSync/UnrealGameSyncShared/Utility.cs
--- a/Engine/Source/Programs/UnrealGameSync/UnrealGameSyncShared/Utility.cs
+++ b/Engine/Source/Programs/UnrealGameSync/UnrealGameSyncShared/Utility.cs
@@ -36,14 +36,16 @@ namespace UnrealGameSync
 		public bool ContainsCode;
 		public bool ContainsContent;
 
-		public PerforceChangeDetails(DescribeRecord describeRecord)
+		public PerforceChangeDetails(DescribeRecord describeRecord, Func<string, bool>? isCodeFile = null)
 		{
+			isCodeFile ??= IsCodeFile;
+
 			Description = describeRecord.Description;
 
 			// Check whether the files are code or content
 			foreach (DescribeFileRecord file in describeRecord.Files)
 			{
-				if (PerforceUtils.CodeExtensions.Any(extension => file.DepotFile.EndsWith(extension, StringComparison.OrdinalIgnoreCase)))
+				if (isCodeFile(file.DepotFile))
 				{
 					ContainsCode = true;
 				}
@@ -58,6 +60,11 @@ namespace UnrealGameSync
 				}
 			}
 		}
+
+		public static bool IsCodeFile(string depotFile)
+		{
+			return PerforceUtils.CodeExtensions.Any(extension => depotFile.EndsWith(extension, StringComparison.OrdinalIgnoreCase));
+		}
 	}
 
 	public static class Utility
