UGS: Fix null reference exception when querying for custom code filters.

diff --git a/Engine/Source/Programs/UnrealGameSync/UnrealGameSync/PerforceMonitor.cs b/Engine/Source/Programs/UnrealGameSync/UnrealGameSync/PerforceMonitor.cs
--- a/Engine/Source/Programs/UnrealGameSync/UnrealGameSync/PerforceMonitor.cs
+++ b/Engine/Source/Programs/UnrealGameSync/UnrealGameSync/PerforceMonitor.cs
@@ -419,7 +419,7 @@ namespace UnrealGameSync
 			// Get the filter for code changes
 			ConfigSection projectConfigSection = LatestProjectConfigFile.FindSection("Perforce");
 
-			string[] codeRules = projectConfigSection.GetValues("CodeFilter", new string[0]);
+			string[] codeRules = projectConfigSection?.GetValues("CodeFilter", (string[]?)null) ?? Array.Empty<string>();
 			if (!Enumerable.SequenceEqual(codeRules, prevCodeRules))
 			{
 				_changeDetails.Clear();
