UGS: Show changes for all projects by default, to reduce load on Perforce server.

diff --git a/Engine/Source/Programs/UnrealGameSync/UnrealGameSync/PerforceMonitor.cs b/Engine/Source/Programs/UnrealGameSync/UnrealGameSync/PerforceMonitor.cs
--- a/Engine/Source/Programs/UnrealGameSync/UnrealGameSync/PerforceMonitor.cs
+++ b/Engine/Source/Programs/UnrealGameSync/UnrealGameSync/PerforceMonitor.cs
@@ -34,6 +34,7 @@ namespace UnrealGameSync
 		}
 
 		public int InitialMaxChangesValue = 100;
+		public bool ShowChangesForAllProjects = true;
 
 		IPerforceSettings _perforceSettings;
 		readonly string _branchClientPath;
@@ -257,7 +258,7 @@ namespace UnrealGameSync
 
 			// Build a full list of all the paths to sync
 			List<string> depotPaths = new List<string>();
-			if (_selectedClientFileName.EndsWith(".uprojectdirs", StringComparison.InvariantCultureIgnoreCase))
+			if (ShowChangesForAllProjects || _selectedClientFileName.EndsWith(".uprojectdirs", StringComparison.InvariantCultureIgnoreCase))
 			{
 				depotPaths.Add(String.Format("{0}/...", _branchClientPath));
 			}
