UGS: Swallow exception clearing contents of log control.

diff --git a/Engine/Source/Programs/UnrealGameSync/UnrealGameSync/Controls/LogControl.cs b/Engine/Source/Programs/UnrealGameSync/UnrealGameSync/Controls/LogControl.cs
--- a/Engine/Source/Programs/UnrealGameSync/UnrealGameSync/Controls/LogControl.cs
+++ b/Engine/Source/Programs/UnrealGameSync/UnrealGameSync/Controls/LogControl.cs
@@ -236,7 +236,11 @@ namespace UnrealGameSync
 			_queuedLines = new ConcurrentQueue<string>();
 			if(_logFileStream != null)
 			{
-				_logFileStream.SetLength(0);
+				try
+				{
+					_logFileStream.SetLength(0);
+				}
+				catch { }
 			}
 		}
 
