UGS: Include revision range in p4 changes query, as well as specifying -e.

diff --git a/Engine/Source/Programs/UnrealGameSync/UnrealGameSync/PerforceMonitor.cs b/Engine/Source/Programs/UnrealGameSync/UnrealGameSync/PerforceMonitor.cs
--- a/Engine/Source/Programs/UnrealGameSync/UnrealGameSync/PerforceMonitor.cs
+++ b/Engine/Source/Programs/UnrealGameSync/UnrealGameSync/PerforceMonitor.cs
@@ -291,8 +291,9 @@ namespace UnrealGameSync
 			}
 			else
 			{
+				List<string> depotPathsWithRange = depotPaths.ConvertAll(x => $"{x}@{newestChangeNumber + 1},#head");
 //				newChanges = await perforce.GetChangesAsync(ChangesOptions.IncludeTimes | ChangesOptions.LongOutput, maxChanges, ChangeStatus.Submitted, depotPaths.Select(x => $"{x}@>{newestChangeNumber}").ToArray(), cancellationToken);
-				newChanges = await perforce.GetChangesAsync(ChangesOptions.IncludeTimes | ChangesOptions.LongOutput, clientName: null, minChangeNumber: newestChangeNumber + 1, maxChanges: maxChanges, status: ChangeStatus.Submitted, userName: null, fileSpecs: depotPaths, cancellationToken);
+				newChanges = await perforce.GetChangesAsync(ChangesOptions.IncludeTimes | ChangesOptions.LongOutput, clientName: null, minChangeNumber: newestChangeNumber + 1, maxChanges: maxChanges, status: ChangeStatus.Submitted, userName: null, fileSpecs: depotPathsWithRange, cancellationToken);
 			}
 
 			// Remove anything we already have
