UGS: Publish release configuration when building UGS for distribution.

diff --git a/Engine/Source/Programs/UnrealGameSync/BuildUGS.xml b/Engine/Source/Programs/UnrealGameSync/BuildUGS.xml
--- a/Engine/Source/Programs/UnrealGameSync/BuildUGS.xml
+++ b/Engine/Source/Programs/UnrealGameSync/BuildUGS.xml
@@ -26,10 +26,10 @@
 	
 		<Node Name="Build UnrealGameSync">
 			<Delete Files="$(UgsOutputDir)/..."/>
-			<DotNet Arguments="publish &quot;$(UgsDir)/UnrealGameSyncCmd/UnrealGameSyncCmd.csproj&quot; --output=&quot;$(UgsPublishDir)&quot; --no-self-contained -r win-x64 $(VersionArguments)"/>
-			<DotNet Arguments="publish &quot;$(UgsDir)/UnrealGameSyncCmd/UnrealGameSyncCmd.csproj&quot; --output=&quot;$(UgsPublishDir)&quot; --no-self-contained -r osx-x64 $(VersionArguments)"/>
-			<DotNet Arguments="publish &quot;$(UgsDir)/UnrealGameSyncCmd/UnrealGameSyncCmd.csproj&quot; --output=&quot;$(UgsPublishDir)&quot; --no-self-contained -r linux-x64 $(VersionArguments)"/>
-			<DotNet Arguments="publish &quot;$(UgsDir)/UnrealGameSync/UnrealGameSync.csproj&quot; --output=&quot;$(UgsPublishDir)&quot; --no-self-contained -r win-x64 $(VersionArguments)"/>
+			<DotNet Arguments="publish &quot;$(UgsDir)/UnrealGameSyncCmd/UnrealGameSyncCmd.csproj&quot; --output=&quot;$(UgsPublishDir)&quot; --no-self-contained -c Release -r win-x64 $(VersionArguments)"/>
+			<DotNet Arguments="publish &quot;$(UgsDir)/UnrealGameSyncCmd/UnrealGameSyncCmd.csproj&quot; --output=&quot;$(UgsPublishDir)&quot; --no-self-contained -c Release -r osx-x64 $(VersionArguments)"/>
+			<DotNet Arguments="publish &quot;$(UgsDir)/UnrealGameSyncCmd/UnrealGameSyncCmd.csproj&quot; --output=&quot;$(UgsPublishDir)&quot; --no-self-contained -c Release -r linux-x64 $(VersionArguments)"/>
+			<DotNet Arguments="publish &quot;$(UgsDir)/UnrealGameSync/UnrealGameSync.csproj&quot; --output=&quot;$(UgsPublishDir)&quot; --no-self-contained -c Release -r win-x64 $(VersionArguments)"/>
 		</Node>
 		<Node Name="Publish Unstable" Requires="Build UnrealGameSync">
 			<Submit Description="Updating UGS from CL $(Change)" Workspace="Temp" Files="$(UgsPublishDir)/..." RootDir="$(UgsOutputDir)" Stream="$(SubmitToStream)" Branch="$(SubmitToBranch)/UnstableRelease"/>
