UGS: Handle exceptions when reading config files and post them to Sentry.

diff --git a/Engine/Source/Programs/UnrealGameSync/UnrealGameSync/Program.cs b/Engine/Source/Programs/UnrealGameSync/UnrealGameSync/Program.cs
--- a/Engine/Source/Programs/UnrealGameSync/UnrealGameSync/Program.cs
+++ b/Engine/Source/Programs/UnrealGameSync/UnrealGameSync/Program.cs
@@ -185,6 +185,8 @@ namespace UnrealGameSync
 					logger.LogInformation("Application version: {Version}", Assembly.GetExecutingAssembly().GetName().Version);
 					logger.LogInformation("Started at {Time}", DateTime.Now.ToString());
 
+					Utility.TraceException += ex => TraceException(ex, logger);
+
 					string sessionId = Guid.NewGuid().ToString();
 					logger.LogInformation("SessionId: {SessionId}", sessionId);
 
@@ -241,6 +243,15 @@ namespace UnrealGameSync
 			}
 		}
 
+		private static void TraceException(Exception ex, ILogger logger)
+		{
+			if (DeploymentSettings.SentryDsn != null)
+			{
+				SentrySdk.CaptureException(ex);
+			}
+			logger.LogError(ex, "Trace exception: {Ex}", ex.ToString());
+		}
+
 		private static void CurrentDomain_UnhandledException(object sender, UnhandledExceptionEventArgs args)
 		{
 			Exception? ex = args.ExceptionObject as Exception;
diff --git a/Engine/Source/Programs/UnrealGameSync/UnrealGameSyncShared/UserSettings.cs b/Engine/Source/Programs/UnrealGameSync/UnrealGameSyncShared/UserSettings.cs
--- a/Engine/Source/Programs/UnrealGameSync/UnrealGameSyncShared/UserSettings.cs
+++ b/Engine/Source/Programs/UnrealGameSync/UnrealGameSyncShared/UserSettings.cs
@@ -561,15 +561,16 @@ namespace UnrealGameSync
 
 		public static bool TryLoad(FileReference configFile, [NotNullWhen(true)] out UserProjectSettings? settings)
 		{
-			if (FileReference.Exists(configFile))
+			if (Utility.TryLoadJson(configFile, out settings))
 			{
-				settings = Utility.LoadJson<UserProjectSettings>(configFile);
 				settings.ConfigFile = configFile;
 				return true;
 			}
-
-			settings = null;
-			return false;
+			else
+			{
+				settings = null;
+				return false;
+			}
 		}
 
 		public bool Save(ILogger logger)
diff --git a/Engine/Source/Programs/UnrealGameSync/UnrealGameSyncShared/Utility.cs b/Engine/Source/Programs/UnrealGameSync/UnrealGameSyncShared/Utility.cs
--- a/Engine/Source/Programs/UnrealGameSync/UnrealGameSyncShared/Utility.cs
+++ b/Engine/Source/Programs/UnrealGameSync/UnrealGameSyncShared/Utility.cs
@@ -69,6 +69,8 @@ namespace UnrealGameSync
 
 	public static class Utility
 	{
+		public static event Action<Exception>? TraceException;
+
 		static JsonSerializerOptions GetDefaultJsonSerializerOptions()
 		{
 			JsonSerializerOptions options = new JsonSerializerOptions();
@@ -94,8 +96,9 @@ namespace UnrealGameSync
 				obj = LoadJson<T>(file);
 				return true;
 			}
-			catch
+			catch(Exception ex)
 			{
+				TraceException?.Invoke(ex);
 				obj = null;
 				return false;
 			}
