Modified UnrealGameSync to support not building C++ UHT based on a setting.

diff --git a/Engine/Programs/UnrealGameSync/UnrealGameSync.ini b/Engine/Programs/UnrealGameSync/UnrealGameSync.ini
--- a/Engine/Programs/UnrealGameSync/UnrealGameSync.ini
+++ b/Engine/Programs/UnrealGameSync/UnrealGameSync.ini
@@ -3,6 +3,7 @@ UseFastModularVersioning=1
 UseFastModularVersioningV2=1
 UseCrashReportClientEditor=1
 DebugGameHasSeparateExecutable=1
+BuildUnrealHeaderTool=0
 +ExcludeChanges=^Rebuild lightmaps\.$
 +ExcludeChanges=^Rebuild lightmaps & texture streaming\.$
 +ExcludeChanges=^Added collection '.*'\s*
diff --git a/Engine/Source/Programs/Shared/EpicGames.Perforce/PerforceResponse.cs b/Engine/Source/Programs/Shared/EpicGames.Perforce/PerforceResponse.cs
--- a/Engine/Source/Programs/Shared/EpicGames.Perforce/PerforceResponse.cs
+++ b/Engine/Source/Programs/Shared/EpicGames.Perforce/PerforceResponse.cs
@@ -28,7 +28,9 @@ namespace EpicGames.Perforce
 		/// <summary>
 		/// True if the response is successful
 		/// </summary>
-		public bool Succeeded => InternalData is not PerforceError;
+#pragma warning disable IDE0083 // Use pattern matching - "is not" is not supported by version 8 of the language and this library is used by UGS
+		public bool Succeeded => !(InternalData is PerforceError);
+#pragma warning restore IDE0083 // Use pattern matching
 
 		/// <summary>
 		/// True if the response is an error
diff --git a/Engine/Source/Programs/UnrealGameSync/UnrealGameSyncShared/ConfigUtils.cs b/Engine/Source/Programs/UnrealGameSync/UnrealGameSyncShared/ConfigUtils.cs
--- a/Engine/Source/Programs/UnrealGameSync/UnrealGameSyncShared/ConfigUtils.cs
+++ b/Engine/Source/Programs/UnrealGameSync/UnrealGameSyncShared/ConfigUtils.cs
@@ -314,7 +314,10 @@ namespace UnrealGameSync
 			}
 
 			List<BuildStep> defaultBuildSteps = new List<BuildStep>();
-			defaultBuildSteps.Add(new BuildStep(new Guid("{01F66060-73FA-4CC8-9CB3-E217FBBA954E}"), 0, "Compile UnrealHeaderTool", "Compiling UnrealHeaderTool...", 1, "UnrealHeaderTool", hostPlatform, "Development", "", !shouldSyncPrecompiledEditor));
+			if (latestProjectConfigFile.GetValue("Options.BuildUnrealHeaderTool", true))
+			{
+				defaultBuildSteps.Add(new BuildStep(new Guid("{01F66060-73FA-4CC8-9CB3-E217FBBA954E}"), 0, "Compile UnrealHeaderTool", "Compiling UnrealHeaderTool...", 1, "UnrealHeaderTool", hostPlatform, "Development", "", !shouldSyncPrecompiledEditor));
+			}
 			defaultBuildSteps.Add(new BuildStep(new Guid("{F097FF61-C916-4058-8391-35B46C3173D5}"), 1, $"Compile {editorTarget}", $"Compiling {editorTarget}...", 10, editorTarget, hostPlatform, editorConfig.ToString(), projectArgument, !shouldSyncPrecompiledEditor));
 			defaultBuildSteps.Add(new BuildStep(new Guid("{C6E633A1-956F-4AD3-BC95-6D06D131E7B4}"), 2, "Compile ShaderCompileWorker", "Compiling ShaderCompileWorker...", 1, "ShaderCompileWorker", hostPlatform, "Development", "", !shouldSyncPrecompiledEditor));
 			defaultBuildSteps.Add(new BuildStep(new Guid("{24FFD88C-7901-4899-9696-AE1066B4B6E8}"), 3, "Compile UnrealLightmass", "Compiling UnrealLightmass...", 1, "UnrealLightmass", hostPlatform, "Development", "", !shouldSyncPrecompiledEditor));
