Fix an issue with UGS filter evaluation that could prevent 'AdditionalPathsToSync' from being synced at all

#rnx

### Problem

- When combing the paths from 'AdditionalPathsToSync with the sync filter that the user creates themselves via the UGS interface, the 'AdditionalPathsToSync' paths will come first and any duplicates removed from the rest of the filter. This means if the user had specifically added one of the additional paths in their filter, it will now appear in a different order than intended which can cause problems if the ordering was intentional and important.
- For example If the users sync filter is set up to exclude everything under a specific directory BUT opt back in to syncing a specific directory (as this is easier than opting out of every sub directory except the desired one) with the following:
 -/Root/...
 /Root/FooBar/...
 BUT the 'FooBar' directory is present in the 'AdditionalPathsToSync' the final combined result will look like
 /Root/FooBar/...
 -/Root/...
 which has inverted the order and will cause the opt in of 'FooBar' to be ignored.

### Fix
- We no longer use union when combining the AdditionalPathsToSync with the filter, but concat instead as duplicate entries do not matter. This is also now done in GlobalSettingsFile.GetCombinedSyncFilter rather than relying on each caller of the method adding AdditionalPathsToSync onto it afterwards. This can cause duplicate entries in the sync filter, but should not affect the outcome.
- Add ::LatestPerforceConfigSection to PerforceMonitor to encapsulate accessing the perforce config section of the current projects config file.
- CleanWorkspace was not taking the AdditionalPathsToSync into consideration, which is now fixed by moving the additional path logic into GlobalSettingsFile.GetCombinedSyncFilter.
- Viewing the sync filter now takes AdditionalPathsToSync into account.

diff --git a/Engine/Source/Programs/UnrealGameSync/UnrealGameSync/Controls/WorkspaceControl.cs b/Engine/Source/Programs/UnrealGameSync/UnrealGameSync/Controls/WorkspaceControl.cs
--- a/Engine/Source/Programs/UnrealGameSync/UnrealGameSync/Controls/WorkspaceControl.cs
+++ b/Engine/Source/Programs/UnrealGameSync/UnrealGameSync/Controls/WorkspaceControl.cs
@@ -920,14 +920,7 @@ namespace UnrealGameSync
 				}
 			}
 
-			string[] combinedSyncFilter = UserSettings.GetCombinedSyncFilter(GetSyncCategories(), _settings.Global.Filter, _workspaceSettings.Filter);
-
-			ConfigSection perforceSection = _perforceMonitor.LatestProjectConfigFile.FindSection("Perforce");
-			if (perforceSection != null)
-			{
-				IEnumerable<string> additionalPaths = perforceSection.GetValues("AdditionalPathsToSync", new string[0]);
-				combinedSyncFilter = additionalPaths.Union(combinedSyncFilter).ToArray();
-			}
+			string[] combinedSyncFilter = UserSettings.GetCombinedSyncFilter(GetSyncCategories(), _settings.Global.Filter, _workspaceSettings.Filter, _perforceMonitor.LatestPerforceConfigSection());
 
 			WorkspaceUpdateContext context = new WorkspaceUpdateContext(changeNumber, options, GetEditorBuildConfig(), combinedSyncFilter, _projectSettings.BuildSteps, null);
 			if (options.HasFlag(WorkspaceUpdateOptions.SyncArchives))
@@ -5216,7 +5209,7 @@ namespace UnrealGameSync
 				extraSafeToDeleteExtensions = "";
 			}
 
-			string[] combinedSyncFilter = UserSettings.GetCombinedSyncFilter(GetSyncCategories(), _settings.Global.Filter, _workspaceSettings.Filter);
+			string[] combinedSyncFilter = UserSettings.GetCombinedSyncFilter(GetSyncCategories(), _settings.Global.Filter, _workspaceSettings.Filter, _perforceMonitor.LatestPerforceConfigSection());
 			List<string> syncPaths = WorkspaceUpdate.GetSyncPaths(_workspace.Project, _workspaceSettings.Filter.AllProjects ?? _settings.Global.Filter.AllProjects ?? false, combinedSyncFilter);
 
 			CleanWorkspaceWindow.DoClean(ParentForm, _perforceSettings, BranchDirectoryName, _workspace.Project.ClientRootPath, syncPaths, extraSafeToDeleteFolders.Split('\n'), extraSafeToDeleteExtensions.Split('\n'), _serviceProvider.GetRequiredService<ILogger<CleanWorkspaceWindow>>());
@@ -5514,7 +5507,7 @@ namespace UnrealGameSync
 
 		private void OptionsContextMenu_SyncFilter_Click(object sender, EventArgs e)
 		{
-			SyncFilter filter = new SyncFilter(GetSyncCategories(), _settings.Global.Filter, _workspaceSettings.Filter);
+			SyncFilter filter = new SyncFilter(GetSyncCategories(), _settings.Global.Filter, _workspaceSettings.Filter, _perforceMonitor.LatestPerforceConfigSection());
 			if (filter.ShowDialog() == DialogResult.OK)
 			{
 				_settings.Global.Filter = filter.GlobalFilter;
diff --git a/Engine/Source/Programs/UnrealGameSync/UnrealGameSync/Forms/SyncFilter.cs b/Engine/Source/Programs/UnrealGameSync/UnrealGameSync/Forms/SyncFilter.cs
--- a/Engine/Source/Programs/UnrealGameSync/UnrealGameSync/Forms/SyncFilter.cs
+++ b/Engine/Source/Programs/UnrealGameSync/UnrealGameSync/Forms/SyncFilter.cs
@@ -17,14 +17,16 @@ namespace UnrealGameSync
 		Dictionary<Guid, WorkspaceSyncCategory> _uniqueIdToCategory;
 		public FilterSettings GlobalFilter;
 		public FilterSettings WorkspaceFilter;
+		ConfigSection _perforceSection;
 
-		public SyncFilter(Dictionary<Guid, WorkspaceSyncCategory> inUniqueIdToCategory, FilterSettings inGlobalFilter, FilterSettings inWorkspaceFilter)
+		public SyncFilter(Dictionary<Guid, WorkspaceSyncCategory> inUniqueIdToCategory, FilterSettings inGlobalFilter, FilterSettings inWorkspaceFilter, ConfigSection inPerforceSection)
 		{
 			InitializeComponent();
 
 			_uniqueIdToCategory = inUniqueIdToCategory;
 			GlobalFilter = inGlobalFilter;
 			WorkspaceFilter = inWorkspaceFilter;
+			_perforceSection = inPerforceSection;
 
 			Dictionary<Guid, bool> syncCategories = WorkspaceSyncCategory.GetDefault(_uniqueIdToCategory.Values);
 
@@ -156,7 +158,7 @@ namespace UnrealGameSync
 		{
 			GetSettings(out FilterSettings newGlobalFilter, out FilterSettings newWorkspaceFilter);
 
-			string[] filter = UserSettings.GetCombinedSyncFilter(_uniqueIdToCategory, newGlobalFilter, newWorkspaceFilter);
+			string[] filter = UserSettings.GetCombinedSyncFilter(_uniqueIdToCategory, newGlobalFilter, newWorkspaceFilter, _perforceSection);
 			if(filter.Length == 0)
 			{
 				filter = new string[]{ "All files will be synced." };
diff --git a/Engine/Source/Programs/UnrealGameSync/UnrealGameSync/PerforceMonitor.cs b/Engine/Source/Programs/UnrealGameSync/UnrealGameSync/PerforceMonitor.cs
--- a/Engine/Source/Programs/UnrealGameSync/UnrealGameSync/PerforceMonitor.cs
+++ b/Engine/Source/Programs/UnrealGameSync/UnrealGameSync/PerforceMonitor.cs
@@ -596,6 +596,11 @@ namespace UnrealGameSync
 			}
 		}
 
+		public ConfigSection LatestPerforceConfigSection()
+		{
+			return LatestProjectConfigFile.FindSection("Perforce");
+		}
+
 		public int LastChangeByCurrentUser
 		{
 			get;
diff --git a/Engine/Source/Programs/UnrealGameSync/UnrealGameSyncCmd/Program.cs b/Engine/Source/Programs/UnrealGameSync/UnrealGameSyncCmd/Program.cs
--- a/Engine/Source/Programs/UnrealGameSync/UnrealGameSyncCmd/Program.cs
+++ b/Engine/Source/Programs/UnrealGameSync/UnrealGameSyncCmd/Program.cs
@@ -337,14 +337,9 @@ namespace UnrealGameSyncCmd
 		static string[] ReadSyncFilter(UserWorkspaceSettings workspaceSettings, GlobalSettingsFile userSettings, ConfigFile projectConfig)
 		{
 			Dictionary<Guid, WorkspaceSyncCategory> syncCategories = ConfigUtils.GetSyncCategories(projectConfig);
-			string[] combinedSyncFilter = GlobalSettingsFile.GetCombinedSyncFilter(syncCategories, userSettings.Global.Filter, workspaceSettings.Filter);
-
 			ConfigSection perforceSection = projectConfig.FindSection("Perforce");
-			if (perforceSection != null)
-			{
-				IEnumerable<string> additionalPaths = perforceSection.GetValues("AdditionalPathsToSync", new string[0]);
-				combinedSyncFilter = additionalPaths.Union(combinedSyncFilter).ToArray();
-			}
+
+			string[] combinedSyncFilter = GlobalSettingsFile.GetCombinedSyncFilter(syncCategories, userSettings.Global.Filter, workspaceSettings.Filter, perforceSection);
 
 			return combinedSyncFilter;
 		}
diff --git a/Engine/Source/Programs/UnrealGameSync/UnrealGameSyncShared/GlobalSettings.cs b/Engine/Source/Programs/UnrealGameSync/UnrealGameSyncShared/GlobalSettings.cs
--- a/Engine/Source/Programs/UnrealGameSync/UnrealGameSyncShared/GlobalSettings.cs
+++ b/Engine/Source/Programs/UnrealGameSync/UnrealGameSyncShared/GlobalSettings.cs
@@ -174,7 +174,7 @@ namespace UnrealGameSync
 			return settings;
 		}
 
-		public static string[] GetCombinedSyncFilter(Dictionary<Guid, WorkspaceSyncCategory> uniqueIdToFilter, FilterSettings globalFilter, FilterSettings workspaceFilter)
+		public static string[] GetCombinedSyncFilter(Dictionary<Guid, WorkspaceSyncCategory> uniqueIdToFilter, FilterSettings globalFilter, FilterSettings workspaceFilter, ConfigSection perforceSection)
 		{
 			List<string> lines = new List<string>();
 			foreach (string viewLine in Enumerable.Concat(globalFilter.View, workspaceFilter.View).Select(x => x.Trim()).Where(x => x.Length > 0 && !x.StartsWith(";")))
@@ -216,6 +216,14 @@ namespace UnrealGameSync
 				}
 			}
 
+			// If there are no filtering lines then we can assume that AdditionalPathsToSync is covered and we do not
+			// need to add them manually.
+			if (lines.Count > 0 && perforceSection != null)
+			{
+				IEnumerable<string> additionalPaths = perforceSection.GetValues("AdditionalPathsToSync", new string[0]);
+				lines = lines.Concat(additionalPaths).ToList();
+			}
+
 			return lines.ToArray();
 		}
 
