UGS: Change automation server to using an async/await model, and properly support cancellation. Fixes "task exception not awaited on shutdown" errors.

diff --git a/Engine/Source/Programs/UnrealGameSync/UnrealGameSync/AutomationServer.cs b/Engine/Source/Programs/UnrealGameSync/UnrealGameSync/AutomationServer.cs
--- a/Engine/Source/Programs/UnrealGameSync/UnrealGameSync/AutomationServer.cs
+++ b/Engine/Source/Programs/UnrealGameSync/UnrealGameSync/AutomationServer.cs
@@ -7,8 +7,10 @@ using System.IO;
 using System.IO.Pipes;
 using System.Net;
 using System.Net.Sockets;
+using System.Runtime.CompilerServices;
 using System.Text;
 using System.Threading;
+using System.Threading.Tasks;
 using System.Windows.Forms;
 
 namespace UnrealGameSync
@@ -127,36 +129,31 @@ namespace UnrealGameSync
 		}
 	}
 
-	class AutomationServer : IDisposable
+	class AutomationServer : IAsyncDisposable, IDisposable
 	{
-		TcpListener? _listener;
-		public const int DefaultPortNumber = 30422;
+		static readonly UnicodeEncoding _streamEncoding = new UnicodeEncoding();
+
+		CancellationTokenSource _cancellationSource = new CancellationTokenSource();
 
-		NamedPipeServerStream? _ipcStream;
-		static UnicodeEncoding _streamEncoding = new UnicodeEncoding();
-		const string UgsChannel = @"\.\pipe\UGSChannel";
+		const string IpcChannel = @"\.\pipe\UGSChannel";
+		ConfiguredTaskAwaitable _ipcTask;
 
-		Thread? _uriThread;
-		Thread? _tcpThread;
+		public const int DefaultPortNumber = 30422;
+		ConfiguredTaskAwaitable? _tcpTask;
 
-		EventWaitHandle _shutdownEvent;		
 		Action<AutomationRequest> _postRequest;
 
-		bool _disposing;
 		ILogger _logger;
-		string? _commandLineUri;
 
 		public AutomationServer(Action<AutomationRequest> postRequest, string? uri, ILogger<AutomationServer> logger)
 		{
-			_shutdownEvent = new ManualResetEvent(false);
 			this._postRequest = postRequest;
-			this._commandLineUri = uri;
 			this._logger = logger;
 
 			try
 			{
 				// IPC named pipe
-				_ipcStream = new NamedPipeServerStream(UgsChannel, PipeDirection.In, 1, PipeTransmissionMode.Message, PipeOptions.Asynchronous);
+				_ipcTask = RunIpcAsync(uri, _cancellationSource.Token).ConfigureAwait(false);
 
 				// TCP listener setup
 				int portNumber = GetPortNumber();
@@ -164,23 +161,13 @@ namespace UnrealGameSync
 				{
 					try
 					{
-						_listener = new TcpListener(IPAddress.Loopback, portNumber);
-						_listener.Start();
-						_tcpThread = new Thread(() => RunTcp());
-						_tcpThread.IsBackground = true;
-						_tcpThread.Start();
-
+						_tcpTask = RunTcpAsync(portNumber, _cancellationSource.Token).ConfigureAwait(false);
 					}
 					catch (Exception ex)
 					{
-						_listener = null;
 						logger.LogError(ex, "Unable to start automation server tcp listener");
 					}
 				}
-
-				_uriThread = new Thread(() => RunUri());
-				_uriThread.IsBackground = true;
-				_uriThread.Start();
 			}
 			catch (Exception ex)
 			{
@@ -213,113 +200,92 @@ namespace UnrealGameSync
 			}
 		}
 
-		void RunUri()
+		async Task RunIpcAsync(string? commandLineUri, CancellationToken cancellationToken)
 		{
 			// Handle main process command line URI request
-			if (!string.IsNullOrEmpty(_commandLineUri))
+			if (!string.IsNullOrEmpty(commandLineUri))
 			{
-				HandleUri(_commandLineUri);
+				HandleUri(commandLineUri);
 			}
 
-			for (;;)
+			using NamedPipeServerStream ipcStream = new NamedPipeServerStream(IpcChannel, PipeDirection.In, 1, PipeTransmissionMode.Message, PipeOptions.Asynchronous);
+			while (!cancellationToken.IsCancellationRequested)
 			{
 				try
 				{
-
-					IAsyncResult ipcResult = _ipcStream!.BeginWaitForConnection(null, null);
-
-					int waitResult = WaitHandle.WaitAny(new WaitHandle[] { _shutdownEvent, ipcResult.AsyncWaitHandle });
-
-					// Shutting down
-					if (waitResult == 0)
-					{
-						break;
-					}
-
+					await ipcStream.WaitForConnectionAsync(cancellationToken).ConfigureAwait(false);
 					try
 					{
-						_ipcStream.EndWaitForConnection(ipcResult);
-
 						_logger.LogInformation("Accepted Uri connection");
 
 						// Read URI
-						string uri = ReadString(_ipcStream);
+						string uri = ReadString(ipcStream);
 
 						_logger.LogInformation("Received Uri: {Uri}", uri);
 
-						_ipcStream.Disconnect();
-
 						HandleUri(uri);
-
 					}
-					catch (Exception ex)
+					finally
 					{
-						_logger.LogError(ex, "Error during automation connection");
+						ipcStream.Disconnect();
 					}
 				}
+				catch (OperationCanceledException) when (cancellationToken.IsCancellationRequested)
+				{
+					break;
+				}
 				catch (Exception ex)
 				{
-					if (!_disposing)
-					{
-						_logger.LogError(ex, "Error during automation connection");
-					}
+					_logger.LogError(ex, "Error during automation connection");
 				}
 			}
 		}
 
-
-		void RunTcp()
+		async Task RunTcpAsync(int portNumber, CancellationToken cancellationToken)
 		{
-			
-			for (;;)
+			TcpListener listener = new TcpListener(IPAddress.Loopback, portNumber);
+			using (IDisposable disposable = cancellationToken.Register(() => listener.Stop()))
 			{
-				try
+				listener.Start();
+				while (!cancellationToken.IsCancellationRequested)
 				{
-
-					IAsyncResult tcpResult = _listener!.BeginAcceptTcpClient(null, null);
-
-					int waitResult = WaitHandle.WaitAny(new WaitHandle[] { _shutdownEvent, tcpResult.AsyncWaitHandle });
-
-					// Shutting down
-					if (waitResult == 0)
-					{
-						break;
-					}
-
 					try
 					{
-						TcpClient client = _listener.EndAcceptTcpClient(tcpResult);
+						TcpClient client = await listener.AcceptTcpClientAsync().ConfigureAwait(false);
+						try
+						{
+							_logger.LogInformation("Accepted connection from {Remote}", client.Client.RemoteEndPoint);
 
-						_logger.LogInformation("Accepted connection from {Remote}", client.Client.RemoteEndPoint);
+							NetworkStream stream = client.GetStream();
 
-						NetworkStream stream = client.GetStream();
+							AutomationRequestInput input = AutomationRequestInput.Read(stream);
+							_logger.LogInformation("Received input: {Type} (+{NumBytes} bytes)", input.Type, input.Data.Length);
 
-						AutomationRequestInput input = AutomationRequestInput.Read(stream);
-						_logger.LogInformation("Received input: {Type} (+{NumBytes} bytes)", input.Type, input.Data.Length);
+							AutomationRequestOutput output;
+							using (AutomationRequest request = new AutomationRequest(input))
+							{
+								_postRequest(request);
+								request.Complete.Wait();
+								output = request.Output!;
+							}
 
-						AutomationRequestOutput output;
-						using (AutomationRequest request = new AutomationRequest(input))
+							output.Write(stream);
+							_logger.LogInformation("Sent output: {Result} (+{NumBytes} bytes)", output.Result, output.Data.Length);
+						}
+						catch (Exception ex)
 						{
-							_postRequest(request);
-							request.Complete.Wait();
-							output = request.Output!;
+							_logger.LogError(ex, "Exception during automation");
+						}
+						finally
+						{
+							_logger.LogInformation("Closed connection.");
 						}
-
-						output.Write(stream);
-						_logger.LogInformation("Sent output: {Result} (+{NumBytes} bytes)", output.Result, output.Data.Length);
-					}
-					catch (Exception ex)
-					{
-						_logger.LogError(ex, "Exception during automation");
 					}
-					finally
+					catch when (cancellationToken.IsCancellationRequested)
 					{
-						_logger.LogInformation("Closed connection.");
+						break;
 					}
-				}
-				catch (Exception ex)
-				{
-					if (!_disposing)
+					catch (Exception ex)
 					{
 						_logger.LogError(ex, "Exception during automation operation");
 					}
@@ -357,7 +323,7 @@ namespace UnrealGameSync
 		/// </summary>		
 		public static void SendUri(string uri)
 		{
-			using (NamedPipeClientStream clientStream = new NamedPipeClientStream(".", UgsChannel, PipeDirection.Out, PipeOptions.None))
+			using (NamedPipeClientStream clientStream = new NamedPipeClientStream(".", IpcChannel, PipeDirection.Out, PipeOptions.None))
 			{
 				try
 				{
@@ -398,46 +364,37 @@ namespace UnrealGameSync
 			stream.Flush();
 		}
 
-
 		public void Dispose()
 		{
-			const int timeout = 5000;
-			_disposing = true;
-			_shutdownEvent.Set();
-
-			if (_uriThread != null)
-			{
-				if (!_uriThread.Join(timeout))
-				{
-					try { _uriThread.Abort(); } catch { }
-				}
+			DisposeAsync().AsTask().Wait();
+		}
 
-				_uriThread = null;
-			}
+		public async ValueTask DisposeAsync()
+		{
+			_cancellationSource.Cancel();
 
-			if (_tcpThread != null)
+			try
 			{
-				if (!_tcpThread.Join(timeout))
-				{
-					try { _tcpThread.Abort(); } catch { }
-				}
-
-				_tcpThread = null;
+				await _ipcTask;
 			}
-
-
-			// clean up IPC stream
-			if (_ipcStream != null)
+			catch (Exception ex)
 			{
-				_ipcStream.Dispose();
+				_logger.LogError(ex, "Error awaiting IPC background task");
 			}
 
-			if (_listener != null)
+			if (_tcpTask != null)
 			{
-				_listener.Stop();
-				_listener = null;
+				try
+				{
+					await _tcpTask.Value;
+				}
+				catch (Exception ex)
+				{
+					_logger.LogError(ex, "Error awaiting TCP background task");
+				}
 			}
 
+			_cancellationSource.Dispose();
 		}
 	}
 }
