UGS: Fix some nullable warnings required for NET 6 upgrade.

diff --git a/Engine/Source/Programs/UnrealGameSync/UnrealGameSync/Automation/VisualStudioAutomation.cs b/Engine/Source/Programs/UnrealGameSync/UnrealGameSync/Automation/VisualStudioAutomation.cs
--- a/Engine/Source/Programs/UnrealGameSync/UnrealGameSync/Automation/VisualStudioAutomation.cs
+++ b/Engine/Source/Programs/UnrealGameSync/UnrealGameSync/Automation/VisualStudioAutomation.cs
@@ -62,7 +62,7 @@ namespace UnrealGameSync
 			errorMessage = "";
 			VisualStudioInstallation? install = VisualStudioInstallations.GetPreferredInstallation();
 
-			if (install == null)
+			if (install == null || install.DevEnvPath == null)
 			{
 				errorMessage = string.Format("Unable to get Visual Studio installation");
 				return false;
@@ -215,7 +215,7 @@ namespace UnrealGameSync
 				return CachedInstalls.First();
 			}
 
-			VisualStudioInstallation installation = CachedInstalls.FirstOrDefault(install => { return install.MajorVersion == majorVersion; });
+			VisualStudioInstallation? installation = CachedInstalls.FirstOrDefault(install => { return install.MajorVersion == majorVersion; });
 
 			return installation;
 
diff --git a/Engine/Source/Programs/UnrealGameSync/UnrealGameSync/Controls/BuildListControl.cs b/Engine/Source/Programs/UnrealGameSync/UnrealGameSync/Controls/BuildListControl.cs
--- a/Engine/Source/Programs/UnrealGameSync/UnrealGameSync/Controls/BuildListControl.cs
+++ b/Engine/Source/Programs/UnrealGameSync/UnrealGameSync/Controls/BuildListControl.cs
@@ -57,7 +57,7 @@ namespace UnrealGameSync
 
 		public BuildListControl()
 		{
-			Font = SystemFonts.IconTitleFont;
+			Font = SystemFonts.IconTitleFont!;
 		}
 
 		public int GetFirstVisibleIndex()
diff --git a/Engine/Source/Programs/UnrealGameSync/UnrealGameSync/Controls/CustomListViewControl.cs b/Engine/Source/Programs/UnrealGameSync/UnrealGameSync/Controls/CustomListViewControl.cs
--- a/Engine/Source/Programs/UnrealGameSync/UnrealGameSync/Controls/CustomListViewControl.cs
+++ b/Engine/Source/Programs/UnrealGameSync/UnrealGameSync/Controls/CustomListViewControl.cs
@@ -400,7 +400,10 @@ namespace UnrealGameSync
 
 		public void DrawNormalSubItem(DrawListViewSubItemEventArgs e)
 		{
-			DrawText(e.Graphics, e.SubItem.Bounds, Columns[e.ColumnIndex].TextAlign, SystemColors.WindowText, e.SubItem.Text);
+			if (e.SubItem != null)
+			{
+				DrawText(e.Graphics, e.SubItem.Bounds, Columns[e.ColumnIndex].TextAlign, SystemColors.WindowText, e.SubItem.Text);
+			}
 		}
 
 		public void DrawCustomSubItem(Graphics graphics, ListViewItem.ListViewSubItem subItem)
diff --git a/Engine/Source/Programs/UnrealGameSync/UnrealGameSync/Controls/LogSplitContainer.cs b/Engine/Source/Programs/UnrealGameSync/UnrealGameSync/Controls/LogSplitContainer.cs
--- a/Engine/Source/Programs/UnrealGameSync/UnrealGameSync/Controls/LogSplitContainer.cs
+++ b/Engine/Source/Programs/UnrealGameSync/UnrealGameSync/Controls/LogSplitContainer.cs
@@ -49,7 +49,7 @@ namespace UnrealGameSync
 
 		protected Font CaptionFont
 		{
-			get { return SystemFonts.IconTitleFont; }
+			get { return SystemFonts.IconTitleFont!; }
 		}
 
 		protected override void OnCreateControl()
@@ -68,7 +68,7 @@ namespace UnrealGameSync
 			Invalidate();
 		}
 
-		public void OnSplitterMoved(object obj, EventArgs args)
+		public void OnSplitterMoved(object? obj, EventArgs args)
 		{
 			_splitterMoved = true;
 
diff --git a/Engine/Source/Programs/UnrealGameSync/UnrealGameSync/Controls/StatusPanel.cs b/Engine/Source/Programs/UnrealGameSync/UnrealGameSync/Controls/StatusPanel.cs
--- a/Engine/Source/Programs/UnrealGameSync/UnrealGameSync/Controls/StatusPanel.cs
+++ b/Engine/Source/Programs/UnrealGameSync/UnrealGameSync/Controls/StatusPanel.cs
@@ -767,7 +767,7 @@ namespace UnrealGameSync
 				}
 				if(_alert != null)
 				{
-					e.Graphics.DrawLine(_alertDividerPen, 0, _alertDividerY, Width, _alertDividerY);
+					e.Graphics.DrawLine(_alertDividerPen!, 0, _alertDividerY, Width, _alertDividerY);
 					_alert.Draw(e.Graphics, _resources!);
 				}
 			}
diff --git a/Engine/Source/Programs/UnrealGameSync/UnrealGameSync/Controls/WorkspaceControl.cs b/Engine/Source/Programs/UnrealGameSync/UnrealGameSync/Controls/WorkspaceControl.cs
--- a/Engine/Source/Programs/UnrealGameSync/UnrealGameSync/Controls/WorkspaceControl.cs
+++ b/Engine/Source/Programs/UnrealGameSync/UnrealGameSync/Controls/WorkspaceControl.cs
@@ -949,7 +949,7 @@ namespace UnrealGameSync
 					string[]? zippedBinariesSyncFilter;
 					if (TryGetProjectSetting(_perforceMonitor.LatestProjectConfigFile, "ZippedBinariesSyncFilter", out zippedBinariesSyncFilter) && zippedBinariesSyncFilter.Length > 0)
 					{
-						context.SyncFilter = Enumerable.Concat(context.SyncFilter, zippedBinariesSyncFilter).ToArray();
+						context.SyncFilter = Enumerable.Concat(context.SyncFilter ?? Enumerable.Empty<string>(), zippedBinariesSyncFilter).ToArray();
 					}
 
 					context.ArchiveTypeToArchive[archive.Type] = new Tuple<IArchiveInfo, string>(archive, archivePath);
@@ -2154,7 +2154,12 @@ namespace UnrealGameSync
 
 		private void BuildList_DrawSubItem(object sender, DrawListViewSubItemEventArgs e)
 		{
-			ChangesRecord change = (ChangesRecord)e.Item.Tag;
+			if (e.Item == null || e.SubItem == null)
+			{
+				return;
+			}
+
+			ChangesRecord? change = (ChangesRecord?)e.Item?.Tag;
 			if (change == null)
 			{
 				return;
@@ -2212,7 +2217,7 @@ namespace UnrealGameSync
 						int firstPass, firstFail;
 						GetRemainingBisectRange(out firstPass, out firstFail);
 
-						BisectEntry entry = _workspaceState.BisectChanges.FirstOrDefault(x => x.Change == change.Number);
+						BisectEntry? entry = _workspaceState.BisectChanges.FirstOrDefault(x => x.Change == change.Number);
 						if (entry == null || entry.State == BisectState.Exclude)
 						{
 							qualityIcon = new Rectangle(0, 0, 0, 0);
@@ -2618,7 +2623,7 @@ namespace UnrealGameSync
 				ProcessStartInfo startInfo = new ProcessStartInfo();
 				startInfo.FileName = url;
 				startInfo.UseShellExecute = true;
-				using Process _ = Process.Start(startInfo);
+				using Process? _ = Process.Start(startInfo);
 			}
 			catch
 			{
@@ -2634,7 +2639,7 @@ namespace UnrealGameSync
 				startInfo.FileName = url;
 				startInfo.Arguments = arguments;
 				startInfo.UseShellExecute = true;
-				using Process _ = Process.Start(startInfo);
+				using Process? _ = Process.Start(startInfo);
 			}
 			catch
 			{
@@ -2651,7 +2656,7 @@ namespace UnrealGameSync
 				startInfo.Arguments = arguments;
 				startInfo.WorkingDirectory = workingDir;
 				startInfo.UseShellExecute = true;
-				using Process _ = Process.Start(startInfo);
+				using Process? _ = Process.Start(startInfo);
 			}
 			catch
 			{
@@ -4172,14 +4177,14 @@ namespace UnrealGameSync
 				bool isBusy = _workspace.IsBusy();
 				bool isCurrentChange = (_contextMenuChange.Number == _workspace.CurrentChangeNumber);
 				BuildListContextMenu_Sync.Visible = !isBusy;
-				BuildListContextMenu_Sync.Font = new Font(SystemFonts.MenuFont, isCurrentChange ? FontStyle.Regular : FontStyle.Bold);
+				BuildListContextMenu_Sync.Font = new Font(SystemFonts.MenuFont!, isCurrentChange ? FontStyle.Regular : FontStyle.Bold);
 				BuildListContextMenu_SyncContentOnly.Visible = !isBusy && ShouldSyncPrecompiledEditor;
 				BuildListContextMenu_SyncOnlyThisChange.Visible = !isBusy && !isCurrentChange && _contextMenuChange.Number > _workspace.CurrentChangeNumber && _workspace.CurrentChangeNumber != -1;
 				BuildListContextMenu_Build.Visible = !isBusy && isCurrentChange && !ShouldSyncPrecompiledEditor;
 				BuildListContextMenu_Rebuild.Visible = !isBusy && isCurrentChange && !ShouldSyncPrecompiledEditor;
 				BuildListContextMenu_GenerateProjectFiles.Visible = !isBusy && isCurrentChange;
 				BuildListContextMenu_LaunchEditor.Visible = !isBusy && _contextMenuChange.Number == _workspace.CurrentChangeNumber;
-				BuildListContextMenu_LaunchEditor.Font = new Font(SystemFonts.MenuFont, FontStyle.Bold);
+				BuildListContextMenu_LaunchEditor.Font = new Font(SystemFonts.MenuFont!, FontStyle.Bold);
 				BuildListContextMenu_OpenVisualStudio.Visible = !isBusy && isCurrentChange;
 				BuildListContextMenu_Cancel.Visible = isBusy;
 
@@ -4350,7 +4355,7 @@ namespace UnrealGameSync
 			Dictionary<string, KeyValuePair<IArchiveInfo, int>> archiveTypeToSelection = new Dictionary<string, KeyValuePair<IArchiveInfo, int>>();
 			foreach (IArchiveInfo archive in archives)
 			{
-				ArchiveSettings archiveSettings = _settings.Archives.FirstOrDefault(x => x.Type == archive.Type);
+				ArchiveSettings? archiveSettings = _settings.Archives.FirstOrDefault(x => x.Type == archive.Type);
 				if (archiveSettings != null && archiveSettings.Enabled)
 				{
 					int preference = archiveSettings.Order.IndexOf(archive.Name);
@@ -4382,7 +4387,7 @@ namespace UnrealGameSync
 
 		private void SetSelectedArchive(IArchiveInfo archive, bool selected)
 		{
-			ArchiveSettings archiveSettings = _settings.Archives.FirstOrDefault(x => x.Type == archive.Type);
+			ArchiveSettings? archiveSettings = _settings.Archives.FirstOrDefault(x => x.Type == archive.Type);
 			if (archiveSettings == null)
 			{
 				archiveSettings = new ArchiveSettings(selected, archive.Type, new string[] { archive.Name });
@@ -5040,7 +5045,7 @@ namespace UnrealGameSync
 			{
 				// Check the project config if they have a scheduled setting selected.
 				string? scheduledSyncTypeId = null;
-				UserSelectedProjectSettings projectSetting = _settings.ScheduleProjects.FirstOrDefault(x => x.LocalPath != null && x.LocalPath.Equals(SelectedProject.LocalPath, StringComparison.OrdinalIgnoreCase));
+				UserSelectedProjectSettings? projectSetting = _settings.ScheduleProjects.FirstOrDefault(x => x.LocalPath != null && x.LocalPath.Equals(SelectedProject.LocalPath, StringComparison.OrdinalIgnoreCase));
 				if (projectSetting != null)
 				{
 					scheduledSyncTypeId = projectSetting.ScheduledSyncTypeId;
@@ -5324,7 +5329,7 @@ namespace UnrealGameSync
 						while (stack.Count > 0)
 						{
 							Guid id = stack.Pop();
-							BuildStep nextStep = allSteps.FirstOrDefault(x => x.UniqueId == id);
+							BuildStep? nextStep = allSteps.FirstOrDefault(x => x.UniqueId == id);
 							if (nextStep != null)
 							{
 								foreach (Guid requiresId in nextStep.Requires)
@@ -5472,7 +5477,7 @@ namespace UnrealGameSync
 		{
 			if (OptionsContextMenu_SyncPrecompiledBinaries.DropDownItems.Count == 0)
 			{
-				IArchiveInfo editorArchive = GetArchives().FirstOrDefault(x => x.Type == IArchiveInfo.EditorArchiveType);
+				IArchiveInfo? editorArchive = GetArchives().FirstOrDefault(x => x.Type == IArchiveInfo.EditorArchiveType);
 				if (editorArchive != null)
 				{
 					SetSelectedArchive(editorArchive, !OptionsContextMenu_SyncPrecompiledBinaries.Checked);
@@ -6209,7 +6214,7 @@ namespace UnrealGameSync
 			ConfigFile config = new ConfigFile();
 			config.Load(newConfigFile);
 
-			ConfigSection section = config.FindSection("/Script/EditorStyle.EditorStyleSettings");
+			ConfigSection? section = config.FindSection("/Script/EditorStyle.EditorStyleSettings");
 			if (section == null)
 			{
 				return null;
@@ -6249,7 +6254,7 @@ namespace UnrealGameSync
 			ConfigFile config = new ConfigFile();
 			config.Load(fileName);
 
-			ConfigSection section = config.FindSection("/Script/EditorStyle.EditorStyleSettings");
+			ConfigSection? section = config.FindSection("/Script/EditorStyle.EditorStyleSettings");
 			if (section == null)
 			{
 				return null;
diff --git a/Engine/Source/Programs/UnrealGameSync/UnrealGameSync/EventMonitor.cs b/Engine/Source/Programs/UnrealGameSync/UnrealGameSync/EventMonitor.cs
--- a/Engine/Source/Programs/UnrealGameSync/UnrealGameSync/EventMonitor.cs
+++ b/Engine/Source/Programs/UnrealGameSync/UnrealGameSync/EventMonitor.cs
@@ -1142,19 +1142,19 @@ namespace UnrealGameSync
 		public bool IsUnderInvestigation(int changeNumber)
 		{
 			UpdateActiveInvestigations();
-			return _activeInvestigations.Any(x => x.Change <= changeNumber);
+			return _activeInvestigations?.Any(x => x.Change <= changeNumber) ?? false;
 		}
 
 		public bool IsUnderInvestigationByCurrentUser(int changeNumber)
 		{
 			UpdateActiveInvestigations();
-			return _activeInvestigations.Any(x => x.Change <= changeNumber && String.Compare(x.UserName, _currentUserName, true) == 0);
+			return _activeInvestigations?.Any(x => x.Change <= changeNumber && String.Compare(x.UserName, _currentUserName, true) == 0) ?? false;
 		}
 
 		public IEnumerable<string> GetInvestigatingUsers(int changeNumber)
 		{
 			UpdateActiveInvestigations();
-			return _activeInvestigations.Where(x => changeNumber >= x.Change).Select(x => x.UserName);
+			return _activeInvestigations?.Where(x => changeNumber >= x.Change).Select(x => x.UserName) ?? Array.Empty<string>();
 		}
 
 		public int GetInvestigationStartChangeNumber(int lastChangeNumber)
diff --git a/Engine/Source/Programs/UnrealGameSync/UnrealGameSync/Forms/BuildStepWindow.cs b/Engine/Source/Programs/UnrealGameSync/UnrealGameSync/Forms/BuildStepWindow.cs
--- a/Engine/Source/Programs/UnrealGameSync/UnrealGameSync/Forms/BuildStepWindow.cs
+++ b/Engine/Source/Programs/UnrealGameSync/UnrealGameSync/Forms/BuildStepWindow.cs
@@ -230,7 +230,7 @@ namespace UnrealGameSync
 			}
 		}
 
-		private void OnClosedVariablesWindow(object sender, EventArgs e)
+		private void OnClosedVariablesWindow(object? sender, EventArgs e)
 		{
 			_variablesWindow = null;
 		}
diff --git a/Engine/Source/Programs/UnrealGameSync/UnrealGameSync/Forms/CleanWorkspaceWindow.cs b/Engine/Source/Programs/UnrealGameSync/UnrealGameSync/Forms/CleanWorkspaceWindow.cs
--- a/Engine/Source/Programs/UnrealGameSync/UnrealGameSync/Forms/CleanWorkspaceWindow.cs
+++ b/Engine/Source/Programs/UnrealGameSync/UnrealGameSync/Forms/CleanWorkspaceWindow.cs
@@ -512,7 +512,7 @@ namespace UnrealGameSync
 			}
 			else if (nodeData.File != null)
 			{
-				Process.Start("explorer.exe", String.Format("\"{0}\"", nodeData.File.Directory.FullName));
+				Process.Start("explorer.exe", String.Format("\"{0}\"", nodeData.File.Directory!.FullName));
 			}
 		}
 	}
diff --git a/Engine/Source/Programs/UnrealGameSync/UnrealGameSync/Forms/IssueBrowserWindow.cs b/Engine/Source/Programs/UnrealGameSync/UnrealGameSync/Forms/IssueBrowserWindow.cs
--- a/Engine/Source/Programs/UnrealGameSync/UnrealGameSync/Forms/IssueBrowserWindow.cs
+++ b/Engine/Source/Programs/UnrealGameSync/UnrealGameSync/Forms/IssueBrowserWindow.cs
@@ -262,7 +262,7 @@ namespace UnrealGameSync
 
 		public static void Show(Form owner, IssueMonitor issueMonitor, IPerforceSettings perforceSettings, TimeSpan? serverTimeOffset, IServiceProvider serviceProvider, string? currentStream, Dictionary<string, Func<IssueData, bool>> customFilters, string? defaultFilter)
 		{
-			IssueBrowserWindow window = _existingWindows.FirstOrDefault(x => x._issueMonitor == issueMonitor);
+			IssueBrowserWindow? window = _existingWindows.FirstOrDefault(x => x._issueMonitor == issueMonitor);
 			if(window == null)
 			{
 				window = new IssueBrowserWindow(issueMonitor, perforceSettings, serverTimeOffset, serviceProvider, currentStream, customFilters, defaultFilter);
@@ -403,6 +403,11 @@ namespace UnrealGameSync
 
 		private void IssueListView_DrawSubItem(object sender, DrawListViewSubItemEventArgs e)
 		{
+			if (e.Item == null || e.SubItem == null)
+			{
+				return;
+			}
+
 			IssueData issue = (IssueData)e.Item.Tag;
 			if (e.ColumnIndex == IconHeader.Index)
 			{
diff --git a/Engine/Source/Programs/UnrealGameSync/UnrealGameSync/Forms/IssueDetailsWindow.cs b/Engine/Source/Programs/UnrealGameSync/UnrealGameSync/Forms/IssueDetailsWindow.cs
--- a/Engine/Source/Programs/UnrealGameSync/UnrealGameSync/Forms/IssueDetailsWindow.cs
+++ b/Engine/Source/Programs/UnrealGameSync/UnrealGameSync/Forms/IssueDetailsWindow.cs
@@ -455,7 +455,7 @@ namespace UnrealGameSync
 			foreach (IGrouping<long, IssueDiagnosticData> group in diagnostics.GroupBy(x => x.BuildId ?? -1))
 			{
 				// Step 'Foo'
-				IssueBuildData build = issueBuilds.FirstOrDefault(x => x.Id == group.Key);
+				IssueBuildData? build = issueBuilds.FirstOrDefault(x => x.Id == group.Key);
 				if (build != null)
 				{
 					richText.Append(@"\pard");   // Paragraph default
@@ -527,7 +527,7 @@ namespace UnrealGameSync
 		{
 			UpdateSummaryTextIfChanged(SummaryTextBox, _issue.Summary.ToString());
 
-			IssueBuildData firstFailingBuild = _issueBuilds.FirstOrDefault(x => x.ErrorUrl != null);
+			IssueBuildData? firstFailingBuild = _issueBuilds.FirstOrDefault(x => x.ErrorUrl != null);
 			BuildLinkLabel.Text = (firstFailingBuild != null)? firstFailingBuild.JobName : "Unknown";
 
 			StringBuilder status = new StringBuilder();
@@ -1044,7 +1044,7 @@ namespace UnrealGameSync
 
 		public static void Show(Form owner, IssueMonitor issueMonitor, IPerforceSettings perforceSettings, TimeSpan? serverTimeOffset, IssueData issue, List<IssueBuildData> issueBuilds, IServiceProvider serviceProvider, string? currentStream)
 		{
-			IssueDetailsWindow window = _existingWindows.FirstOrDefault(x => x._issueMonitor == issueMonitor && x._issue.Id == issue.Id);
+			IssueDetailsWindow? window = _existingWindows.FirstOrDefault(x => x._issueMonitor == issueMonitor && x._issue.Id == issue.Id);
 			if(window == null)
 			{
 				List<IssueDiagnosticData> diagnostics = new List<IssueDiagnosticData>();
@@ -1141,6 +1141,11 @@ namespace UnrealGameSync
 
 		private void BuildListView_DrawSubItem(object sender, DrawListViewSubItemEventArgs e)
 		{
+			if (e.Item == null || e.SubItem == null)
+			{
+				return;
+			}
+
 			Font currentFont = this.Font;//BuildFont;(Change.Number == Workspace.PendingChangeNumber || Change.Number == Workspace.CurrentChangeNumber)? SelectedBuildFont : BuildFont;
 
 			Color textColor = SystemColors.WindowText;
@@ -1369,7 +1374,7 @@ namespace UnrealGameSync
 
 		private void DescriptionLabel_LinkClicked(object sender, LinkLabelLinkClickedEventArgs e)
 		{
-			IssueBuildData lastBuild = _issueBuilds.Where(x => x.Stream == _selectedStream).OrderByDescending(x => x.Change).ThenByDescending(x => x.ErrorUrl).FirstOrDefault();
+			IssueBuildData? lastBuild = _issueBuilds.Where(x => x.Stream == _selectedStream).OrderByDescending(x => x.Change).ThenByDescending(x => x.ErrorUrl).FirstOrDefault();
 			if(lastBuild != null)
 			{
 				Utility.OpenUrl(lastBuild.ErrorUrl);
@@ -1417,7 +1422,7 @@ namespace UnrealGameSync
 
 		private void BuildLinkLabel_LinkClicked(object sender, LinkLabelLinkClickedEventArgs e)
 		{
-			IssueBuildData build = _issueBuilds.FirstOrDefault(x => x.ErrorUrl != null);
+			IssueBuildData? build = _issueBuilds.FirstOrDefault(x => x.ErrorUrl != null);
 			if (build != null)
 			{
 				Utility.OpenUrl(build.ErrorUrl);
@@ -1426,7 +1431,10 @@ namespace UnrealGameSync
 
 		private void DetailsTextBox_LinkClicked(object sender, LinkClickedEventArgs e)
 		{
-			Utility.OpenUrl(e.LinkText);
+			if (e.LinkText != null)
+			{
+				Utility.OpenUrl(e.LinkText);
+			}
 		}
 	}
 }
diff --git a/Engine/Source/Programs/UnrealGameSync/UnrealGameSync/Forms/MainWindow.cs b/Engine/Source/Programs/UnrealGameSync/UnrealGameSync/Forms/MainWindow.cs
--- a/Engine/Source/Programs/UnrealGameSync/UnrealGameSync/Forms/MainWindow.cs
+++ b/Engine/Source/Programs/UnrealGameSync/UnrealGameSync/Forms/MainWindow.cs
@@ -1525,7 +1525,7 @@ namespace UnrealGameSync
 						}
 					}
 
-					IssueAlertWindow alertWindow = _alertWindows.FirstOrDefault(x => x.IssueMonitor == issueMonitor && x.Issue.Id == issue.Id);
+					IssueAlertWindow? alertWindow = _alertWindows.FirstOrDefault(x => x.IssueMonitor == issueMonitor && x.Issue.Id == issue.Id);
 					if(alertWindow == null)
 					{
 						if(reason != 0)
@@ -1677,7 +1677,7 @@ namespace UnrealGameSync
 
 		public IssueMonitor CreateIssueMonitor(string? apiUrl, string userName)
 		{
-			WorkspaceIssueMonitor workspaceIssueMonitor = _workspaceIssueMonitors.FirstOrDefault(x => String.Compare(x.IssueMonitor.ApiUrl, apiUrl, StringComparison.OrdinalIgnoreCase) == 0 && String.Compare(x.IssueMonitor.UserName, userName, StringComparison.OrdinalIgnoreCase) == 0);
+			WorkspaceIssueMonitor? workspaceIssueMonitor = _workspaceIssueMonitors.FirstOrDefault(x => String.Compare(x.IssueMonitor.ApiUrl, apiUrl, StringComparison.OrdinalIgnoreCase) == 0 && String.Compare(x.IssueMonitor.UserName, userName, StringComparison.OrdinalIgnoreCase) == 0);
 			if (workspaceIssueMonitor == null)
 			{
 				string serverId = apiUrl != null ? Regex.Replace(apiUrl, @"^.*://", "") : "noserver";
diff --git a/Engine/Source/Programs/UnrealGameSync/UnrealGameSync/Forms/ModifyBuildStepsWindow.cs b/Engine/Source/Programs/UnrealGameSync/UnrealGameSync/Forms/ModifyBuildStepsWindow.cs
--- a/Engine/Source/Programs/UnrealGameSync/UnrealGameSync/Forms/ModifyBuildStepsWindow.cs
+++ b/Engine/Source/Programs/UnrealGameSync/UnrealGameSync/Forms/ModifyBuildStepsWindow.cs
@@ -142,10 +142,9 @@ namespace UnrealGameSync
 
 		private void BuildStepList_DrawSubItem(object sender, DrawListViewSubItemEventArgs e)
 		{
-			if(e.ColumnIndex == 0)
+			if(e.ColumnIndex == 0 || e.Item == null)
 			{
 				e.DrawDefault = true;
-				return;
 			}
 			else
 			{
diff --git a/Engine/Source/Programs/UnrealGameSync/UnrealGameSync/Forms/NotificationWindow.cs b/Engine/Source/Programs/UnrealGameSync/UnrealGameSync/Forms/NotificationWindow.cs
--- a/Engine/Source/Programs/UnrealGameSync/UnrealGameSync/Forms/NotificationWindow.cs
+++ b/Engine/Source/Programs/UnrealGameSync/UnrealGameSync/Forms/NotificationWindow.cs
@@ -234,7 +234,7 @@ namespace UnrealGameSync
 			e.Graphics.DrawRectangle(Pens.Black, new Rectangle(0, 0, Width - 1, Height - 1));
 
 			e.Graphics.InterpolationMode = InterpolationMode.HighQualityBicubic;
-			e.Graphics.DrawImage(_logoBitmap, _logoBounds);
+			e.Graphics.DrawImage(_logoBitmap!, _logoBounds);
 		}
 
 		protected override void OnPaint(PaintEventArgs e)
diff --git a/Engine/Source/Programs/UnrealGameSync/UnrealGameSync/Forms/ScheduleWindow.cs b/Engine/Source/Programs/UnrealGameSync/UnrealGameSync/Forms/ScheduleWindow.cs
--- a/Engine/Source/Programs/UnrealGameSync/UnrealGameSync/Forms/ScheduleWindow.cs
+++ b/Engine/Source/Programs/UnrealGameSync/UnrealGameSync/Forms/ScheduleWindow.cs
@@ -76,7 +76,7 @@ namespace UnrealGameSync
 			}
 		}
 
-		private void ProjectListBox_ItemCheck(object sender, ItemCheckEventArgs e)
+		private void ProjectListBox_ItemCheck(object? sender, ItemCheckEventArgs e)
 		{
 			bool isAnyOpenProjectIndex = e.Index == 0;
 			if (e.NewValue == CheckState.Checked && !isAnyOpenProjectIndex)
@@ -90,7 +90,7 @@ namespace UnrealGameSync
 			}
 		}
 
-		private void SyncTypeDropDown_Closed(object sender, ToolStripDropDownClosedEventArgs e)
+		private void SyncTypeDropDown_Closed(object? sender, ToolStripDropDownClosedEventArgs e)
 		{
 			if (e.CloseReason != ToolStripDropDownCloseReason.ItemClicked)
 			{
diff --git a/Engine/Source/Programs/UnrealGameSync/UnrealGameSync/Forms/SdkInfoWindow.cs b/Engine/Source/Programs/UnrealGameSync/UnrealGameSync/Forms/SdkInfoWindow.cs
--- a/Engine/Source/Programs/UnrealGameSync/UnrealGameSync/Forms/SdkInfoWindow.cs
+++ b/Engine/Source/Programs/UnrealGameSync/UnrealGameSync/Forms/SdkInfoWindow.cs
@@ -145,7 +145,11 @@ namespace UnrealGameSync
 
 		private void SdkListView_DrawSubItem(object sender, DrawListViewSubItemEventArgs e)
 		{
-			if(e.ColumnIndex != columnHeader3.Index)
+			if (e.Item == null || e.SubItem == null)
+			{
+				e.DrawDefault = true;
+			}
+			else if(e.ColumnIndex != columnHeader3.Index)
 			{
 				TextRenderer.DrawText(e.Graphics, e.SubItem.Text, SdkListView.Font, e.Bounds, SdkListView.ForeColor, TextFormatFlags.EndEllipsis | TextFormatFlags.SingleLine | TextFormatFlags.Left | TextFormatFlags.VerticalCenter | TextFormatFlags.NoPrefix);
 			}
diff --git a/Engine/Source/Programs/UnrealGameSync/UnrealGameSync/Forms/SelectProjectFromWorkspaceWindow.cs b/Engine/Source/Programs/UnrealGameSync/UnrealGameSync/Forms/SelectProjectFromWorkspaceWindow.cs
--- a/Engine/Source/Programs/UnrealGameSync/UnrealGameSync/Forms/SelectProjectFromWorkspaceWindow.cs
+++ b/Engine/Source/Programs/UnrealGameSync/UnrealGameSync/Forms/SelectProjectFromWorkspaceWindow.cs
@@ -225,7 +225,7 @@ namespace UnrealGameSync
 
 		private void ProjectTreeView_AfterSelect(object sender, TreeViewEventArgs e)
 		{
-			OkBtn.Enabled = (e.Node.Tag != null);
+			OkBtn.Enabled = (e.Node?.Tag != null);
 		}
 
 		private void OkBtn_Click(object sender, EventArgs e)
diff --git a/Engine/Source/Programs/UnrealGameSync/UnrealGameSync/Forms/SyncFilter.cs b/Engine/Source/Programs/UnrealGameSync/UnrealGameSync/Forms/SyncFilter.cs
--- a/Engine/Source/Programs/UnrealGameSync/UnrealGameSync/Forms/SyncFilter.cs
+++ b/Engine/Source/Programs/UnrealGameSync/UnrealGameSync/Forms/SyncFilter.cs
@@ -17,16 +17,16 @@ namespace UnrealGameSync
 		Dictionary<Guid, WorkspaceSyncCategory> _uniqueIdToCategory;
 		public FilterSettings GlobalFilter;
 		public FilterSettings WorkspaceFilter;
-		ConfigSection _perforceSection;
+		ConfigSection? _perforceSection;
 
-		public SyncFilter(Dictionary<Guid, WorkspaceSyncCategory> inUniqueIdToCategory, FilterSettings inGlobalFilter, FilterSettings inWorkspaceFilter, ConfigSection inPerforceSection)
+		public SyncFilter(Dictionary<Guid, WorkspaceSyncCategory> uniqueIdToCategory, FilterSettings globalFilter, FilterSettings workspaceFilter, ConfigSection? perforceSection)
 		{
 			InitializeComponent();
 
-			_uniqueIdToCategory = inUniqueIdToCategory;
-			GlobalFilter = inGlobalFilter;
-			WorkspaceFilter = inWorkspaceFilter;
-			_perforceSection = inPerforceSection;
+			_uniqueIdToCategory = uniqueIdToCategory;
+			GlobalFilter = globalFilter;
+			WorkspaceFilter = workspaceFilter;
+			_perforceSection = perforceSection;
 
 			Dictionary<Guid, bool> syncCategories = WorkspaceSyncCategory.GetDefault(_uniqueIdToCategory.Values);
 
diff --git a/Engine/Source/Programs/UnrealGameSync/UnrealGameSync/IssueMonitor.cs b/Engine/Source/Programs/UnrealGameSync/UnrealGameSync/IssueMonitor.cs
--- a/Engine/Source/Programs/UnrealGameSync/UnrealGameSync/IssueMonitor.cs
+++ b/Engine/Source/Programs/UnrealGameSync/UnrealGameSync/IssueMonitor.cs
@@ -427,7 +427,7 @@ namespace UnrealGameSync
 						{
 							newIssue.IsWarning = !builds.Any(x => x.Outcome != IssueBuildOutcome.Warning);
 
-							IssueBuildData lastBuild = builds.OrderByDescending(x => x.Change).FirstOrDefault();
+							IssueBuildData? lastBuild = builds.OrderByDescending(x => x.Change).FirstOrDefault();
 							if (lastBuild != null && !String.IsNullOrEmpty(lastBuild.ErrorUrl))
 							{
 								newIssue.BuildUrl = lastBuild.ErrorUrl;
diff --git a/Engine/Source/Programs/UnrealGameSync/UnrealGameSync/PerforceMonitor.cs b/Engine/Source/Programs/UnrealGameSync/UnrealGameSync/PerforceMonitor.cs
--- a/Engine/Source/Programs/UnrealGameSync/UnrealGameSync/PerforceMonitor.cs
+++ b/Engine/Source/Programs/UnrealGameSync/UnrealGameSync/PerforceMonitor.cs
@@ -273,7 +273,7 @@ namespace UnrealGameSync
 				}
 
 				// Add in additional paths property
-				ConfigSection projectConfigSection = LatestProjectConfigFile.FindSection("Perforce");
+				ConfigSection? projectConfigSection = LatestProjectConfigFile.FindSection("Perforce");
 				if (projectConfigSection != null)
 				{
 					IEnumerable<string> additionalPaths = projectConfigSection.GetValues("AdditionalPathsToSync", new string[0]);
@@ -399,7 +399,11 @@ namespace UnrealGameSync
 
 		Func<string, bool>? CreateCodeFilterFunc()
 		{
-			ConfigSection projectConfigSection = LatestProjectConfigFile.FindSection("Perforce");
+			ConfigSection? projectConfigSection = LatestProjectConfigFile.FindSection("Perforce");
+			if (projectConfigSection == null)
+			{
+				return null;
+			}
 
 			string[] rules = projectConfigSection.GetValues("CodeFilter", new string[0]);
 			if (rules.Length == 0)
@@ -419,7 +423,7 @@ namespace UnrealGameSync
 		public async Task<bool> UpdateChangeTypesAsync(IPerforceConnection perforce, CancellationToken cancellationToken)
 		{
 			// Get the filter for code changes
-			ConfigSection projectConfigSection = LatestProjectConfigFile.FindSection("Perforce");
+			ConfigSection? projectConfigSection = LatestProjectConfigFile.FindSection("Perforce");
 
 			string[] codeRules = projectConfigSection?.GetValues("CodeFilter", (string[]?)null) ?? Array.Empty<string>();
 			if (!Enumerable.SequenceEqual(codeRules, prevCodeRules))
@@ -596,7 +600,7 @@ namespace UnrealGameSync
 			}
 		}
 
-		public ConfigSection LatestPerforceConfigSection()
+		public ConfigSection? LatestPerforceConfigSection()
 		{
 			return LatestProjectConfigFile.FindSection("Perforce");
 		}
diff --git a/Engine/Source/Programs/UnrealGameSync/UnrealGameSync/Program.cs b/Engine/Source/Programs/UnrealGameSync/UnrealGameSync/Program.cs
--- a/Engine/Source/Programs/UnrealGameSync/UnrealGameSync/Program.cs
+++ b/Engine/Source/Programs/UnrealGameSync/UnrealGameSync/Program.cs
@@ -147,7 +147,7 @@ namespace UnrealGameSync
 			{
 				if (File.Exists(updateSpawn))
 				{
-					Directory.SetCurrentDirectory(Path.GetDirectoryName(updateSpawn));
+					Directory.SetCurrentDirectory(Path.GetDirectoryName(updateSpawn)!);
 				}
 			}
 
diff --git a/Engine/Source/Programs/UnrealGameSync/UnrealGameSync/ProgramApplicationContext.cs b/Engine/Source/Programs/UnrealGameSync/UnrealGameSync/ProgramApplicationContext.cs
--- a/Engine/Source/Programs/UnrealGameSync/UnrealGameSync/ProgramApplicationContext.cs
+++ b/Engine/Source/Programs/UnrealGameSync/UnrealGameSync/ProgramApplicationContext.cs
@@ -261,24 +261,24 @@ namespace UnrealGameSync
 			_tokenStore = TokenStoreFactory.CreateTokenStore();
 			List<(DirectoryInfo, DirectoryInfo?)> configurationLocations = new List<(DirectoryInfo, DirectoryInfo?)>();
 			List<string> allowedProviders = new List<string>();
-			foreach ((UserSelectedProjectSettings, ModalTask<OpenProjectInfo>) StartupTask in startupTasks)
+			foreach ((UserSelectedProjectSettings, ModalTask<OpenProjectInfo>) startupTask in startupTasks)
 			{
 				try
 				{
 
-					ConfigFile ConfigFile = StartupTask.Item2.Result.LatestProjectConfigFile;
-					if (ConfigFile == null)
+					ConfigFile configFile = startupTask.Item2.Result.LatestProjectConfigFile;
+					if (configFile == null)
 					{
 						continue;
 					}
 
-					ConfigSection ProviderSection = ConfigFile.FindSection("OIDCProvider");
-					if (ProviderSection == null)
+					ConfigSection? providerSection = configFile.FindSection("OIDCProvider");
+					if (providerSection == null)
 					{
 						continue;
 					}
 
-					string[] oidcAllowedProviders = ProviderSection.GetValues("OidcProviderAllowList", Array.Empty<string>());
+					string[] oidcAllowedProviders = providerSection.GetValues("OidcProviderAllowList", Array.Empty<string>());
 					if (oidcAllowedProviders.Length == 0)
 					{
 						continue;
@@ -286,7 +286,7 @@ namespace UnrealGameSync
 
 					allowedProviders.AddRange(oidcAllowedProviders);
 					
-					ProjectInfo projectInfo = StartupTask.Item2.Result.ProjectInfo;
+					ProjectInfo projectInfo = startupTask.Item2.Result.ProjectInfo;
 					configurationLocations.Add((projectInfo.EngineDir.ToDirectoryInfo(), projectInfo.ProjectDir?.ToDirectoryInfo()));
 				}
 				catch (Exception)
@@ -333,7 +333,7 @@ namespace UnrealGameSync
 			_mainWindowInstance.FormClosed += MainWindowInstance_FormClosed;
 		}
 
-		private void MainWindowInstance_FormClosed(object sender, FormClosedEventArgs e)
+		private void MainWindowInstance_FormClosed(object? sender, FormClosedEventArgs e)
 		{
 			ExitThread();
 		}
@@ -418,7 +418,7 @@ namespace UnrealGameSync
 			}
 		}
 
-		private void NotifyIcon_MouseDown(object sender, MouseEventArgs e)
+		private void NotifyIcon_MouseDown(object? sender, MouseEventArgs e)
 		{
 			// Have to set up this stuff here, because the menu is laid out before Opening() is called on it after mouse-up.
 			bool canSyncNow = _mainWindowInstance != null && _mainWindowInstance.CanSyncNow();
diff --git a/Engine/Source/Programs/UnrealGameSync/UnrealGameSync/UriHandlers/UriHandler.cs b/Engine/Source/Programs/UnrealGameSync/UnrealGameSync/UriHandlers/UriHandler.cs
--- a/Engine/Source/Programs/UnrealGameSync/UnrealGameSync/UriHandlers/UriHandler.cs
+++ b/Engine/Source/Programs/UnrealGameSync/UnrealGameSync/UriHandlers/UriHandler.cs
@@ -297,7 +297,7 @@ namespace UnrealGameSync
 
 		public static string CurrentProcessFilePath
 		{
-			get { return Path.ChangeExtension(Path.GetFullPath(Process.GetCurrentProcess().MainModule.FileName), ".exe"); }
+			get { return Path.ChangeExtension(Path.GetFullPath(Process.GetCurrentProcess().MainModule!.FileName!), ".exe"); }
 		}
 
 		static List<RegistrySetting> GetGlobalRegistrySettings()
diff --git a/Engine/Source/Programs/UnrealGameSync/UnrealGameSyncCmd/Program.cs b/Engine/Source/Programs/UnrealGameSync/UnrealGameSyncCmd/Program.cs
--- a/Engine/Source/Programs/UnrealGameSync/UnrealGameSyncCmd/Program.cs
+++ b/Engine/Source/Programs/UnrealGameSync/UnrealGameSyncCmd/Program.cs
@@ -337,7 +337,7 @@ namespace UnrealGameSyncCmd
 		static string[] ReadSyncFilter(UserWorkspaceSettings workspaceSettings, GlobalSettingsFile userSettings, ConfigFile projectConfig)
 		{
 			Dictionary<Guid, WorkspaceSyncCategory> syncCategories = ConfigUtils.GetSyncCategories(projectConfig);
-			ConfigSection perforceSection = projectConfig.FindSection("Perforce");
+			ConfigSection? perforceSection = projectConfig.FindSection("Perforce");
 
 			string[] combinedSyncFilter = GlobalSettingsFile.GetCombinedSyncFilter(syncCategories, userSettings.Global.Filter, workspaceSettings.Filter, perforceSection);
 
diff --git a/Engine/Source/Programs/UnrealGameSync/UnrealGameSyncShared/ArchiveInfo.cs b/Engine/Source/Programs/UnrealGameSync/UnrealGameSyncShared/ArchiveInfo.cs
--- a/Engine/Source/Programs/UnrealGameSync/UnrealGameSyncShared/ArchiveInfo.cs
+++ b/Engine/Source/Programs/UnrealGameSync/UnrealGameSyncShared/ArchiveInfo.cs
@@ -148,7 +148,7 @@ namespace UnrealGameSync
 			List<PerforceArchiveInfo> newArchives = new List<PerforceArchiveInfo>();
 
 			// Find all the zipped binaries under this stream
-			ConfigSection projectConfigSection = latestProjectConfigFile.FindSection(projectIdentifier);
+			ConfigSection? projectConfigSection = latestProjectConfigFile.FindSection(projectIdentifier);
 			if (projectConfigSection != null)
 			{
 				// Legacy
diff --git a/Engine/Source/Programs/UnrealGameSync/UnrealGameSyncShared/ConfigFile.cs b/Engine/Source/Programs/UnrealGameSync/UnrealGameSyncShared/ConfigFile.cs
--- a/Engine/Source/Programs/UnrealGameSync/UnrealGameSyncShared/ConfigFile.cs
+++ b/Engine/Source/Programs/UnrealGameSync/UnrealGameSyncShared/ConfigFile.cs
@@ -591,14 +591,14 @@ namespace UnrealGameSync
 			FileReference.Move(tempFileName, fileName);
 		}
 
-		public ConfigSection FindSection(string name)
+		public ConfigSection? FindSection(string name)
 		{
 			return _sections.FirstOrDefault(x => String.Compare(x.Name, name, true) == 0);
 		}
 
 		public ConfigSection FindOrAddSection(string name)
 		{
-			ConfigSection section = FindSection(name);
+			ConfigSection? section = FindSection(name);
 			if(section == null)
 			{
 				section = new ConfigSection(name);
@@ -638,14 +638,14 @@ namespace UnrealGameSync
 		public bool GetValue(string key, bool defaultValue)
 		{
 			int dotIdx = key.IndexOf('.');
-			ConfigSection section = FindSection(key.Substring(0, dotIdx));
+			ConfigSection? section = FindSection(key.Substring(0, dotIdx));
 			return (section == null)? defaultValue : section.GetValue(key.Substring(dotIdx + 1), defaultValue);
 		}
 
 		public int GetValue(string key, int defaultValue)
 		{
 			int dotIdx = key.IndexOf('.');
-			ConfigSection section = FindSection(key.Substring(0, dotIdx));
+			ConfigSection? section = FindSection(key.Substring(0, dotIdx));
 			return (section == null)? defaultValue : section.GetValue(key.Substring(dotIdx + 1), defaultValue);
 		}
 
@@ -653,7 +653,7 @@ namespace UnrealGameSync
 		public string? GetValue(string key, string? defaultValue)
 		{
 			int dotIdx = key.IndexOf('.');
-			ConfigSection section = FindSection(key.Substring(0, dotIdx));
+			ConfigSection? section = FindSection(key.Substring(0, dotIdx));
 			return (section == null)? defaultValue : section.GetValue(key.Substring(dotIdx + 1), defaultValue);
 		}
 
@@ -661,7 +661,7 @@ namespace UnrealGameSync
 		public string[]? GetValues(string key, string[]? defaultValue)
 		{
 			int dotIdx = key.IndexOf('.');
-			ConfigSection section = FindSection(key.Substring(0, dotIdx));
+			ConfigSection? section = FindSection(key.Substring(0, dotIdx));
 			return (section == null)? defaultValue : section.GetValues(key.Substring(dotIdx + 1), defaultValue);
 		}
 
@@ -669,7 +669,7 @@ namespace UnrealGameSync
 		public Guid[]? GetGuidValues(string key, Guid[]? defaultValue)
 		{
 			int dotIdx = key.IndexOf('.');
-			ConfigSection section = FindSection(key.Substring(0, dotIdx));
+			ConfigSection? section = FindSection(key.Substring(0, dotIdx));
 			return (section == null)? defaultValue : section.GetValues(key.Substring(dotIdx + 1), defaultValue);
 		}
 	}
diff --git a/Engine/Source/Programs/UnrealGameSync/UnrealGameSyncShared/ConfigUtils.cs b/Engine/Source/Programs/UnrealGameSync/UnrealGameSyncShared/ConfigUtils.cs
--- a/Engine/Source/Programs/UnrealGameSync/UnrealGameSyncShared/ConfigUtils.cs
+++ b/Engine/Source/Programs/UnrealGameSync/UnrealGameSyncShared/ConfigUtils.cs
@@ -386,7 +386,7 @@ namespace UnrealGameSync
 			string path = selectedProjectIdentifier;
 			for (; ; )
 			{
-				ConfigSection projectSection = projectConfigFile.FindSection(path);
+				ConfigSection? projectSection = projectConfigFile.FindSection(path);
 				if (projectSection != null)
 				{
 					string? newValue = projectSection.GetValue(name, null);
@@ -406,7 +406,7 @@ namespace UnrealGameSync
 				path = path.Substring(0, lastSlash);
 			}
 
-			ConfigSection defaultSection = projectConfigFile.FindSection("Default");
+			ConfigSection? defaultSection = projectConfigFile.FindSection("Default");
 			if (defaultSection != null)
 			{
 				string? newValue = defaultSection.GetValue(name, null);
diff --git a/Engine/Source/Programs/UnrealGameSync/UnrealGameSyncShared/FileFilter.cs b/Engine/Source/Programs/UnrealGameSync/UnrealGameSyncShared/FileFilter.cs
--- a/Engine/Source/Programs/UnrealGameSync/UnrealGameSyncShared/FileFilter.cs
+++ b/Engine/Source/Programs/UnrealGameSync/UnrealGameSyncShared/FileFilter.cs
@@ -266,7 +266,7 @@ namespace UnrealGameSync
 			FileFilterNode lastNode = _rootNode;
 			foreach (string branchPattern in branchPatterns)
 			{
-				FileFilterNode nextNode = lastNode.Branches.FirstOrDefault(x => x.Pattern == branchPattern);
+				FileFilterNode? nextNode = lastNode.Branches.FirstOrDefault(x => x.Pattern == branchPattern);
 				if (nextNode == null)
 				{
 					nextNode = new FileFilterNode(lastNode, branchPattern);
diff --git a/Engine/Source/Programs/UnrealGameSync/UnrealGameSyncShared/GlobalSettings.cs b/Engine/Source/Programs/UnrealGameSync/UnrealGameSyncShared/GlobalSettings.cs
--- a/Engine/Source/Programs/UnrealGameSync/UnrealGameSyncShared/GlobalSettings.cs
+++ b/Engine/Source/Programs/UnrealGameSync/UnrealGameSyncShared/GlobalSettings.cs
@@ -174,7 +174,7 @@ namespace UnrealGameSync
 			return settings;
 		}
 
-		public static string[] GetCombinedSyncFilter(Dictionary<Guid, WorkspaceSyncCategory> uniqueIdToFilter, FilterSettings globalFilter, FilterSettings workspaceFilter, ConfigSection perforceSection)
+		public static string[] GetCombinedSyncFilter(Dictionary<Guid, WorkspaceSyncCategory> uniqueIdToFilter, FilterSettings globalFilter, FilterSettings workspaceFilter, ConfigSection? perforceSection)
 		{
 			List<string> lines = new List<string>();
 			foreach (string viewLine in Enumerable.Concat(globalFilter.View, workspaceFilter.View).Select(x => x.Trim()).Where(x => x.Length > 0 && !x.StartsWith(";")))
diff --git a/Engine/Source/Programs/UnrealGameSync/UnrealGameSyncShared/Telemetry.cs b/Engine/Source/Programs/UnrealGameSync/UnrealGameSyncShared/Telemetry.cs
--- a/Engine/Source/Programs/UnrealGameSync/UnrealGameSyncShared/Telemetry.cs
+++ b/Engine/Source/Programs/UnrealGameSync/UnrealGameSyncShared/Telemetry.cs
@@ -180,7 +180,7 @@ namespace UnrealGameSync
 				catch (WebException ex)
 				{
 					// Handle errors. Any non-200 responses automatically generate a WebException.
-					HttpWebResponse response = (HttpWebResponse)ex.Response;
+					HttpWebResponse? response = (HttpWebResponse?)ex.Response;
 					if (response == null)
 					{
 						_logger.LogError(ex, "Exception while attempting to send event");
diff --git a/Engine/Source/Programs/UnrealGameSync/UnrealGameSyncShared/UserSettings.cs b/Engine/Source/Programs/UnrealGameSync/UnrealGameSyncShared/UserSettings.cs
--- a/Engine/Source/Programs/UnrealGameSync/UnrealGameSyncShared/UserSettings.cs
+++ b/Engine/Source/Programs/UnrealGameSync/UnrealGameSyncShared/UserSettings.cs
@@ -477,7 +477,7 @@ namespace UnrealGameSync
 
 		public void SetBisectState(int change, BisectState state)
 		{
-			BisectEntry entry = BisectChanges.FirstOrDefault(x => x.Change == change);
+			BisectEntry? entry = BisectChanges.FirstOrDefault(x => x.Change == change);
 			if (entry == null)
 			{
 				entry = new BisectEntry();
@@ -925,7 +925,7 @@ namespace UnrealGameSync
 			SyncOptions.NumSyncErrorRetries = _configFile.GetValue("Perforce.NumSyncErrorRetries", PerforceSyncOptions.DefaultNumSyncErrorRetries);
 		}
 
-		static Dictionary<Guid, bool> GetCategorySettings(ConfigSection section, string includedKey, string excludedKey)
+		static Dictionary<Guid, bool> GetCategorySettings(ConfigSection? section, string includedKey, string excludedKey)
 		{
 			Dictionary<Guid, bool> result = new Dictionary<Guid, bool>();
 			if (section != null)
@@ -1010,7 +1010,7 @@ namespace UnrealGameSync
 		protected override void ImportWorkspaceState(DirectoryReference rootDir, string clientName, string branchPath, UserWorkspaceState currentWorkspace)
 		{
 			// Read the workspace settings
-			ConfigSection workspaceSection = _configFile.FindSection(clientName + branchPath);
+			ConfigSection? workspaceSection = _configFile.FindSection(clientName + branchPath);
 			if(workspaceSection == null)
 			{
 				string legacyBranchAndClientKey = clientName + branchPath;
@@ -1109,7 +1109,7 @@ namespace UnrealGameSync
 
 		protected override void ImportWorkspaceSettings(DirectoryReference rootDir, string clientName, string branchPath, UserWorkspaceSettings currentWorkspace)
 		{
-			ConfigSection workspaceSection = _configFile.FindSection(clientName + branchPath);
+			ConfigSection? workspaceSection = _configFile.FindSection(clientName + branchPath);
 			if (workspaceSection != null)
 			{
 				currentWorkspace.Filter.View = workspaceSection.GetValues("SyncFilter", new string[0]).ToList();
diff --git a/Engine/Source/Programs/UnrealGameSync/UnrealGameSyncShared/Utility.cs b/Engine/Source/Programs/UnrealGameSync/UnrealGameSyncShared/Utility.cs
--- a/Engine/Source/Programs/UnrealGameSync/UnrealGameSyncShared/Utility.cs
+++ b/Engine/Source/Programs/UnrealGameSync/UnrealGameSyncShared/Utility.cs
@@ -130,7 +130,7 @@ namespace UnrealGameSync
 
 		public static string GetPathWithCorrectCase(FileInfo info)
 		{
-			DirectoryInfo parentInfo = info.Directory;
+			DirectoryInfo parentInfo = info.Directory!;
 			if(info.Exists)
 			{
 				return Path.Combine(GetPathWithCorrectCase(parentInfo), parentInfo.GetFiles(info.Name)[0].Name); 
@@ -143,7 +143,7 @@ namespace UnrealGameSync
 
 		public static string GetPathWithCorrectCase(DirectoryInfo info)
 		{
-			DirectoryInfo parentInfo = info.Parent;
+			DirectoryInfo? parentInfo = info.Parent;
 			if(parentInfo == null)
 			{
 				return info.FullName.ToUpperInvariant();
@@ -372,7 +372,7 @@ namespace UnrealGameSync
 
 			if (projectFile.Name.EndsWith(".uproject", StringComparison.OrdinalIgnoreCase))
 			{
-				AddLocalConfigPaths_WithExtensionDirs(projectFile.Directory, "Build", "UnrealGameSync.ini", searchPaths);
+				AddLocalConfigPaths_WithExtensionDirs(projectFile.Directory!, "Build", "UnrealGameSync.ini", searchPaths);
 			}
 			else
 			{
@@ -596,7 +596,7 @@ namespace UnrealGameSync
 			ProcessStartInfo startInfo = new ProcessStartInfo();
 			startInfo.FileName = url;
 			startInfo.UseShellExecute = true;
-			using Process _ = Process.Start(startInfo);
+			using Process? _ = Process.Start(startInfo);
 		}
 	}
 }
diff --git a/Engine/Source/Programs/UnrealGameSync/UnrealGameSyncShared/WorkspaceUpdate.cs b/Engine/Source/Programs/UnrealGameSync/UnrealGameSyncShared/WorkspaceUpdate.cs
--- a/Engine/Source/Programs/UnrealGameSync/UnrealGameSyncShared/WorkspaceUpdate.cs
+++ b/Engine/Source/Programs/UnrealGameSync/UnrealGameSyncShared/WorkspaceUpdate.cs
@@ -595,7 +595,7 @@ namespace UnrealGameSync
 
 					// Check if the new sync filter matches the previous one. If not, we'll enumerate all files in the workspace and make sure there's nothing extra there.
 					string? nextSyncFilterHash = null;
-					using (SHA1Managed sha = new SHA1Managed())
+					using (SHA1 sha = SHA1.Create())
 					{
 						StringBuilder combinedFilter = new StringBuilder();
 						foreach (string relativeSyncPath in relativeSyncPaths)
@@ -1795,7 +1795,7 @@ namespace UnrealGameSync
 				return null;
 			}
 
-			return response.Data.Contents.Select(x => x.Trim()).Where(x => x.Length > 0).ToList().AsReadOnly();
+			return response.Data.Contents?.Select(x => x.Trim()).Where(x => x.Length > 0).ToList().AsReadOnly();
 		}
 
 		static string FormatTime(long seconds)
@@ -1958,7 +1958,7 @@ namespace UnrealGameSync
 				}
 				else
 				{
-					Directory.CreateDirectory(Path.GetDirectoryName(localPath));
+					Directory.CreateDirectory(Path.GetDirectoryName(localPath)!);
 					Utility.ForceDeleteFile(localPath);
 					if (depotPath != null)
 					{
