UGS: Download PCBs to folder under workspace root, to avoid issues with lack of space in the temp folder causing errors.

diff --git a/Engine/Source/Programs/UnrealGameSync/UnrealGameSyncShared/ArchiveInfo.cs b/Engine/Source/Programs/UnrealGameSync/UnrealGameSyncShared/ArchiveInfo.cs
--- a/Engine/Source/Programs/UnrealGameSync/UnrealGameSyncShared/ArchiveInfo.cs
+++ b/Engine/Source/Programs/UnrealGameSync/UnrealGameSyncShared/ArchiveInfo.cs
@@ -116,7 +116,10 @@ namespace UnrealGameSync
 
 		public async Task<bool> DownloadArchive(IPerforceConnection perforce, string archiveKey, DirectoryReference localRootPath, FileReference manifestFileName, ILogger logger, ProgressValue progress, CancellationToken cancellationToken)
 		{
-			FileReference tempZipFileName = new FileReference(Path.GetTempFileName());
+			DirectoryReference configDir = UserSettings.GetConfigDir(localRootPath);
+			UserSettings.CreateConfigDir(configDir);
+
+			FileReference tempZipFileName = FileReference.Combine(configDir, "archive.zip");
 			try
 			{
 				PrintRecord record = await perforce.PrintAsync(tempZipFileName.FullName, archiveKey, cancellationToken);
