UGS: Detect external changes to the state config file, and refresh the UI to match. Allows using command line client to update state without closing GUI.

diff --git a/Engine/Source/Programs/UnrealGameSync/UnrealGameSync/Controls/WorkspaceControl.cs b/Engine/Source/Programs/UnrealGameSync/UnrealGameSync/Controls/WorkspaceControl.cs
--- a/Engine/Source/Programs/UnrealGameSync/UnrealGameSync/Controls/WorkspaceControl.cs
+++ b/Engine/Source/Programs/UnrealGameSync/UnrealGameSync/Controls/WorkspaceControl.cs
@@ -211,7 +211,6 @@ namespace UnrealGameSync
 		ProjectInfo ProjectInfo { get; }
 
 		UserSettings _settings;
-		UserWorkspaceState _workspaceState;
 		UserWorkspaceSettings _workspaceSettings;
 		UserProjectSettings _projectSettings;
 
@@ -307,8 +306,7 @@ namespace UnrealGameSync
 			_perforceSettings = openProjectInfo.PerforceSettings;
 			ProjectInfo = openProjectInfo.ProjectInfo;
 			_settings = inSettings;
-			this._workspaceSettings = openProjectInfo.WorkspaceSettings;
-			this._workspaceState = openProjectInfo.WorkspaceState;
+			_workspaceSettings = openProjectInfo.WorkspaceSettings;
 			_projectSettings = inSettings.FindOrAddProjectSettings(openProjectInfo.ProjectInfo, openProjectInfo.WorkspaceSettings, _logger);
 
 			DesiredTaskbarState = Tuple.Create(TaskbarState.NoProgress, 0.0f);
@@ -366,7 +364,8 @@ namespace UnrealGameSync
 
 			ProjectInfo project = openProjectInfo.ProjectInfo;
 
-			_workspace = new Workspace(perforceClientSettings, project, _workspaceState, openProjectInfo.WorkspaceProjectConfigFile, openProjectInfo.WorkspaceProjectStreamFilter, new LogControlTextWriter(SyncLog), _serviceProvider);
+			_workspace = new Workspace(perforceClientSettings, project, openProjectInfo.WorkspaceStateWrapper, openProjectInfo.WorkspaceProjectConfigFile, openProjectInfo.WorkspaceProjectStreamFilter, new LogControlTextWriter(SyncLog), _serviceProvider);
+			_workspace.OnStateChanged += StateChangedCallback;
 			_workspace.OnUpdateComplete += UpdateCompleteCallback;
 
 			FileReference projectLogBaseName = FileReference.Combine(_workspaceDataFolder, "sync.log");
@@ -1015,13 +1014,6 @@ namespace UnrealGameSync
 		{
 			if (_workspace.IsBusy() && MessageBox.Show("Are you sure you want to cancel the current operation?", "Cancel operation", MessageBoxButtons.YesNo) == DialogResult.Yes)
 			{
-				_workspaceState.LastSyncChangeNumber = _workspace.PendingChangeNumber;
-				_workspaceState.LastSyncResult = WorkspaceUpdateResult.Canceled;
-				_workspaceState.LastSyncResultMessage = null;
-				_workspaceState.LastSyncTime = null;
-				_workspaceState.LastSyncDurationSeconds = 0;
-				_workspaceState.Save(_logger);
-
 				_workspace.CancelUpdate();
 
 				if (_updateCallback != null)
@@ -1040,6 +1032,13 @@ namespace UnrealGameSync
 			}
 		}
 
+		void StateChangedCallback(ReadOnlyWorkspaceState state)
+		{
+			_changeNumberToLayoutInfo.Clear();
+			UpdateBuildList();
+			UpdateStatusPanel();
+		}
+
 		void UpdateCompleteCallback(WorkspaceUpdateContext context, WorkspaceUpdateResult result, string resultMessage)
 		{
 			_mainThreadSynchronizationContext.Post((o) => { if (!_isDisposing) { UpdateComplete(context, result, resultMessage); } }, null);
@@ -1466,7 +1465,7 @@ namespace UnrealGameSync
 					return false;
 				}
 			}
-			if (IsBisectModeEnabled() && !_workspaceState.BisectChanges.Any(x => x.Change == change.Number))
+			if (IsBisectModeEnabled() && !_workspace.State.BisectChanges.Any(x => x.Change == change.Number))
 			{
 				return false;
 			}
@@ -2131,7 +2130,7 @@ namespace UnrealGameSync
 		private void GetRemainingBisectRange(out int outPassChangeNumber, out int outFailChangeNumber)
 		{
 			int passChangeNumber = -1;
-			foreach (BisectEntry entry in _workspaceState.BisectChanges)
+			foreach (BisectEntry entry in _workspace.State.BisectChanges)
 			{
 				if (entry.State == BisectState.Pass && (entry.Change > passChangeNumber || passChangeNumber == -1))
 				{
@@ -2140,7 +2139,7 @@ namespace UnrealGameSync
 			}
 
 			int failChangeNumber = -1;
-			foreach (BisectEntry entry in _workspaceState.BisectChanges)
+			foreach (BisectEntry entry in _workspace.State.BisectChanges)
 			{
 				if (entry.State == BisectState.Fail && entry.Change > passChangeNumber && (entry.Change < failChangeNumber || failChangeNumber == -1))
 				{
@@ -2179,7 +2178,7 @@ namespace UnrealGameSync
 
 			bool allowSync = CanSyncChange(change.Number);
 			int badgeAlpha = allowSync ? 255 : 128;
-			Color textColor = (allowSync || change.Number == _workspace.PendingChangeNumber || change.Number == _workspace.CurrentChangeNumber || _workspaceState.AdditionalChangeNumbers.Contains(change.Number)) ? SystemColors.WindowText : Blend(SystemColors.Window, SystemColors.WindowText, 0.25f);
+			Color textColor = (allowSync || change.Number == _workspace.PendingChangeNumber || change.Number == _workspace.CurrentChangeNumber || _workspace.State.AdditionalChangeNumbers.Contains(change.Number)) ? SystemColors.WindowText : Blend(SystemColors.Window, SystemColors.WindowText, 0.25f);
 
 			const int fadeRange = 6;
 			if (e.ItemIndex >= BuildList.Items.Count - fadeRange && _numChanges >= _perforceMonitor.CurrentMaxChanges && !IsBisectModeEnabled())
@@ -2198,7 +2197,7 @@ namespace UnrealGameSync
 				{
 					e.Graphics.DrawImage(Properties.Resources.Icons, minX * dpiScaleX, iconY, _previousSyncIcon, GraphicsUnit.Pixel);
 				}
-				else if (_workspaceSettings != null && _workspaceState.AdditionalChangeNumbers.Contains(change.Number))
+				else if (_workspaceSettings != null && _workspace != null && _workspace.State.AdditionalChangeNumbers.Contains(change.Number))
 				{
 					e.Graphics.DrawImage(Properties.Resources.Icons, minX * dpiScaleX, iconY, _additionalSyncIcon, GraphicsUnit.Pixel);
 				}
@@ -2217,7 +2216,7 @@ namespace UnrealGameSync
 						int firstPass, firstFail;
 						GetRemainingBisectRange(out firstPass, out firstFail);
 
-						BisectEntry? entry = _workspaceState.BisectChanges.FirstOrDefault(x => x.Change == change.Number);
+						BisectEntry? entry = _workspace?.State.BisectChanges.FirstOrDefault(x => x.Change == change.Number);
 						if (entry == null || entry.State == BisectState.Exclude)
 						{
 							qualityIcon = new Rectangle(0, 0, 0, 0);
@@ -2363,7 +2362,7 @@ namespace UnrealGameSync
 					}
 
 					string? summaryText;
-					if (_workspaceState.LastSyncChangeNumber == -1 || _workspaceState.LastSyncChangeNumber != change.Number || !GetLastUpdateMessage(_workspaceState.LastSyncResult, _workspaceState.LastSyncResultMessage, out summaryText))
+					if (_workspace.State.LastSyncChangeNumber == -1 || _workspace.State.LastSyncChangeNumber != change.Number || !GetLastUpdateMessage(_workspace.State.LastSyncResult, _workspace.State.LastSyncResultMessage, out summaryText))
 					{
 						StringBuilder summaryTextBuilder = new StringBuilder();
 
@@ -3527,19 +3526,19 @@ namespace UnrealGameSync
 				lines.Add(programsLine);
 
 				// Get the summary of the last sync
-				if (_workspaceState.LastSyncChangeNumber > 0)
+				if (_workspace.State.LastSyncChangeNumber > 0)
 				{
 					string? summaryText;
-					if (_workspaceState.LastSyncChangeNumber == _workspace.CurrentChangeNumber && _workspaceState.LastSyncResult == WorkspaceUpdateResult.Success && _workspaceState.LastSyncTime.HasValue)
+					if (_workspace.State.LastSyncChangeNumber == _workspace.CurrentChangeNumber && _workspace.State.LastSyncResult == WorkspaceUpdateResult.Success && _workspace.State.LastSyncTime.HasValue)
 					{
 						lines.Add(new StatusLine() { LineHeight = 0.5f });
 
 						StatusLine successLine = new StatusLine();
 						successLine.AddIcon(Properties.Resources.StatusIcons, new Size(16, 16), 0);
-						successLine.AddText(String.Format("  Sync took {0}{1}s, completed at {2}.", (_workspaceState.LastSyncDurationSeconds >= 60) ? String.Format("{0}m ", _workspaceState.LastSyncDurationSeconds / 60) : "", _workspaceState.LastSyncDurationSeconds % 60, _workspaceState.LastSyncTime.Value.ToLocalTime().ToString("h\\:mmtt").ToLowerInvariant()));
+						successLine.AddText(String.Format("  Sync took {0}{1}s, completed at {2}.", (_workspace.State.LastSyncDurationSeconds >= 60) ? String.Format("{0}m ", _workspace.State.LastSyncDurationSeconds / 60) : "", _workspace.State.LastSyncDurationSeconds % 60, _workspace.State.LastSyncTime.Value.ToLocalTime().ToString("h\\:mmtt").ToLowerInvariant()));
 						lines.Add(successLine);
 					}
-					else if (GetLastUpdateMessage(_workspaceState.LastSyncResult, _workspaceState.LastSyncResultMessage, out summaryText))
+					else if (GetLastUpdateMessage(_workspace.State.LastSyncResult, _workspace.State.LastSyncResultMessage, out summaryText))
 					{
 						lines.Add(new StatusLine() { LineHeight = 0.5f });
 
@@ -3945,10 +3944,10 @@ namespace UnrealGameSync
 		private void ViewLastSyncStatus()
 		{
 			string? summaryText;
-			if (GetLastUpdateMessage(_workspaceState.LastSyncResult, _workspaceState.LastSyncResultMessage, out summaryText))
+			if (GetLastUpdateMessage(_workspace.State.LastSyncResult, _workspace.State.LastSyncResultMessage, out summaryText))
 			{
 				string? captionText;
-				if (!GetGenericLastUpdateMessage(_workspaceState.LastSyncResult, out captionText))
+				if (!GetGenericLastUpdateMessage(_workspace.State.LastSyncResult, out captionText))
 				{
 					captionText = "Sync error";
 				}
@@ -4188,7 +4187,7 @@ namespace UnrealGameSync
 				BuildListContextMenu_OpenVisualStudio.Visible = !isBusy && isCurrentChange;
 				BuildListContextMenu_Cancel.Visible = isBusy;
 
-				BisectState state = _workspaceState.BisectChanges.FirstOrDefault(x => x.Change == _contextMenuChange.Number)?.State ?? BisectState.Include;
+				BisectState state = _workspace.State.BisectChanges.FirstOrDefault(x => x.Change == _contextMenuChange.Number)?.State ?? BisectState.Include;
 				bool isBisectMode = IsBisectModeEnabled();
 				BuildListContextMenu_Bisect_Pass.Visible = isBisectMode && state != BisectState.Pass;
 				BuildListContextMenu_Bisect_Fail.Visible = isBisectMode && state != BisectState.Fail;
@@ -5870,7 +5869,7 @@ namespace UnrealGameSync
 
 		private bool IsBisectModeEnabled()
 		{
-			return _workspaceState.BisectChanges.Count >= 2;
+			return _workspace.State.BisectChanges.Count >= 2;
 		}
 
 		private void EnableBisectMode()
@@ -5890,40 +5889,28 @@ namespace UnrealGameSync
 				changeNumberToBisectState[changeNumberToBisectState.Keys.Min()] = BisectState.Pass;
 				changeNumberToBisectState[changeNumberToBisectState.Keys.Max()] = BisectState.Fail;
 
-				_workspaceState.BisectChanges = changeNumberToBisectState.Select(x => new BisectEntry { Change = x.Key, State = x.Value }).ToList();
-				_workspaceState.Save(_logger);
-
-				UpdateBuildList();
-				UpdateStatusPanel();
+				_workspace.ModifyState(x => x.BisectChanges = changeNumberToBisectState.Select(x => new BisectEntry { Change = x.Key, State = x.Value }).ToList());
 			}
 		}
 
 		private void CancelBisectMode()
 		{
-			_workspaceState.BisectChanges.Clear();
-			_workspaceState.Save(_logger);
-
-			UpdateBuildList();
-			UpdateStatusPanel();
+			_workspace.ModifyState(x => x.BisectChanges.Clear());
 		}
 
 		private void SetBisectStateForSelection(BisectState state)
 		{
-			foreach (ListViewItem? selectedItem in BuildList.SelectedItems)
+			_workspace.ModifyState(x =>
 			{
-				ChangesRecord? change = selectedItem?.Tag as ChangesRecord;
-				if (change != null)
+				foreach (ListViewItem? selectedItem in BuildList.SelectedItems)
 				{
-					_workspaceState.SetBisectState(change.Number, state);
+					ChangesRecord? change = selectedItem?.Tag as ChangesRecord;
+					if (change != null)
+					{
+						x.SetBisectState(change.Number, state);
+					}
 				}
-			}
-
-			_workspaceState.Save(_logger);
-
-			_changeNumberToLayoutInfo.Clear();
-			BuildList.Invalidate();
-
-			UpdateStatusPanel();
+			});
 		}
 
 		private int GetBisectChangeNumber()
@@ -5933,7 +5920,7 @@ namespace UnrealGameSync
 			GetRemainingBisectRange(out passChangeNumber, out failChangeNumber);
 
 			List<int> changeNumbers = new List<int>();
-			foreach (BisectEntry entry in _workspaceState.BisectChanges)
+			foreach (BisectEntry entry in _workspace.State.BisectChanges)
 			{
 				if (entry.State == BisectState.Include && entry.Change > passChangeNumber && entry.Change < failChangeNumber)
 				{
diff --git a/Engine/Source/Programs/UnrealGameSync/UnrealGameSync/Workspace.cs b/Engine/Source/Programs/UnrealGameSync/UnrealGameSync/Workspace.cs
--- a/Engine/Source/Programs/UnrealGameSync/UnrealGameSync/Workspace.cs
+++ b/Engine/Source/Programs/UnrealGameSync/UnrealGameSync/Workspace.cs
@@ -25,7 +25,8 @@ namespace UnrealGameSync
 	public sealed class Workspace : IDisposable
 	{
 		public IPerforceSettings PerforceSettings { get; }
-		public UserWorkspaceState State { get; }
+		WorkspaceStateWrapper _stateWrapper;
+		public ReadOnlyWorkspaceState State { get; private set; }
 		public ProjectInfo Project { get; }
 		public SynchronizationContext SynchronizationContext { get; }
 		ILogger _logger;
@@ -36,13 +37,18 @@ namespace UnrealGameSync
 
 		public event Action<WorkspaceUpdateContext, WorkspaceUpdateResult, string>? OnUpdateComplete;
 
+		public event Action<ReadOnlyWorkspaceState>? OnStateChanged;
+
 		IAsyncDisposer _asyncDisposer;
 
-		public Workspace(IPerforceSettings inPerfoceSettings, ProjectInfo inProject, UserWorkspaceState inState, ConfigFile projectConfigFile, IReadOnlyList<string>? projectStreamFilter, ILogger logger, IServiceProvider serviceProvider)
+		public Workspace(IPerforceSettings perforceSettings, ProjectInfo project, WorkspaceStateWrapper stateWrapper, ConfigFile projectConfigFile, IReadOnlyList<string>? projectStreamFilter, ILogger logger, IServiceProvider serviceProvider)
 		{
-			PerforceSettings = inPerfoceSettings;
-			Project = inProject;
-			State = inState;
+			PerforceSettings = perforceSettings;
+			Project = project;
+			_stateWrapper = stateWrapper;
+			_stateWrapper.OnModified += OnStateChangedInternal;
+			State = stateWrapper.Current;
+
 			this.SynchronizationContext = SynchronizationContext.Current!;
 			this._logger = logger;
 			this._asyncDisposer = serviceProvider.GetRequiredService<IAsyncDisposer>();
@@ -51,6 +57,26 @@ namespace UnrealGameSync
 			this.ProjectStreamFilter = projectStreamFilter;
 		}
 
+		public void ModifyState(Action<WorkspaceState> action)
+		{
+			_stateWrapper.Modify(x =>
+			{
+				x.ResetForProject(Project);
+				action(x);
+			});
+		}
+
+		private void OnStateChangedInternal(ReadOnlyWorkspaceState state)
+		{
+			SynchronizationContext.Post(x => OnStateChangedInternalMainThread(state), null);
+		}
+
+		private void OnStateChangedInternalMainThread(ReadOnlyWorkspaceState state)
+		{
+			State = state.ResetForProject(Project);
+			OnStateChanged?.Invoke(state);
+		}
+
 		public void Dispose()
 		{
 			CancelUpdate();
@@ -58,6 +84,7 @@ namespace UnrealGameSync
 			{
 				_asyncDisposer.Add(_prevUpdateTask);
 			}
+			_stateWrapper.Dispose();
 		}
 
 		public ConfigFile ProjectConfigFile
@@ -120,7 +147,7 @@ namespace UnrealGameSync
 
 			try
 			{
-				(result, statusMessage) = await update.ExecuteAsync(PerforceSettings, Project, State, _logger, cancellationToken);
+				(result, statusMessage) = await update.ExecuteAsync(PerforceSettings, Project, _stateWrapper, _logger, cancellationToken);
 				if (result != WorkspaceUpdateResult.Success)
 				{
 					_logger.LogError("{Message}", statusMessage);
@@ -149,8 +176,18 @@ namespace UnrealGameSync
 			{
 				WorkspaceUpdateContext context = update.Context;
 
-				State.SetLastSyncState(result, context, statusMessage);
-				State.Save(_logger);
+				ModifyState(x =>
+				{
+					if (result == WorkspaceUpdateResult.Canceled)
+					{
+						x.LastSyncChangeNumber = update.Context.ChangeNumber;
+						x.LastSyncResult = WorkspaceUpdateResult.Canceled;
+						x.LastSyncResultMessage = null;
+						x.LastSyncTime = null;
+						x.LastSyncDurationSeconds = 0;
+					}
+					x.SetLastSyncState(result, context, statusMessage);
+				});
 
 				_currentUpdate = null;
 
diff --git a/Engine/Source/Programs/UnrealGameSync/UnrealGameSyncShared/GlobalSettings.cs b/Engine/Source/Programs/UnrealGameSync/UnrealGameSyncShared/GlobalSettings.cs
--- a/Engine/Source/Programs/UnrealGameSync/UnrealGameSyncShared/GlobalSettings.cs
+++ b/Engine/Source/Programs/UnrealGameSync/UnrealGameSyncShared/GlobalSettings.cs
@@ -8,6 +8,7 @@ using System.Linq;
 using System.Text;
 using System.Text.Json;
 using System.Text.Json.Serialization;
+using System.Threading;
 using System.Threading.Tasks;
 
 namespace UnrealGameSync
@@ -92,7 +93,7 @@ namespace UnrealGameSync
 		{
 		}
 
-		protected virtual void ImportWorkspaceState(DirectoryReference rootDir, string clientName, string branchPath, UserWorkspaceState workspaceState)
+		protected virtual void ImportWorkspaceState(DirectoryReference rootDir, string clientName, string branchPath, WorkspaceState workspaceState)
 		{
 		}
 
@@ -126,33 +127,26 @@ namespace UnrealGameSync
 			}
 		}
 
-		public UserWorkspaceState FindOrAddWorkspaceState(UserWorkspaceSettings settings, ILogger logger)
+		public WorkspaceStateWrapper FindOrAddWorkspaceState(UserWorkspaceSettings settings, ILogger logger)
 		{
 			return FindOrAddWorkspaceState(settings.RootDir, settings.ClientName, settings.BranchPath, logger);
 		}
 
-		public UserWorkspaceState FindOrAddWorkspaceState(DirectoryReference rootDir, string clientName, string branchPath, ILogger logger)
+		public WorkspaceStateWrapper FindOrAddWorkspaceState(DirectoryReference rootDir, string clientName, string branchPath, ILogger logger)
 		{
-			UserWorkspaceState? state;
-			if (!UserWorkspaceState.TryLoad(rootDir, out state))
+			return new WorkspaceStateWrapper(rootDir, () =>
 			{
-				state = new UserWorkspaceState();
-				state.RootDir = rootDir;
+				WorkspaceState state = new WorkspaceState();
 				ImportWorkspaceState(rootDir, clientName, branchPath, state);
-				state.Save(logger);
-			}
-			return state;
+				return state;
+			});
 		}
 
-		public UserWorkspaceState FindOrAddWorkspaceState(ProjectInfo projectInfo, UserWorkspaceSettings settings, ILogger logger)
+		public WorkspaceStateWrapper FindOrAddWorkspaceState(ProjectInfo projectInfo, UserWorkspaceSettings settings, ILogger logger)
 		{
-			UserWorkspaceState state = FindOrAddWorkspaceState(projectInfo.LocalRootPath, projectInfo.ClientName, projectInfo.BranchPath, logger);
-			if (!state.IsValid(projectInfo))
-			{
-				state = new UserWorkspaceState();
-			}
-			state.UpdateCachedProjectInfo(projectInfo, settings.LastModifiedTimeUtc);
-			return state;
+			WorkspaceStateWrapper wrapper = FindOrAddWorkspaceState(projectInfo.LocalRootPath, projectInfo.ClientName, projectInfo.BranchPath, logger);
+			wrapper.Modify(x => x.UpdateCachedProjectInfo(projectInfo, settings.LastModifiedTimeUtc));
+			return wrapper;
 		}
 
 		public UserWorkspaceSettings FindOrAddWorkspaceSettings(DirectoryReference rootDir, string? serverAndPort, string? userName, string clientName, string branchPath, string projectPath, ILogger logger)
diff --git a/Engine/Source/Programs/UnrealGameSync/UnrealGameSyncShared/OpenProjectInfo.cs b/Engine/Source/Programs/UnrealGameSync/UnrealGameSyncShared/OpenProjectInfo.cs
--- a/Engine/Source/Programs/UnrealGameSync/UnrealGameSyncShared/OpenProjectInfo.cs
+++ b/Engine/Source/Programs/UnrealGameSync/UnrealGameSyncShared/OpenProjectInfo.cs
@@ -21,20 +21,20 @@ namespace UnrealGameSync
 		public IPerforceSettings PerforceSettings { get; }
 		public ProjectInfo ProjectInfo { get; }
 		public UserWorkspaceSettings WorkspaceSettings { get; }
-		public UserWorkspaceState WorkspaceState { get; }
+		public WorkspaceStateWrapper WorkspaceStateWrapper { get; }
 		public ConfigFile LatestProjectConfigFile { get; }
 		public ConfigFile WorkspaceProjectConfigFile { get; }
 		public IReadOnlyList<string>? WorkspaceProjectStreamFilter { get; }
 		public List<KeyValuePair<FileReference, DateTime>> LocalConfigFiles { get; }
 
-		public OpenProjectInfo(UserSelectedProjectSettings selectedProject, IPerforceSettings perforceSettings, ProjectInfo projectInfo, UserWorkspaceSettings workspaceSettings, UserWorkspaceState workspaceState, ConfigFile latestProjectConfigFile, ConfigFile workspaceProjectConfigFile, IReadOnlyList<string>? workspaceProjectStreamFilter, List<KeyValuePair<FileReference, DateTime>> localConfigFiles)
+		public OpenProjectInfo(UserSelectedProjectSettings selectedProject, IPerforceSettings perforceSettings, ProjectInfo projectInfo, UserWorkspaceSettings workspaceSettings, WorkspaceStateWrapper workspaceStateWrapper, ConfigFile latestProjectConfigFile, ConfigFile workspaceProjectConfigFile, IReadOnlyList<string>? workspaceProjectStreamFilter, List<KeyValuePair<FileReference, DateTime>> localConfigFiles)
 		{
 			this.SelectedProject = selectedProject;
 
 			this.PerforceSettings = perforceSettings;
 			this.ProjectInfo = projectInfo;
 			this.WorkspaceSettings = workspaceSettings;
-			this.WorkspaceState = workspaceState;
+			this.WorkspaceStateWrapper = workspaceStateWrapper;
 			this.LatestProjectConfigFile = latestProjectConfigFile;
 			this.WorkspaceProjectConfigFile = workspaceProjectConfigFile;
 			this.WorkspaceProjectStreamFilter = workspaceProjectStreamFilter;
@@ -218,7 +218,7 @@ namespace UnrealGameSync
 				ProjectInfo projectInfo = await ProjectInfo.CreateAsync(perforceClient, userWorkspaceSettings, cancellationToken);
 
 				// Update the cached workspace state
-				UserWorkspaceState userWorkspaceState = userSettings.FindOrAddWorkspaceState(projectInfo, userWorkspaceSettings, logger);
+				WorkspaceStateWrapper workspaceStateWrapper = userSettings.FindOrAddWorkspaceState(projectInfo, userWorkspaceSettings, logger);
 
 				// Read the initial config file
 				List<KeyValuePair<FileReference, DateTime>> localConfigFiles = new List<KeyValuePair<FileReference, DateTime>>();
@@ -228,7 +228,7 @@ namespace UnrealGameSync
 				ConfigFile workspaceProjectConfigFile = await WorkspaceUpdate.ReadProjectConfigFile(branchDirectoryName, newSelectedFileName, logger);
 				IReadOnlyList<string>? workspaceProjectStreamFilter = await WorkspaceUpdate.ReadProjectStreamFilter(perforceClient, workspaceProjectConfigFile, logger, cancellationToken);
 
-				OpenProjectInfo workspaceSettings = new OpenProjectInfo(selectedProject, perforceSettings, projectInfo, userWorkspaceSettings, userWorkspaceState, latestProjectConfigFile, workspaceProjectConfigFile, workspaceProjectStreamFilter, localConfigFiles);
+				OpenProjectInfo workspaceSettings = new OpenProjectInfo(selectedProject, perforceSettings, projectInfo, userWorkspaceSettings, workspaceStateWrapper, latestProjectConfigFile, workspaceProjectConfigFile, workspaceProjectStreamFilter, localConfigFiles);
 
 				return workspaceSettings;
 			}
diff --git a/Engine/Source/Programs/UnrealGameSync/UnrealGameSyncShared/UserSettings.cs b/Engine/Source/Programs/UnrealGameSync/UnrealGameSyncShared/UserSettings.cs
--- a/Engine/Source/Programs/UnrealGameSync/UnrealGameSyncShared/UserSettings.cs
+++ b/Engine/Source/Programs/UnrealGameSync/UnrealGameSyncShared/UserSettings.cs
@@ -405,140 +405,6 @@ namespace UnrealGameSync
 		}
 	}
 
-	public class UserWorkspaceState
-	{
-		[JsonIgnore]
-		public DirectoryReference RootDir { get; set; } = null!;
-
-		// Cached state about the project, configured using UserWorkspaceSettings and taken from computed values in ProjectInfo. Assumed valid unless manually updated.
-		public long SettingsTimeUtc { get; set; }
-		public string ClientName { get; set; } = String.Empty;
-		public string BranchPath { get; set; } = String.Empty;
-		public string ProjectPath { get; set; } = String.Empty;
-		public string? StreamName { get; set; }
-		public string ProjectIdentifier { get; set; } = String.Empty; // Should be fully reset if this changes
-		public bool IsEnterpriseProject { get; set; }
-
-		// Settings for the currently synced project in this workspace. CurrentChangeNumber is only valid for this workspace if CurrentProjectPath is the current project.
-		public int CurrentChangeNumber { get; set; } = -1;
-		public int CurrentCodeChangeNumber { get; set; } = -1;
-		public string? CurrentSyncFilterHash { get; set; }
-		public List<int> AdditionalChangeNumbers { get; set; } = new List<int>();
-
-		// Settings for the last attempted sync. These values are set to persist error messages between runs.
-		public int LastSyncChangeNumber { get; set; }
-		public WorkspaceUpdateResult LastSyncResult { get; set; }
-		public string? LastSyncResultMessage { get; set; }
-		public DateTime? LastSyncTime { get; set; }
-		public int LastSyncDurationSeconds { get; set; }
-
-		// The last successful build, regardless of whether a failed sync has happened in the meantime. Used to determine whether to force a clean due to entries in the project config file.
-		public int LastBuiltChangeNumber { get; set; }
-
-		// Expanded archives in the workspace
-		public string[]? ExpandedArchiveTypes { get; set; }
-
-		// The changes that we're regressing at the moment
-		public List<BisectEntry> BisectChanges { get; set; } = new List<BisectEntry>();
-
-		// Path to the config file
-		[JsonIgnore]
-		public FileReference ConfigFile => GetConfigFile(RootDir);
-
-		public static UserWorkspaceState CreateNew(DirectoryReference rootDir)
-		{
-			UserWorkspaceState state = new UserWorkspaceState();
-			state.RootDir = rootDir;
-			return state;
-		}
-
-		public ProjectInfo CreateProjectInfo()
-		{
-			return new ProjectInfo(RootDir, ClientName, BranchPath, ProjectPath, StreamName, ProjectIdentifier, IsEnterpriseProject);
-		}
-
-		public void UpdateCachedProjectInfo(ProjectInfo projectInfo, long settingsTimeUtc)
-		{
-			this.SettingsTimeUtc = settingsTimeUtc;
-
-			RootDir = projectInfo.LocalRootPath;
-			ClientName = projectInfo.ClientName;
-			BranchPath = projectInfo.BranchPath;
-			ProjectPath = projectInfo.ProjectPath;
-			StreamName = projectInfo.StreamName;
-			ProjectIdentifier = projectInfo.ProjectIdentifier;
-			IsEnterpriseProject = projectInfo.IsEnterpriseProject;
-		}
-
-		public bool IsValid(ProjectInfo projectInfo)
-		{
-			return ProjectIdentifier.Equals(projectInfo.ProjectIdentifier, StringComparison.OrdinalIgnoreCase);
-		}
-
-		public void SetBisectState(int change, BisectState state)
-		{
-			BisectEntry? entry = BisectChanges.FirstOrDefault(x => x.Change == change);
-			if (entry == null)
-			{
-				entry = new BisectEntry();
-				entry.Change = change;
-				BisectChanges.Add(entry);
-			}
-			entry.State = state;
-		}
-
-		public static bool TryLoad(DirectoryReference rootDir, [NotNullWhen(true)] out UserWorkspaceState? state)
-		{
-			FileReference configFile = GetConfigFile(rootDir);
-			if (Utility.TryLoadJson(configFile, out state))
-			{
-				state.RootDir = rootDir;
-				return true;
-			}
-			else
-			{
-				state = null;
-				return false;
-			}
-		}
-
-		public void SetLastSyncState(WorkspaceUpdateResult result, WorkspaceUpdateContext context, string statusMessage)
-		{
-			LastSyncChangeNumber = context.ChangeNumber;
-			LastSyncResult = result;
-			LastSyncResultMessage = statusMessage;
-			LastSyncTime = DateTime.UtcNow;
-			LastSyncDurationSeconds = (int)(LastSyncTime.Value - context.StartTime).TotalSeconds;
-		}
-
-		static object _syncRoot = new object();
-
-		public bool Save(ILogger logger)
-		{
-			try
-			{
-				SaveInternal();
-				return true;
-			}
-			catch (Exception ex)
-			{
-				logger.LogError(ex, "Unable to save {ConfigFile}: {Message}", ConfigFile, ex.Message);
-				return false;
-			}
-		}
-
-		private void SaveInternal()
-		{
-			lock (_syncRoot)
-			{
-				UserSettings.CreateConfigDir(ConfigFile.Directory);
-				Utility.SaveJson(ConfigFile, this);
-			}
-		}
-
-		public static FileReference GetConfigFile(DirectoryReference rootDir) => FileReference.Combine(UserSettings.GetConfigDir(rootDir), "state.json");
-	}
-
 	public class UserProjectSettings
 	{
 		[JsonIgnore]
@@ -661,7 +527,6 @@ namespace UnrealGameSync
 		public int NotifyUnresolvedMinutes;
 
 		// Project settings
-		Dictionary<DirectoryReference, UserWorkspaceState> _workspaceDirToState = new Dictionary<DirectoryReference, UserWorkspaceState>();
 		Dictionary<DirectoryReference, UserWorkspaceSettings> _workspaceDirToSettings = new Dictionary<DirectoryReference, UserWorkspaceSettings>();
 		Dictionary<FileReference, UserProjectSettings> _projectKeyToSettings = new Dictionary<FileReference, UserProjectSettings>();
 
@@ -1007,7 +872,7 @@ namespace UnrealGameSync
 			}
 		}
 
-		protected override void ImportWorkspaceState(DirectoryReference rootDir, string clientName, string branchPath, UserWorkspaceState currentWorkspace)
+		protected override void ImportWorkspaceState(DirectoryReference rootDir, string clientName, string branchPath, WorkspaceState currentWorkspace)
 		{
 			// Read the workspace settings
 			ConfigSection? workspaceSection = _configFile.FindSection(clientName + branchPath);
@@ -1374,7 +1239,6 @@ namespace UnrealGameSync
 		{
 			List<FileReference> files = new List<FileReference>();
 			files.AddRange(_workspaceDirToSettings.Values.Select(x => x.ConfigFile));
-			files.AddRange(_workspaceDirToState.Values.Select(x => x.ConfigFile));
 			files.AddRange(_projectKeyToSettings.Keys);
 			return files;
 		}
diff --git a/Engine/Source/Programs/UnrealGameSync/UnrealGameSyncShared/Utility.cs b/Engine/Source/Programs/UnrealGameSync/UnrealGameSyncShared/Utility.cs
--- a/Engine/Source/Programs/UnrealGameSync/UnrealGameSyncShared/Utility.cs
+++ b/Engine/Source/Programs/UnrealGameSync/UnrealGameSyncShared/Utility.cs
@@ -104,13 +104,26 @@ namespace UnrealGameSync
 			}
 		}
 
+		public static T? TryDeserializeJson<T>(byte[] data) where T : class
+		{
+			try
+			{
+				return JsonSerializer.Deserialize<T>(data, DefaultJsonSerializerOptions)!;
+			}
+			catch (Exception ex)
+			{
+				TraceException?.Invoke(ex);
+				return null;
+			}
+		}
+
 		public static T LoadJson<T>(FileReference file)
 		{
 			byte[] data = FileReference.ReadAllBytes(file);
 			return JsonSerializer.Deserialize<T>(data, DefaultJsonSerializerOptions)!;
 		}
 
-		public static void SaveJson<T>(FileReference file, T obj)
+		public static byte[] SerializeJson<T>(T obj)
 		{
 			JsonSerializerOptions options = new JsonSerializerOptions { IgnoreNullValues = true, WriteIndented = true };
 			options.Converters.Add(new JsonStringEnumConverter());
@@ -125,6 +138,12 @@ namespace UnrealGameSync
 				buffer = stream.ToArray();
 			}
 
+			return buffer;
+		}
+
+		public static void SaveJson<T>(FileReference file, T obj)
+		{
+			byte[] buffer = SerializeJson(obj);
 			FileReference.WriteAllBytes(file, buffer);
 		}
 
diff --git a/Engine/Source/Programs/UnrealGameSync/UnrealGameSyncShared/WorkspaceState.cs b/Engine/Source/Programs/UnrealGameSync/UnrealGameSyncShared/WorkspaceState.cs
new file mode 100644
--- /dev/null
+++ b/Engine/Source/Programs/UnrealGameSync/UnrealGameSyncShared/WorkspaceState.cs
@@ -0,0 +1,294 @@
+using EpicGames.Core;
+using JetBrains.Annotations;
+using System;
+using System.Collections.Generic;
+using System.IO;
+using System.Linq;
+using System.Reflection;
+using System.Text;
+using System.Threading;
+
+namespace UnrealGameSync
+{
+	/// <summary>
+	/// The current workspace state. Use <see cref="WorkspaceStateWrapper"/> to access state for a workspace.
+	/// </summary>
+	public class WorkspaceState
+	{
+		// Cached state about the project, configured using UserWorkspaceSettings and taken from computed values in ProjectInfo. Assumed valid unless manually updated.
+		public long SettingsTimeUtc { get; set; }
+		public string ClientName { get; set; } = String.Empty;
+		public string BranchPath { get; set; } = String.Empty;
+		public string ProjectPath { get; set; } = String.Empty;
+		public string? StreamName { get; set; }
+		public string ProjectIdentifier { get; set; } = String.Empty; // Should be fully reset if this changes
+		public bool IsEnterpriseProject { get; set; }
+
+		// Settings for the currently synced project in this workspace. CurrentChangeNumber is only valid for this workspace if CurrentProjectPath is the current project.
+		public int CurrentChangeNumber { get; set; } = -1;
+		public int CurrentCodeChangeNumber { get; set; } = -1;
+		public string? CurrentSyncFilterHash { get; set; }
+		public List<int> AdditionalChangeNumbers { get; set; } = new List<int>();
+
+		// Settings for the last attempted sync. These values are set to persist error messages between runs.
+		public int LastSyncChangeNumber { get; set; }
+		public WorkspaceUpdateResult LastSyncResult { get; set; }
+		public string? LastSyncResultMessage { get; set; }
+		public DateTime? LastSyncTime { get; set; }
+		public int LastSyncDurationSeconds { get; set; }
+
+		// The last successful build, regardless of whether a failed sync has happened in the meantime. Used to determine whether to force a clean due to entries in the project config file.
+		public int LastBuiltChangeNumber { get; set; }
+
+		// Expanded archives in the workspace
+		public string[]? ExpandedArchiveTypes { get; set; }
+
+		// The changes that we're regressing at the moment
+		public List<BisectEntry> BisectChanges { get; set; } = new List<BisectEntry>();
+
+		public void UpdateCachedProjectInfo(ProjectInfo projectInfo, long settingsTimeUtc)
+		{
+			this.SettingsTimeUtc = settingsTimeUtc;
+
+			ClientName = projectInfo.ClientName;
+			BranchPath = projectInfo.BranchPath;
+			ProjectPath = projectInfo.ProjectPath;
+			StreamName = projectInfo.StreamName;
+			ProjectIdentifier = projectInfo.ProjectIdentifier;
+			IsEnterpriseProject = projectInfo.IsEnterpriseProject;
+		}
+
+		public bool IsValid(ProjectInfo projectInfo)
+		{
+			return ProjectIdentifier.Equals(projectInfo.ProjectIdentifier, StringComparison.OrdinalIgnoreCase);
+		}
+
+		public void ResetForProject(ProjectInfo projectInfo)
+		{
+			if (!IsValid(projectInfo))
+			{
+				CopyFrom(new WorkspaceState());
+				UpdateCachedProjectInfo(projectInfo, 0);
+			}
+		}
+
+		public void CopyFrom(WorkspaceState other)
+		{
+			foreach (PropertyInfo propertyInfo in GetType().GetProperties(BindingFlags.Instance | BindingFlags.Public | BindingFlags.SetProperty | BindingFlags.GetProperty))
+			{
+				propertyInfo.SetValue(this, propertyInfo.GetValue(other));
+			}
+		}
+
+		public void SetBisectState(int change, BisectState state)
+		{
+			BisectEntry entry = BisectChanges.FirstOrDefault(x => x.Change == change);
+			if (entry == null)
+			{
+				entry = new BisectEntry();
+				entry.Change = change;
+				BisectChanges.Add(entry);
+			}
+			entry.State = state;
+		}
+
+		public void SetLastSyncState(WorkspaceUpdateResult result, WorkspaceUpdateContext context, string statusMessage)
+		{
+			LastSyncChangeNumber = context.ChangeNumber;
+			LastSyncResult = result;
+			LastSyncResultMessage = statusMessage;
+			LastSyncTime = DateTime.UtcNow;
+			LastSyncDurationSeconds = (int)(LastSyncTime.Value - context.StartTime).TotalSeconds;
+		}
+	}
+
+	/// <summary>
+	/// Read-only wrapper around <see cref="WorkspaceState"/>. To modify state, use <see cref="WorkspaceStateWrapper.Modify(Action{WorkspaceState})"/>.
+	/// </summary>
+	public class ReadOnlyWorkspaceState
+	{
+		readonly WorkspaceState _inner;
+
+		public long SettingsTimeUtc => _inner.SettingsTimeUtc;
+		public string ClientName => _inner.ClientName;
+		public string BranchPath => _inner.BranchPath;
+		public string ProjectPath => _inner.ProjectPath;
+		public string? StreamName => _inner.StreamName;
+		public string ProjectIdentifier => _inner.ProjectIdentifier;
+		public bool IsEnterpriseProject => _inner.IsEnterpriseProject;
+
+		// Settings for the currently synced project in this workspace. CurrentChangeNumber is only valid for this workspace if CurrentProjectPath is the current project.
+		public int CurrentChangeNumber => _inner.CurrentChangeNumber;
+		public int CurrentCodeChangeNumber => _inner.CurrentCodeChangeNumber;
+		public string? CurrentSyncFilterHash => _inner.CurrentSyncFilterHash;
+		public IReadOnlyList<int> AdditionalChangeNumbers => _inner.AdditionalChangeNumbers;
+
+		// Settings for the last attempted sync. These values are set to persist error messages between runs.
+		public int LastSyncChangeNumber => _inner.LastSyncChangeNumber;
+		public WorkspaceUpdateResult LastSyncResult => _inner.LastSyncResult;
+		public string? LastSyncResultMessage => _inner.LastSyncResultMessage;
+		public DateTime? LastSyncTime => _inner.LastSyncTime;
+		public int LastSyncDurationSeconds => _inner.LastSyncDurationSeconds;
+
+		// The last successful build, regardless of whether a failed sync has happened in the meantime. Used to determine whether to force a clean due to entries in the project config file.
+		public int LastBuiltChangeNumber => _inner.LastBuiltChangeNumber;
+
+		// Expanded archives in the workspace
+		public IReadOnlyList<string> ExpandedArchiveTypes => _inner.ExpandedArchiveTypes ?? Array.Empty<string>();
+
+		// The changes that we're regressing at the moment
+		public IReadOnlyList<BisectEntry> BisectChanges => _inner.BisectChanges;
+
+		internal ReadOnlyWorkspaceState(WorkspaceState inner) => _inner = inner;
+
+		public ReadOnlyWorkspaceState ResetForProject(ProjectInfo projectInfo)
+		{
+			if (_inner.IsValid(projectInfo))
+			{
+				return this;
+			}
+			else
+			{
+				WorkspaceState state = new WorkspaceState();
+				state.ResetForProject(projectInfo);
+				return new ReadOnlyWorkspaceState(state);
+			}
+		}
+
+		public WorkspaceState MutableCopy()
+		{
+			WorkspaceState state = new WorkspaceState();
+			state.CopyFrom(_inner);
+			return state;
+		}
+	}
+
+	/// <summary>
+	/// Monitors state of a workspace state file.
+	/// </summary>
+	public class WorkspaceStateWrapper : IDisposable
+	{
+		/// <summary>
+		/// Directory containing the config file
+		/// </summary>
+		public DirectoryReference RootDir { get; }
+
+		/// <summary>
+		/// Path to the modified file
+		/// </summary>
+		public FileReference File { get; }
+
+		/// <summary>
+		/// Current workspace state
+		/// </summary>
+		public ReadOnlyWorkspaceState Current { get; private set; }
+
+		/// <summary>
+		/// Callback for modifications to the workspace state
+		/// </summary>
+		public event Action<ReadOnlyWorkspaceState>? OnModified;
+
+		readonly FileSystemWatcher _watcher;
+		readonly Func<WorkspaceState> _createNew;
+
+		/// <summary>
+		/// Constructor
+		/// </summary>
+		/// <param name="file"></param>
+		public WorkspaceStateWrapper(DirectoryReference rootDir, Func<WorkspaceState> createNew)
+		{
+			RootDir = rootDir;
+
+			DirectoryReference configDir = UserSettings.GetConfigDir(rootDir);
+			File = FileReference.Combine(configDir, "state.json");
+
+			UserSettings.CreateConfigDir(RootDir);
+
+			_watcher = new FileSystemWatcher(configDir.FullName, File.GetFileName());
+			_watcher.NotifyFilter = NotifyFilters.LastWrite;
+			_watcher.Changed += OnModifiedInternal;
+			_watcher.Created += OnModifiedInternal;
+			_watcher.Deleted += OnModifiedInternal;
+			_watcher.EnableRaisingEvents = true;
+
+			_createNew = createNew;
+
+			Current = GetSnapshot();
+		}
+
+		/// <inheritdoc/>
+		public void Dispose()
+		{
+			_watcher.Dispose();
+		}
+
+		/// <summary>
+		/// Modify the current workspace state
+		/// </summary>
+		/// <param name="action">Action to be called with the state to be modified</param>
+		public ReadOnlyWorkspaceState Modify(Action<WorkspaceState> action)
+		{
+			for (; ; )
+			{
+				try
+				{
+					using (FileStream stream = FileReference.Open(File, FileMode.OpenOrCreate, FileAccess.ReadWrite, FileShare.None))
+					{
+						stream.Position = 0;
+						WorkspaceState state = ReadState(stream);
+						action(state);
+
+						stream.Position = 0;
+						stream.Write(Utility.SerializeJson(state));
+						stream.SetLength(stream.Position);
+
+						return new ReadOnlyWorkspaceState(state);
+					}
+				}
+				catch (IOException ex) when ((ex.HResult & 0xffff) == ERROR_SHARING_VIOLATION)
+				{
+					Thread.Sleep(10);
+				}
+			}
+		}
+
+		void OnModifiedInternal(object sender, FileSystemEventArgs e)
+		{
+			Current = GetSnapshot();
+			OnModified?.Invoke(Current);
+		}
+
+		const int ERROR_SHARING_VIOLATION = 32;
+
+		private ReadOnlyWorkspaceState GetSnapshot()
+		{
+			for (; ; )
+			{
+				try
+				{
+					using (FileStream stream = FileReference.Open(File, FileMode.OpenOrCreate, FileAccess.Read, FileShare.Read))
+					{
+						WorkspaceState state = ReadState(stream);
+						return new ReadOnlyWorkspaceState(state);
+					}
+				}
+				catch (IOException ex) when ((ex.HResult & 0xffff) == ERROR_SHARING_VIOLATION)
+				{
+					Thread.Sleep(10);
+				}
+			}
+		}
+
+		WorkspaceState ReadState(FileStream stream)
+		{
+			if (stream.Length == 0)
+			{
+				return new WorkspaceState();
+			}
+
+			byte[] buffer = new byte[stream.Length];
+			stream.Read(buffer);
+			return Utility.TryDeserializeJson<WorkspaceState>(buffer) ?? _createNew();
+		}
+	}
+}
diff --git a/Engine/Source/Programs/UnrealGameSync/UnrealGameSyncShared/WorkspaceUpdate.cs b/Engine/Source/Programs/UnrealGameSync/UnrealGameSyncShared/WorkspaceUpdate.cs
--- a/Engine/Source/Programs/UnrealGameSync/UnrealGameSyncShared/WorkspaceUpdate.cs
+++ b/Engine/Source/Programs/UnrealGameSync/UnrealGameSyncShared/WorkspaceUpdate.cs
@@ -200,6 +200,11 @@ namespace UnrealGameSync
 		public static DirectoryReference GetDataFolder(DirectoryReference workspaceDir) => DirectoryReference.Combine(workspaceDir, ".ugs");
 		public static DirectoryReference GetCacheFolder(DirectoryReference workspaceDir) => DirectoryReference.Combine(workspaceDir, ".ugs", "cache");
 
+		public ProjectInfo(DirectoryReference localRootPath, ReadOnlyWorkspaceState state)
+			: this(localRootPath, state.ClientName, state.BranchPath, state.ProjectPath, state.StreamName, state.ProjectIdentifier, state.IsEnterpriseProject)
+		{
+		}
+
 		public ProjectInfo(DirectoryReference localRootPath, string clientName, string branchPath, string projectPath, string? streamName, string projectIdentifier, bool isEnterpriseProject)
 		{
 			ValidateBranchPath(branchPath);
@@ -562,10 +567,12 @@ namespace UnrealGameSync
 			}
 		}
 
-		public async Task<(WorkspaceUpdateResult, string)> ExecuteAsync(IPerforceSettings perforceSettings, ProjectInfo project, UserWorkspaceState state, ILogger logger, CancellationToken cancellationToken)
+		public async Task<(WorkspaceUpdateResult, string)> ExecuteAsync(IPerforceSettings perforceSettings, ProjectInfo project, WorkspaceStateWrapper stateMgr, ILogger logger, CancellationToken cancellationToken)
 		{
 			using IPerforceConnection perforce = await PerforceConnection.CreateAsync(perforceSettings, logger);
 
+			ReadOnlyWorkspaceState state = stateMgr.Current;
+
 			List<Tuple<string, TimeSpan>> times = new List<Tuple<string, TimeSpan>>();
 
 			int numFilesSynced = 0;
@@ -696,8 +703,7 @@ namespace UnrealGameSync
 							if (removeDepotPaths.Count > 0)
 							{
 								// Clear the current sync filter hash. If the sync is canceled, we'll be in an indeterminate state, and we should always clean next time round.
-								state.CurrentSyncFilterHash = "INVALID";
-								state.Save(logger);
+								state = stateMgr.Modify(x => x.CurrentSyncFilterHash = "INVALID");
 
 								// Find all the depot paths that will be synced
 								HashSet<string> remainingDepotPathsToRemove = new HashSet<string>(StringComparer.OrdinalIgnoreCase);
@@ -715,8 +721,7 @@ namespace UnrealGameSync
 							}
 
 							// Update the sync filter hash. We've removed any files we need to at this point.
-							state.CurrentSyncFilterHash = nextSyncFilterHash;
-							state.Save(logger);
+							state = stateMgr.Modify(x => x.CurrentSyncFilterHash = nextSyncFilterHash);
 						}
 					}
 
@@ -857,9 +862,11 @@ namespace UnrealGameSync
 					// Clear the current sync changelist, in case we cancel
 					if (!Context.Options.HasFlag(WorkspaceUpdateOptions.SyncSingleChange))
 					{
-						state.CurrentChangeNumber = -1;
-						state.AdditionalChangeNumbers.Clear();
-						state.Save(logger);
+						state = stateMgr.Modify(x =>
+						{
+							x.CurrentChangeNumber = -1;
+							x.AdditionalChangeNumbers.Clear();
+						});
 					}
 
 					// Find all the depot paths that will be synced
@@ -1079,16 +1086,18 @@ namespace UnrealGameSync
 					}
 
 					// Update the current state
+					state = stateMgr.Modify(x =>
+					{
 					if (Context.Options.HasFlag(WorkspaceUpdateOptions.SyncSingleChange))
 					{
-						state.AdditionalChangeNumbers.Add(Context.ChangeNumber);
+							x.AdditionalChangeNumbers.Add(Context.ChangeNumber);
 					}
 					else
 					{
-						state.CurrentChangeNumber = Context.ChangeNumber;
-						state.CurrentCodeChangeNumber = versionChangeNumber;
+							x.CurrentChangeNumber = Context.ChangeNumber;
+							x.CurrentCodeChangeNumber = versionChangeNumber;
 					}
-					state.Save(logger);
+					});
 
 					// Update the timing info
 					times.Add(new Tuple<string, TimeSpan>("Sync", syncTelemetryStopwatch.Stop("Success")));
@@ -1148,8 +1157,7 @@ namespace UnrealGameSync
 					}
 
 					// Update the state
-					state.ExpandedArchiveTypes = Context.ArchiveTypeToArchive.Where(x => x.Value != null).Select(x => x.Key).ToArray();
-					state.Save(logger);
+					state = stateMgr.Modify(x => x.ExpandedArchiveTypes = Context.ArchiveTypeToArchive.Where(x => x.Value != null).Select(x => x.Key).ToArray());
 
 					// Add the finish time
 					times.Add(new Tuple<string, TimeSpan>("Archive", stopwatch.Stop("Success")));
@@ -1397,8 +1405,7 @@ namespace UnrealGameSync
 				// Update the last successful build change number
 				if (Context.CustomBuildSteps == null || Context.CustomBuildSteps.Count == 0)
 				{
-					state.LastBuiltChangeNumber = state.CurrentChangeNumber;
-					state.Save(logger);
+					state = stateMgr.Modify(x => x.LastBuiltChangeNumber = state.CurrentChangeNumber);
 				}
 			}
 
