UGS: Create a global mutex when updating a workspace, ensuring that only one application is updating it at a time.

diff --git a/Engine/Source/Programs/Shared/EpicGames.Core/AsyncUtils.cs b/Engine/Source/Programs/Shared/EpicGames.Core/AsyncUtils.cs
--- a/Engine/Source/Programs/Shared/EpicGames.Core/AsyncUtils.cs
+++ b/Engine/Source/Programs/Shared/EpicGames.Core/AsyncUtils.cs
@@ -199,5 +199,34 @@ namespace EpicGames.Core
 				return Prefetch(source, count - 1, cancellationToken);
 			}
 		}
+
+		/// <summary>
+		/// Waits for a native wait handle to be signalled
+		/// </summary>
+		/// <param name="handle">Handle to wait for</param>
+		/// <param name="cancellationToken">Cancellation token for the operation</param>
+		public static Task WaitOneAsync(this WaitHandle handle, CancellationToken cancellationToken = default) => handle.WaitOneAsync(-1, cancellationToken);
+
+		/// <summary>
+		/// Waits for a native wait handle to be signalled
+		/// </summary>
+		/// <param name="handle">Handle to wait for</param>
+		/// <param name="timeoutMs">Timeout for the wait</param>
+		/// <param name="cancellationToken">Cancellation token for the operation</param>
+		public static async Task WaitOneAsync(this WaitHandle handle, int timeoutMs, CancellationToken cancellationToken = default)
+		{
+			TaskCompletionSource<bool> completionSource = new TaskCompletionSource<bool>();
+
+			RegisteredWaitHandle waitHandle = ThreadPool.RegisterWaitForSingleObject(handle, (state, timedOut) => ((TaskCompletionSource<bool>)state!).TrySetResult(!timedOut), completionSource, timeoutMs, true);
+			try
+			{
+				using IDisposable registration = cancellationToken.Register(x => ((TaskCompletionSource<bool>)x!).SetCanceled(), completionSource);
+				await completionSource.Task;
+			}
+			finally
+			{
+				waitHandle.Unregister(null);
+			}
+		}
 	}
 }
diff --git a/Engine/Source/Programs/UnrealGameSync/UnrealGameSync/Controls/WorkspaceControl.cs b/Engine/Source/Programs/UnrealGameSync/UnrealGameSync/Controls/WorkspaceControl.cs
--- a/Engine/Source/Programs/UnrealGameSync/UnrealGameSync/Controls/WorkspaceControl.cs
+++ b/Engine/Source/Programs/UnrealGameSync/UnrealGameSync/Controls/WorkspaceControl.cs
@@ -3330,6 +3330,12 @@ namespace UnrealGameSync
 				if (progress.Item2 > 0.0f) progressLine.AddProgressBar(progress.Item2);
 				lines.Add(progressLine);
 			}
+			else if (_workspace.IsExternalSyncActive())
+			{
+				StatusLine summaryLine = new StatusLine();
+				summaryLine.AddText("Command-line sync in progress...");
+				lines.Add(summaryLine);
+			}
 			else
 			{
 				// Project
diff --git a/Engine/Source/Programs/UnrealGameSync/UnrealGameSync/Workspace.cs b/Engine/Source/Programs/UnrealGameSync/UnrealGameSync/Workspace.cs
--- a/Engine/Source/Programs/UnrealGameSync/UnrealGameSync/Workspace.cs
+++ b/Engine/Source/Programs/UnrealGameSync/UnrealGameSync/Workspace.cs
@@ -29,6 +29,7 @@ namespace UnrealGameSync
 		public ReadOnlyWorkspaceState State { get; private set; }
 		public ProjectInfo Project { get; }
 		public SynchronizationContext SynchronizationContext { get; }
+		public WorkspaceLock Lock { get; }
 		ILogger _logger;
 
 		bool Syncing => _currentUpdate != null;
@@ -49,6 +50,9 @@ namespace UnrealGameSync
 			_stateWrapper.OnModified += OnStateChangedInternal;
 			State = stateWrapper.Current;
 
+			Lock = new WorkspaceLock(project.LocalRootPath);
+			Lock.OnChange += OnLockChangedInternal;
+
 			this.SynchronizationContext = SynchronizationContext.Current!;
 			this._logger = logger;
 			this._asyncDisposer = serviceProvider.GetRequiredService<IAsyncDisposer>();
@@ -57,6 +61,8 @@ namespace UnrealGameSync
 			this.ProjectStreamFilter = projectStreamFilter;
 		}
 
+		public bool IsExternalSyncActive() => Lock.IsLockedByOtherProcess();
+
 		public void ModifyState(Action<WorkspaceState> action)
 		{
 			_stateWrapper.Modify(x =>
@@ -66,6 +72,16 @@ namespace UnrealGameSync
 			});
 		}
 
+		private void OnLockChangedInternal(bool locked)
+		{
+			SynchronizationContext.Post(OnLockChangedInternalMainThread, null);
+		}
+
+		private void OnLockChangedInternalMainThread(object? obj)
+		{
+			OnStateChanged?.Invoke(State);
+		}
+
 		private void OnStateChangedInternal(ReadOnlyWorkspaceState state)
 		{
 			SynchronizationContext.Post(x => OnStateChangedInternalMainThread(state), null);
@@ -85,6 +101,7 @@ namespace UnrealGameSync
 				_asyncDisposer.Add(_prevUpdateTask);
 			}
 			_stateWrapper.Dispose();
+			Lock.Dispose();
 		}
 
 		public ConfigFile ProjectConfigFile
@@ -145,23 +162,35 @@ namespace UnrealGameSync
 			string statusMessage;
 			WorkspaceUpdateResult result = WorkspaceUpdateResult.FailedToSync;
 
-			try
-			{
-				(result, statusMessage) = await update.ExecuteAsync(PerforceSettings, Project, _stateWrapper, _logger, cancellationToken);
-				if (result != WorkspaceUpdateResult.Success)
-				{
-					_logger.LogError("{Message}", statusMessage);
-				}
-			}
-			catch (OperationCanceledException)
+			if (!Lock.TryAcquire())
 			{
-				statusMessage = "Canceled.";
-				_logger.LogError("Canceled.");
+				statusMessage = "Command line sync already in progress";
+				_logger.LogError("Another process is already syncing this workspace.");
 			}
-			catch (Exception ex)
+			else
 			{
-				statusMessage = "Failed with exception - " + ex.ToString();
-				_logger.LogError(ex, "Failed with exception");
+				try
+				{
+					(result, statusMessage) = await update.ExecuteAsync(PerforceSettings, Project, _stateWrapper, _logger, cancellationToken);
+					if (result != WorkspaceUpdateResult.Success)
+					{
+						_logger.LogError("{Message}", statusMessage);
+					}
+				}
+				catch (OperationCanceledException)
+				{
+					statusMessage = "Canceled.";
+					_logger.LogError("Canceled.");
+				}
+				catch (Exception ex)
+				{
+					statusMessage = "Failed with exception - " + ex.ToString();
+					_logger.LogError(ex, "Failed with exception");
+				}
+				finally
+				{
+					Lock.Release();
+				}
 			}
 
 			ProjectConfigFile = context.ProjectConfigFile;
diff --git a/Engine/Source/Programs/UnrealGameSync/UnrealGameSyncCmd/Program.cs b/Engine/Source/Programs/UnrealGameSync/UnrealGameSyncCmd/Program.cs
--- a/Engine/Source/Programs/UnrealGameSync/UnrealGameSyncCmd/Program.cs
+++ b/Engine/Source/Programs/UnrealGameSync/UnrealGameSyncCmd/Program.cs
@@ -299,15 +299,14 @@ namespace UnrealGameSyncCmd
 			return settings;
 		}
 
-		public static async Task<UserWorkspaceState> ReadWorkspaceState(IPerforceConnection perforceClient, UserWorkspaceSettings settings, GlobalSettingsFile userSettings, ILogger logger)
+		public static async Task<WorkspaceStateWrapper> ReadWorkspaceState(IPerforceConnection perforceClient, UserWorkspaceSettings settings, GlobalSettingsFile userSettings, ILogger logger)
 		{
-			UserWorkspaceState state = userSettings.FindOrAddWorkspaceState(settings, logger);
-			if (state.SettingsTimeUtc != settings.LastModifiedTimeUtc)
+			WorkspaceStateWrapper state = userSettings.FindOrAddWorkspaceState(settings, logger);
+			if (state.Current.SettingsTimeUtc != settings.LastModifiedTimeUtc)
 			{
 				logger.LogDebug("Updating state due to modified settings timestamp");
 				ProjectInfo info = await ProjectInfo.CreateAsync(perforceClient, settings, CancellationToken.None);
-				state.UpdateCachedProjectInfo(info, settings.LastModifiedTimeUtc);
-				state.Save(logger);
+				state.Modify(x => x.UpdateCachedProjectInfo(info, settings.LastModifiedTimeUtc));
 			}
 			return state;
 		}
@@ -613,11 +612,11 @@ namespace UnrealGameSyncCmd
 
 				UserWorkspaceSettings settings = ReadRequiredUserWorkspaceSettings();
 				using IPerforceConnection perforceClient = await ConnectAsync(settings, context.LoggerFactory);
-				UserWorkspaceState state = await ReadWorkspaceState(perforceClient, settings, context.UserSettings, logger);
+				WorkspaceStateWrapper state = await ReadWorkspaceState(perforceClient, settings, context.UserSettings, logger);
 
 				changeString ??= "latest";
 
-				ProjectInfo projectInfo = state.CreateProjectInfo();
+				ProjectInfo projectInfo = new ProjectInfo(settings.RootDir, state.Current);
 				UserProjectSettings projectSettings = context.UserSettings.FindOrAddProjectSettings(projectInfo, settings, logger);
 
 				ConfigFile projectConfig = await ConfigUtils.ReadProjectConfigFileAsync(perforceClient, projectInfo, logger, CancellationToken.None);
@@ -659,10 +658,17 @@ namespace UnrealGameSyncCmd
 
 				string[] syncFilter = ReadSyncFilter(settings, context.UserSettings, projectConfig);
 
+				using WorkspaceLock workspaceLock = new WorkspaceLock(settings.RootDir);
+				if (!workspaceLock.TryAcquire())
+				{
+					logger.LogError("Another process is already syncing this workspace.");
+					return;
+				}
+
 				WorkspaceUpdateContext updateContext = new WorkspaceUpdateContext(change, options, BuildConfig.Development, syncFilter, projectSettings.BuildSteps, null);
 				if (syncOptions.Binaries)
 				{
-					List<PerforceArchiveInfo> archives = await PerforceArchive.EnumerateAsync(perforceClient, projectConfig, state.ProjectIdentifier, CancellationToken.None);
+					List<PerforceArchiveInfo> archives = await PerforceArchive.EnumerateAsync(perforceClient, projectConfig, state.Current.ProjectIdentifier, CancellationToken.None);
 
 					PerforceArchiveInfo? editorArchiveInfo = archives.FirstOrDefault(x => x.Name == IArchiveInfo.EditorArchiveType);
 					if (editorArchiveInfo == null)
@@ -719,8 +725,7 @@ namespace UnrealGameSyncCmd
 					logger.LogError("{Message} (Result: {Result})", message, result);
 				}
 
-				state.SetLastSyncState(result, updateContext, message);
-				state.Save(logger);
+				state.Modify(x => x.SetLastSyncState(result, updateContext, message));
 			}
 		}
 
@@ -759,9 +764,9 @@ namespace UnrealGameSyncCmd
 
 				UserWorkspaceSettings settings = ReadRequiredUserWorkspaceSettings();
 				using IPerforceConnection perforceClient = await ConnectAsync(settings, context.LoggerFactory);
-				UserWorkspaceState state = await ReadWorkspaceState(perforceClient, settings, context.UserSettings, logger);
+				WorkspaceStateWrapper state = await ReadWorkspaceState(perforceClient, settings, context.UserSettings, logger);
 
-				ProjectInfo projectInfo = state.CreateProjectInfo();
+				ProjectInfo projectInfo = new ProjectInfo(settings.RootDir, state.Current);
 				ConfigFile projectConfig = await ConfigUtils.ReadProjectConfigFileAsync(perforceClient, projectInfo, logger, CancellationToken.None);
 
 				FileReference receiptFile = ConfigUtils.GetEditorReceiptFile(projectInfo, projectConfig, EditorConfig);
@@ -946,8 +951,8 @@ namespace UnrealGameSyncCmd
 
 				UserWorkspaceSettings workspaceSettings = ReadRequiredUserWorkspaceSettings();
 				using IPerforceConnection perforceClient = await ConnectAsync(workspaceSettings, context.LoggerFactory);
-				UserWorkspaceState workspaceState = await ReadWorkspaceState(perforceClient, workspaceSettings, context.UserSettings, logger);
-				ProjectInfo projectInfo = workspaceState.CreateProjectInfo();
+				WorkspaceStateWrapper workspaceState = await ReadWorkspaceState(perforceClient, workspaceSettings, context.UserSettings, logger);
+				ProjectInfo projectInfo = new ProjectInfo(workspaceSettings.RootDir, workspaceState.Current);
 
 				ConfigFile projectConfig = await ConfigUtils.ReadProjectConfigFileAsync(perforceClient, projectInfo, logger, CancellationToken.None);
 				Dictionary<Guid, WorkspaceSyncCategory> syncCategories = ConfigUtils.GetSyncCategories(projectConfig);
@@ -1094,9 +1099,9 @@ namespace UnrealGameSyncCmd
 
 				UserWorkspaceSettings settings = ReadRequiredUserWorkspaceSettings();
 				using IPerforceConnection perforceClient = await ConnectAsync(settings, context.LoggerFactory);
-				UserWorkspaceState state = await ReadWorkspaceState(perforceClient, settings, context.UserSettings, logger);
+				WorkspaceStateWrapper state = await ReadWorkspaceState(perforceClient, settings, context.UserSettings, logger);
 
-				ProjectInfo projectInfo = state.CreateProjectInfo();
+				ProjectInfo projectInfo = new ProjectInfo(settings.RootDir, state.Current);
 
 				if (listOnly)
 				{
@@ -1128,7 +1133,7 @@ namespace UnrealGameSyncCmd
 					steps = new HashSet<Guid> { id };
 				}
 
-				WorkspaceUpdateContext updateContext = new WorkspaceUpdateContext(state.CurrentChangeNumber, WorkspaceUpdateOptions.Build, BuildConfig.Development, null, new List<ConfigObject>(), steps);
+				WorkspaceUpdateContext updateContext = new WorkspaceUpdateContext(state.Current.CurrentChangeNumber, WorkspaceUpdateOptions.Build, BuildConfig.Development, null, new List<ConfigObject>(), steps);
 
 				WorkspaceUpdate update = new WorkspaceUpdate(updateContext);
 				await update.ExecuteAsync(perforceClient.Settings, projectInfo, state, context.Logger, CancellationToken.None);
@@ -1150,25 +1155,25 @@ namespace UnrealGameSyncCmd
 
 				using IPerforceConnection perforceClient = await ConnectAsync(settings, context.LoggerFactory);
 
-				UserWorkspaceState state = await ReadWorkspaceState(perforceClient, settings, context.UserSettings, logger);
+				WorkspaceStateWrapper state = await ReadWorkspaceState(perforceClient, settings, context.UserSettings, logger);
 				if (update)
 				{
 					ProjectInfo newProjectInfo = await ProjectInfo.CreateAsync(perforceClient, settings, CancellationToken.None);
-					state.UpdateCachedProjectInfo(newProjectInfo, settings.LastModifiedTimeUtc);
+					state.Modify(x => x.UpdateCachedProjectInfo(newProjectInfo, settings.LastModifiedTimeUtc));
 				}
 
-				string streamOrBranchName = state.StreamName ?? settings.BranchPath.TrimStart('/');
-				if (state.LastSyncResultMessage == null)
+				string streamOrBranchName = state.Current.StreamName ?? settings.BranchPath.TrimStart('/');
+				if (state.Current.LastSyncResultMessage == null)
 				{
 					logger.LogInformation("Not currently synced to {Stream}", streamOrBranchName);
 				}
-				else if (state.LastSyncResult == WorkspaceUpdateResult.Success)
+				else if (state.Current.LastSyncResult == WorkspaceUpdateResult.Success)
 				{
-					logger.LogInformation("Synced to {Stream} CL {Change}", streamOrBranchName, state.LastSyncChangeNumber);
+					logger.LogInformation("Synced to {Stream} CL {Change}", streamOrBranchName, state.Current.LastSyncChangeNumber);
 				}
 				else
 				{
-					logger.LogWarning("Last sync to {Stream} CL {Change} failed: {Result}", streamOrBranchName, state.LastSyncChangeNumber, state.LastSyncResultMessage);
+					logger.LogWarning("Last sync to {Stream} CL {Change} failed: {Result}", streamOrBranchName, state.Current.LastSyncChangeNumber, state.Current.LastSyncResultMessage);
 				}
 			}
 		}
diff --git a/Engine/Source/Programs/UnrealGameSync/UnrealGameSyncShared/WorkspaceLock.cs b/Engine/Source/Programs/UnrealGameSync/UnrealGameSyncShared/WorkspaceLock.cs
new file mode 100644
--- /dev/null
+++ b/Engine/Source/Programs/UnrealGameSync/UnrealGameSyncShared/WorkspaceLock.cs
@@ -0,0 +1,176 @@
+using EpicGames.Core;
+using System;
+using System.Collections.Generic;
+using System.Runtime.CompilerServices;
+using System.Security.Cryptography;
+using System.Text;
+using System.Threading;
+using System.Threading.Tasks;
+
+namespace UnrealGameSync
+{
+	/// <summary>
+	/// Encapsulates the state of cross-process workspace lock
+	/// </summary>
+	public class WorkspaceLock : IDisposable
+	{
+		const string Prefix = @"Global\ugs-workspace";
+
+		readonly Mutex _mutex;
+
+		readonly object _lockObject = new object();
+		readonly string _objectName;
+		int _acquireCount;
+
+		bool _locked;
+		EventWaitHandle? _lockedEvent;
+
+		Thread? _monitorThread;
+		ManualResetEvent _cancelMonitorEvent = new ManualResetEvent(false);
+
+		/// <summary>
+		/// Callback for the lock state changing
+		/// </summary>
+		public event Action<bool>? OnChange;
+
+		/// <summary>
+		/// Constructor
+		/// </summary>
+		/// <param name="rootDir">Root directory for the workspace</param>
+		public WorkspaceLock(DirectoryReference rootDir)
+		{
+			using (MD5 md5 = MD5.Create())
+			{
+				byte[] idBytes = Encoding.UTF8.GetBytes(rootDir.FullName.ToUpperInvariant());
+				_objectName = StringUtils.FormatHexString(md5.ComputeHash(idBytes));
+			}
+
+			_mutex = new Mutex(false, $"{Prefix}.{_objectName}.mutex");
+
+			_monitorThread = new Thread(MonitorThread);
+			_monitorThread.IsBackground = true;
+			_monitorThread.Start();
+		}
+
+		/// <inheritdoc/>
+		public void Dispose()
+		{
+			while (_acquireCount > 0)
+			{
+				Release();
+			}
+
+			if (_monitorThread != null)
+			{
+				_cancelMonitorEvent.Set();
+
+				_monitorThread.Join();
+				_monitorThread = null;
+
+				_cancelMonitorEvent.Dispose();
+			}
+
+			_mutex.Dispose();
+		}
+
+		/// <summary>
+		/// Determines if the lock is held by any
+		/// </summary>
+		/// <returns>True if the lock is held by any process</returns>
+		public bool IsLocked() => _locked;
+
+		/// <summary>
+		/// Determines if the lock is held by another process
+		/// </summary>
+		/// <returns>True if the lock is held by another process</returns>
+		public bool IsLockedByOtherProcess() => _acquireCount == 0 && IsLocked();
+
+		/// <summary>
+		/// Attempt to acquire the mutext
+		/// </summary>
+		/// <returns></returns>
+		public bool TryAcquire()
+		{
+			lock (_lockObject)
+			{
+				try
+				{
+					if (!_mutex.WaitOne(0))
+					{
+						return false;
+					}
+				}
+				catch (AbandonedMutexException)
+				{
+				}
+
+				if (++_acquireCount == 1)
+				{
+					_lockedEvent = CreateLockedEvent();
+					_lockedEvent.Set();
+				}
+
+				return true;
+			}
+		}
+
+		/// <summary>
+		/// Release the current mutext
+		/// </summary>
+		public void Release()
+		{
+			lock (_lockObject)
+			{
+				if (_acquireCount > 0)
+				{
+					_mutex.ReleaseMutex();
+					if (--_acquireCount == 0 && _lockedEvent != null)
+					{
+						_lockedEvent.Reset();
+						_lockedEvent.Dispose();
+						_lockedEvent = null;
+					}
+				}
+			}
+		}
+
+		void MonitorThread()
+		{
+			_locked = IsLocked();
+			for (; ;)
+			{
+				if (_locked)
+				{
+					try
+					{
+						int idx = WaitHandle.WaitAny(new WaitHandle[] { _mutex, _cancelMonitorEvent });
+						if (idx == 1)
+						{
+							break;
+						}
+					}
+					catch (AbandonedMutexException)
+					{
+					}
+
+					_mutex.ReleaseMutex();
+				}
+				else
+				{
+					using EventWaitHandle lockedEvent = CreateLockedEvent();
+
+					int idx = WaitHandle.WaitAny(new WaitHandle[] { lockedEvent, _cancelMonitorEvent });
+					if (idx == 1)
+					{
+						break;
+					}
+				}
+
+				_locked ^= true;
+				OnChange?.Invoke(!_locked);
+			}
+		}
+
+		EventWaitHandle CreateLockedEvent() => new EventWaitHandle(false, EventResetMode.ManualReset, $"{Prefix}.{_objectName}.locked");
+	}
+}
