UGS: Respect settings for syncing all projects in a workspace when syncing through the CLI app.

diff --git a/Engine/Source/Programs/UnrealGameSync/UnrealGameSync/Controls/WorkspaceControl.cs b/Engine/Source/Programs/UnrealGameSync/UnrealGameSync/Controls/WorkspaceControl.cs
--- a/Engine/Source/Programs/UnrealGameSync/UnrealGameSync/Controls/WorkspaceControl.cs
+++ b/Engine/Source/Programs/UnrealGameSync/UnrealGameSync/Controls/WorkspaceControl.cs
@@ -963,14 +963,7 @@ namespace UnrealGameSync
 			{
 				context.Options |= WorkspaceUpdateOptions.AutoResolveChanges;
 			}
-			if (_workspaceSettings.Filter.AllProjects ?? _settings.Global.Filter.AllProjects ?? false)
-			{
-				context.Options |= WorkspaceUpdateOptions.SyncAllProjects | WorkspaceUpdateOptions.IncludeAllProjectsInSolution;
-			}
-			if (_workspaceSettings.Filter.AllProjectsInSln ?? _settings.Global.Filter.AllProjectsInSln ?? false)
-			{
-				context.Options |= WorkspaceUpdateOptions.IncludeAllProjectsInSolution;
-			}
+			context.Options |= WorkspaceUpdateContext.GetOptionsFromConfig(_settings.Global, _workspaceSettings);
 
 			_updateCallback = callback;
 
diff --git a/Engine/Source/Programs/UnrealGameSync/UnrealGameSyncCmd/Program.cs b/Engine/Source/Programs/UnrealGameSync/UnrealGameSyncCmd/Program.cs
--- a/Engine/Source/Programs/UnrealGameSync/UnrealGameSyncCmd/Program.cs
+++ b/Engine/Source/Programs/UnrealGameSync/UnrealGameSyncCmd/Program.cs
@@ -654,6 +654,7 @@ namespace UnrealGameSyncCmd
 				{
 					options |= WorkspaceUpdateOptions.Refilter;
 				}
+				options |= WorkspaceUpdateContext.GetOptionsFromConfig(context.UserSettings.Global, settings);
 				options |= WorkspaceUpdateOptions.RemoveFilteredFiles;
 
 				string[] syncFilter = ReadSyncFilter(settings, context.UserSettings, projectConfig);
diff --git a/Engine/Source/Programs/UnrealGameSync/UnrealGameSyncShared/WorkspaceUpdate.cs b/Engine/Source/Programs/UnrealGameSync/UnrealGameSyncShared/WorkspaceUpdate.cs
--- a/Engine/Source/Programs/UnrealGameSync/UnrealGameSyncShared/WorkspaceUpdate.cs
+++ b/Engine/Source/Programs/UnrealGameSync/UnrealGameSyncShared/WorkspaceUpdate.cs
@@ -117,6 +117,20 @@ namespace UnrealGameSync
 			UserBuildStepObjects = inUserBuildSteps;
 			CustomBuildSteps = inCustomBuildSteps;
 		}
+
+		public static WorkspaceUpdateOptions GetOptionsFromConfig(GlobalSettings globalSettings, UserWorkspaceSettings workspaceSettings)
+		{
+			WorkspaceUpdateOptions options = 0;
+			if (workspaceSettings.Filter.AllProjects ?? globalSettings.Filter.AllProjects ?? false)
+			{
+				options |= WorkspaceUpdateOptions.SyncAllProjects | WorkspaceUpdateOptions.IncludeAllProjectsInSolution;
+			}
+			if (workspaceSettings.Filter.AllProjectsInSln ?? globalSettings.Filter.AllProjectsInSln ?? false)
+			{
+				options |= WorkspaceUpdateOptions.IncludeAllProjectsInSolution;
+			}
+			return options;
+		}
 	}
 
 	public class WorkspaceSyncCategory
