UGS: Fix nullable static analyzer warning.

diff --git a/Engine/Source/Programs/UnrealGameSync/UnrealGameSyncShared/WorkspaceState.cs b/Engine/Source/Programs/UnrealGameSync/UnrealGameSyncShared/WorkspaceState.cs
--- a/Engine/Source/Programs/UnrealGameSync/UnrealGameSyncShared/WorkspaceState.cs
+++ b/Engine/Source/Programs/UnrealGameSync/UnrealGameSyncShared/WorkspaceState.cs
@@ -84,7 +84,7 @@ namespace UnrealGameSync
 
 		public void SetBisectState(int change, BisectState state)
 		{
-			BisectEntry entry = BisectChanges.FirstOrDefault(x => x.Change == change);
+			BisectEntry? entry = BisectChanges.FirstOrDefault(x => x.Change == change);
 			if (entry == null)
 			{
 				entry = new BisectEntry();
