UGS: Use HttpClient rather than WebRequest for HTTP requests (deprecated in newer NET versions).

diff --git a/Engine/Source/Programs/UnrealGameSync/UnrealGameSync/RESTApi.cs b/Engine/Source/Programs/UnrealGameSync/UnrealGameSync/RESTApi.cs
--- a/Engine/Source/Programs/UnrealGameSync/UnrealGameSync/RESTApi.cs
+++ b/Engine/Source/Programs/UnrealGameSync/UnrealGameSync/RESTApi.cs
@@ -4,6 +4,8 @@ using System;
 using System.Collections.Generic;
 using System.IO;
 using System.Net;
+using System.Net.Http;
+using System.Net.Http.Headers;
 using System.Net.Mime;
 using System.Text;
 using System.Text.Json;
@@ -14,7 +16,7 @@ namespace UnrealGameSync
 {
 	class RestException : Exception
 	{
-		public RestException(string method, string uri, Exception innerException)
+		public RestException(HttpMethod method, string uri, Exception innerException)
 			: base(String.Format("Error executing {0} {1}", method, uri), innerException)
 		{
 		}
@@ -27,51 +29,44 @@ namespace UnrealGameSync
 
 	public static class RestApi
 	{
-		private static async Task<string> SendRequestInternal(string url, string method, string? requestBody, CancellationToken cancellationToken)
-		{
-			HttpWebRequest request = (HttpWebRequest)WebRequest.Create(url);
-			request.ContentType = "application/json";
-			request.Method = method;
+		static HttpClient s_httpClient = new HttpClient();
 
-			// Add json to request body
-			if (!string.IsNullOrEmpty(requestBody))
+		private static async Task<string> SendRequestInternalAsync(string url, HttpMethod method, string? requestBody, CancellationToken cancellationToken)
+		{
+			using (HttpRequestMessage request = new HttpRequestMessage(method, url))
 			{
-				if (method == "POST" || method == "PUT")
+				// Add json to request body
+				if (!string.IsNullOrEmpty(requestBody))
 				{
-					byte[] bytes = Encoding.UTF8.GetBytes(requestBody);
-					using (Stream requestStream = request.GetRequestStream())
+					if (method == HttpMethod.Post || method == HttpMethod.Put)
 					{
-						await requestStream.WriteAsync(bytes, 0, bytes.Length, cancellationToken);
+						request.Content = new ByteArrayContent(Encoding.UTF8.GetBytes(requestBody));
+						request.Content.Headers.ContentType = new MediaTypeHeaderValue("application/json");
 					}
 				}
-			}
-			try
-			{
-				using (WebResponse response = request.GetResponse())
+
+				try
 				{
-					byte[] data;
-					using (MemoryStream buffer = new MemoryStream())
+					using (HttpResponseMessage response = await s_httpClient.SendAsync(request))
 					{
-						await response.GetResponseStream().CopyToAsync(buffer, cancellationToken);
-						data = buffer.ToArray();
+						return await response.Content.ReadAsStringAsync();
 					}
-					return Encoding.UTF8.GetString(data);
 				}
-			}
-			catch (Exception ex)
-			{
-				throw new RestException(method, request.RequestUri.ToString(), ex);
+				catch (Exception ex)
+				{
+					throw new RestException(method, request.RequestUri.ToString(), ex);
+				}
 			}
 		}
 
 		public static Task<string> PostAsync(string url, string requestBody, CancellationToken cancellationToken)
 		{
-			return SendRequestInternal(url, "POST", requestBody, cancellationToken);
+			return SendRequestInternalAsync(url, HttpMethod.Post, requestBody, cancellationToken);
 		}
 
 		public static Task<string> GetAsync(string url, CancellationToken cancellationToken)
 		{
-			return SendRequestInternal(url, "GET", null, cancellationToken);
+			return SendRequestInternalAsync(url, HttpMethod.Get, null, cancellationToken);
 		}
 
 		public static async Task<T> GetAsync<T>(string url, CancellationToken cancellationToken)
@@ -81,7 +76,7 @@ namespace UnrealGameSync
 
 		public static Task<string> PutAsync<T>(string url, T obj, CancellationToken cancellationToken)
 		{
-			return SendRequestInternal(url, "PUT", JsonSerializer.Serialize(obj), cancellationToken);
+			return SendRequestInternalAsync(url, HttpMethod.Put, JsonSerializer.Serialize(obj), cancellationToken);
 		}
 	}
 }
diff --git a/Engine/Source/Programs/UnrealGameSync/UnrealGameSyncShared/Telemetry.cs b/Engine/Source/Programs/UnrealGameSync/UnrealGameSyncShared/Telemetry.cs
--- a/Engine/Source/Programs/UnrealGameSync/UnrealGameSyncShared/Telemetry.cs
+++ b/Engine/Source/Programs/UnrealGameSync/UnrealGameSyncShared/Telemetry.cs
@@ -4,10 +4,17 @@ using Microsoft.Extensions.Logging;
 using System;
 using System.Collections.Generic;
 using System.IO;
+using System.Linq;
 using System.Net;
+using System.Net.Http;
+using System.Net.Http.Headers;
+using System.Reflection;
+using System.Runtime.CompilerServices;
 using System.Text;
 using System.Text.Json;
 using System.Threading;
+using System.Threading.Channels;
+using System.Threading.Tasks;
 using System.Web;
 
 namespace UnrealGameSync
@@ -46,30 +53,11 @@ namespace UnrealGameSync
 	/// </summary>
 	public class EpicTelemetrySink : ITelemetrySink
 	{
-		/// <summary>
-		/// Combined url to post event streams to
-		/// </summary>
-		string _url;
-
-		/// <summary>
-		/// Lock used to modify the event queue
-		/// </summary>
-		object _lockObject = new object();
-
-		/// <summary>
-		/// Whether a flush is queued
-		/// </summary>
-		bool _hasPendingFlush = false;
-
-		/// <summary>
-		/// List of pending events
-		/// </summary>
-		List<string> _pendingEvents = new List<string>();
-
-		/// <summary>
-		/// The log writer to use
-		/// </summary>
-		ILogger _logger;
+		readonly string _url;
+		readonly ILogger _logger;
+		readonly HttpClient _httpClient = new HttpClient();
+		readonly Channel<string> _channel = Channel.CreateUnbounded<string>(new UnboundedChannelOptions());
+		readonly Task _backgroundTask;
 
 		/// <summary>
 		/// Constructor
@@ -80,28 +68,22 @@ namespace UnrealGameSync
 			this._logger = logger;
 
 			logger.LogInformation("Posting to URL: {Url}", url);
-		}
 
-		/// <inheritdoc/>
-		public void Dispose()
-		{
-			Flush();
+			_backgroundTask = Task.Run(WriteEventsAsync);
 		}
 
 		/// <inheritdoc/>
-		public void Flush()
+		public void Dispose()
 		{
-			for (; ; )
+			_channel.Writer.TryComplete();
+			try
+			{
+				_backgroundTask.Wait();
+			}
+			catch (OperationCanceledException)
 			{
-				lock (_lockObject)
-				{
-					if (!_hasPendingFlush)
-					{
-						break;
-					}
-				}
-				Thread.Sleep(10);
 			}
+			_httpClient.Dispose();
 		}
 
 		/// <inheritdoc/>
@@ -114,38 +96,24 @@ namespace UnrealGameSync
 			}
 
 			string eventText = attributesText.Insert(1, String.Format("\"EventName\":\"{0}\",", HttpUtility.JavaScriptStringEncode(eventName)));
-			lock (_pendingEvents)
-			{
-				_pendingEvents.Add(eventText);
-				if (!_hasPendingFlush)
-				{
-					ThreadPool.QueueUserWorkItem(obj => BackgroundFlush());
-					_hasPendingFlush = true;
-				}
-			}
+			_channel.Writer.TryWrite(eventText);
 		}
 
 		/// <summary>
 		/// Synchronously sends a telemetry event
 		/// </summary>
-		void BackgroundFlush()
+		async Task WriteEventsAsync()
 		{
-			for (; ; )
+			string version = Assembly.GetExecutingAssembly().GetName().Version?.ToString() ?? "0.0";
+			while (await _channel.Reader.WaitToReadAsync())
 			{
 				try
 				{
 					// Generate the content for this event
-					List<string> events = new List<string>();
-					lock (_lockObject)
+					List<string> events = await _channel.Reader.ReadAllAsync().ToListAsync();
+					if (events.Count == 0)
 					{
-						if (_pendingEvents.Count == 0)
-						{
-							_hasPendingFlush = false;
-							break;
-						}
-
-						events.AddRange(_pendingEvents);
-						_pendingEvents.Clear();
+						continue;
 					}
 
 					// Print all the events we're sending
@@ -159,42 +127,24 @@ namespace UnrealGameSync
 					byte[] content = Encoding.UTF8.GetBytes(contentText);
 
 					// Post the event data
-					HttpWebRequest request = (HttpWebRequest)WebRequest.Create(_url);
-					request.Method = "POST";
-					request.ContentType = "application/json";
-					request.UserAgent = "ue/ugs";
-					request.Timeout = 5000;
-					request.ContentLength = content.Length;
-					request.ContentType = "application/json";
-					using (Stream requestStream = request.GetRequestStream())
+					using (HttpRequestMessage request = new HttpRequestMessage(HttpMethod.Post, _url))
 					{
-						requestStream.Write(content, 0, content.Length);
-					}
+						request.Headers.UserAgent.Add(new ProductInfoHeaderValue("UnrealGameSync", version));
+						request.Content = new ByteArrayContent(content);
+						request.Content.Headers.ContentType = new MediaTypeHeaderValue("application/json");
 
-					// Wait for the response and dispose of it immediately
-					using (HttpWebResponse response = (HttpWebResponse)request.GetResponse())
-					{
-						_logger.LogInformation("Response: {StatusCode}", (int)response.StatusCode);
-					}
-				}
-				catch (WebException ex)
-				{
-					// Handle errors. Any non-200 responses automatically generate a WebException.
-					HttpWebResponse? response = (HttpWebResponse?)ex.Response;
-					if (response == null)
-					{
-						_logger.LogError(ex, "Exception while attempting to send event");
-					}
-					else
-					{
-						string responseText;
-						using (Stream responseStream = response.GetResponseStream())
+						using (HttpResponseMessage response = await _httpClient.SendAsync(request))
 						{
-							MemoryStream memoryStream = new MemoryStream();
-							responseStream.CopyTo(memoryStream);
-							responseText = Encoding.UTF8.GetString(memoryStream.ToArray());
+							if (response.IsSuccessStatusCode)
+							{
+								_logger.LogInformation("Response: {StatusCode}", (int)response.StatusCode);
+							}
+							else
+							{
+								string responseContent = response.Content.ReadAsStringAsync().Result;
+								_logger.LogError("Unable to send telemetry data to server ({Code}): {Message}", response.StatusCode, responseContent);
+							}
 						}
-						_logger.LogError("Failed to send analytics event. Code = {Code}. Desc = {Dec}. Response = {Response}.", (int)response.StatusCode, response.StatusDescription, responseText);
 					}
 				}
 				catch (Exception ex)
