Horde: Add a prompt to UGS to update the launcher if it's below a certain version. Requires UnrealGameSync.ini to be included in the distribution; can be enabled with -set:IncludeInstaller=true to BuildGraph script.

diff --git a/Engine/Source/Programs/UnrealGameSync/BuildUGS.xml b/Engine/Source/Programs/UnrealGameSync/BuildUGS.xml
--- a/Engine/Source/Programs/UnrealGameSync/BuildUGS.xml
+++ b/Engine/Source/Programs/UnrealGameSync/BuildUGS.xml
@@ -6,6 +6,7 @@
 	<Option Name="SubmitToStream" DefaultValue="" Description="Stream to submit binaries to"/>
 	<Option Name="SubmitToBranch" DefaultValue="" Description="Branch to submit binaries to"/>
 	<Option Name="SubmitPath" DefaultValue="" Description="Relative path within the stream/branch to submit binaries"/>
+	<Option Name="IncludeInstaller" DefaultValue="false" Description="Whether to include the launcher installer in the regular client distribution"/>
 
 	<Property Name="UgsDir" Value="$(RootDir)/Engine/Source/Programs/UnrealGameSync"/>
 	<Property Name="UgsOutputDir" Value="$(RootDir)/LocalBuilds/UnrealGameSync"/>
@@ -18,18 +19,36 @@
 	<Property Name="VersionArguments" Value="/p:Version=$(Version).0 /p:InformationalVersion=$(InformationalVersion)"/>
 	
 	<Agent Name="UnrealGameSync Agent" Type="Win64">
+		<Property Name="InstallerDir" Value="$(RootDir)/Engine/Source/Programs/UnrealGameSync/Installer"/>
+
 		<Property Name="UgsPublishDir" Value="$(UgsOutputDir)"/>
 		<Property Name="UgsPublishDir" Value="$(UgsPublishDir)/$(SubmitPath)" If="'$(SubmitPath)' != ''"/>
 
 		<Property Name="UgsZipPublishDir" Value="$(UgsZipOutputDir)"/>
 		<Property Name="UgsZipPublishDir" Value="$(UgsZipPublishDir)/$(SubmitPath)" If="'$(SubmitPath)' != ''"/>
-	
-		<Node Name="Build UnrealGameSync">
+
+		<Node Name="Build UnrealGameSync Installer">
+			<Property Name="WixDir" Value="$(RootDir)/Engine/Source/ThirdParty/WiX/3.8"/>
+			<Property Name="UgsLauncherPublishDir" Value="$(UgsDir)/UnrealGameSyncLauncher/bin/Release/netcoreapp3.1/win-x64/publish"/>
+
+			<Delete Files="$(UgsLauncherPublishDir)/..."/>
+			<DotNet Arguments="publish &quot;$(UgsDir)/UnrealGameSyncLauncher/UnrealGameSyncLauncher.csproj&quot; --output=&quot;$(UgsLauncherPublishDir)&quot; --no-self-contained -c Release -r win-x64 $(VersionArguments)"/>
+
+			<Spawn Exe="$(WixDir)/heat.exe" Arguments="dir &quot;$(UgsLauncherPublishDir)&quot; -cg UGSLauncher_Project -dr INSTALLFOLDER -scom -sreg -srd -var var.BasePath -gg -sfrag -out obj/Release/UGSLauncher.wxs" WorkingDir="$(InstallerDir)"/>
+			<Spawn Exe="$(WixDir)/candle.exe" Arguments="-dBasePath=&quot;$(UgsLauncherPublishDir)&quot; -out obj\Release\ -dConfiguration=Release -dPlatform=x64 -arch x86 -ext ..\..\..\ThirdParty\WiX\3.8\WixUtilExtension.dll Product.wxs obj/Release/UGSLauncher.wxs" WorkingDir="$(InstallerDir)"/>
+			<Spawn Exe="$(WixDir)/light.exe" Arguments="-out &quot;bin/Release/UnrealGameSync.msi&quot; -pdbout &quot;bin/Release/UnrealGameSync.wixpdb&quot; -cultures:null -ext ..\..\..\ThirdParty\WiX\3.8\WixUtilExtension.dll -sice:ICE69 &quot;obj\Release\Product.wixobj&quot; &quot;obj\Release\UGSLauncher.wixobj&quot;" WorkingDir="$(InstallerDir)"/>
+		</Node>
+		
+		<Property Name="InstallerDependency" Value=""/>
+		<Property Name="InstallerDependency" Value="Build UnrealGameSync Installer" If="IncludeInstaller"/>
+		
+		<Node Name="Build UnrealGameSync" Requires="$(InstallerDependency)">
 			<Delete Files="$(UgsOutputDir)/..."/>
 			<DotNet Arguments="publish &quot;$(UgsDir)/UnrealGameSyncCmd/UnrealGameSyncCmd.csproj&quot; --output=&quot;$(UgsPublishDir)&quot; --no-self-contained -c Release -r win-x64 $(VersionArguments)"/>
 			<DotNet Arguments="publish &quot;$(UgsDir)/UnrealGameSyncCmd/UnrealGameSyncCmd.csproj&quot; --output=&quot;$(UgsPublishDir)&quot; --no-self-contained -c Release -r osx-x64 $(VersionArguments)"/>
 			<DotNet Arguments="publish &quot;$(UgsDir)/UnrealGameSyncCmd/UnrealGameSyncCmd.csproj&quot; --output=&quot;$(UgsPublishDir)&quot; --no-self-contained -c Release -r linux-x64 $(VersionArguments)"/>
 			<DotNet Arguments="publish &quot;$(UgsDir)/UnrealGameSync/UnrealGameSync.csproj&quot; --output=&quot;$(UgsPublishDir)&quot; --no-self-contained -c Release -r win-x64 $(VersionArguments)"/>
+			<Copy From="$(InstallerDir)/bin/Release/UnrealGameSync.msi" To="$(UgsPublishDir)/UnrealGameSync.msi" If="$(IncludeInstaller)"/>
 		</Node>
 		<Node Name="Publish Unstable" Requires="Build UnrealGameSync">
 			<Submit Description="Updating UGS from CL $(Change)" Workspace="Temp" Files="$(UgsPublishDir)/..." RootDir="$(UgsOutputDir)" Stream="$(SubmitToStream)" Branch="$(SubmitToBranch)/UnstableRelease"/>
@@ -46,21 +65,6 @@
 			<Submit Description="Updating UGS from CL $(Change) (Zip)" Workspace="Temp" Files="$(UgsZipPublishDir)/Release.zip" RootDir="$(UgsZipOutputDir)" Stream="$(SubmitToStream)" Branch="$(SubmitToBranch)"/>
 		</Node>
 	</Agent>
-	
-	<Agent Name="UnrealGameSync Installer" Type="Win64">
-		<Node Name="Build HordeInstaller">
-			<Property Name="WixDir" Value="$(RootDir)/Engine/Source/ThirdParty/WiX/3.8"/>
-			<Property Name="SourceDir" Value="$(RootDir)/Engine/Source/Programs/UnrealGameSync/Installer"/>
-			<Property Name="UgsLauncherPublishDir" Value="$(UgsDir)/UnrealGameSyncLauncher/bin/Release/netcoreapp3.1/win-x64/publish"/>
-
-			<Delete Files="$(UgsLauncherPublishDir)/..."/>
-			<DotNet Arguments="publish &quot;$(UgsDir)/UnrealGameSyncLauncher/UnrealGameSyncLauncher.csproj&quot; --output=&quot;$(UgsLauncherPublishDir)&quot; --no-self-contained -c Release -r win-x64 $(VersionArguments)"/>
-
-			<Spawn Exe="$(WixDir)/heat.exe" Arguments="dir &quot;$(UgsLauncherPublishDir)&quot; -cg UGSLauncher_Project -dr INSTALLFOLDER -scom -sreg -srd -var var.BasePath -gg -sfrag -out obj/Release/UGSLauncher.wxs" WorkingDir="$(SourceDir)"/>
-			<Spawn Exe="$(WixDir)/candle.exe" Arguments="-dBasePath=&quot;$(UgsLauncherPublishDir)&quot; -out obj\Release\ -dConfiguration=Release -dPlatform=x64 -arch x86 -ext ..\..\..\ThirdParty\WiX\3.8\WixUtilExtension.dll Product.wxs obj/Release/UGSLauncher.wxs" WorkingDir="$(SourceDir)"/>
-			<Spawn Exe="$(WixDir)/light.exe" Arguments="-out &quot;bin/Release/UnrealGameSync.msi&quot; -pdbout &quot;bin/Release/UnrealGameSync.wixpdb&quot; -cultures:null -ext ..\..\..\ThirdParty\WiX\3.8\WixUtilExtension.dll -sice:ICE69 &quot;obj\Release\Product.wixobj&quot; &quot;obj\Release\UGSLauncher.wixobj&quot;" WorkingDir="$(SourceDir)"/>
-		</Node>
-	</Agent>
 
 	<Aggregate Name="Publish Unstable With Zip" Requires="Publish Unstable;Publish Unstable Zip"/>
 	<Aggregate Name="Publish Live With Zip" Requires="Publish Live;Publish Live Zip"/>
diff --git a/Engine/Source/Programs/UnrealGameSync/UnrealGameSync/Forms/LauncherUpdateWindow.cs b/Engine/Source/Programs/UnrealGameSync/UnrealGameSync/Forms/LauncherUpdateWindow.cs
new file mode 100644
--- /dev/null
+++ b/Engine/Source/Programs/UnrealGameSync/UnrealGameSync/Forms/LauncherUpdateWindow.cs
@@ -0,0 +1,23 @@
+using System;
+using System.Collections.Generic;
+using System.ComponentModel;
+using System.Data;
+using System.Drawing;
+using System.Text;
+using System.Windows.Forms;
+
+namespace UnrealGameSync.Forms
+{
+	public partial class LauncherUpdateWindow : Form
+	{
+		public LauncherUpdateWindow()
+		{
+			InitializeComponent();
+		}
+
+		private void button1_Click(object sender, EventArgs e)
+		{
+
+		}
+	}
+}
diff --git a/Engine/Source/Programs/UnrealGameSync/UnrealGameSync/Forms/MainWindow.cs b/Engine/Source/Programs/UnrealGameSync/UnrealGameSync/Forms/MainWindow.cs
--- a/Engine/Source/Programs/UnrealGameSync/UnrealGameSync/Forms/MainWindow.cs
+++ b/Engine/Source/Programs/UnrealGameSync/UnrealGameSync/Forms/MainWindow.cs
@@ -74,6 +74,11 @@ namespace UnrealGameSync
 			Quiet,
 		}
 
+		/// <summary>
+		/// Minimum version of the launcher that should be installed. Users will be prompted to upgrade it if the installed version is older than this.
+		/// </summary>
+		static readonly int s_minLauncherVersion = MakeLauncherVersion(1, 19);
+
 		[DllImport("uxtheme.dll", CharSet = CharSet.Unicode)]
 		static extern int SetWindowTheme(IntPtr hWnd, string pszSubAppName, string pszSubIdList);
 
@@ -793,7 +798,7 @@ namespace UnrealGameSync
 
 		public bool CanPerformUpdate()
 		{
-			if(ContainsFocus || Form.ActiveForm == this)
+			if (ContainsFocus || Form.ActiveForm == this || !CanFocus)
 			{
 				return false;
 			}
@@ -917,6 +922,82 @@ namespace UnrealGameSync
 			}
 		}
 
+		private void CheckLauncherVersionTimer_Tick(object sender, EventArgs e)
+		{
+			if (!CanPerformUpdate())
+			{
+				return;
+			}
+
+			CheckLauncherVersionTimer.Enabled = false;
+
+			FileReference executableFile = new FileReference(Assembly.GetExecutingAssembly()!.Location);
+			FileReference installerFile = FileReference.Combine(executableFile.Directory, "UnrealGameSync.msi");
+			if (!FileReference.Exists(installerFile))
+			{
+				return;
+			}
+
+			DateTime now = DateTime.UtcNow;
+			if (now.Ticks < _settings.NextLauncherVersionCheck)
+			{
+				return;
+			}
+
+			int version = GetLauncherVersion();
+			if (version >= s_minLauncherVersion)
+			{
+				return;
+			}
+
+			LauncherUpdateWindow update = new LauncherUpdateWindow();
+			if (update.ShowDialog(this) == DialogResult.Ignore)
+			{
+				_settings.NextLauncherVersionCheck = (now + TimeSpan.FromDays(1.0)).Ticks;
+				_settings.Save(_logger);
+
+				CheckLauncherVersionTimer.Enabled = true;
+				return;
+			}
+
+			using (Process childProcess = new Process())
+			{
+				childProcess.StartInfo.FileName = installerFile.FullName;
+				childProcess.StartInfo.UseShellExecute = true;
+				childProcess.Start();
+			}
+
+			ForceClose();
+		}
+
+		static int GetLauncherVersion()
+		{
+			try
+			{
+				DirectoryReference? programFiles = DirectoryReference.GetSpecialFolder(Environment.SpecialFolder.ProgramFilesX86);
+				if (programFiles == null)
+				{
+					return 0;
+				}
+
+				FileReference launcherFile = FileReference.Combine(programFiles, "UnrealGameSync", "UnrealGameSyncLauncher.exe");
+				if (!FileReference.Exists(launcherFile))
+				{
+					return 0;
+				}
+
+				FileVersionInfo versionInfo = FileVersionInfo.GetVersionInfo(launcherFile.FullName);
+				return MakeLauncherVersion(versionInfo.FileMajorPart, versionInfo.FileMinorPart);
+			}
+			catch (Exception ex)
+			{
+				Program.CaptureException(ex);
+				return 0;
+			}
+		}
+
+		static int MakeLauncherVersion(int major, int minor) => (major << 16) + minor;
+
 		private void StopScheduleTimer()
 		{
 			if(_scheduleTimer != null)
diff --git a/Engine/Source/Programs/UnrealGameSync/UnrealGameSyncLauncher/Properties/AssemblyInfo.cs b/Engine/Source/Programs/UnrealGameSync/UnrealGameSyncLauncher/Properties/AssemblyInfo.cs
--- a/Engine/Source/Programs/UnrealGameSync/UnrealGameSyncLauncher/Properties/AssemblyInfo.cs
+++ b/Engine/Source/Programs/UnrealGameSync/UnrealGameSyncLauncher/Properties/AssemblyInfo.cs
@@ -34,5 +34,5 @@ using System.Runtime.InteropServices;
 // You can specify all the values or you can default the Build and Revision Numbers 
 // by using the '*' as shown below:
 // [assembly: AssemblyVersion("1.0.*")]
-[assembly: AssemblyVersion("1.18.0.0")]
-[assembly: AssemblyFileVersion("1.18.0.0")]
+[assembly: AssemblyVersion("1.19.0.0")]
+[assembly: AssemblyFileVersion("1.19.0.0")]
diff --git a/Engine/Source/Programs/UnrealGameSync/UnrealGameSyncShared/ConfigFile.cs b/Engine/Source/Programs/UnrealGameSync/UnrealGameSyncShared/ConfigFile.cs
--- a/Engine/Source/Programs/UnrealGameSync/UnrealGameSyncShared/ConfigFile.cs
+++ b/Engine/Source/Programs/UnrealGameSync/UnrealGameSyncShared/ConfigFile.cs
@@ -355,6 +355,11 @@ namespace UnrealGameSync
 			Pairs[key] = value.ToString();
 		}
 
+		public void SetValue(string key, long value)
+		{
+			Pairs[key] = value.ToString();
+		}
+
 		public void SetValue(string key, bool value)
 		{
 			Pairs[key] = value? "1" : "0";
@@ -428,6 +433,20 @@ namespace UnrealGameSync
 			return defaultValue;
 		}
 
+		public long GetValue(string key, long defaultValue)
+		{
+			string? valueString = GetValue(key);
+			if (valueString != null)
+			{
+				long value;
+				if (long.TryParse(valueString, out value))
+				{
+					return value;
+				}
+			}
+			return defaultValue;
+		}
+
 		public bool GetValue(string key, bool defaultValue)
 		{
 			return GetValue(key, defaultValue? 1 : 0) != 0;
diff --git a/Engine/Source/Programs/UnrealGameSync/UnrealGameSyncShared/UserSettings.cs b/Engine/Source/Programs/UnrealGameSync/UnrealGameSyncShared/UserSettings.cs
--- a/Engine/Source/Programs/UnrealGameSync/UnrealGameSyncShared/UserSettings.cs
+++ b/Engine/Source/Programs/UnrealGameSync/UnrealGameSyncShared/UserSettings.cs
@@ -498,6 +498,7 @@ namespace UnrealGameSync
 		public string SyncTypeId;
 		public BuildConfig CompiledEditorBuildConfig; // NB: This assumes not using precompiled editor. See CurrentBuildConfig.
 		public TabLabels TabLabels;
+		public long NextLauncherVersionCheck;
 
 		// Precompiled binaries
 		public List<ArchiveSettings> Archives = new List<ArchiveSettings>();
@@ -625,6 +626,7 @@ namespace UnrealGameSync
 			AutoResolveConflicts = (_configFile.GetValue("General.AutoResolveConflicts", "1") != "0");
 			ShowUnreviewedChanges = _configFile.GetValue("General.ShowUnreviewed", true);
 			ShowAutomatedChanges = _configFile.GetValue("General.ShowAutomated", false);
+			NextLauncherVersionCheck = _configFile.GetValue("General.NextLauncherVersionCheck", 0);
 
 			// safely parse the filter enum
 			ShowRobomerge = RobomergeShowChangesOption.All;
@@ -1037,6 +1039,7 @@ namespace UnrealGameSync
 			generalSection.SetValue("FilterIndex", FilterIndex);
 			generalSection.SetValues("RecentProjects", RecentProjects.Select(x => x.ToConfigEntry()).ToArray());
 			generalSection.SetValue("SyncTypeID", SyncTypeId);
+			generalSection.SetValue("NextLauncherVersionCheck", NextLauncherVersionCheck);
 
 			// Build configuration
 			generalSection.SetValue("BuildConfig", CompiledEditorBuildConfig.ToString());
