UGS: Missing copyright notice.

diff --git a/Engine/Source/Programs/UnrealGameSync/UnrealGameSync/Forms/LauncherUpdateWindow.cs b/Engine/Source/Programs/UnrealGameSync/UnrealGameSync/Forms/LauncherUpdateWindow.cs
--- a/Engine/Source/Programs/UnrealGameSync/UnrealGameSync/Forms/LauncherUpdateWindow.cs
+++ b/Engine/Source/Programs/UnrealGameSync/UnrealGameSync/Forms/LauncherUpdateWindow.cs
@@ -1,3 +1,5 @@
+// Copyright Epic Games, Inc. All Rights Reserved.
+
 using System;
 using System.Collections.Generic;
 using System.ComponentModel;
