UGS: Move Perforce settings into the global config file, so they can be accessed by the command-line client.

diff --git a/Engine/Source/Programs/UnrealGameSync/UnrealGameSync/Forms/ApplicationSettingsWindow.cs b/Engine/Source/Programs/UnrealGameSync/UnrealGameSync/Forms/ApplicationSettingsWindow.cs
--- a/Engine/Source/Programs/UnrealGameSync/UnrealGameSync/Forms/ApplicationSettingsWindow.cs
+++ b/Engine/Source/Programs/UnrealGameSync/UnrealGameSync/Forms/ApplicationSettingsWindow.cs
@@ -93,7 +93,7 @@ namespace UnrealGameSync
 			this.UserNameTextBox.Select(UserNameTextBox.TextLength, 0);
 			this.UserNameTextBox.CueBanner = $"Default ({defaultPerforceSettings.UserName})";
 
-			this.ParallelSyncThreadsSpinner.Value = Math.Max(Math.Min(settings.SyncOptions.NumThreads, ParallelSyncThreadsSpinner.Maximum), ParallelSyncThreadsSpinner.Minimum);
+			this.ParallelSyncThreadsSpinner.Value = Math.Max(Math.Min(settings.SyncOptions.NumThreads ?? PerforceSyncOptions.DefaultNumThreads, ParallelSyncThreadsSpinner.Maximum), ParallelSyncThreadsSpinner.Minimum);
 
 			this.DepotPathTextBox.Text = _initialDepotPath;
 			this.DepotPathTextBox.Select(DepotPathTextBox.TextLength, 0);
@@ -223,7 +223,8 @@ namespace UnrealGameSync
 
 			if (_settings.KeepInTray != KeepInTrayCheckBox.Checked || _settings.SyncOptions.NumThreads != ParallelSyncThreadsSpinner.Value)
 			{
-				_settings.SyncOptions.NumThreads = (int)ParallelSyncThreadsSpinner.Value;
+				int numThreads = (int)ParallelSyncThreadsSpinner.Value;
+				_settings.SyncOptions.NumThreads = (numThreads != PerforceSyncOptions.DefaultNumThreads)? (int?)numThreads : null;
 				_settings.KeepInTray = KeepInTrayCheckBox.Checked;
 				_settings.Save(_logger);
 			}
diff --git a/Engine/Source/Programs/UnrealGameSync/UnrealGameSync/Forms/PerforceSyncSettingsWindow.cs b/Engine/Source/Programs/UnrealGameSync/UnrealGameSync/Forms/PerforceSyncSettingsWindow.cs
--- a/Engine/Source/Programs/UnrealGameSync/UnrealGameSync/Forms/PerforceSyncSettingsWindow.cs
+++ b/Engine/Source/Programs/UnrealGameSync/UnrealGameSync/Forms/PerforceSyncSettingsWindow.cs
@@ -22,28 +22,27 @@ namespace UnrealGameSync
 		private void PerforceSettingsWindow_Load(object sender, EventArgs e)
 		{
 			PerforceSyncOptions syncOptions = _settings.SyncOptions;
-			numericUpDownNumRetries.Value = (syncOptions.NumRetries > 0) ? syncOptions.NumRetries : PerforceSyncOptions.DefaultNumRetries;
-			numericUpDownTcpBufferSize.Value = (syncOptions.TcpBufferSize > 0) ? syncOptions.TcpBufferSize / 1024 : PerforceSyncOptions.DefaultTcpBufferSize / 1024;
-			numericUpDownFileBufferSize.Value = (syncOptions.FileBufferSize > 0) ? syncOptions.FileBufferSize / 1024 : PerforceSyncOptions.DefaultFileBufferSize / 1024;
-			numericUpDownMaxCommandsPerBatch.Value = (syncOptions.MaxCommandsPerBatch > 0) ? syncOptions.MaxCommandsPerBatch : PerforceSyncOptions.DefaultMaxCommandsPerBatch;
-			numericUpDownMaxSizePerBatch.Value = (syncOptions.MaxSizePerBatch > 0) ? syncOptions.MaxSizePerBatch / 1024 / 1024 : PerforceSyncOptions.DefaultMaxSizePerBatch / 1024 / 1024;
-			numericUpDownRetriesOnSyncError.Value = (syncOptions.NumSyncErrorRetries > 0) ? syncOptions.NumSyncErrorRetries : PerforceSyncOptions.DefaultNumSyncErrorRetries;
+			numericUpDownMaxCommandsPerBatch.Value = syncOptions.MaxCommandsPerBatch ?? PerforceSyncOptions.DefaultMaxCommandsPerBatch;
+			numericUpDownMaxSizePerBatch.Value = (syncOptions.MaxSizePerBatch ?? PerforceSyncOptions.DefaultMaxSizePerBatch) / 1024 / 1024;
+			numericUpDownRetriesOnSyncError.Value = syncOptions.NumSyncErrorRetries ?? PerforceSyncOptions.DefaultNumSyncErrorRetries;
 		}
 
 		private void OkButton_Click(object sender, EventArgs e)
 		{
-			_settings.SyncOptions.NumRetries = (int)numericUpDownNumRetries.Value;
-			_settings.SyncOptions.TcpBufferSize = (int)numericUpDownTcpBufferSize.Value * 1024;
-			_settings.SyncOptions.FileBufferSize = (int)numericUpDownFileBufferSize.Value * 1024;
-			_settings.SyncOptions.MaxCommandsPerBatch = (int)numericUpDownMaxCommandsPerBatch.Value;
-			_settings.SyncOptions.MaxSizePerBatch = (int)numericUpDownMaxSizePerBatch.Value * 1024 * 1024;
-			_settings.SyncOptions.NumSyncErrorRetries = (int)numericUpDownRetriesOnSyncError.Value;
+			_settings.SyncOptions.MaxCommandsPerBatch = GetValueIfNotDefault((int)numericUpDownMaxCommandsPerBatch.Value, PerforceSyncOptions.DefaultMaxCommandsPerBatch);
+			_settings.SyncOptions.MaxSizePerBatch = GetValueIfNotDefault((int)numericUpDownMaxSizePerBatch.Value, PerforceSyncOptions.DefaultMaxSizePerBatch) * 1024 * 1024;
+			_settings.SyncOptions.NumSyncErrorRetries = GetValueIfNotDefault((int)numericUpDownRetriesOnSyncError.Value, PerforceSyncOptions.DefaultNumSyncErrorRetries);
 			_settings.Save(_logger);
 
 			DialogResult = System.Windows.Forms.DialogResult.OK;
 			Close();
 		}
 
+		private static int? GetValueIfNotDefault(int value, int defaultValue)
+		{
+			return (value == defaultValue) ? (int?)null : value;
+		}
+
 		private void CancButton_Click(object sender, EventArgs e)
 		{
 			DialogResult = System.Windows.Forms.DialogResult.Cancel;
@@ -53,9 +52,6 @@ namespace UnrealGameSync
 		private void ResetButton_Click(object sender, EventArgs e)
 		{
 			PerforceSyncOptions syncOptions = _settings.SyncOptions;
-			numericUpDownNumRetries.Value = PerforceSyncOptions.DefaultNumRetries;
-			numericUpDownTcpBufferSize.Value = PerforceSyncOptions.DefaultTcpBufferSize / 1024;
-			numericUpDownFileBufferSize.Value = PerforceSyncOptions.DefaultFileBufferSize / 1024;
 			numericUpDownMaxCommandsPerBatch.Value = PerforceSyncOptions.DefaultMaxCommandsPerBatch;
 			numericUpDownMaxSizePerBatch.Value = PerforceSyncOptions.DefaultMaxSizePerBatch / 1024 / 1024;
 			numericUpDownRetriesOnSyncError.Value = PerforceSyncOptions.DefaultNumSyncErrorRetries;
diff --git a/Engine/Source/Programs/UnrealGameSync/UnrealGameSyncCmd/Program.cs b/Engine/Source/Programs/UnrealGameSync/UnrealGameSyncCmd/Program.cs
--- a/Engine/Source/Programs/UnrealGameSync/UnrealGameSyncCmd/Program.cs
+++ b/Engine/Source/Programs/UnrealGameSync/UnrealGameSyncCmd/Program.cs
@@ -667,6 +667,8 @@ namespace UnrealGameSyncCmd
 				}
 
 				WorkspaceUpdateContext updateContext = new WorkspaceUpdateContext(change, options, BuildConfig.Development, syncFilter, projectSettings.BuildSteps, null);
+				updateContext.PerforceSyncOptions = context.UserSettings.Global.Perforce;
+
 				if (syncOptions.Binaries)
 				{
 					List<PerforceArchiveInfo> archives = await PerforceArchive.EnumerateAsync(perforceClient, projectConfig, state.Current.ProjectIdentifier, CancellationToken.None);
diff --git a/Engine/Source/Programs/UnrealGameSync/UnrealGameSyncShared/ConfigFile.cs b/Engine/Source/Programs/UnrealGameSync/UnrealGameSyncShared/ConfigFile.cs
--- a/Engine/Source/Programs/UnrealGameSync/UnrealGameSyncShared/ConfigFile.cs
+++ b/Engine/Source/Programs/UnrealGameSync/UnrealGameSyncShared/ConfigFile.cs
@@ -419,7 +419,9 @@ namespace UnrealGameSync
 			Pairs.Remove(key);
 		}
 
-		public int GetValue(string key, int defaultValue)
+		public int GetValue(string key, int defaultValue) => GetOptionalIntValue(key, defaultValue) ?? defaultValue;
+
+		public int? GetOptionalIntValue(string key, int? defaultValue)
 		{
 			string? valueString = GetValue(key);
 			if(valueString != null)
@@ -626,6 +628,15 @@ namespace UnrealGameSync
 			return section;
 		}
 
+		public void RemoveSection(string name)
+		{
+			ConfigSection? section = FindSection(name);
+			if (section != null)
+			{
+				_sections.Remove(section);
+			}
+		}
+
 		public void SetValue(string key, int value)
 		{
 			int dotIdx = key.IndexOf('.');
@@ -668,6 +679,13 @@ namespace UnrealGameSync
 			return (section == null)? defaultValue : section.GetValue(key.Substring(dotIdx + 1), defaultValue);
 		}
 
+		public int? GetOptionalIntValue(string key, int? defaultValue)
+		{
+			int dotIdx = key.IndexOf('.');
+			ConfigSection? section = FindSection(key.Substring(0, dotIdx));
+			return (section == null) ? defaultValue : section.GetOptionalIntValue(key.Substring(dotIdx + 1), defaultValue);
+		}
+
 		[return: NotNullIfNotNull("defaultValue")]
 		public string? GetValue(string key, string? defaultValue)
 		{
diff --git a/Engine/Source/Programs/UnrealGameSync/UnrealGameSyncShared/GlobalSettings.cs b/Engine/Source/Programs/UnrealGameSync/UnrealGameSyncShared/GlobalSettings.cs
--- a/Engine/Source/Programs/UnrealGameSync/UnrealGameSyncShared/GlobalSettings.cs
+++ b/Engine/Source/Programs/UnrealGameSync/UnrealGameSyncShared/GlobalSettings.cs
@@ -53,6 +53,7 @@ namespace UnrealGameSync
 
 	public class GlobalSettings
 	{
+		public PerforceSyncOptions Perforce { get; set; } = new PerforceSyncOptions();
 		public FilterSettings Filter { get; set; } = new FilterSettings();
 	}
 
diff --git a/Engine/Source/Programs/UnrealGameSync/UnrealGameSyncShared/UserSettings.cs b/Engine/Source/Programs/UnrealGameSync/UnrealGameSyncShared/UserSettings.cs
--- a/Engine/Source/Programs/UnrealGameSync/UnrealGameSyncShared/UserSettings.cs
+++ b/Engine/Source/Programs/UnrealGameSync/UnrealGameSyncShared/UserSettings.cs
@@ -532,7 +532,7 @@ namespace UnrealGameSync
 		Dictionary<FileReference, UserProjectSettings> _projectKeyToSettings = new Dictionary<FileReference, UserProjectSettings>();
 
 		// Perforce settings
-		public PerforceSyncOptions SyncOptions = new PerforceSyncOptions();
+		public PerforceSyncOptions SyncOptions => base.Global.Perforce;
 
 		private List<UserSelectedProjectSettings> ReadProjectList(string settingName, string legacySettingName)
 		{
@@ -771,25 +771,22 @@ namespace UnrealGameSync
 			NotifyUnresolvedMinutes = _configFile.GetValue("Notifications.NotifyUnresolvedMinutes", -1);
 
 			// Perforce settings
-			if(!int.TryParse(_configFile.GetValue("Perforce.NumRetries", "0"), out SyncOptions.NumRetries))
+			string? numThreadsStr = _configFile.GetValue("Perforce.NumThreads", null);
+			if (numThreadsStr != null)
 			{
-				SyncOptions.NumRetries = 0;
-			}
-
-			int numThreads;
-			if(int.TryParse(_configFile.GetValue("Perforce.NumThreads", "0"), out numThreads) && numThreads > 0)
-			{
-				if(Version >= UserSettingsVersion.DefaultNumberOfThreads || numThreads > 1)
+				int numThreads;
+				if (int.TryParse(numThreadsStr, out numThreads) && numThreads > 0)
 				{
-					SyncOptions.NumThreads = numThreads;
+					if (Version >= UserSettingsVersion.DefaultNumberOfThreads || numThreads > 1)
+					{
+						SyncOptions.NumThreads = numThreads;
+					}
 				}
 			}
 
-			SyncOptions.TcpBufferSize = _configFile.GetValue("Perforce.TcpBufferSize", PerforceSyncOptions.DefaultTcpBufferSize);
-			SyncOptions.FileBufferSize = _configFile.GetValue("Perforce.FileBufferSize", PerforceSyncOptions.DefaultFileBufferSize);
-			SyncOptions.MaxCommandsPerBatch = _configFile.GetValue("Perforce.MaxCommandsPerBatch", PerforceSyncOptions.DefaultMaxCommandsPerBatch);
-			SyncOptions.MaxSizePerBatch = _configFile.GetValue("Perforce.MaxSizePerBatch", PerforceSyncOptions.DefaultMaxSizePerBatch);
-			SyncOptions.NumSyncErrorRetries = _configFile.GetValue("Perforce.NumSyncErrorRetries", PerforceSyncOptions.DefaultNumSyncErrorRetries);
+			SyncOptions.MaxCommandsPerBatch = _configFile.GetOptionalIntValue("Perforce.MaxCommandsPerBatch", SyncOptions.MaxCommandsPerBatch);
+			SyncOptions.MaxSizePerBatch = _configFile.GetOptionalIntValue("Perforce.MaxSizePerBatch", SyncOptions.MaxSizePerBatch);
+			SyncOptions.NumSyncErrorRetries = _configFile.GetOptionalIntValue("Perforce.NumSyncErrorRetries", SyncOptions.NumSyncErrorRetries);
 		}
 
 		static Dictionary<Guid, bool> GetCategorySettings(ConfigSection? section, string includedKey, string excludedKey)
@@ -1112,36 +1109,7 @@ namespace UnrealGameSync
 			}
 
 			// Perforce settings
-			ConfigSection perforceSection = _configFile.FindOrAddSection("Perforce");
-			perforceSection.Clear();
-			if(SyncOptions.NumRetries > 0 && SyncOptions.NumRetries != PerforceSyncOptions.DefaultNumRetries)
-			{
-				perforceSection.SetValue("NumRetries", SyncOptions.NumRetries);
-			}
-			if(SyncOptions.NumThreads > 0 && SyncOptions.NumThreads != PerforceSyncOptions.DefaultNumThreads)
-			{
-				perforceSection.SetValue("NumThreads", SyncOptions.NumThreads);
-			}
-			if(SyncOptions.TcpBufferSize > 0 && SyncOptions.TcpBufferSize != PerforceSyncOptions.DefaultTcpBufferSize)
-			{
-				perforceSection.SetValue("TcpBufferSize", SyncOptions.TcpBufferSize);
-			}
-			if (SyncOptions.FileBufferSize > 0 && SyncOptions.FileBufferSize != PerforceSyncOptions.DefaultFileBufferSize)
-			{
-				perforceSection.SetValue("FileBufferSize", SyncOptions.FileBufferSize);
-			}
-			if (SyncOptions.MaxCommandsPerBatch > 0 && SyncOptions.MaxCommandsPerBatch != PerforceSyncOptions.DefaultMaxCommandsPerBatch)
-			{
-				perforceSection.SetValue("MaxCommandsPerBatch", SyncOptions.MaxCommandsPerBatch);
-			}
-			if (SyncOptions.MaxSizePerBatch > 0 && SyncOptions.MaxSizePerBatch != PerforceSyncOptions.DefaultMaxSizePerBatch)
-			{
-				perforceSection.SetValue("MaxSizePerBatch", SyncOptions.MaxSizePerBatch);
-			}
-			if (SyncOptions.NumSyncErrorRetries > 0 && SyncOptions.NumSyncErrorRetries != PerforceSyncOptions.DefaultNumSyncErrorRetries)
-			{
-				perforceSection.SetValue("NumSyncErrorRetries", SyncOptions.NumSyncErrorRetries);
-			}
+			_configFile.RemoveSection("Perforce");
 
 			// Save the file
 			try
diff --git a/Engine/Source/Programs/UnrealGameSync/UnrealGameSyncShared/Utility.cs b/Engine/Source/Programs/UnrealGameSync/UnrealGameSyncShared/Utility.cs
--- a/Engine/Source/Programs/UnrealGameSync/UnrealGameSyncShared/Utility.cs
+++ b/Engine/Source/Programs/UnrealGameSync/UnrealGameSyncShared/Utility.cs
@@ -125,7 +125,7 @@ namespace UnrealGameSync
 
 		public static byte[] SerializeJson<T>(T obj)
 		{
-			JsonSerializerOptions options = new JsonSerializerOptions { IgnoreNullValues = true, WriteIndented = true };
+			JsonSerializerOptions options = new JsonSerializerOptions { WriteIndented = true, DefaultIgnoreCondition = JsonIgnoreCondition.WhenWritingDefault };
 			options.Converters.Add(new JsonStringEnumConverter());
 
 			byte[] buffer;
diff --git a/Engine/Source/Programs/UnrealGameSync/UnrealGameSyncShared/WorkspaceUpdate.cs b/Engine/Source/Programs/UnrealGameSync/UnrealGameSyncShared/WorkspaceUpdate.cs
--- a/Engine/Source/Programs/UnrealGameSync/UnrealGameSyncShared/WorkspaceUpdate.cs
+++ b/Engine/Source/Programs/UnrealGameSync/UnrealGameSyncShared/WorkspaceUpdate.cs
@@ -70,16 +70,12 @@ namespace UnrealGameSync
 		public const int DefaultMaxSizePerBatch = 128 * 1024 * 1024;
 		public const int DefaultNumSyncErrorRetries = 0;
 
-		public int NumRetries = DefaultNumRetries;
-		public int NumThreads = DefaultNumThreads;
-		public int TcpBufferSize = DefaultTcpBufferSize;
-		public int FileBufferSize = DefaultFileBufferSize;
-
-		public int MaxCommandsPerBatch = DefaultMaxCommandsPerBatch;
-		public int MaxSizePerBatch = DefaultMaxSizePerBatch;
-
-		public int NumSyncErrorRetries = DefaultNumSyncErrorRetries;
+		public int? NumThreads { get; set; }
 
+		public int? MaxCommandsPerBatch { get; set; }
+		public int? MaxSizePerBatch { get; set; }
+		
+		public int? NumSyncErrorRetries { get; set; }
 
 		public PerforceSyncOptions Clone()
 		{
@@ -1678,7 +1674,7 @@ namespace UnrealGameSync
 				WorkspaceUpdateResult result = WorkspaceUpdateResult.FailedToSync;
 				string statusMessage = "";
 
-				int retries = (null != context.PerforceSyncOptions) ? context.PerforceSyncOptions!.NumSyncErrorRetries : 0;
+				int retries = context.PerforceSyncOptions!.NumSyncErrorRetries ?? PerforceSyncOptions.DefaultNumSyncErrorRetries;
 
 				while (retries >= 0 && WorkspaceUpdateResult.FailedToSync == result)
 				{
