UGS: Add .editorconfig file for static analysis warnings in UGS (+lots of fixes).

diff --git a/Engine/Source/Programs/UnrealGameSync/UnrealGameSync/ActivationListener.cs b/Engine/Source/Programs/UnrealGameSync/UnrealGameSync/ActivationListener.cs
--- a/Engine/Source/Programs/UnrealGameSync/UnrealGameSync/ActivationListener.cs
+++ b/Engine/Source/Programs/UnrealGameSync/UnrealGameSync/ActivationListener.cs
@@ -1,19 +1,15 @@
 // Copyright Epic Games, Inc. All Rights Reserved.
 
 using System;
-using System.Collections.Generic;
-using System.Linq;
-using System.Text;
 using System.Threading;
-using System.Threading.Tasks;
 
 namespace UnrealGameSync
 {
 	class ActivationListener : IDisposable
 	{
 		Thread? _backgroundThread;
-		EventWaitHandle _activateEventHandle;
-		EventWaitHandle _quitEventHandle;
+		readonly EventWaitHandle _activateEventHandle;
+		readonly EventWaitHandle _quitEventHandle;
 
 		public event Action? OnActivate;
 
diff --git a/Engine/Source/Programs/UnrealGameSync/UnrealGameSync/AsyncDisposer.cs b/Engine/Source/Programs/UnrealGameSync/UnrealGameSync/AsyncDisposer.cs
--- a/Engine/Source/Programs/UnrealGameSync/UnrealGameSync/AsyncDisposer.cs
+++ b/Engine/Source/Programs/UnrealGameSync/UnrealGameSync/AsyncDisposer.cs
@@ -3,8 +3,6 @@
 using Microsoft.Extensions.Logging;
 using System;
 using System.Collections.Generic;
-using System.Text;
-using System.Threading;
 using System.Threading.Tasks;
 
 namespace UnrealGameSync
@@ -16,13 +14,13 @@ namespace UnrealGameSync
 
 	internal class AsyncDisposer : IAsyncDisposer, IAsyncDisposable
 	{
-		object _lockObject = new object();
-		List<Task> _tasks = new List<Task>();
-		ILogger _logger;
+		readonly object _lockObject = new object();
+		readonly List<Task> _tasks = new List<Task>();
+		readonly ILogger _logger;
 
 		public AsyncDisposer(ILogger<AsyncDisposer> logger)
 		{
-			this._logger = logger;
+			_logger = logger;
 		}
 
 		public void Add(Task task)
diff --git a/Engine/Source/Programs/UnrealGameSync/UnrealGameSync/Automation/P4Automation.cs b/Engine/Source/Programs/UnrealGameSync/UnrealGameSync/Automation/P4Automation.cs
--- a/Engine/Source/Programs/UnrealGameSync/UnrealGameSync/Automation/P4Automation.cs
+++ b/Engine/Source/Programs/UnrealGameSync/UnrealGameSync/Automation/P4Automation.cs
@@ -1,6 +1,5 @@
 // Copyright Epic Games, Inc. All Rights Reserved.
 
-using EpicGames.Core;
 using EpicGames.Perforce;
 using Microsoft.Extensions.Logging;
 using System;
@@ -10,10 +9,8 @@ using System.Threading.Tasks;
 
 namespace UnrealGameSync
 {
-
 	static class P4Automation
 	{
-
 		public static IPerforceSettings GetConnectionSettings()
 		{
 			// Read the settings
@@ -25,7 +22,6 @@ namespace UnrealGameSync
 			GlobalPerforceSettings.ReadGlobalPerforceSettings(ref serverAndPort, ref userName, ref depotPathSettings, ref preview);
 
 			return Utility.OverridePerforceSettings(PerforceSettings.Default, serverAndPort, userName);
-
 		}
 
 		public static Task<string> PrintToTempFile(IPerforceConnection? connection, string depotPath, ILogger logger)
@@ -36,7 +32,6 @@ namespace UnrealGameSync
 		public static async Task<string> PrintToTempFileAsync(IPerforceConnection? connection, string depotPath, CancellationToken cancellationToken, ILogger logger)
 		{
 			bool createNewConnection = (connection == null);
-
 			try
 			{
 				if (connection == null)
@@ -57,7 +52,7 @@ namespace UnrealGameSync
 
 				string cl = depotFileName.Substring(index + 1);
 				string fileName = depotFileName.Substring(0, index);
-				string tempFileName = string.Format("{0}@{1}{2}", Path.GetFileNameWithoutExtension(fileName), cl, Path.GetExtension(fileName));
+				string tempFileName = String.Format("{0}@{1}{2}", Path.GetFileNameWithoutExtension(fileName), cl, Path.GetExtension(fileName));
 
 				tempFileName = Path.Combine(Path.GetTempPath(), tempFileName);
 				await connection.PrintAsync(tempFileName, depotFileName, cancellationToken);
@@ -70,11 +65,8 @@ namespace UnrealGameSync
 				if (createNewConnection)
 				{
 					connection?.Dispose();
-					connection = null;
 				}
 			}
 		}
-
 	}
-
 }
\ No newline at end of file
diff --git a/Engine/Source/Programs/UnrealGameSync/UnrealGameSync/Automation/VisualStudioAutomation.cs b/Engine/Source/Programs/UnrealGameSync/UnrealGameSync/Automation/VisualStudioAutomation.cs
--- a/Engine/Source/Programs/UnrealGameSync/UnrealGameSync/Automation/VisualStudioAutomation.cs
+++ b/Engine/Source/Programs/UnrealGameSync/UnrealGameSync/Automation/VisualStudioAutomation.cs
@@ -36,10 +36,10 @@ namespace UnrealGameSync
 			}
 
 			// unable to get DTE connection, so launch nesw VS instance
-			string arguments = string.Format("\"{0}\"", fileName);
+			string arguments = String.Format("\"{0}\"", fileName);
 			if (line != -1)
 			{
-				arguments += string.Format(" /command \"edit.goto {0}\"", line);
+				arguments += String.Format(" /command \"edit.goto {0}\"", line);
 			}
 
 			// Launch new visual studio instance
@@ -64,7 +64,7 @@ namespace UnrealGameSync
 
 			if (install == null || install.DevEnvPath == null)
 			{
-				errorMessage = string.Format("Unable to get Visual Studio installation");
+				errorMessage = String.Format("Unable to get Visual Studio installation");
 				return false;
 			}
 			try
@@ -116,7 +116,7 @@ namespace UnrealGameSync
 					try
 					{
 						currentMoniker.GetDisplayName(bindContext, null, out outDisplayName);
-						if (string.IsNullOrEmpty(outDisplayName) || !IsVisualStudioDteMoniker(outDisplayName))
+						if (String.IsNullOrEmpty(outDisplayName) || !IsVisualStudioDteMoniker(outDisplayName))
 						{
 							continue;
 						}
@@ -143,9 +143,9 @@ namespace UnrealGameSync
 
 		static bool IsVisualStudioDteMoniker(string inName)
 		{
-			VisualStudioInstallation[] installs = VisualStudioInstallations.Installs;
+			IReadOnlyList<VisualStudioInstallation> installs = VisualStudioInstallations.Installs;
 
-			for (int idx = 0; idx < installs.Length; idx++)
+			for (int idx = 0; idx < installs.Count; idx++)
 			{
 				string? moniker = installs[idx].RotMoniker;
 				if (moniker != null && inName.StartsWith(moniker))
@@ -162,13 +162,11 @@ namespace UnrealGameSync
 			return result >= 0;
 		}
 
-
 		[DllImport("ole32.dll")]
 		public static extern int CreateBindCtx(int reserved, out IBindCtx bindCtx);
 
 		[DllImport("ole32.dll")]
 		public static extern int GetRunningObjectTable(int reserved, out IRunningObjectTable rot);
-
 	}
 
 	class VisualStudioInstallation
@@ -176,24 +174,22 @@ namespace UnrealGameSync
 		/// <summary>
 		/// Base directory for the installation
 		/// </summary>
-		public string? BaseDir;
-
+		public string? BaseDir { get; set; }
 
 		/// <summary>
 		/// Path of the devenv executable
 		/// </summary>
-		public string? DevEnvPath;
+		public string? DevEnvPath { get; set; }
 
 		/// <summary>
 		/// Visual Studio major version number
 		/// </summary>
-		public int MajorVersion;
+		public int MajorVersion { get; set; }
 
 		/// <summary>
 		/// Running Object Table moniker for this installation
 		/// </summary>
-		public string? RotMoniker;
-
+		public string? RotMoniker { get; set; }
 	}
 
 	/// <summary>
@@ -205,26 +201,23 @@ namespace UnrealGameSync
 		public static VisualStudioInstallation? GetPreferredInstallation(int majorVersion = 0)
 		{
 
-			if (CachedInstalls.Count == 0)
+			if (s_cachedInstalls.Count == 0)
 			{
 				return null;
 			}
 
 			if (majorVersion == 0)
 			{
-				return CachedInstalls.First();
+				return s_cachedInstalls.First();
 			}
 
-			VisualStudioInstallation? installation = CachedInstalls.FirstOrDefault(install => { return install.MajorVersion == majorVersion; });
+			VisualStudioInstallation? installation = s_cachedInstalls.FirstOrDefault(install => { return install.MajorVersion == majorVersion; });
 
 			return installation;
 
 		}
 
-		public static VisualStudioInstallation[] Installs
-		{
-			get { return CachedInstalls.ToArray(); }
-		}
+		public static IReadOnlyList<VisualStudioInstallation> Installs => s_cachedInstalls;
 
 		public static void Refresh()
 		{
@@ -233,7 +226,7 @@ namespace UnrealGameSync
 
 		static List<VisualStudioInstallation> GetVisualStudioInstallations()
 		{
-			CachedInstalls.Clear();
+			s_cachedInstalls.Clear();
 
 			try
 			{
@@ -266,32 +259,31 @@ namespace UnrealGameSync
 						string installationPath = instance.GetInstallationPath();
 						string devEnvPath = Path.Combine(installationPath, "Common7\\IDE\\devenv.exe");
 
-						if (!int.TryParse(components[0], out majorVersion) || (majorVersion != 15 && majorVersion != 16))
+						if (!Int32.TryParse(components[0], out majorVersion) || (majorVersion != 15 && majorVersion != 16))
 						{
 							continue;
 						}
 
-
 						if (!File.Exists(devEnvPath))
 						{
 							continue;
 						}
 
-						VisualStudioInstallation installation = new VisualStudioInstallation() { BaseDir = installationPath, DevEnvPath = devEnvPath, MajorVersion = majorVersion, RotMoniker = string.Format("!VisualStudio.DTE.{0}.0", majorVersion) };
+						VisualStudioInstallation installation = new VisualStudioInstallation() { BaseDir = installationPath, DevEnvPath = devEnvPath, MajorVersion = majorVersion, RotMoniker = String.Format("!VisualStudio.DTE.{0}.0", majorVersion) };
 
-						CachedInstalls.Add(installation);
+						s_cachedInstalls.Add(installation);
 					}
 				}
 			}
 			catch (Exception ex)
 			{
-				MessageBox.Show(string.Format("Exception while finding Visual Studio installations {0}", ex.Message));
+				MessageBox.Show(String.Format("Exception while finding Visual Studio installations {0}", ex.Message));
 			}
 
 			// prefer newer versions
-			CachedInstalls.Sort((a, b) => { return -a.MajorVersion.CompareTo(b.MajorVersion); });
+			s_cachedInstalls.Sort((a, b) => { return -a.MajorVersion.CompareTo(b.MajorVersion); });
 
-			return CachedInstalls;
+			return s_cachedInstalls;
 		}
 
 		static VisualStudioInstallations()
@@ -299,8 +291,6 @@ namespace UnrealGameSync
 			GetVisualStudioInstallations();
 		}
 
-		static readonly List<VisualStudioInstallation> CachedInstalls = new List<VisualStudioInstallation>();
-
+		static readonly List<VisualStudioInstallation> s_cachedInstalls = new List<VisualStudioInstallation>();
 	}
-
 }
\ No newline at end of file
diff --git a/Engine/Source/Programs/UnrealGameSync/UnrealGameSync/AutomationServer.cs b/Engine/Source/Programs/UnrealGameSync/UnrealGameSync/AutomationServer.cs
--- a/Engine/Source/Programs/UnrealGameSync/UnrealGameSync/AutomationServer.cs
+++ b/Engine/Source/Programs/UnrealGameSync/UnrealGameSync/AutomationServer.cs
@@ -31,8 +31,8 @@ namespace UnrealGameSync
 
 		public AutomationRequestInput(AutomationRequestType type, byte[] data)
 		{
-			this.Type = type;
-			this.Data = data;
+			Type = type;
+			Data = data;
 		}
 
 		public static AutomationRequestInput Read(Stream inputStream)
@@ -73,14 +73,14 @@ namespace UnrealGameSync
 
 		public AutomationRequestOutput(AutomationRequestResult result)
 		{
-			this.Result = result;
-			this.Data = new byte[0];
+			Result = result;
+			Data = new byte[0];
 		}
 
 		public AutomationRequestOutput(AutomationRequestResult result, byte[] data)
 		{
-			this.Result = result;
-			this.Data = data;
+			Result = result;
+			Data = data;
 		}
 
 		public static AutomationRequestOutput Read(Stream inputStream)
@@ -113,13 +113,13 @@ namespace UnrealGameSync
 
 		public AutomationRequest(AutomationRequestInput input)
 		{
-			this.Input = input;
-			this.Complete = new ManualResetEventSlim(false);
+			Input = input;
+			Complete = new ManualResetEventSlim(false);
 		}
 
 		public void SetOutput(AutomationRequestOutput output)
 		{
-			this.Output = output;
+			Output = output;
 			Complete?.Set();
 		}
 
@@ -133,22 +133,22 @@ namespace UnrealGameSync
 	{
 		static readonly UnicodeEncoding _streamEncoding = new UnicodeEncoding();
 
-		CancellationTokenSource _cancellationSource = new CancellationTokenSource();
+		readonly CancellationTokenSource _cancellationSource = new CancellationTokenSource();
 
 		const string IpcChannel = @"\.\pipe\UGSChannel";
-		ConfiguredTaskAwaitable _ipcTask;
+		readonly ConfiguredTaskAwaitable _ipcTask;
 
 		public const int DefaultPortNumber = 30422;
-		ConfiguredTaskAwaitable? _tcpTask;
+		readonly ConfiguredTaskAwaitable? _tcpTask;
 
-		Action<AutomationRequest> _postRequest;
+		readonly Action<AutomationRequest> _postRequest;
 
-		ILogger _logger;
+		readonly ILogger _logger;
 
 		public AutomationServer(Action<AutomationRequest> postRequest, string? uri, ILogger<AutomationServer> logger)
 		{
-			this._postRequest = postRequest;
-			this._logger = logger;
+			_postRequest = postRequest;
+			_logger = logger;
 
 			try
 			{
@@ -190,9 +190,9 @@ namespace UnrealGameSync
 		public static int GetPortNumber()
 		{
 			object? portValue = Registry.GetValue("HKEY_CURRENT_USER\\Software\\Epic Games\\UnrealGameSync", "AutomationPort", null);
-			if (portValue != null && portValue is int)
+			if (portValue != null && portValue is int portValueInt)
 			{
-				return (int)portValue;
+				return portValueInt;
 			}
 			else
 			{
@@ -203,7 +203,7 @@ namespace UnrealGameSync
 		async Task RunIpcAsync(string? commandLineUri, CancellationToken cancellationToken)
 		{
 			// Handle main process command line URI request
-			if (!string.IsNullOrEmpty(commandLineUri))
+			if (!String.IsNullOrEmpty(commandLineUri))
 			{
 				HandleUri(commandLineUri);
 			}
@@ -300,7 +300,7 @@ namespace UnrealGameSync
 				UriResult result = UriHandler.HandleUri(uri);
 				if (!result.Success)
 				{
-					if (!string.IsNullOrEmpty(result.Error))
+					if (!String.IsNullOrEmpty(result.Error))
 					{
 						MessageBox.Show(String.Format("Error handling uri: {0}", result.Error));
 					}
@@ -353,9 +353,9 @@ namespace UnrealGameSync
 
 			int len = outBuffer.Length;
 
-			if (len > ushort.MaxValue)
+			if (len > UInt16.MaxValue)
 			{
-				len = ushort.MaxValue;
+				len = UInt16.MaxValue;
 			}
 
 			stream.WriteByte((byte)(len / 256));
diff --git a/Engine/Source/Programs/UnrealGameSync/UnrealGameSync/Controls/AlertButtonControl.cs b/Engine/Source/Programs/UnrealGameSync/UnrealGameSync/Controls/AlertButtonControl.cs
--- a/Engine/Source/Programs/UnrealGameSync/UnrealGameSync/Controls/AlertButtonControl.cs
+++ b/Engine/Source/Programs/UnrealGameSync/UnrealGameSync/Controls/AlertButtonControl.cs
@@ -1,16 +1,10 @@
 // Copyright Epic Games, Inc. All Rights Reserved.
 
 using System;
-using System.Collections.Generic;
 using System.ComponentModel;
 using System.Drawing;
 using System.Drawing.Drawing2D;
-using System.Data;
-using System.Linq;
-using System.Text;
-using System.Threading.Tasks;
 using System.Windows.Forms;
-using System.Windows.Forms.VisualStyles;
 
 namespace UnrealGameSync
 {
@@ -77,14 +71,22 @@ namespace UnrealGameSync
 
 		public AlertButtonTheme Theme
 		{
-			get { return _themeValue; }
-			set { _themeValue = value; UpdateThemeColors(); }
+			get => _themeValue;
+			set 
+			{ 
+				_themeValue = value; 
+				UpdateThemeColors(); 
+			}
 		}
 
 		public AlertButtonColors CustomColors
 		{
-			get { return _customColorsValue; }
-			set { _customColorsValue = value; UpdateThemeColors(); }
+			get => _customColorsValue;
+			set 
+			{ 
+				_customColorsValue = value; 
+				UpdateThemeColors();
+			}
 		}
 
 		public AlertButtonControl()
diff --git a/Engine/Source/Programs/UnrealGameSync/UnrealGameSync/Controls/AlertImageControl.cs b/Engine/Source/Programs/UnrealGameSync/UnrealGameSync/Controls/AlertImageControl.cs
--- a/Engine/Source/Programs/UnrealGameSync/UnrealGameSync/Controls/AlertImageControl.cs
+++ b/Engine/Source/Programs/UnrealGameSync/UnrealGameSync/Controls/AlertImageControl.cs
@@ -1,16 +1,9 @@
 // Copyright Epic Games, Inc. All Rights Reserved.
 
 using System;
-using System.Collections.Generic;
-using System.ComponentModel;
 using System.Drawing;
 using System.Drawing.Drawing2D;
-using System.Data;
-using System.Linq;
-using System.Text;
-using System.Threading.Tasks;
 using System.Windows.Forms;
-using System.Windows.Forms.VisualStyles;
 
 namespace UnrealGameSync
 {
@@ -20,8 +13,12 @@ namespace UnrealGameSync
 
 		public Image? Image
 		{
-			get { return _imageValue; }
-			set { _imageValue = value; Invalidate(); }
+			get => _imageValue;
+			set 
+			{ 
+				_imageValue = value; 
+				Invalidate(); 
+			}
 		}
 
 		public AlertImageControl()
diff --git a/Engine/Source/Programs/UnrealGameSync/UnrealGameSync/Controls/BuildListControl.cs b/Engine/Source/Programs/UnrealGameSync/UnrealGameSync/Controls/BuildListControl.cs
--- a/Engine/Source/Programs/UnrealGameSync/UnrealGameSync/Controls/BuildListControl.cs
+++ b/Engine/Source/Programs/UnrealGameSync/UnrealGameSync/Controls/BuildListControl.cs
@@ -1,13 +1,9 @@
 // Copyright Epic Games, Inc. All Rights Reserved.
 
 using System;
-using System.Collections.Generic;
 using System.ComponentModel;
 using System.Drawing;
-using System.Linq;
 using System.Runtime.InteropServices;
-using System.Text;
-using System.Threading.Tasks;
 using System.Windows.Forms;
 
 namespace UnrealGameSync
@@ -29,13 +25,13 @@ namespace UnrealGameSync
 		const int WmVscroll = 0x115;
 		const int WmMousewheel = 0x020A;
 
-		const int SbHorz = 0;
+		//const int SbHorz = 0;
 		const int SbVert = 1;
 
 		const int SifRange = 0x0001;
 		const int SifPage = 0x0002;
 		const int SifPos = 0x0004;
-		const int SifDisablenoscroll = 0x0008;
+		//const int SifDisablenoscroll = 0x0008;
 		const int SifTrackpos = 0x0010;
 		const int SifAll = (SifRange | SifPage | SifPos | SifTrackpos);        
 
@@ -49,7 +45,7 @@ namespace UnrealGameSync
 		private static extern int SetScrollInfo(IntPtr hwnd, int fnBar, Scrollinfo lpsi, int redraw);
 
 		[DllImport("user32.dll", CharSet = CharSet.Auto)]
-		private static extern IntPtr SendMessage(IntPtr hWnd, UInt32 msg, IntPtr wParam, IntPtr lParam);
+		private static extern IntPtr SendMessage(IntPtr hWnd, uint msg, IntPtr wParam, IntPtr lParam);
 
 		public delegate void OnScrollDelegate();
 
@@ -148,8 +144,8 @@ namespace UnrealGameSync
 		[DesignerSerializationVisibility(DesignerSerializationVisibility.Hidden)]
 		public override Font Font
 		{
-			get { return base.Font; }
-			set { base.Font = value; }
+			get => base.Font;
+			set => base.Font = value;
 		}
 	}
 }
diff --git a/Engine/Source/Programs/UnrealGameSync/UnrealGameSync/Controls/CustomListViewControl.cs b/Engine/Source/Programs/UnrealGameSync/UnrealGameSync/Controls/CustomListViewControl.cs
--- a/Engine/Source/Programs/UnrealGameSync/UnrealGameSync/Controls/CustomListViewControl.cs
+++ b/Engine/Source/Programs/UnrealGameSync/UnrealGameSync/Controls/CustomListViewControl.cs
@@ -1,13 +1,7 @@
 // Copyright Epic Games, Inc. All Rights Reserved.
 
 using System;
-using System.Collections.Generic;
-using System.ComponentModel;
 using System.Drawing;
-using System.Data;
-using System.Linq;
-using System.Text;
-using System.Threading.Tasks;
 using System.Windows.Forms;
 using System.Windows.Forms.VisualStyles;
 
@@ -31,8 +25,8 @@ namespace UnrealGameSync
 
 		public CustomListViewWidget(ListViewItem item)
 		{
-			this.Item = item;
-			this.Cursor = Cursors.Arrow;
+			Item = item;
+			Cursor = Cursors.Arrow;
 		}
 
 		public void Invalidate()
@@ -79,8 +73,8 @@ namespace UnrealGameSync
 
 	class StatusLineListViewWidget : CustomListViewWidget
 	{
-		public StatusElementResources Resources;
-		public StatusLine Line = new StatusLine();
+		public StatusElementResources Resources { get; }
+		public StatusLine Line { get; } = new StatusLine();
 
 		public HorizontalAlignment HorizontalAlignment
 		{
@@ -94,8 +88,8 @@ namespace UnrealGameSync
 		public StatusLineListViewWidget(ListViewItem item, StatusElementResources resources)
 			: base(item)
 		{
-			this.Resources = resources;
-			this.HorizontalAlignment = HorizontalAlignment.Center;
+			Resources = resources;
+			HorizontalAlignment = HorizontalAlignment.Center;
 		}
 
 		public override bool RequiresLayout()
@@ -171,7 +165,7 @@ namespace UnrealGameSync
 
 		public override void Layout(Graphics graphics, Rectangle bounds)
 		{
-			this._bounds = bounds;
+			_bounds = bounds;
 
 			int offsetY = bounds.Y + bounds.Height / 2;
 			Line.Layout(graphics, new Point(bounds.X, offsetY), Resources);
@@ -198,8 +192,9 @@ namespace UnrealGameSync
 
 	partial class CustomListViewControl : ListView
 	{
-		VisualStyleRenderer? _selectedItemRenderer;
-		VisualStyleRenderer? _trackedItemRenderer;
+		readonly VisualStyleRenderer? _selectedItemRenderer;
+		readonly VisualStyleRenderer? _trackedItemRenderer;
+		
 		public int HoverItem = -1;
 
 		CustomListViewWidget? _mouseOverWidget;
@@ -336,7 +331,7 @@ namespace UnrealGameSync
 
 		protected override void OnMouseMove(MouseEventArgs e)
 		{
-			ListViewHitTestInfo hitTest = this.HitTest(e.Location);
+			ListViewHitTestInfo hitTest = HitTest(e.Location);
 
 			CustomListViewWidget? newMouseOverWidget = FindWidget(hitTest);
 			if(_mouseOverWidget != null && _mouseOverWidget != newMouseOverWidget)
diff --git a/Engine/Source/Programs/UnrealGameSync/UnrealGameSync/Controls/ErrorPanel.cs b/Engine/Source/Programs/UnrealGameSync/UnrealGameSync/Controls/ErrorPanel.cs
--- a/Engine/Source/Programs/UnrealGameSync/UnrealGameSync/Controls/ErrorPanel.cs
+++ b/Engine/Source/Programs/UnrealGameSync/UnrealGameSync/Controls/ErrorPanel.cs
@@ -1,11 +1,7 @@
 ﻿// Copyright Epic Games, Inc. All Rights Reserved.
 
 using System;
-using System.Collections.Generic;
 using System.Drawing;
-using System.Linq;
-using System.Text;
-using System.Threading.Tasks;
 
 namespace UnrealGameSync
 {
@@ -58,15 +54,9 @@ namespace UnrealGameSync
 		{
 		}
 
-		public Color? TintColor
-		{
-			get { return null; }
-		}
+		public Color? TintColor => null;
 
-		public Tuple<TaskbarState, float> DesiredTaskbarState
-		{
-			get { return Tuple.Create(TaskbarState.Normal, 0.0f); }
-		}
+		public Tuple<TaskbarState, float> DesiredTaskbarState => Tuple.Create(TaskbarState.Normal, 0.0f);
 
 		public UserSelectedProjectSettings SelectedProject
 		{
diff --git a/Engine/Source/Programs/UnrealGameSync/UnrealGameSync/Controls/LogControl.cs b/Engine/Source/Programs/UnrealGameSync/UnrealGameSync/Controls/LogControl.cs
--- a/Engine/Source/Programs/UnrealGameSync/UnrealGameSync/Controls/LogControl.cs
+++ b/Engine/Source/Programs/UnrealGameSync/UnrealGameSync/Controls/LogControl.cs
@@ -2,12 +2,10 @@
 
 using System;
 using System.Collections.Generic;
-using System.ComponentModel;
 using System.Drawing;
 using System.Data;
 using System.Linq;
 using System.Text;
-using System.Threading.Tasks;
 using System.Windows.Forms;
 using System.Runtime.InteropServices;
 using System.Collections.Concurrent;
@@ -66,7 +64,7 @@ namespace UnrealGameSync
 		static extern int GetScrollInfo(IntPtr hwnd, ScrollBarType fnBar, Scrollinfo lpsi);
 
 		[DllImport("user32.dll", CharSet = CharSet.Auto)]
-		static extern IntPtr SendMessage(IntPtr hWnd, UInt32 msg, IntPtr wParam, IntPtr lParam);
+		static extern IntPtr SendMessage(IntPtr hWnd, uint msg, IntPtr wParam, IntPtr lParam);
 
 		struct TextLocation
 		{
@@ -107,7 +105,7 @@ namespace UnrealGameSync
 			}
 		}
 
-		List<string> _lines = new List<string>();
+		readonly List<string> _lines = new List<string>();
 		int _maxLineLength;
 
 		int _scrollLine;
@@ -288,7 +286,7 @@ namespace UnrealGameSync
 						}
 						catch(Exception ex)
 						{
-							textToAppend += String.Format("Failed to write to log file ({0}): {1}\n", _logFileStream.Name, ex.ToString());
+							newLines.Add(String.Format("Failed to write to log file ({0}): {1}\n", _logFileStream.Name, ex.ToString()));
 						}
 					}
 
@@ -348,15 +346,6 @@ namespace UnrealGameSync
 			Invalidate();
 		}
 
-		private Scrollinfo GetScrollInfo(ScrollBarType type, ScrollInfoMask mask)
-		{
-			Scrollinfo scrollInfo = new Scrollinfo();
-			scrollInfo.cbSize = Marshal.SizeOf(VerticalScroll);
-			scrollInfo.fMask = mask;
-			GetScrollInfo(Handle, ScrollBarType.SbVert, scrollInfo);
-			return scrollInfo;
-		}
-
 		private bool IsTrackingLastLine()
 		{
 			Scrollinfo verticalScroll = new Scrollinfo();
@@ -495,7 +484,7 @@ namespace UnrealGameSync
 		{
 			using(Graphics graphics = CreateGraphics())
 			{
-				_fontSize = TextRenderer.MeasureText(graphics, "A", Font, new Size(int.MaxValue, int.MaxValue), TextFormatFlags.NoPadding);
+				_fontSize = TextRenderer.MeasureText(graphics, "A", Font, new Size(Int32.MaxValue, Int32.MaxValue), TextFormatFlags.NoPadding);
 			}
 		}
 
@@ -924,7 +913,7 @@ namespace UnrealGameSync
 			public void Dispose() { }
 		}
 
-		LogControl _logControl;
+		readonly LogControl _logControl;
 
 		public LogControlTextWriter(LogControl inLogControl)
 		{
diff --git a/Engine/Source/Programs/UnrealGameSync/UnrealGameSync/Controls/LogSplitContainer.cs b/Engine/Source/Programs/UnrealGameSync/UnrealGameSync/Controls/LogSplitContainer.cs
--- a/Engine/Source/Programs/UnrealGameSync/UnrealGameSync/Controls/LogSplitContainer.cs
+++ b/Engine/Source/Programs/UnrealGameSync/UnrealGameSync/Controls/LogSplitContainer.cs
@@ -1,16 +1,9 @@
 // Copyright Epic Games, Inc. All Rights Reserved.
 
 using System;
-using System.Collections.Generic;
 using System.ComponentModel;
 using System.Drawing;
-using System.Data;
-using System.Linq;
-using System.Text;
-using System.Threading.Tasks;
 using System.Windows.Forms;
-using System.Runtime.InteropServices;
-using System.Collections.Concurrent;
 
 namespace UnrealGameSync
 {
@@ -47,10 +40,7 @@ namespace UnrealGameSync
 			SplitterMoved += OnSplitterMoved;
 		}
 
-		protected Font CaptionFont
-		{
-			get { return SystemFonts.IconTitleFont!; }
-		}
+		protected Font CaptionFont => SystemFonts.IconTitleFont!;
 
 		protected override void OnCreateControl()
 		{
@@ -215,10 +205,7 @@ namespace UnrealGameSync
 			}
 		}
 
-		protected override bool ShowFocusCues
-		{
-			get { return false; }
-		}
+		protected override bool ShowFocusCues => false;
 
 		protected override void OnResize(EventArgs e)
 		{
@@ -301,9 +288,6 @@ namespace UnrealGameSync
 			}
 		}
 
-		Rectangle CloseButtonRectangle
-		{
-			get { return new Rectangle(ClientRectangle.Right - 20, SplitterDistance, 20, SplitterWidth); }
-		}
+		Rectangle CloseButtonRectangle => new Rectangle(ClientRectangle.Right - 20, SplitterDistance, 20, SplitterWidth);
 	}
 }
diff --git a/Engine/Source/Programs/UnrealGameSync/UnrealGameSync/Controls/OIDCControl.cs b/Engine/Source/Programs/UnrealGameSync/UnrealGameSync/Controls/OIDCControl.cs
--- a/Engine/Source/Programs/UnrealGameSync/UnrealGameSync/Controls/OIDCControl.cs
+++ b/Engine/Source/Programs/UnrealGameSync/UnrealGameSync/Controls/OIDCControl.cs
@@ -10,7 +10,7 @@ namespace UnrealGameSync.Controls
 {
 	public partial class OidcControl : UserControl
 	{
-		private Font _badgeFont;
+		private readonly Font _badgeFont;
 		private readonly OidcTokenManager _oidcManager;
 		private readonly string _providerIdentifier;
 
@@ -22,7 +22,7 @@ namespace UnrealGameSync.Controls
 			_oidcManager = inOidcManager;
 			_providerIdentifier = inProviderIdentifier;
 
-			_badgeFont = new Font(this.Font.FontFamily, this.Font.SizeInPoints - 2, FontStyle.Bold);
+			_badgeFont = new Font(Font.FontFamily, Font.SizeInPoints - 2, FontStyle.Bold);
 
 			OidcStatus serviceStatus = _oidcManager.GetStatusForProvider(_providerIdentifier);
 			LoginButton.Enabled = serviceStatus == OidcStatus.NotLoggedIn;
diff --git a/Engine/Source/Programs/UnrealGameSync/UnrealGameSync/Controls/RichTextBox41.cs b/Engine/Source/Programs/UnrealGameSync/UnrealGameSync/Controls/RichTextBox41.cs
--- a/Engine/Source/Programs/UnrealGameSync/UnrealGameSync/Controls/RichTextBox41.cs
+++ b/Engine/Source/Programs/UnrealGameSync/UnrealGameSync/Controls/RichTextBox41.cs
@@ -1,11 +1,7 @@
 // Copyright Epic Games, Inc. All Rights Reserved.
 
 using System;
-using System.Collections.Generic;
-using System.Linq;
 using System.Runtime.InteropServices;
-using System.Text;
-using System.Threading.Tasks;
 using System.Windows.Forms;
 
 namespace UnrealGameSync
diff --git a/Engine/Source/Programs/UnrealGameSync/UnrealGameSync/Controls/StatusPanel.cs b/Engine/Source/Programs/UnrealGameSync/UnrealGameSync/Controls/StatusPanel.cs
--- a/Engine/Source/Programs/UnrealGameSync/UnrealGameSync/Controls/StatusPanel.cs
+++ b/Engine/Source/Programs/UnrealGameSync/UnrealGameSync/Controls/StatusPanel.cs
@@ -2,14 +2,11 @@
 
 using System;
 using System.Collections.Generic;
-using System.ComponentModel;
 using System.Diagnostics.CodeAnalysis;
 using System.Drawing;
 using System.Drawing.Drawing2D;
 using System.Linq;
 using System.Runtime.InteropServices;
-using System.Text;
-using System.Threading.Tasks;
 using System.Windows.Forms;
 
 namespace UnrealGameSync
@@ -27,7 +24,7 @@ namespace UnrealGameSync
 
 	class StatusElementResources
 	{
-		Dictionary<FontStyle, Font> _fontCache = new Dictionary<FontStyle, Font>();
+		readonly Dictionary<FontStyle, Font> _fontCache = new Dictionary<FontStyle, Font>();
 		public readonly Font BadgeFont;
 
 		public StatusElementResources(Font baseFont)
@@ -67,10 +64,10 @@ namespace UnrealGameSync
 
 	abstract class StatusElement
 	{
-		public Cursor Cursor = Cursors.Arrow;
-		public bool MouseOver;
-		public bool MouseDown;
-		public Rectangle Bounds;
+		public Cursor Cursor { get; set; } = Cursors.Arrow;
+		public bool MouseOver { get; set; }
+		public bool MouseDown { get; set; }
+		public Rectangle Bounds { get; set; }
 
 		public Point Layout(Graphics graphics, Point location, StatusElementResources resources)
 		{
@@ -89,7 +86,7 @@ namespace UnrealGameSync
 
 	class IconStatusElement : StatusElement
 	{
-		Image _icon;
+		readonly Image _icon;
 
 		public IconStatusElement(Image inIcon)
 		{
@@ -109,15 +106,15 @@ namespace UnrealGameSync
 
 	class IconStripStatusElement : StatusElement
 	{
-		Image _strip;
-		Size _iconSize;
-		int _index;
+		readonly Image _strip;
+		readonly Size _iconSize;
+		readonly int _index;
 
-		public IconStripStatusElement(Image inStrip, Size inIconSize, int inIndex)
+		public IconStripStatusElement(Image strip, Size iconSize, int index)
 		{
-			_strip = inStrip;
-			_iconSize = inIconSize;
-			_index = inIndex;
+			_strip = strip;
+			_iconSize = iconSize;
+			_index = index;
 		}
 
 		public override Size Measure(Graphics graphics, StatusElementResources resources)
@@ -133,9 +130,9 @@ namespace UnrealGameSync
 
 	class TextStatusElement : StatusElement
 	{
-		string _text;
-		Color _color;
-		FontStyle _style;
+		readonly string _text;
+		readonly Color _color;
+		readonly FontStyle _style;
 
 		public TextStatusElement(string inText, Color inColor, FontStyle inStyle)
 		{
@@ -146,7 +143,7 @@ namespace UnrealGameSync
 
 		public override Size Measure(Graphics graphics, StatusElementResources resources)
 		{
-			return TextRenderer.MeasureText(graphics, _text, resources.FindOrAddFont(_style), new Size(int.MaxValue, int.MaxValue), TextFormatFlags.NoPadding);
+			return TextRenderer.MeasureText(graphics, _text, resources.FindOrAddFont(_style), new Size(Int32.MaxValue, Int32.MaxValue), TextFormatFlags.NoPadding);
 		}
 
 		public override void Draw(Graphics graphics, StatusElementResources resources)
@@ -157,15 +154,15 @@ namespace UnrealGameSync
 
 	class LinkStatusElement : StatusElement
 	{
-		string _text;
-		FontStyle _style;
-		Action<Point, Rectangle> _linkAction;
+		readonly string _text;
+		readonly FontStyle _style;
+		readonly Action<Point, Rectangle> _linkAction;
 
-		public LinkStatusElement(string inText, FontStyle inStyle, Action<Point, Rectangle> inLinkAction)
+		public LinkStatusElement(string text, FontStyle style, Action<Point, Rectangle> linkAction)
 		{
-			_text = inText;
-			_style = inStyle;
-			_linkAction = inLinkAction;
+			_text = text;
+			_style = style;
+			_linkAction = linkAction;
 			Cursor = NativeCursors.Hand;
 		}
 
@@ -176,7 +173,7 @@ namespace UnrealGameSync
 
 		public override Size Measure(Graphics graphics, StatusElementResources resources)
 		{
-			return TextRenderer.MeasureText(graphics, _text, resources.FindOrAddFont(_style), new Size(int.MaxValue, int.MaxValue), TextFormatFlags.NoPadding | TextFormatFlags.NoPrefix);
+			return TextRenderer.MeasureText(graphics, _text, resources.FindOrAddFont(_style), new Size(Int32.MaxValue, Int32.MaxValue), TextFormatFlags.NoPadding | TextFormatFlags.NoPrefix);
 		}
 
 		public override void Draw(Graphics graphics, StatusElementResources resources)
@@ -196,10 +193,10 @@ namespace UnrealGameSync
 
 	class BadgeStatusElement : StatusElement
 	{
-		string _name;
-		Color _backgroundColor;
-		Color _hoverBackgroundColor;
-		Action<Point, Rectangle>? _clickAction;
+		readonly string _name;
+		readonly Color _backgroundColor;
+		readonly Color _hoverBackgroundColor;
+		readonly Action<Point, Rectangle>? _clickAction;
 
 		public bool MergeLeft
 		{
@@ -273,7 +270,7 @@ namespace UnrealGameSync
 
 	class ProgressBarStatusElement : StatusElement
 	{
-		float _progress;
+		readonly float _progress;
 
 		public ProgressBarStatusElement(float inProgress)
 		{
@@ -304,7 +301,7 @@ namespace UnrealGameSync
 	class StatusLine
 	{
 		bool _modified;
-		List<StatusElement> _elements = new List<StatusElement>();
+		readonly List<StatusElement> _elements = new List<StatusElement>();
 
 		public StatusLine()
 		{
@@ -381,8 +378,8 @@ namespace UnrealGameSync
 			_elements.Add(new BadgeStatusElement(inText, inBackgroundColor, inClickAction));
 			if(_elements.Count >= 2)
 			{
-				BadgeStatusElement? prevBadge = _elements[_elements.Count - 2] as BadgeStatusElement;
-				BadgeStatusElement? nextBadge = _elements[_elements.Count - 1] as BadgeStatusElement;
+				BadgeStatusElement? prevBadge = _elements[^2] as BadgeStatusElement;
+				BadgeStatusElement? nextBadge = _elements[^1] as BadgeStatusElement;
 				if(prevBadge != null && nextBadge != null)
 				{
 					prevBadge.MergeRight = true;
@@ -446,7 +443,7 @@ namespace UnrealGameSync
 		bool _disposeProjectLogo;
 		Rectangle _projectLogoBounds;
 		StatusElementResources? _resources;
-		List<StatusLine> _lines = new List<StatusLine>();
+		readonly List<StatusLine> _lines = new List<StatusLine>();
 		StatusLine? _caption;
 		Pen? _alertDividerPen;
 		int _alertDividerY;
@@ -543,10 +540,7 @@ namespace UnrealGameSync
 
 		public void Set(IEnumerable<StatusLine> newLines, StatusLine? newCaption, StatusLine? newAlert, Color? newTintColor)
 		{
-			if(_resources == null)
-			{
-				_resources = new StatusElementResources(Font);
-			}
+			_resources ??= new StatusElementResources(Font);
 
 			if(_tintColor != newTintColor)
 			{
@@ -761,10 +755,9 @@ namespace UnrealGameSync
 				{
 					line.Draw(e.Graphics, _resources!);
 				}
-				if(_caption != null)
-				{
-					_caption.Draw(e.Graphics, _resources!);
-				}
+
+				_caption?.Draw(e.Graphics, _resources!);
+
 				if(_alert != null)
 				{
 					e.Graphics.DrawLine(_alertDividerPen!, 0, _alertDividerY, Width, _alertDividerY);
diff --git a/Engine/Source/Programs/UnrealGameSync/UnrealGameSync/Controls/SyncFilterControl.cs b/Engine/Source/Programs/UnrealGameSync/UnrealGameSync/Controls/SyncFilterControl.cs
--- a/Engine/Source/Programs/UnrealGameSync/UnrealGameSync/Controls/SyncFilterControl.cs
+++ b/Engine/Source/Programs/UnrealGameSync/UnrealGameSync/Controls/SyncFilterControl.cs
@@ -1,21 +1,15 @@
 // Copyright Epic Games, Inc. All Rights Reserved.
 
-using System;
 using System.Collections.Generic;
-using System.ComponentModel;
-using System.Drawing;
 using System.Data;
 using System.Linq;
-using System.Text;
-using System.Threading.Tasks;
 using System.Windows.Forms;
-using System.Runtime.InteropServices;
 
 namespace UnrealGameSync.Controls
 {
 	public partial class SyncFilterControl : UserControl
 	{
-		HashSet<string> _stripLines = new HashSet<string>()
+		static readonly HashSet<string> s_stripLines = new HashSet<string>()
 		{
 			"; Rules are specified one per line, and may use any standard Perforce wildcards:",
 			";    ?    Matches one character.",
@@ -39,7 +33,7 @@ namespace UnrealGameSync.Controls
 
 		public void SetView(string[] view)
 		{
-			ViewTextBox.Lines = view.Where(x => !_stripLines.Contains(x.Trim())).SkipWhile(x => x.Trim().Length == 0 || x.Trim() == ";").ToArray();
+			ViewTextBox.Lines = view.Where(x => !s_stripLines.Contains(x.Trim())).SkipWhile(x => x.Trim().Length == 0 || x.Trim() == ";").ToArray();
 			ViewTextBox.Select(ViewTextBox.Text.Length, 0);
 		}
 
diff --git a/Engine/Source/Programs/UnrealGameSync/UnrealGameSync/Controls/TabControl.cs b/Engine/Source/Programs/UnrealGameSync/UnrealGameSync/Controls/TabControl.cs
--- a/Engine/Source/Programs/UnrealGameSync/UnrealGameSync/Controls/TabControl.cs
+++ b/Engine/Source/Programs/UnrealGameSync/UnrealGameSync/Controls/TabControl.cs
@@ -4,8 +4,6 @@ using System;
 using System.Collections.Generic;
 using System.Drawing;
 using System.Linq;
-using System.Text;
-using System.Threading.Tasks;
 using System.Windows.Forms;
 using System.Drawing.Drawing2D;
 
@@ -18,39 +16,39 @@ namespace UnrealGameSync
 
 		class TabData
 		{
-			public string Name;
-			public object Data;
-			public int MinX;
-			public int Width;
-			public int CloseButtonWidth;
-			public Size TextSize;
-			public Tuple<Color, float>? Highlight;
-			public Color? TintColor;
+			public string _name;
+			public object _data;
+			public int _minX;
+			public int _width;
+			public int _closeButtonWidth;
+			public Size _textSize;
+			public Tuple<Color, float>? _highlight;
+			public Color? _tintColor;
 
 			public TabData(string name, object data)
 			{
-				this.Name = name;
-				this.Data = data;
+				_name = name;
+				_data = data;
 			}
 		}
 
 		class TabDragData
 		{
-			public TabData Tab;
-			public int InitialIdx;
-			public int MouseX;
-			public int RelativeMouseX;
+			public TabData _tab;
+			public int _initialIdx;
+			public int _mouseX;
+			public int _relativeMouseX;
 
 			public TabDragData(TabData tab, int initialIdx, int mouseX, int relativeMouseX)
 			{
-				this.Tab = tab;
-				this.InitialIdx = initialIdx;
-				this.MouseX = mouseX;
-				this.RelativeMouseX = relativeMouseX;
+				_tab = tab;
+				_initialIdx = initialIdx;
+				_mouseX = mouseX;
+				_relativeMouseX = relativeMouseX;
 			}
 		}
 
-		List<TabData> _tabs = new List<TabData>();
+		readonly List<TabData> _tabs = new List<TabData>();
 		int _selectedTabIdx;
 		int _hoverTabIdx;
 		int _highlightTabIdx;
@@ -102,7 +100,7 @@ namespace UnrealGameSync
 			int result = -1;
 			for(int tabIdx = 0; tabIdx < _tabs.Count - 1; tabIdx++)
 			{
-				if(_tabs[tabIdx].Data == data)
+				if(_tabs[tabIdx]._data == data)
 				{
 					result = tabIdx;
 					break;
@@ -113,12 +111,12 @@ namespace UnrealGameSync
 
 		public void SetHighlight(int tabIdx, Tuple<Color, float>? highlight)
 		{
-			Tuple<Color, float>? currentHighlight = _tabs[tabIdx].Highlight;
+			Tuple<Color, float>? currentHighlight = _tabs[tabIdx]._highlight;
 			if(highlight == null || currentHighlight == null)
 			{
 				if(highlight != currentHighlight)
 				{
-					_tabs[tabIdx].Highlight = highlight;
+					_tabs[tabIdx]._highlight = highlight;
 					Invalidate();
 				}
 			}
@@ -126,7 +124,7 @@ namespace UnrealGameSync
 			{
 				if(highlight.Item1 != currentHighlight.Item1 || highlight.Item2 != currentHighlight.Item2)
 				{
-					_tabs[tabIdx].Highlight = highlight;
+					_tabs[tabIdx]._highlight = highlight;
 					Invalidate();
 				}
 			}
@@ -134,9 +132,9 @@ namespace UnrealGameSync
 
 		public void SetTint(int tabIdx, Color? tintColor)
 		{
-			if(_tabs[tabIdx].TintColor != tintColor)
+			if(_tabs[tabIdx]._tintColor != tintColor)
 			{
-				_tabs[tabIdx].TintColor = tintColor;
+				_tabs[tabIdx]._tintColor = tintColor;
 				Invalidate();
 			}
 		}
@@ -153,7 +151,7 @@ namespace UnrealGameSync
 				_selectedTabIdx++;
 			}
 
-			_tabs.Insert(idx, new TabData(name, data){ TintColor = tintColor });
+			_tabs.Insert(idx, new TabData(name, data){ _tintColor = tintColor });
 			LayoutTabs();
 			Invalidate();
 			return idx;
@@ -166,19 +164,19 @@ namespace UnrealGameSync
 
 		public void SetTabName(int tabIdx, string name)
 		{
-			_tabs[tabIdx].Name = name;
+			_tabs[tabIdx]._name = name;
 			LayoutTabs();
 			Invalidate();
 		}
 
 		public object GetTabData(int tabIdx)
 		{
-			return _tabs[tabIdx].Data;
+			return _tabs[tabIdx]._data;
 		}
 
 		public void SetTabData(int tabIdx, object data)
 		{
-			_tabs[tabIdx].Data = data;
+			_tabs[tabIdx]._data = data;
 		}
 
 		public int GetSelectedTabIndex()
@@ -188,7 +186,7 @@ namespace UnrealGameSync
 
 		public object GetSelectedTabData()
 		{
-			return _tabs[_selectedTabIdx].Data;
+			return _tabs[_selectedTabIdx]._data;
 		}
 
 		public void SelectTab(int tabIdx)
@@ -213,14 +211,14 @@ namespace UnrealGameSync
 				}
 				else
 				{
-					OnTabChanged(_tabs[_selectedTabIdx].Data);
+					OnTabChanged(_tabs[_selectedTabIdx]._data);
 				}
 			}
 		}
 
 		public void RemoveTab(int tabIdx)
 		{
-			object tabData = _tabs[tabIdx].Data;
+			object tabData = _tabs[tabIdx]._data;
 			if(OnTabClosing == null || OnTabClosing(tabData))
 			{
 				_tabs.RemoveAt(tabIdx);
@@ -257,10 +255,7 @@ namespace UnrealGameSync
 				LayoutTabs();
 				Invalidate();
 
-				if(OnTabClosed != null)
-				{
-					OnTabClosed(tabData);
-				}
+				OnTabClosed?.Invoke(tabData);
 			}
 		}
 
@@ -274,12 +269,12 @@ namespace UnrealGameSync
 				for(int idx = 0; idx < _tabs.Count; idx++)
 				{
 					TabData tab = _tabs[idx];
-					tab.TextSize = TextRenderer.MeasureText(graphics, tab.Name, Font, new Size(int.MaxValue, int.MaxValue), TextFormatFlags.NoPadding);
-					tab.Width = TabPadding + tab.TextSize.Width + (int)(TabPadding * dpiScaleX);
+					tab._textSize = TextRenderer.MeasureText(graphics, tab._name, Font, new Size(Int32.MaxValue, Int32.MaxValue), TextFormatFlags.NoPadding);
+					tab._width = TabPadding + tab._textSize.Width + (int)(TabPadding * dpiScaleX);
 					if(idx == _selectedTabIdx)
 					{
-						tab.CloseButtonWidth = (int)(TabCloseButtonWidth * dpiScaleX);
-						tab.Width += tab.CloseButtonWidth;
+						tab._closeButtonWidth = (int)(TabCloseButtonWidth * dpiScaleX);
+						tab._width += tab._closeButtonWidth;
 					}
 				}
 
@@ -287,45 +282,45 @@ namespace UnrealGameSync
 				for(int idx = 0; idx < _tabs.Count; idx++)
 				{
 					TabData tab = _tabs[idx];
-					tab.MinX = leftX;
-					leftX += tab.Width;
+					tab._minX = leftX;
+					leftX += tab._width;
 				}
 
 				int rightX = Width - 1;
 				if(leftX > rightX)
 				{
-					int usedWidth = _tabs.Take(_tabs.Count - 1).Sum(x => x.Width);
-					int remainingWidth = rightX + 1 - _tabs[_tabs.Count - 1].Width;
+					int usedWidth = _tabs.Take(_tabs.Count - 1).Sum(x => x._width);
+					int remainingWidth = rightX + 1 - _tabs[^1]._width;
 					if(_selectedTabIdx != -1)
 					{
-						usedWidth -= _tabs[_selectedTabIdx].Width;
-						remainingWidth -= _tabs[_selectedTabIdx].Width;
+						usedWidth -= _tabs[_selectedTabIdx]._width;
+						remainingWidth -= _tabs[_selectedTabIdx]._width;
 					}
 
 					int newX = 0;
 					for(int idx = 0; idx < _tabs.Count; idx++)
 					{
-						_tabs[idx].MinX = newX;
-						int prevWidth = _tabs[idx].Width;
+						_tabs[idx]._minX = newX;
+						int prevWidth = _tabs[idx]._width;
 						if(idx != _selectedTabIdx && idx != _tabs.Count - 1)
 						{
-							_tabs[idx].Width = Math.Max((_tabs[idx].Width * remainingWidth) / usedWidth, TabPadding * 3);
+							_tabs[idx]._width = Math.Max((_tabs[idx]._width * remainingWidth) / usedWidth, TabPadding * 3);
 						}
-						_tabs[idx].TextSize.Width -= prevWidth - _tabs[idx].Width;
-						newX += _tabs[idx].Width;
+						_tabs[idx]._textSize.Width -= prevWidth - _tabs[idx]._width;
+						newX += _tabs[idx]._width;
 					}
 				}
 
 				if(_dragState != null)
 				{
-					int minOffset = -_dragState.Tab.MinX;
-					int maxOffset = _tabs[_tabs.Count - 1].MinX - _dragState.Tab.Width - _dragState.Tab.MinX;
+					int minOffset = -_dragState._tab._minX;
+					int maxOffset = _tabs[^1]._minX - _dragState._tab._width - _dragState._tab._minX;
 
-					int offset = (_dragState.MouseX - _dragState.RelativeMouseX) - _dragState.Tab.MinX;
+					int offset = (_dragState._mouseX - _dragState._relativeMouseX) - _dragState._tab._minX;
 					offset = Math.Max(offset, minOffset);
 					offset = Math.Min(offset, maxOffset);
 
-					_dragState.Tab.MinX += offset;
+					_dragState._tab._minX += offset;
 				}
 			}
 		}
@@ -350,13 +345,13 @@ namespace UnrealGameSync
 				{
 					if(_hoverTabIdx == _selectedTabIdx)
 					{
-						if(e.Location.X > _tabs[_hoverTabIdx].MinX + _tabs[_hoverTabIdx].Width - _tabs[_hoverTabIdx].CloseButtonWidth)
+						if(e.Location.X > _tabs[_hoverTabIdx]._minX + _tabs[_hoverTabIdx]._width - _tabs[_hoverTabIdx]._closeButtonWidth)
 						{
 							RemoveTab(_hoverTabIdx);
 						}
 						else
 						{
-							_dragState = new TabDragData(_tabs[_selectedTabIdx], _selectedTabIdx, e.Location.X, e.Location.X - _tabs[_selectedTabIdx].MinX);
+							_dragState = new TabDragData(_tabs[_selectedTabIdx], _selectedTabIdx, e.Location.X, e.Location.X - _tabs[_selectedTabIdx]._minX);
 						}
 					}
 					else
@@ -386,7 +381,7 @@ namespace UnrealGameSync
 
 			if(OnTabClicked != null)
 			{
-				object? tabData = (_hoverTabIdx == -1)? null : _tabs[_hoverTabIdx].Data;
+				object? tabData = (_hoverTabIdx == -1)? null : _tabs[_hoverTabIdx]._data;
 				OnTabClicked(tabData, e.Location, e.Button);
 			}
 		}
@@ -415,20 +410,20 @@ namespace UnrealGameSync
 
 			if(_dragState != null)
 			{
-				_dragState.MouseX = e.Location.X;
+				_dragState._mouseX = e.Location.X;
 
-				_tabs.Remove(_dragState.Tab);
+				_tabs.Remove(_dragState._tab);
 
 				int totalWidth = 0;
 				for(int insertIdx = 0;;insertIdx++)
 				{
-					if(insertIdx == _tabs.Count - 1 || _dragState.MouseX - _dragState.RelativeMouseX < totalWidth + _tabs[insertIdx].Width / 2)
+					if(insertIdx == _tabs.Count - 1 || _dragState._mouseX - _dragState._relativeMouseX < totalWidth + _tabs[insertIdx]._width / 2)
 					{
 						_hoverTabIdx = _selectedTabIdx = insertIdx;
-						_tabs.Insert(insertIdx, _dragState.Tab);
+						_tabs.Insert(insertIdx, _dragState._tab);
 						break;
 					}
-					totalWidth += _tabs[insertIdx].Width;
+					totalWidth += _tabs[insertIdx]._width;
 				}
 
 				LayoutTabs();
@@ -441,7 +436,7 @@ namespace UnrealGameSync
 				{
 					for(int idx = 0; idx < _tabs.Count; idx++)
 					{
-						if(e.Location.X > _tabs[idx].MinX && e.Location.X < _tabs[idx].MinX + _tabs[idx].Width)
+						if(e.Location.X > _tabs[idx]._minX && e.Location.X < _tabs[idx]._minX + _tabs[idx]._width)
 						{
 							newHoverTabIdx = idx;
 							break;
@@ -491,11 +486,11 @@ namespace UnrealGameSync
 					TabData tab = _tabs[idx];
 					if(idx == _hoverTabIdx || idx == _highlightTabIdx)
 					{
-						DrawBackground(e.Graphics, tab, SystemBrushes.Window, hoverBrush, tab.Highlight);
+						DrawBackground(e.Graphics, tab, SystemBrushes.Window, hoverBrush, tab._highlight);
 					}
 					else
 					{
-						DrawBackground(e.Graphics, tab, SystemBrushes.Control, SystemBrushes.Control, tab.Highlight);
+						DrawBackground(e.Graphics, tab, SystemBrushes.Control, SystemBrushes.Control, tab._highlight);
 					}
 				}
 			}
@@ -507,15 +502,15 @@ namespace UnrealGameSync
 					TabData tab = _tabs[idx];
 					if((idx > 0 && idx < _tabs.Count - 1) || idx >= _tabs.Count || idx == _selectedTabIdx || idx == _hoverTabIdx || idx == _highlightTabIdx || _tabs.Count == 1)
 					{
-						e.Graphics.DrawLine(separatorPen, tab.MinX, 0, tab.MinX, ClientSize.Height - 2);
+						e.Graphics.DrawLine(separatorPen, tab._minX, 0, tab._minX, ClientSize.Height - 2);
 					}
 					if(idx < _tabs.Count - 1 || idx >= _tabs.Count || idx == _hoverTabIdx || idx == _highlightTabIdx || _tabs.Count == 1)
 					{
-						e.Graphics.DrawLine(separatorPen, tab.MinX + tab.Width, 0, tab.MinX + tab.Width, ClientSize.Height - 2);
+						e.Graphics.DrawLine(separatorPen, tab._minX + tab._width, 0, tab._minX + tab._width, ClientSize.Height - 2);
 					}
 					if(idx == _hoverTabIdx || idx == _highlightTabIdx || idx >= _tabs.Count)
 					{
-						e.Graphics.DrawLine(separatorPen, tab.MinX, 0, tab.MinX + tab.Width, 0);
+						e.Graphics.DrawLine(separatorPen, tab._minX, 0, tab._minX + tab._width, 0);
 					}
 				}
 			}
@@ -528,14 +523,13 @@ namespace UnrealGameSync
 				}
 			}
 
-
 			if(_selectedTabIdx != -1)
 			{
 				TabData selectedTab = _tabs[_selectedTabIdx];
 
 				using(SolidBrush selectedBrush = new SolidBrush(Color.FromArgb(0, 136, 204)))
 				{
-					DrawBackground(e.Graphics, selectedTab, SystemBrushes.Window, selectedBrush, selectedTab.Highlight);
+					DrawBackground(e.Graphics, selectedTab, SystemBrushes.Window, selectedBrush, selectedTab._highlight);
 				}
 
 				DrawText(e.Graphics, selectedTab);
@@ -547,7 +541,7 @@ namespace UnrealGameSync
 				float dpiScaleX = e.Graphics.DpiX / 96.0f;
 				float dpiScaleY = e.Graphics.DpiY / 96.0f;
 
-				int closeMidX = selectedTab.MinX + selectedTab.Width - (int)((TabPadding + TabCloseButtonWidth) * dpiScaleX / 2);// + (13 / 2);
+				int closeMidX = selectedTab._minX + selectedTab._width - (int)((TabPadding + TabCloseButtonWidth) * dpiScaleX / 2);// + (13 / 2);
 				int closeMidY = (ClientSize.Height - (int)(4 * dpiScaleY)) / 2;
 
 				Rectangle closeButton = new Rectangle(closeMidX - (int)((13 / 2) * dpiScaleX), closeMidY - (int)((13 / 2) * dpiScaleY), (int)(13 * dpiScaleX), (int)(13 * dpiScaleY));
@@ -563,9 +557,9 @@ namespace UnrealGameSync
 				e.Graphics.SmoothingMode = prevSmoothingMode;
 
 				// Border
-				e.Graphics.DrawLine(SystemPens.ControlDark, selectedTab.MinX, 0, selectedTab.MinX + selectedTab.Width, 0);
-				e.Graphics.DrawLine(SystemPens.ControlDark, selectedTab.MinX, 0, selectedTab.MinX, ClientSize.Height - 2);
-				e.Graphics.DrawLine(SystemPens.ControlDark, selectedTab.MinX + selectedTab.Width, 0, selectedTab.MinX + selectedTab.Width, ClientSize.Height - 2);
+				e.Graphics.DrawLine(SystemPens.ControlDark, selectedTab._minX, 0, selectedTab._minX + selectedTab._width, 0);
+				e.Graphics.DrawLine(SystemPens.ControlDark, selectedTab._minX, 0, selectedTab._minX, ClientSize.Height - 2);
+				e.Graphics.DrawLine(SystemPens.ControlDark, selectedTab._minX + selectedTab._width, 0, selectedTab._minX + selectedTab._width, ClientSize.Height - 2);
 			}
 
 			e.Graphics.DrawLine(SystemPens.ControlDarkDark, 0, ClientSize.Height - 2,  ClientSize.Width, ClientSize.Height - 2);
@@ -578,15 +572,15 @@ namespace UnrealGameSync
 
 			graphics.SmoothingMode = System.Drawing.Drawing2D.SmoothingMode.None;
 
-			graphics.FillRectangle(backgroundBrush, tab.MinX + 1, 1, tab.Width - 1, ClientSize.Height - 3);
-			graphics.FillRectangle(stripeBrush, tab.MinX + 1, ClientSize.Height - (int)(5 * dpiScaleY), tab.Width - 1, (int)(5 * dpiScaleY) - 2);
+			graphics.FillRectangle(backgroundBrush, tab._minX + 1, 1, tab._width - 1, ClientSize.Height - 3);
+			graphics.FillRectangle(stripeBrush, tab._minX + 1, ClientSize.Height - (int)(5 * dpiScaleY), tab._width - 1, (int)(5 * dpiScaleY) - 2);
 
-			if (tab.TintColor.HasValue)
+			if (tab._tintColor.HasValue)
 			{
-				int midX = tab.MinX + tab.Width / 2;
-				using (LinearGradientBrush brush = new LinearGradientBrush(new Rectangle(midX, 1, 1, ClientSize.Height), Color.FromArgb(64, tab.TintColor.Value), Color.FromArgb(32, tab.TintColor.Value), LinearGradientMode.Vertical))
+				int midX = tab._minX + tab._width / 2;
+				using (LinearGradientBrush brush = new LinearGradientBrush(new Rectangle(midX, 1, 1, ClientSize.Height), Color.FromArgb(64, tab._tintColor.Value), Color.FromArgb(32, tab._tintColor.Value), LinearGradientMode.Vertical))
 				{
-					graphics.FillRectangle(brush, tab.MinX + 1, 1, tab.Width - 1, ClientSize.Height - 3);
+					graphics.FillRectangle(brush, tab._minX + 1, 1, tab._width - 1, ClientSize.Height - 3);
 				}
 			}
 
@@ -594,7 +588,7 @@ namespace UnrealGameSync
 			{
 				using(SolidBrush brush = new SolidBrush(highlight.Item1))
 				{
-					graphics.FillRectangle(brush, tab.MinX + 1, ClientSize.Height - (int)(5 * dpiScaleY), (int)((tab.Width - 2) * highlight.Item2), (int)(5 * dpiScaleY) - 2);
+					graphics.FillRectangle(brush, tab._minX + 1, ClientSize.Height - (int)(5 * dpiScaleY), (int)((tab._width - 2) * highlight.Item2), (int)(5 * dpiScaleY) - 2);
 				}
 			}
 		}
@@ -603,7 +597,7 @@ namespace UnrealGameSync
 		{
 			float dpiScaleX = graphics.DpiX / 96.0f;
 			float dpiScaleY = graphics.DpiY / 96.0f;
-			TextRenderer.DrawText(graphics, tab.Name, Font, new Rectangle(tab.MinX + (int)(TabPadding * dpiScaleX), 0, tab.TextSize.Width, ClientSize.Height - (int)(3 * dpiScaleY)), Color.Black, TextFormatFlags.NoPadding | TextFormatFlags.EndEllipsis | TextFormatFlags.VerticalCenter);
+			TextRenderer.DrawText(graphics, tab._name, Font, new Rectangle(tab._minX + (int)(TabPadding * dpiScaleX), 0, tab._textSize.Width, ClientSize.Height - (int)(3 * dpiScaleY)), Color.Black, TextFormatFlags.NoPadding | TextFormatFlags.EndEllipsis | TextFormatFlags.VerticalCenter);
 		}
 	}
 }
diff --git a/Engine/Source/Programs/UnrealGameSync/UnrealGameSync/Controls/TextBoxWithCueBanner.cs b/Engine/Source/Programs/UnrealGameSync/UnrealGameSync/Controls/TextBoxWithCueBanner.cs
--- a/Engine/Source/Programs/UnrealGameSync/UnrealGameSync/Controls/TextBoxWithCueBanner.cs
+++ b/Engine/Source/Programs/UnrealGameSync/UnrealGameSync/Controls/TextBoxWithCueBanner.cs
@@ -1,11 +1,7 @@
 // Copyright Epic Games, Inc. All Rights Reserved.
 
 using System;
-using System.Collections.Generic;
-using System.Linq;
 using System.Runtime.InteropServices;
-using System.Text;
-using System.Threading.Tasks;
 using System.Windows.Forms;
 
 namespace UnrealGameSync
@@ -21,8 +17,12 @@ namespace UnrealGameSync
 
 		public string CueBanner
 		{
-			get { return _cueBannerValue; }
-			set { _cueBannerValue = value; UpdateCueBanner(); }
+			get => _cueBannerValue;
+			set 
+			{ 
+				_cueBannerValue = value; 
+				UpdateCueBanner();
+			}
 		}
 
 		private void UpdateCueBanner()
