UGS: Add .editorconfig file for static analysis warnings in UGS (+lots of fixes).

diff --git a/Engine/Source/Programs/UnrealGameSync/UnrealGameSync/Controls/WorkspaceControl.cs b/Engine/Source/Programs/UnrealGameSync/UnrealGameSync/Controls/WorkspaceControl.cs
--- a/Engine/Source/Programs/UnrealGameSync/UnrealGameSync/Controls/WorkspaceControl.cs
+++ b/Engine/Source/Programs/UnrealGameSync/UnrealGameSync/Controls/WorkspaceControl.cs
@@ -7,7 +7,6 @@ using System.Drawing;
 using System.Drawing.Drawing2D;
 using System.IO;
 using System.Linq;
-using System.Reflection;
 using System.Runtime.InteropServices;
 using System.Text;
 using System.Text.Json;
@@ -78,16 +77,16 @@ namespace UnrealGameSync
 
 		class BadgeInfo
 		{
-			public string Label;
-			public string? Group;
-			public string? UniqueId;
-			public int Offset;
-			public int Width;
-			public int Height;
-			public Color BackgroundColor;
-			public Color HoverBackgroundColor;
-			public Action? ClickHandler;
-			public string? ToolTip;
+			public string _label;
+			public string? _group;
+			public string? _uniqueId;
+			public int _offset;
+			public int _width;
+			public int _height;
+			public Color _backgroundColor;
+			public Color _hoverBackgroundColor;
+			public Action? _clickHandler;
+			public string? _toolTip;
 
 			public BadgeInfo(string label, string group, Color badgeColor)
 				: this(label, group, null, badgeColor, badgeColor, null)
@@ -96,39 +95,39 @@ namespace UnrealGameSync
 
 			public BadgeInfo(string label, string? group, string? uniqueId, Color backgroundColor, Color hoverBackgroundColor, Action? clickHandler)
 			{
-				this.Label = label;
-				this.Group = group;
-				this.UniqueId = uniqueId;
-				this.BackgroundColor = backgroundColor;
-				this.HoverBackgroundColor = hoverBackgroundColor;
-				this.ClickHandler = clickHandler;
+				_label = label;
+				_group = group;
+				_uniqueId = uniqueId;
+				_backgroundColor = backgroundColor;
+				_hoverBackgroundColor = hoverBackgroundColor;
+				_clickHandler = clickHandler;
 			}
 
 			public BadgeInfo(BadgeInfo other)
-				: this(other.Label, other.Group, other.UniqueId, other.BackgroundColor, other.HoverBackgroundColor, other.ClickHandler)
+				: this(other._label, other._group, other._uniqueId, other._backgroundColor, other._hoverBackgroundColor, other._clickHandler)
 			{
-				this.Offset = other.Offset;
-				this.Width = other.Width;
-				this.Height = other.Height;
+				_offset = other._offset;
+				_width = other._width;
+				_height = other._height;
 			}
 
 			public Rectangle GetBounds(Point listLocation)
 			{
-				return new Rectangle(listLocation.X + Offset, listLocation.Y, Width, Height);
+				return new Rectangle(listLocation.X + _offset, listLocation.Y, _width, _height);
 			}
 		}
 
 		class ServiceBadgeInfo : IEquatable<ServiceBadgeInfo>
 		{
-			public string Name;
-			public string? Url;
-			public BadgeResult Result;
+			public string Name { get; }
+			public string? Url { get; }
+			public BadgeResult Result { get; }
 
 			public ServiceBadgeInfo(string name, string? url, BadgeResult result)
 			{
-				this.Name = name;
-				this.Url = url;
-				this.Result = result;
+				Name = name;
+				Url = url;
+				Result = result;
 			}
 
 			public bool Equals(ServiceBadgeInfo? otherBadge)
@@ -149,10 +148,10 @@ namespace UnrealGameSync
 
 		class ChangeLayoutInfo
 		{
-			public List<BadgeInfo> DescriptionBadges = new List<BadgeInfo>();
-			public List<BadgeInfo> TypeBadges = new List<BadgeInfo>();
-			public List<BadgeInfo> BuildBadges = new List<BadgeInfo>();
-			public Dictionary<string, List<BadgeInfo>> CustomBadges = new Dictionary<string, List<BadgeInfo>>();
+			public List<BadgeInfo> _descriptionBadges = new List<BadgeInfo>();
+			public List<BadgeInfo> _typeBadges = new List<BadgeInfo>();
+			public List<BadgeInfo> _buildBadges = new List<BadgeInfo>();
+			public Dictionary<string, List<BadgeInfo>> _customBadges = new Dictionary<string, List<BadgeInfo>>();
 		}
 
 		class BuildListItemComparer : System.Collections.IComparer
@@ -176,43 +175,41 @@ namespace UnrealGameSync
 			}
 		}
 
-		public static Rectangle GoodBuildIcon = new Rectangle(0, 0, 16, 16);
-		public static Rectangle MixedBuildIcon = new Rectangle(16, 0, 16, 16);
-		public static Rectangle BadBuildIcon = new Rectangle(32, 0, 16, 16);
-		public static Rectangle DefaultBuildIcon = new Rectangle(48, 0, 16, 16);
-		static Rectangle _promotedBuildIcon = new Rectangle(64, 0, 16, 16);
-		static Rectangle _detailsIcon = new Rectangle(80, 0, 16, 16);
-		static Rectangle _infoIcon = new Rectangle(96, 0, 16, 16);
-		static Rectangle _cancelIcon = new Rectangle(112, 0, 16, 16);
-		static Rectangle _syncIcon = new Rectangle(128, 0, 32, 16);
-		static Rectangle _happyIcon = new Rectangle(160, 0, 16, 16);
-		static Rectangle _disabledHappyIcon = new Rectangle(176, 0, 16, 16);
-		static Rectangle _frownIcon = new Rectangle(192, 0, 16, 16);
-		static Rectangle _disabledFrownIcon = new Rectangle(208, 0, 16, 16);
-		static Rectangle _previousSyncIcon = new Rectangle(224, 0, 16, 16);
-		static Rectangle _additionalSyncIcon = new Rectangle(240, 0, 16, 16);
-		static Rectangle _bisectPassIcon = new Rectangle(256, 0, 16, 16);
-		static Rectangle _bisectFailIcon = new Rectangle(273, 0, 16, 16);
-		static Rectangle _bisectImplicitPassIcon = new Rectangle(290, 0, 16, 16);
-		static Rectangle _bisectImplicitFailIcon = new Rectangle(307, 0, 16, 16);
+		public static Rectangle GoodBuildIcon { get; } = new Rectangle(0, 0, 16, 16);
+		public static Rectangle MixedBuildIcon { get; } = new Rectangle(16, 0, 16, 16);
+		public static Rectangle BadBuildIcon { get; } = new Rectangle(32, 0, 16, 16);
+		public static Rectangle DefaultBuildIcon { get; } = new Rectangle(48, 0, 16, 16);
+		static Rectangle PromotedBuildIcon { get; } = new Rectangle(64, 0, 16, 16);
+		static Rectangle InfoIcon { get; } = new Rectangle(96, 0, 16, 16);
+		static Rectangle CancelIcon { get; } = new Rectangle(112, 0, 16, 16);
+		static Rectangle HappyIcon { get; } = new Rectangle(160, 0, 16, 16);
+		static Rectangle DisabledHappyIcon { get; } = new Rectangle(176, 0, 16, 16);
+		static Rectangle FrownIcon { get; } = new Rectangle(192, 0, 16, 16);
+		static Rectangle DisabledFrownIcon { get; } = new Rectangle(208, 0, 16, 16);
+		static Rectangle PreviousSyncIcon { get; } = new Rectangle(224, 0, 16, 16);
+		static Rectangle AdditionalSyncIcon { get; } = new Rectangle(240, 0, 16, 16);
+		static Rectangle BisectPassIcon { get; } = new Rectangle(256, 0, 16, 16);
+		static Rectangle BisectFailIcon { get; } = new Rectangle(273, 0, 16, 16);
+		static Rectangle BisectImplicitPassIcon { get; } = new Rectangle(290, 0, 16, 16);
+		static Rectangle BisectImplicitFailIcon { get; } = new Rectangle(307, 0, 16, 16);
 
 		[DllImport("uxtheme.dll", CharSet = CharSet.Unicode)]
 		static extern int SetWindowTheme(IntPtr hWnd, string pszSubAppName, string pszSubIdList);
 
 		const int BuildListExpandCount = 250;
 
-		IWorkspaceControlOwner _owner;
-		DirectoryReference _appDataFolder;
-		string? _apiUrl;
-		DirectoryReference _workspaceDataFolder;
-		IServiceProvider _serviceProvider;
-		ILogger _logger;
-		IPerforceSettings _perforceSettings;
+		readonly IWorkspaceControlOwner _owner;
+		readonly DirectoryReference _appDataFolder;
+		readonly string? _apiUrl;
+		readonly DirectoryReference _workspaceDataFolder;
+		readonly IServiceProvider _serviceProvider;
+		readonly ILogger _logger;
+		readonly IPerforceSettings _perforceSettings;
 		ProjectInfo ProjectInfo { get; }
 
-		UserSettings _settings;
-		UserWorkspaceSettings _workspaceSettings;
-		UserProjectSettings _projectSettings;
+		readonly UserSettings _settings;
+		readonly UserWorkspaceSettings _workspaceSettings;
+		readonly UserProjectSettings _projectSettings;
 
 		bool _userHasOpenIssues = false;
 
@@ -224,7 +221,7 @@ namespace UnrealGameSync
 		public string? StreamName => ProjectInfo.StreamName;
 		public string ClientName => ProjectInfo.ClientName;
 
-		SynchronizationContext _mainThreadSynchronizationContext;
+		readonly SynchronizationContext _mainThreadSynchronizationContext;
 		bool _isDisposing;
 
 		JupiterMonitor _jupiterMonitor;
@@ -232,19 +229,19 @@ namespace UnrealGameSync
 		Workspace _workspace;
 		IssueMonitor _issueMonitor;
 		EventMonitor _eventMonitor;
-		System.Windows.Forms.Timer _updateTimer;
+		readonly System.Windows.Forms.Timer _updateTimer;
 		HashSet<int> _promotedChangeNumbers = new HashSet<int>();
 		List<int> _listIndexToChangeIndex = new List<int>();
 		List<int> _sortedChangeNumbers = new List<int>();
-		Dictionary<string, Dictionary<int, string?>> _archiveToChangeNumberToArchiveKey = new Dictionary<string, Dictionary<int, string?>>();
-		Dictionary<int, ChangeLayoutInfo> _changeNumberToLayoutInfo = new Dictionary<int, ChangeLayoutInfo>();
-		List<ContextMenuStrip> _customStatusPanelMenus = new List<ContextMenuStrip>();
-		List<(string, Action<Point, Rectangle>)> _customStatusPanelLinks = new List<(string, Action<Point, Rectangle>)>();
+		readonly Dictionary<string, Dictionary<int, string?>> _archiveToChangeNumberToArchiveKey = new Dictionary<string, Dictionary<int, string?>>();
+		readonly Dictionary<int, ChangeLayoutInfo> _changeNumberToLayoutInfo = new Dictionary<int, ChangeLayoutInfo>();
+		readonly List<ContextMenuStrip> _customStatusPanelMenus = new List<ContextMenuStrip>();
+		readonly List<(string, Action<Point, Rectangle>)> _customStatusPanelLinks = new List<(string, Action<Point, Rectangle>)>();
 		int _numChanges;
 		int _pendingSelectedChangeNumber = -1;
 		bool _hasBuildSteps = false;
 
-		Dictionary<string, int> _notifiedBuildTypeToChangeNumber = new Dictionary<string, int>();
+		readonly Dictionary<string, int> _notifiedBuildTypeToChangeNumber = new Dictionary<string, int>();
 
 		bool _mouseOverExpandLink;
 
@@ -255,7 +252,7 @@ namespace UnrealGameSync
 		Font? _selectedBuildFont;
 		Font? _badgeFont;
 		List<KeyValuePair<string, string>> _badgeNameAndGroupPairs = new List<KeyValuePair<string, string>>();
-		Dictionary<string, Size> _badgeLabelToSize = new Dictionary<string, Size>();
+		readonly Dictionary<string, Size> _badgeLabelToSize = new Dictionary<string, Size>();
 		List<ServiceBadgeInfo> _serviceBadges = new List<ServiceBadgeInfo>();
 
 		NotificationWindow _notificationWindow;
@@ -289,7 +286,7 @@ namespace UnrealGameSync
 		string _authorFilterText = "";
 
 		// Placeholder text that is in the control and cleared when the user starts editing.
-		static string _authorFilterPlaceholderText = "<username>";
+		const string AuthorFilterPlaceholderText = "<username>";
 
 		public WorkspaceControl(IWorkspaceControlOwner inOwner, DirectoryReference inAppDataFolder, string? inApiUrl, OpenProjectInfo openProjectInfo, IServiceProvider inServiceProvider, UserSettings inSettings, OidcTokenManager? inOidcTokenManager)
 		{
@@ -400,7 +397,13 @@ namespace UnrealGameSync
 			_perforceMonitor.Start();
 			_eventMonitor.Start();
 
-			_startupTimer = new System.Threading.Timer(x => _mainThreadSynchronizationContext.Post((o) => { if (!IsDisposed) { StartupTimerElapsed(false); } }, null), null, TimeSpan.FromSeconds(20.0), TimeSpan.FromMilliseconds(-1.0));
+			_startupTimer = new System.Threading.Timer(x => _mainThreadSynchronizationContext.Post((o) => 
+			{ 
+				if (!IsDisposed) 
+				{ 
+					StartupTimerElapsed(false); 
+				} 
+			}, null), null, TimeSpan.FromSeconds(20.0), TimeSpan.FromMilliseconds(-1.0));
 			_startupCallbacks = new List<WorkspaceStartupCallback>();
 
 			string? issuesApiUrl = GetIssuesApiUrl();
@@ -427,7 +430,13 @@ namespace UnrealGameSync
 
 		public void IssueMonitor_OnIssuesChangedAsync()
 		{
-			_mainThreadSynchronizationContext.Post((o) => { if (_issueMonitor != null) { IssueMonitor_OnIssuesChanged(); } }, null);
+			_mainThreadSynchronizationContext.Post((o) => 
+			{ 
+				if (_issueMonitor != null) 
+				{ 
+					IssueMonitor_OnIssuesChanged(); 
+				} 
+			}, null);
 		}
 
 		public List<IssueData> GetOpenIssuesForUser()
@@ -455,8 +464,7 @@ namespace UnrealGameSync
 		{
 			if (_startupTimer != null)
 			{
-				int latestChangeNumber;
-				if (FindChangeToSync(_settings.SyncTypeId, out latestChangeNumber))
+				if (FindChangeToSync(_settings.SyncTypeId, out _))
 				{
 					StartupTimerElapsed(false);
 				}
@@ -537,7 +545,7 @@ namespace UnrealGameSync
 								int index = columnConfig.GetValue("Index", -1);
 								if (index < 0 || index > BuildList.Columns.Count)
 								{
-									index = ((_customColumns.Count > 0) ? _customColumns[_customColumns.Count - 1].Index : CISColumn.Index) + 1;
+									index = ((_customColumns.Count > 0) ? _customColumns[^1].Index : CISColumn.Index) + 1;
 								}
 
 								column = new ColumnHeader();
@@ -693,10 +701,7 @@ namespace UnrealGameSync
 				components.Dispose();
 			}
 
-			if (_updateTimer != null)
-			{
-				_updateTimer.Stop();
-			}
+			_updateTimer?.Stop();
 
 			if (_startupCallbacks != null)
 			{
@@ -710,10 +715,7 @@ namespace UnrealGameSync
 			if (_issueMonitor != null)
 			{
 				_issueMonitor.OnIssuesChanged -= IssueMonitor_OnIssuesChangedAsync;
-				if (_owner != null)
-				{
-					_owner.ReleaseIssueMonitor(_issueMonitor);
-				}
+				_owner?.ReleaseIssueMonitor(_issueMonitor);
 				_issueMonitor = null!;
 			}
 			if (_startupTimer != null)
@@ -859,7 +861,7 @@ namespace UnrealGameSync
 				{
 					newPendingMaxChanges = _listIndexToChangeIndex[lastVisibleIndex];
 				}
-				newPendingMaxChanges = _perforceMonitor.InitialMaxChangesValue + ((Math.Max(newPendingMaxChanges - _perforceMonitor.InitialMaxChangesValue, 0) + BuildListExpandCount - 1) / BuildListExpandCount) * BuildListExpandCount;
+				newPendingMaxChanges = PerforceMonitor.InitialMaxChangesValue + ((Math.Max(newPendingMaxChanges - PerforceMonitor.InitialMaxChangesValue, 0) + BuildListExpandCount - 1) / BuildListExpandCount) * BuildListExpandCount;
 
 				// Shrink the number of changes retained by the PerforceMonitor class
 				if (_perforceMonitor.PendingMaxChanges > newPendingMaxChanges)
@@ -990,7 +992,11 @@ namespace UnrealGameSync
 						FileReference receiptFile = ConfigUtils.GetReceiptFile(_workspace.Project, _workspace.ProjectConfigFile, targetFile, config.ToString());
 						if (FileReference.Exists(receiptFile))
 						{
-							try { FileReference.Delete(receiptFile); } catch (Exception) { }
+							try 
+							{ 
+								FileReference.Delete(receiptFile); 
+							} 
+							catch (Exception) { }
 						}
 					}
 				}
@@ -1034,10 +1040,16 @@ namespace UnrealGameSync
 
 		void UpdateCompleteCallback(WorkspaceUpdateContext context, WorkspaceUpdateResult result, string resultMessage)
 		{
-			_mainThreadSynchronizationContext.Post((o) => { if (!_isDisposing) { UpdateComplete(context, result, resultMessage); } }, null);
+			_mainThreadSynchronizationContext.Post((o) => 
+			{ 
+				if (!_isDisposing) 
+				{ 
+					UpdateComplete(context, result); 
+				} 
+			}, null);
 		}
 
-		void UpdateComplete(WorkspaceUpdateContext context, WorkspaceUpdateResult result, string resultMessage)
+		void UpdateComplete(WorkspaceUpdateContext context, WorkspaceUpdateResult result)
 		{
 			if (_isDisposing)
 			{
@@ -1137,7 +1149,14 @@ namespace UnrealGameSync
 			if (!_updateBuildListPosted)
 			{
 				_updateBuildListPosted = true;
-				_mainThreadSynchronizationContext.Post((o) => { _updateBuildListPosted = false; if (!_isDisposing) { UpdateBuildList(); } }, null);
+				_mainThreadSynchronizationContext.Post((o) => 
+				{ 
+					_updateBuildListPosted = false; 
+					if (!_isDisposing) 
+					{ 
+						UpdateBuildList(); 
+					} 
+				}, null);
 			}
 		}
 
@@ -1410,9 +1429,9 @@ namespace UnrealGameSync
 			}
 
 			// if filtering by user, only show changes where the author contains the filter text
-			if (!string.IsNullOrEmpty(this._authorFilterText)
-				&& this._authorFilterText != _authorFilterPlaceholderText
-				&& change.User!.IndexOf(this._authorFilterText, StringComparison.OrdinalIgnoreCase) < 0)
+			if (!String.IsNullOrEmpty(_authorFilterText)
+				&& _authorFilterText != AuthorFilterPlaceholderText
+				&& change.User!.IndexOf(_authorFilterText, StringComparison.OrdinalIgnoreCase) < 0)
 			{
 				return false;
 			}
@@ -1470,7 +1489,14 @@ namespace UnrealGameSync
 			if (!_updateBuildMetadataPosted)
 			{
 				_updateBuildMetadataPosted = true;
-				_mainThreadSynchronizationContext.Post((o) => { _updateBuildMetadataPosted = false; if (!_isDisposing) { UpdateBuildMetadata(); } }, null);
+				_mainThreadSynchronizationContext.Post((o) => 
+				{ 
+					_updateBuildMetadataPosted = false; 
+					if (!_isDisposing) 
+					{ 
+						UpdateBuildMetadata(); 
+					} 
+				}, null);
 			}
 		}
 
@@ -1627,7 +1653,7 @@ namespace UnrealGameSync
 			if (dummyBadges.Count > 0)
 			{
 				LayoutBadges(dummyBadges);
-				width = dummyBadges[dummyBadges.Count - 1].Offset + dummyBadges[dummyBadges.Count - 1].Width;
+				width = dummyBadges[^1]._offset + dummyBadges[^1]._width;
 			}
 			return width;
 		}
@@ -1705,7 +1731,7 @@ namespace UnrealGameSync
 
 		void ParseUncontrolledChangelistsPersistencyFile(string inProjectRoot, HashSet<string> outUncontrolledFiles)
 		{
-			String uncontrolledChangelistPersistencyFilePath = Path.Combine(inProjectRoot, "Saved", "SourceControl", "UncontrolledChangelists.json");
+			string uncontrolledChangelistPersistencyFilePath = Path.Combine(inProjectRoot, "Saved", "SourceControl", "UncontrolledChangelists.json");
 
 			if (File.Exists(uncontrolledChangelistPersistencyFilePath))
 			{
@@ -1736,7 +1762,14 @@ namespace UnrealGameSync
 			if (!_updateReviewsPosted)
 			{
 				_updateReviewsPosted = true;
-				_mainThreadSynchronizationContext.Post((o) => { _updateReviewsPosted = false; if (!_isDisposing) { UpdateReviews(); } }, null);
+				_mainThreadSynchronizationContext.Post((o) => 
+				{ 
+					_updateReviewsPosted = false; 
+					if (!_isDisposing) 
+					{ 
+						UpdateReviews(); 
+					} 
+				}, null);
 			}
 		}
 
@@ -1849,7 +1882,7 @@ namespace UnrealGameSync
 				}
 				if (notifyBuilds.Count > 1)
 				{
-					platformList.AppendFormat(" and {0}", notifyBuilds[notifyBuilds.Count - 1].BuildType);
+					platformList.AppendFormat(" and {0}", notifyBuilds[^1].BuildType);
 				}
 
 				// Show the balloon tooltip
@@ -1868,7 +1901,11 @@ namespace UnrealGameSync
 
 				// Set the link to open the right build pages
 				int highlightChange = notifyBuilds.Max(x => x.ChangeNumber);
-				_notificationWindow.OnMoreInformation = () => { _owner.ShowAndActivate(); SelectChange(highlightChange); };
+				_notificationWindow.OnMoreInformation = () => 
+				{ 
+					_owner.ShowAndActivate(); 
+					SelectChange(highlightChange); 
+				};
 
 				// Don't show messages for this change again
 				foreach (BadgeData notifyBuild in notifyBuilds)
@@ -1885,10 +1922,10 @@ namespace UnrealGameSync
 
 		class ExpandRowLayout
 		{
-			public string? MainText;
-			public Rectangle MainRect;
-			public string? LinkText;
-			public Rectangle LinkRect;
+			public string? _mainText;
+			public Rectangle _mainRect;
+			public string? _linkText;
+			public Rectangle _linkRect;
 		}
 
 		ExpandRowLayout LayoutExpandRow(Graphics graphics, Rectangle bounds)
@@ -1900,35 +1937,35 @@ namespace UnrealGameSync
 			int currentMaxChanges = _perforceMonitor.CurrentMaxChanges;
 			if (_perforceMonitor.PendingMaxChanges > currentMaxChanges)
 			{
-				layout.MainText = String.Format("{0}. Fetching {1} more from server...  ", showingChanges, _perforceMonitor.PendingMaxChanges - currentMaxChanges);
-				layout.LinkText = "Cancel";
+				layout._mainText = String.Format("{0}. Fetching {1} more from server...  ", showingChanges, _perforceMonitor.PendingMaxChanges - currentMaxChanges);
+				layout._linkText = "Cancel";
 			}
 			else if (_perforceMonitor.CurrentMaxChanges > _numChanges)
 			{
-				layout.MainText = showingChanges;
-				layout.LinkText = "";
+				layout._mainText = showingChanges;
+				layout._linkText = "";
 			}
 			else
 			{
-				layout.MainText = String.Format("{0}  ", showingChanges);
-				layout.LinkText = "Show more...";
+				layout._mainText = String.Format("{0}  ", showingChanges);
+				layout._linkText = "Show more...";
 			}
 
-			Size mainTextSize = TextRenderer.MeasureText(graphics, layout.MainText, _buildFont, new Size(1000, 1000), TextFormatFlags.NoPadding);
-			Size linkTextSize = TextRenderer.MeasureText(graphics, layout.LinkText, _buildFont, new Size(1000, 1000), TextFormatFlags.NoPadding);
+			Size mainTextSize = TextRenderer.MeasureText(graphics, layout._mainText, _buildFont, new Size(1000, 1000), TextFormatFlags.NoPadding);
+			Size linkTextSize = TextRenderer.MeasureText(graphics, layout._linkText, _buildFont, new Size(1000, 1000), TextFormatFlags.NoPadding);
 
 			int minX = bounds.Left + (bounds.Width - mainTextSize.Width - linkTextSize.Width) / 2;
 			int minY = bounds.Bottom - mainTextSize.Height - 1;
 
-			layout.MainRect = new Rectangle(new Point(minX, minY), mainTextSize);
-			layout.LinkRect = new Rectangle(new Point(minX + mainTextSize.Width, minY), linkTextSize);
+			layout._mainRect = new Rectangle(new Point(minX, minY), mainTextSize);
+			layout._linkRect = new Rectangle(new Point(minX + mainTextSize.Width, minY), linkTextSize);
 
 			return layout;
 		}
 
 		private void DrawExpandRow(Graphics graphics, ExpandRowLayout layout)
 		{
-			TextRenderer.DrawText(graphics, layout.MainText, _buildFont, layout.MainRect, SystemColors.WindowText, TextFormatFlags.SingleLine | TextFormatFlags.VerticalCenter | TextFormatFlags.NoPrefix | TextFormatFlags.NoPadding);
+			TextRenderer.DrawText(graphics, layout._mainText, _buildFont, layout._mainRect, SystemColors.WindowText, TextFormatFlags.SingleLine | TextFormatFlags.VerticalCenter | TextFormatFlags.NoPrefix | TextFormatFlags.NoPadding);
 
 			Color linkColor = SystemColors.HotTrack;
 			if (_mouseOverExpandLink)
@@ -1936,7 +1973,7 @@ namespace UnrealGameSync
 				linkColor = Color.FromArgb(linkColor.B, linkColor.G, linkColor.R);
 			}
 
-			TextRenderer.DrawText(graphics, layout.LinkText, _buildFont, layout.LinkRect, linkColor, TextFormatFlags.SingleLine | TextFormatFlags.VerticalCenter | TextFormatFlags.NoPrefix | TextFormatFlags.NoPadding);
+			TextRenderer.DrawText(graphics, layout._linkText, _buildFont, layout._linkRect, linkColor, TextFormatFlags.SingleLine | TextFormatFlags.VerticalCenter | TextFormatFlags.NoPrefix | TextFormatFlags.NoPadding);
 		}
 
 		private bool HitTestExpandLink(Point location)
@@ -1948,7 +1985,7 @@ namespace UnrealGameSync
 			using (Graphics graphics = Graphics.FromHwnd(IntPtr.Zero))
 			{
 				ExpandRowLayout layout = LayoutExpandRow(graphics, _expandItem.Bounds);
-				return layout.LinkRect.Contains(location);
+				return layout._linkRect.Contains(location);
 			}
 		}
 
@@ -2018,9 +2055,9 @@ namespace UnrealGameSync
 			return archivePath;
 		}
 
-		private Color Blend(Color first, Color second, float T)
+		private Color Blend(Color first, Color second, float t)
 		{
-			return Color.FromArgb((int)(first.R + (second.R - first.R) * T), (int)(first.G + (second.G - first.G) * T), (int)(first.B + (second.B - first.B) * T));
+			return Color.FromArgb((int)(first.R + (second.R - first.R) * t), (int)(first.G + (second.G - first.G) * t), (int)(first.B + (second.B - first.B) * t));
 		}
 
 		private bool CanSyncChange(int changeNumber)
@@ -2108,12 +2145,12 @@ namespace UnrealGameSync
 			{
 				layoutInfo = new ChangeLayoutInfo();
 
-				layoutInfo.DescriptionBadges = CreateDescriptionBadges(change);
-				layoutInfo.TypeBadges = CreateTypeBadges(change.Number);
+				layoutInfo._descriptionBadges = CreateDescriptionBadges(change);
+				layoutInfo._typeBadges = CreateTypeBadges(change.Number);
 
 				EventSummary? summary = _eventMonitor.GetSummaryForChange(change.Number);
-				layoutInfo.BuildBadges = CreateBuildBadges(change.Number, summary);
-				layoutInfo.CustomBadges = CreateCustomBadges(change.Number, summary);
+				layoutInfo._buildBadges = CreateBuildBadges(change.Number, summary);
+				layoutInfo._customBadges = CreateCustomBadges(change.Number, summary);
 
 				_changeNumberToLayoutInfo[change.Number] = layoutInfo;
 			}
@@ -2173,10 +2210,10 @@ namespace UnrealGameSync
 			int badgeAlpha = allowSync ? 255 : 128;
 			Color textColor = (allowSync || change.Number == _workspace.PendingChangeNumber || change.Number == _workspace.CurrentChangeNumber || _workspace.State.AdditionalChangeNumbers.Contains(change.Number)) ? SystemColors.WindowText : Blend(SystemColors.Window, SystemColors.WindowText, 0.25f);
 
-			const int fadeRange = 6;
-			if (e.ItemIndex >= BuildList.Items.Count - fadeRange && _numChanges >= _perforceMonitor.CurrentMaxChanges && !IsBisectModeEnabled())
+			const int FadeRange = 6;
+			if (e.ItemIndex >= BuildList.Items.Count - FadeRange && _numChanges >= _perforceMonitor.CurrentMaxChanges && !IsBisectModeEnabled())
 			{
-				float opacity = (float)(BuildList.Items.Count - e.ItemIndex - 0.9f) / fadeRange;
+				float opacity = (float)(BuildList.Items.Count - e.ItemIndex - 0.9f) / FadeRange;
 				badgeAlpha = (int)(badgeAlpha * opacity);
 				textColor = Blend(SystemColors.Window, textColor, opacity);
 			}
@@ -2188,17 +2225,17 @@ namespace UnrealGameSync
 				float minX = 4 * dpiScaleX;
 				if ((summary != null && _eventMonitor.WasSyncedByCurrentUser(summary.ChangeNumber)) || (_workspace != null && _workspace.CurrentChangeNumber == change.Number))
 				{
-					e.Graphics.DrawImage(Properties.Resources.Icons, minX * dpiScaleX, iconY, _previousSyncIcon, GraphicsUnit.Pixel);
+					e.Graphics.DrawImage(Properties.Resources.Icons, minX * dpiScaleX, iconY, PreviousSyncIcon, GraphicsUnit.Pixel);
 				}
 				else if (_workspaceSettings != null && _workspace != null && _workspace.State.AdditionalChangeNumbers.Contains(change.Number))
 				{
-					e.Graphics.DrawImage(Properties.Resources.Icons, minX * dpiScaleX, iconY, _additionalSyncIcon, GraphicsUnit.Pixel);
+					e.Graphics.DrawImage(Properties.Resources.Icons, minX * dpiScaleX, iconY, AdditionalSyncIcon, GraphicsUnit.Pixel);
 				}
 				else if (allowSync && ((summary != null && summary.LastStarReview != null && summary.LastStarReview.Type == EventType.Starred) || _promotedChangeNumbers.Contains(change.Number)))
 				{
-					e.Graphics.DrawImage(Properties.Resources.Icons, minX * dpiScaleX, iconY, _promotedBuildIcon, GraphicsUnit.Pixel);
+					e.Graphics.DrawImage(Properties.Resources.Icons, minX * dpiScaleX, iconY, PromotedBuildIcon, GraphicsUnit.Pixel);
 				}
-				minX += _promotedBuildIcon.Width * dpiScaleX;
+				minX += PromotedBuildIcon.Width * dpiScaleX;
 
 				if (allowSync)
 				{
@@ -2216,19 +2253,19 @@ namespace UnrealGameSync
 						}
 						else if (entry.State == BisectState.Pass)
 						{
-							qualityIcon = _bisectPassIcon;
+							qualityIcon = BisectPassIcon;
 						}
 						else if (entry.State == BisectState.Fail)
 						{
-							qualityIcon = _bisectFailIcon;
+							qualityIcon = BisectFailIcon;
 						}
 						else if (firstFail != -1 && change.Number > firstFail)
 						{
-							qualityIcon = _bisectImplicitFailIcon;
+							qualityIcon = BisectImplicitFailIcon;
 						}
 						else if (firstPass != -1 && change.Number < firstPass)
 						{
-							qualityIcon = _bisectImplicitPassIcon;
+							qualityIcon = BisectImplicitPassIcon;
 						}
 					}
 					else
@@ -2271,15 +2308,15 @@ namespace UnrealGameSync
 				ChangeLayoutInfo layout = GetChangeLayoutInfo(change);
 
 				Rectangle remainingBounds = e.Bounds;
-				if (layout.DescriptionBadges.Count > 0)
+				if (layout._descriptionBadges.Count > 0)
 				{
 					e.Graphics.IntersectClip(e.Bounds);
 					e.Graphics.SmoothingMode = SmoothingMode.HighQuality;
 
 					remainingBounds = new Rectangle(remainingBounds.Left, remainingBounds.Top, remainingBounds.Width - (int)(2 * dpiScaleX), remainingBounds.Height);
 
-					Point listLocation = GetBadgeListLocation(layout.DescriptionBadges, remainingBounds, HorizontalAlign.Right, VerticalAlignment.Middle);
-					DrawBadges(e.Graphics, listLocation, layout.DescriptionBadges, badgeAlpha);
+					Point listLocation = GetBadgeListLocation(layout._descriptionBadges, remainingBounds, HorizontalAlign.Right, VerticalAlignment.Middle);
+					DrawBadges(e.Graphics, listLocation, layout._descriptionBadges, badgeAlpha);
 
 					remainingBounds = new Rectangle(remainingBounds.Left, remainingBounds.Top, listLocation.X - remainingBounds.Left - (int)(2 * dpiScaleX), remainingBounds.Height);
 				}
@@ -2289,26 +2326,26 @@ namespace UnrealGameSync
 			else if (e.ColumnIndex == TypeColumn.Index)
 			{
 				ChangeLayoutInfo layout = GetChangeLayoutInfo(change);
-				if (layout.TypeBadges.Count > 0)
+				if (layout._typeBadges.Count > 0)
 				{
 					e.Graphics.IntersectClip(e.Bounds);
 					e.Graphics.SmoothingMode = SmoothingMode.HighQuality;
-					Point typeLocation = GetBadgeListLocation(layout.TypeBadges, e.Bounds, HorizontalAlign.Center, VerticalAlignment.Middle);
-					DrawBadges(e.Graphics, typeLocation, layout.TypeBadges, badgeAlpha);
+					Point typeLocation = GetBadgeListLocation(layout._typeBadges, e.Bounds, HorizontalAlign.Center, VerticalAlignment.Middle);
+					DrawBadges(e.Graphics, typeLocation, layout._typeBadges, badgeAlpha);
 				}
 			}
 			else if (e.ColumnIndex == CISColumn.Index)
 			{
 				ChangeLayoutInfo layout = GetChangeLayoutInfo(change);
-				if (layout.BuildBadges.Count > 0)
+				if (layout._buildBadges.Count > 0)
 				{
 					e.Graphics.IntersectClip(e.Bounds);
 					e.Graphics.SmoothingMode = SmoothingMode.HighQuality;
 
-					Point buildsLocation = GetBadgeListLocation(layout.BuildBadges, e.Bounds, HorizontalAlign.Center, VerticalAlignment.Middle);
+					Point buildsLocation = GetBadgeListLocation(layout._buildBadges, e.Bounds, HorizontalAlign.Center, VerticalAlignment.Middle);
 					buildsLocation.X = Math.Max(buildsLocation.X, e.Bounds.Left);
 
-					DrawBadges(e.Graphics, buildsLocation, layout.BuildBadges, badgeAlpha);
+					DrawBadges(e.Graphics, buildsLocation, layout._buildBadges, badgeAlpha);
 				}
 			}
 			else if (e.ColumnIndex == StatusColumn.Index)
@@ -2319,13 +2356,13 @@ namespace UnrealGameSync
 				{
 					Tuple<string, float> progress = _workspace.CurrentProgress;
 
-					maxX -= _cancelIcon.Width;
-					e.Graphics.DrawImage(Properties.Resources.Icons, maxX, iconY, _cancelIcon, GraphicsUnit.Pixel);
+					maxX -= CancelIcon.Width;
+					e.Graphics.DrawImage(Properties.Resources.Icons, maxX, iconY, CancelIcon, GraphicsUnit.Pixel);
 
 					if (!Splitter.IsLogVisible())
 					{
-						maxX -= _infoIcon.Width;
-						e.Graphics.DrawImage(Properties.Resources.Icons, maxX, iconY, _infoIcon, GraphicsUnit.Pixel);
+						maxX -= InfoIcon.Width;
+						e.Graphics.DrawImage(Properties.Resources.Icons, maxX, iconY, InfoIcon, GraphicsUnit.Pixel);
 					}
 
 					float minX = e.Bounds.Left;
@@ -2341,11 +2378,11 @@ namespace UnrealGameSync
 					{
 						EventData? review = _eventMonitor.GetReviewByCurrentUser(change.Number);
 
-						maxX -= _frownIcon.Width * dpiScaleX;
-						e.Graphics.DrawImage(Properties.Resources.Icons, maxX, iconY, (review == null || !EventMonitor.IsPositiveReview(review.Type)) ? _frownIcon : _disabledFrownIcon, GraphicsUnit.Pixel);
+						maxX -= FrownIcon.Width * dpiScaleX;
+						e.Graphics.DrawImage(Properties.Resources.Icons, maxX, iconY, (review == null || !EventMonitor.IsPositiveReview(review.Type)) ? FrownIcon : DisabledFrownIcon, GraphicsUnit.Pixel);
 
-						maxX -= _happyIcon.Width * dpiScaleX;
-						e.Graphics.DrawImage(Properties.Resources.Icons, maxX, iconY, (review == null || !EventMonitor.IsNegativeReview(review.Type)) ? _happyIcon : _disabledHappyIcon, GraphicsUnit.Pixel);
+						maxX -= HappyIcon.Width * dpiScaleX;
+						e.Graphics.DrawImage(Properties.Resources.Icons, maxX, iconY, (review == null || !EventMonitor.IsNegativeReview(review.Type)) ? HappyIcon : DisabledHappyIcon, GraphicsUnit.Pixel);
 					}
 					else if (e.ItemIndex == BuildList.HoverItem && allowSync)
 					{
@@ -2395,7 +2432,7 @@ namespace UnrealGameSync
 					ChangeLayoutInfo layout = GetChangeLayoutInfo(change);
 
 					List<BadgeInfo>? badges;
-					if (layout.CustomBadges.TryGetValue(column.Text, out badges) && badges.Count > 0)
+					if (layout._customBadges.TryGetValue(column.Text, out badges) && badges.Count > 0)
 					{
 						e.Graphics.IntersectClip(e.Bounds);
 						e.Graphics.SmoothingMode = SmoothingMode.HighQuality;
@@ -2414,15 +2451,15 @@ namespace UnrealGameSync
 			{
 				BadgeInfo badge = badges[idx];
 
-				if (idx > 0 && badge.Group != badges[idx - 1].Group)
+				if (idx > 0 && badge._group != badges[idx - 1]._group)
 				{
 					offset += 6;
 				}
-				badge.Offset = offset;
+				badge._offset = offset;
 
-				Size badgeSize = GetBadgeSize(badge.Label);
-				badge.Width = badgeSize.Width;
-				badge.Height = badgeSize.Height;
+				Size badgeSize = GetBadgeSize(badge._label);
+				badge._width = badgeSize.Width;
+				badge._height = badgeSize.Height;
 
 				offset += badgeSize.Width + 1;
 			}
@@ -2430,22 +2467,21 @@ namespace UnrealGameSync
 
 		private Point GetBadgeListLocation(List<BadgeInfo> badges, Rectangle bounds, HorizontalAlign horzAlign, VerticalAlignment vertAlign)
 		{
-			Point location = bounds.Location;
 			if (badges.Count == 0)
 			{
 				return bounds.Location;
 			}
 
-			BadgeInfo lastBadge = badges[badges.Count - 1];
+			BadgeInfo lastBadge = badges[^1];
 
 			int x = bounds.X;
 			switch (horzAlign)
 			{
 				case HorizontalAlign.Center:
-					x += (bounds.Width - lastBadge.Width - lastBadge.Offset) / 2;
+					x += (bounds.Width - lastBadge._width - lastBadge._offset) / 2;
 					break;
 				case HorizontalAlign.Right:
-					x += (bounds.Width - lastBadge.Width - lastBadge.Offset);
+					x += (bounds.Width - lastBadge._width - lastBadge._offset);
 					break;
 			}
 
@@ -2453,10 +2489,10 @@ namespace UnrealGameSync
 			switch (vertAlign)
 			{
 				case VerticalAlignment.Middle:
-					y += (bounds.Height - lastBadge.Height) / 2;
+					y += (bounds.Height - lastBadge._height) / 2;
 					break;
 				case VerticalAlignment.Bottom:
-					y += (bounds.Height - lastBadge.Height);
+					y += (bounds.Height - lastBadge._height);
 					break;
 			}
 
@@ -2468,17 +2504,17 @@ namespace UnrealGameSync
 			for (int idx = 0; idx < badges.Count; idx++)
 			{
 				BadgeInfo badge = badges[idx];
-				bool mergeLeft = (idx > 0 && badges[idx - 1].Group == badge.Group);
-				bool mergeRight = (idx < badges.Count - 1 && badges[idx + 1].Group == badge.Group);
+				bool mergeLeft = (idx > 0 && badges[idx - 1]._group == badge._group);
+				bool mergeRight = (idx < badges.Count - 1 && badges[idx + 1]._group == badge._group);
 
-				Rectangle bounds = new Rectangle(listLocation.X + badge.Offset, listLocation.Y, badge.Width, badge.Height);
-				if (badge.UniqueId != null && badge.UniqueId == _hoverBadgeUniqueId)
+				Rectangle bounds = new Rectangle(listLocation.X + badge._offset, listLocation.Y, badge._width, badge._height);
+				if (badge._uniqueId != null && badge._uniqueId == _hoverBadgeUniqueId)
 				{
-					DrawBadge(graphics, bounds, badge.Label, Color.FromArgb((alpha * badge.HoverBackgroundColor.A) / 255, badge.HoverBackgroundColor), mergeLeft, mergeRight);
+					DrawBadge(graphics, bounds, badge._label, Color.FromArgb((alpha * badge._hoverBackgroundColor.A) / 255, badge._hoverBackgroundColor), mergeLeft, mergeRight);
 				}
 				else
 				{
-					DrawBadge(graphics, bounds, badge.Label, Color.FromArgb((alpha * badge.BackgroundColor.A) / 255, badge.BackgroundColor), mergeLeft, mergeRight);
+					DrawBadge(graphics, bounds, badge._label, Color.FromArgb((alpha * badge._backgroundColor.A) / 255, badge._backgroundColor), mergeLeft, mergeRight);
 				}
 			}
 		}
@@ -2563,11 +2599,10 @@ namespace UnrealGameSync
 			ShowErrorDialog("{0} Erroring String({1}) \n {2}", errorMessage, erroringDefinition, usageMessage);
 		}
 
-
-		/**
-		 * Array of config specified options for what the Latest Change to Sync should be.
-		 * See PrintLatestChangeTypeUsage for usage information
-		 */
+		/// <summary>
+		/// Array of config specified options for what the Latest Change to Sync should be.
+		/// See PrintLatestChangeTypeUsage for usage information
+		/// </summary>
 		public List<LatestChangeType> GetCustomLatestChangeTypes()
 		{
 			List<LatestChangeType> foundLatestChangeTypes = new List<LatestChangeType>();
@@ -2721,26 +2756,6 @@ namespace UnrealGameSync
 			}
 		}
 
-		private bool TryGetProjectSetting(ConfigFile projectConfigFile, string name, string legacyName, [NotNullWhen(true)] out string? value)
-		{
-			string? newValue;
-			if (TryGetProjectSetting(projectConfigFile, name, out newValue))
-			{
-				value = newValue;
-				return true;
-			}
-
-			newValue = projectConfigFile.GetValue(legacyName, null);
-			if (newValue != null)
-			{
-				value = newValue;
-				return true;
-			}
-
-			value = null;
-			return false;
-		}
-
 		private List<BadgeInfo> CreateBuildBadges(int changeNumber, EventSummary? summary)
 		{
 			List<BadgeInfo> badges = new List<BadgeInfo>();
@@ -2761,10 +2776,10 @@ namespace UnrealGameSync
 					badgeNameToBuildData.TryGetValue(badgeNameAndGroup.Key, out badgeData);
 
 					BadgeInfo badgeInfo = CreateBadge(changeNumber, badgeNameAndGroup.Key, badgeNameAndGroup.Value, badgeData);
-					if (_maxBuildBadgeChars != -1 && badgeInfo.Label.Length > _maxBuildBadgeChars)
+					if (_maxBuildBadgeChars != -1 && badgeInfo._label.Length > _maxBuildBadgeChars)
 					{
-						badgeInfo.ToolTip = badgeInfo.Label;
-						badgeInfo.Label = badgeInfo.Label.Substring(0, _maxBuildBadgeChars);
+						badgeInfo._toolTip = badgeInfo._label;
+						badgeInfo._label = badgeInfo._label.Substring(0, _maxBuildBadgeChars);
 					}
 					badges.Add(badgeInfo);
 				}
@@ -2860,16 +2875,6 @@ namespace UnrealGameSync
 			return new Rectangle(bounds.Right - badgeSize.Width - 2, bounds.Top + (bounds.Height - badgeSize.Height) / 2, badgeSize.Width, badgeSize.Height);
 		}
 
-		private void DrawSingleBadge(Graphics graphics, Rectangle displayRectangle, string badgeText, Color badgeColor)
-		{
-			Size badgeSize = GetBadgeSize(badgeText);
-
-			int x = displayRectangle.Left + (displayRectangle.Width - badgeSize.Width) / 2;
-			int y = displayRectangle.Top + (displayRectangle.Height - badgeSize.Height) / 2;
-
-			DrawBadge(graphics, new Rectangle(x, y, badgeSize.Width, badgeSize.Height), badgeText, badgeColor, false, false);
-		}
-
 		private void DrawBadge(Graphics graphics, Rectangle badgeRect, string badgeText, Color badgeColor, bool mergeLeft, bool mergeRight)
 		{
 			if (badgeColor.A != 0)
@@ -3257,7 +3262,6 @@ namespace UnrealGameSync
 			SafeProcessStart("p4v.exe", commandLine.ToString());
 		}
 
-
 		void RunTool(ToolDefinition tool, ToolLink link)
 		{
 			DirectoryReference? toolDir = _owner.ToolUpdateMonitor.GetToolPath(tool.Name);
@@ -3320,7 +3324,10 @@ namespace UnrealGameSync
 
 				StatusLine progressLine = new StatusLine();
 				progressLine.AddText(String.Format("{0}  ", progress.Item1));
-				if (progress.Item2 > 0.0f) progressLine.AddProgressBar(progress.Item2);
+				if (progress.Item2 > 0.0f)
+				{
+					progressLine.AddProgressBar(progress.Item2);
+				}
 				lines.Add(progressLine);
 			}
 			else if (_workspace.IsExternalSyncActive())
@@ -3387,7 +3394,7 @@ namespace UnrealGameSync
 							summaryLine.AddText(" at changelist ");
 						}
 						summaryLine.AddLink(String.Format("{0}.", _workspace.CurrentChangeNumber), FontStyle.Regular, () => { SelectChange(_workspace.CurrentChangeNumber); });
-						summaryLine.AddLink(" \u25BE", 0, (P, R) =>
+						summaryLine.AddLink(" \u25BE", 0, (point, bounds) =>
 						{
 							// Try to find the CL record in our stored list in order to avoid requesting the info from p4:
 							foreach (ListViewItem? item in BuildList.Items)
@@ -3405,7 +3412,7 @@ namespace UnrealGameSync
 										else if (summary.Number == _workspace.CurrentChangeNumber)
 										{
 											_contextMenuChange = summary;
-											ShowChangelistContextMenu(StatusPanel, showTimeRelatedItems: false, new Point(R.Left, R.Bottom), ToolStripDropDownDirection.BelowRight);
+											ShowChangelistContextMenu(StatusPanel, showTimeRelatedItems: false, new Point(bounds.Left, bounds.Bottom), ToolStripDropDownDirection.BelowRight);
 											return;
 										}
 									}
@@ -3414,21 +3421,21 @@ namespace UnrealGameSync
 
 							// Couldn't find the CL in our BuildList, get the info from p4 :
 							ChangesRecord? foundRecord = null;
-							Func<IPerforceConnection, CancellationToken, Task<bool>> getChangesRecordFunc = async (Perforce, CancellationToken) =>
+							async Task<bool> GetChangesRecordFunc(IPerforceConnection perforce, CancellationToken cancellationToken)
 							{
-								List<ChangesRecord> changes = await Perforce.GetChangesAsync(ChangesOptions.None, 1, ChangeStatus.Submitted, $"//{_perforceSettings.ClientName}/...@{_workspace.CurrentChangeNumber},{_workspace.CurrentChangeNumber}", CancellationToken);
+								List<ChangesRecord> changes = await perforce.GetChangesAsync(ChangesOptions.None, 1, ChangeStatus.Submitted, $"//{_perforceSettings.ClientName}/...@{_workspace.CurrentChangeNumber},{_workspace.CurrentChangeNumber}", cancellationToken);
 								if (changes.Count == 1)
 								{
 									foundRecord = changes[0];
 								}
 								return true;
-							};
+							}
 
-							ModalTask<bool>? getChangesTask = PerforceModalTask.Execute<bool>(this, "Getting CL info", "Please wait...", _perforceSettings, getChangesRecordFunc, _logger);
+							ModalTask<bool>? getChangesTask = PerforceModalTask.Execute<bool>(this, "Getting CL info", "Please wait...", _perforceSettings, GetChangesRecordFunc, _logger);
 							if ((getChangesTask != null) && getChangesTask.Succeeded && (foundRecord != null))
 							{
 								_contextMenuChange = foundRecord;
-								ShowChangelistContextMenu(StatusPanel, showTimeRelatedItems: false, new Point(R.Left, R.Bottom), ToolStripDropDownDirection.BelowRight);
+								ShowChangelistContextMenu(StatusPanel, showTimeRelatedItems: false, new Point(bounds.Left, bounds.Bottom), ToolStripDropDownDirection.BelowRight);
 							}
 						});
 					}
@@ -3909,7 +3916,7 @@ namespace UnrealGameSync
 				{
 					for (int idx = 0; idx < 2; idx++)
 					{
-						Func<IPerforceConnection, CancellationToken, Task<bool>> switchFunc = async (perforce, cancellationToken) =>
+						async Task<bool> SwitchFunc(IPerforceConnection perforce, CancellationToken cancellationToken)
 						{
 							if (idx == 1 || !await perforce.OpenedAsync(OpenedOptions.None, FileSpecList.Any, cancellationToken).AnyAsync())
 							{
@@ -3917,9 +3924,9 @@ namespace UnrealGameSync
 								return true;
 							}
 							return false;
-						};
+						}
 
-						ModalTask<bool>? switchTask = PerforceModalTask.Execute<bool>(this, "Switching streams", "Please wait...", _perforceSettings, switchFunc, _logger);
+						ModalTask<bool>? switchTask = PerforceModalTask.Execute<bool>(this, "Switching streams", "Please wait...", _perforceSettings, SwitchFunc, _logger);
 						if (switchTask == null || !switchTask.Succeeded)
 						{
 							break;
@@ -3999,10 +4006,7 @@ namespace UnrealGameSync
 			if (_perforceMonitor.PendingMaxChanges <= currentMaxChanges && BuildList.SelectedItems.Count == 0)
 			{
 				_perforceMonitor.PendingMaxChanges = currentMaxChanges + BuildListExpandCount;
-				if (_expandItem != null)
-				{
-					_expandItem.EnsureVisible();
-				}
+				_expandItem?.EnsureVisible();
 			}
 		}
 
@@ -4077,17 +4081,17 @@ namespace UnrealGameSync
 						if (DescriptionColumn.Index < hitTest.Item.SubItems.Count && hitTest.Item.SubItems[DescriptionColumn.Index] == hitTest.SubItem)
 						{
 							ChangeLayoutInfo layoutInfo = GetChangeLayoutInfo(change);
-							if (layoutInfo.DescriptionBadges.Count > 0)
+							if (layoutInfo._descriptionBadges.Count > 0)
 							{
-								Point buildListLocation = GetBadgeListLocation(layoutInfo.DescriptionBadges, hitTest.SubItem.Bounds, HorizontalAlign.Right, VerticalAlignment.Middle);
+								Point buildListLocation = GetBadgeListLocation(layoutInfo._descriptionBadges, hitTest.SubItem.Bounds, HorizontalAlign.Right, VerticalAlignment.Middle);
 								buildListLocation.Offset(-2, 0);
 
-								foreach (BadgeInfo badge in layoutInfo.DescriptionBadges)
+								foreach (BadgeInfo badge in layoutInfo._descriptionBadges)
 								{
 									Rectangle badgeBounds = badge.GetBounds(buildListLocation);
-									if (badgeBounds.Contains(args.Location) && badge.ClickHandler != null)
+									if (badgeBounds.Contains(args.Location) && badge._clickHandler != null)
 									{
-										badge.ClickHandler();
+										badge._clickHandler();
 										break;
 									}
 								}
@@ -4097,15 +4101,15 @@ namespace UnrealGameSync
 						if (CISColumn.Index < hitTest.Item.SubItems.Count && hitTest.Item.SubItems[CISColumn.Index] == hitTest.SubItem)
 						{
 							ChangeLayoutInfo layoutInfo = GetChangeLayoutInfo(change);
-							if (layoutInfo.BuildBadges.Count > 0)
+							if (layoutInfo._buildBadges.Count > 0)
 							{
-								Point buildListLocation = GetBadgeListLocation(layoutInfo.BuildBadges, hitTest.SubItem.Bounds, HorizontalAlign.Center, VerticalAlignment.Middle);
+								Point buildListLocation = GetBadgeListLocation(layoutInfo._buildBadges, hitTest.SubItem.Bounds, HorizontalAlign.Center, VerticalAlignment.Middle);
 								buildListLocation.X = Math.Max(buildListLocation.X, hitTest.SubItem.Bounds.Left);
 
-								BadgeInfo? badgeInfo = HitTestBadge(args.Location, layoutInfo.BuildBadges, buildListLocation);
-								if (badgeInfo != null && badgeInfo.ClickHandler != null)
+								BadgeInfo? badgeInfo = HitTestBadge(args.Location, layoutInfo._buildBadges, buildListLocation);
+								if (badgeInfo != null && badgeInfo._clickHandler != null)
 								{
-									badgeInfo.ClickHandler();
+									badgeInfo._clickHandler();
 								}
 							}
 						}
@@ -4117,14 +4121,14 @@ namespace UnrealGameSync
 								ChangeLayoutInfo layoutInfo = GetChangeLayoutInfo((ChangesRecord)hitTest.Item.Tag);
 
 								List<BadgeInfo>? badges;
-								if (layoutInfo.CustomBadges.TryGetValue(customColumn.Text, out badges) && badges.Count > 0)
+								if (layoutInfo._customBadges.TryGetValue(customColumn.Text, out badges) && badges.Count > 0)
 								{
 									Point listLocation = GetBadgeListLocation(badges, hitTest.SubItem.Bounds, HorizontalAlign.Center, VerticalAlignment.Middle);
 
 									BadgeInfo? badgeInfo = HitTestBadge(args.Location, badges, listLocation);
-									if (badgeInfo != null && badgeInfo.ClickHandler != null)
+									if (badgeInfo != null && badgeInfo._clickHandler != null)
 									{
-										badgeInfo.ClickHandler();
+										badgeInfo._clickHandler();
 									}
 								}
 							}
@@ -4267,34 +4271,34 @@ namespace UnrealGameSync
 			}
 
 			string? newHoverBadgeUniqueId = null;
-			if (hitTest.Item != null && hitTest.Item.Tag is ChangesRecord)
+			if (hitTest.Item != null && hitTest.Item.Tag is ChangesRecord changesRecord)
 			{
 				int columnIndex = hitTest.Item.SubItems.IndexOf(hitTest.SubItem);
 				if (columnIndex == DescriptionColumn.Index)
 				{
-					ChangeLayoutInfo layoutInfo = GetChangeLayoutInfo((ChangesRecord)hitTest.Item.Tag);
-					if (layoutInfo.DescriptionBadges.Count > 0)
+					ChangeLayoutInfo layoutInfo = GetChangeLayoutInfo(changesRecord);
+					if (layoutInfo._descriptionBadges.Count > 0)
 					{
-						Point listLocation = GetBadgeListLocation(layoutInfo.DescriptionBadges, hitTest.SubItem.Bounds, HorizontalAlign.Right, VerticalAlignment.Middle);
-						newHoverBadgeUniqueId = HitTestBadge(e.Location, layoutInfo.DescriptionBadges, listLocation)?.UniqueId;
+						Point listLocation = GetBadgeListLocation(layoutInfo._descriptionBadges, hitTest.SubItem.Bounds, HorizontalAlign.Right, VerticalAlignment.Middle);
+						newHoverBadgeUniqueId = HitTestBadge(e.Location, layoutInfo._descriptionBadges, listLocation)?._uniqueId;
 					}
 				}
 				else if (columnIndex == CISColumn.Index)
 				{
-					ChangeLayoutInfo layoutInfo = GetChangeLayoutInfo((ChangesRecord)hitTest.Item.Tag);
-					if (layoutInfo.BuildBadges.Count > 0)
+					ChangeLayoutInfo layoutInfo = GetChangeLayoutInfo(changesRecord);
+					if (layoutInfo._buildBadges.Count > 0)
 					{
-						Point buildListLocation = GetBadgeListLocation(layoutInfo.BuildBadges, hitTest.SubItem.Bounds, HorizontalAlign.Center, VerticalAlignment.Middle);
+						Point buildListLocation = GetBadgeListLocation(layoutInfo._buildBadges, hitTest.SubItem.Bounds, HorizontalAlign.Center, VerticalAlignment.Middle);
 						buildListLocation.X = Math.Max(buildListLocation.X, hitTest.SubItem.Bounds.Left);
 
-						BadgeInfo? badge = HitTestBadge(e.Location, layoutInfo.BuildBadges, buildListLocation);
-						newHoverBadgeUniqueId = (badge != null) ? badge.UniqueId : null;
+						BadgeInfo? badge = HitTestBadge(e.Location, layoutInfo._buildBadges, buildListLocation);
+						newHoverBadgeUniqueId = badge?._uniqueId;
 
 						if (_hoverBadgeUniqueId != newHoverBadgeUniqueId)
 						{
-							if (badge != null && badge.ToolTip != null && badge.BackgroundColor.A != 0)
+							if (badge?._toolTip != null && badge._backgroundColor.A != 0)
 							{
-								BuildListToolTip.Show(badge.ToolTip, BuildList, new Point(buildListLocation.X + badge.Offset, hitTest.Item.Bounds.Bottom + 2));
+								BuildListToolTip.Show(badge._toolTip, BuildList, new Point(buildListLocation.X + badge._offset, hitTest.Item.Bounds.Bottom + 2));
 							}
 							else
 							{
@@ -4310,10 +4314,10 @@ namespace UnrealGameSync
 					ChangeLayoutInfo layoutInfo = GetChangeLayoutInfo((ChangesRecord)hitTest.Item.Tag);
 
 					List<BadgeInfo>? badges;
-					if (layoutInfo.CustomBadges.TryGetValue(column.Text, out badges) && badges.Count > 0)
+					if (layoutInfo._customBadges.TryGetValue(column.Text, out badges) && badges.Count > 0)
 					{
 						Point listLocation = GetBadgeListLocation(badges, hitTest.SubItem.Bounds, HorizontalAlign.Center, VerticalAlignment.Middle);
-						newHoverBadgeUniqueId = HitTestBadge(e.Location, badges, listLocation)?.UniqueId;
+						newHoverBadgeUniqueId = HitTestBadge(e.Location, badges, listLocation)?._uniqueId;
 					}
 				}
 			}
@@ -4480,7 +4484,9 @@ namespace UnrealGameSync
 		{
 			IReadOnlyList<IArchiveInfo>? availableArchives = _jupiterMonitor?.AvailableArchives;
 			if (availableArchives != null && availableArchives.Count != 0)
+			{
 				return availableArchives;
+			}
 
 			// if jupiter had no archives we fallback to the perforce monitor
 			if (_perforceMonitor != null)
@@ -4539,22 +4545,13 @@ namespace UnrealGameSync
 
 		private void BuildList_FontChanged(object? sender, EventArgs e)
 		{
-			if (_buildFont != null)
-			{
-				_buildFont.Dispose();
-			}
+			_buildFont?.Dispose();
 			_buildFont = BuildList.Font;
 
-			if (_selectedBuildFont != null)
-			{
-				_selectedBuildFont.Dispose();
-			}
+			_selectedBuildFont?.Dispose();
 			_selectedBuildFont = new Font(_buildFont, FontStyle.Bold);
 
-			if (_badgeFont != null)
-			{
-				_badgeFont.Dispose();
-			}
+			_badgeFont?.Dispose();
 			_badgeFont = new Font(_buildFont.FontFamily, _buildFont.SizeInPoints - 2, FontStyle.Bold);
 		}
 
@@ -4798,10 +4795,7 @@ namespace UnrealGameSync
 		{
 			UpdateSyncActionCheckboxes();
 
-			if (_perforceMonitor != null)
-			{
-				_perforceMonitor.Refresh();
-			}
+			_perforceMonitor?.Refresh();
 
 			if (DesiredTaskbarState.Item1 == TaskbarState.Error)
 			{
@@ -4843,7 +4837,7 @@ namespace UnrealGameSync
 						if (descriptionBounds.Contains(clientPoint))
 						{
 							ChangeLayoutInfo layoutInfo = GetChangeLayoutInfo(change);
-							if (layoutInfo.DescriptionBadges.Count == 0 || clientPoint.X < GetBadgeListLocation(layoutInfo.DescriptionBadges, descriptionBounds, HorizontalAlign.Right, VerticalAlignment.Middle).X - 2)
+							if (layoutInfo._descriptionBadges.Count == 0 || clientPoint.X < GetBadgeListLocation(layoutInfo._descriptionBadges, descriptionBounds, HorizontalAlign.Right, VerticalAlignment.Middle).X - 2)
 							{
 								BuildListToolTip.Show(change.Description, BuildList, new Point(descriptionBounds.Left, descriptionBounds.Bottom + 2));
 								return;
@@ -5051,7 +5045,7 @@ namespace UnrealGameSync
 
 				_logger.LogInformation("Scheduled sync at {0} for {1}.",
 					DateTime.Now,
-					scheduledSyncTypeId != null ? scheduledSyncTypeId : "Default");
+					scheduledSyncTypeId ?? "Default");
 
 				int changeNumber;
 				if (!FindChangeToSync(scheduledSyncTypeId, out changeNumber))
@@ -5435,10 +5429,7 @@ namespace UnrealGameSync
 			}
 		}
 
-		private bool ShouldSyncPrecompiledEditor
-		{
-			get { return _settings.Archives.Any(x => x.Enabled && x.Type == IArchiveInfo.EditorArchiveType) && GetArchives().Any(x => x.Type == "Editor"); }
-		}
+		private bool ShouldSyncPrecompiledEditor => _settings.Archives.Any(x => x.Enabled && x.Type == IArchiveInfo.EditorArchiveType) && GetArchives().Any(x => x.Type == "Editor");
 
 		public BuildConfig GetEditorBuildConfig()
 		{
@@ -5528,8 +5519,9 @@ namespace UnrealGameSync
 		{
 			SyncContextMenu.Items.Clear();
 			SyncContextMenu.Items.AddRange(new System.Windows.Forms.ToolStripItem[] {
-			this.SyncContexMenu_EnterChangelist,
-			this.toolStripSeparator8});
+				SyncContexMenu_EnterChangelist,
+				toolStripSeparator8
+			});
 
 			foreach (LatestChangeType changeType in GetCustomLatestChangeTypes())
 			{
@@ -5537,7 +5529,7 @@ namespace UnrealGameSync
 				menuItem.Name = changeType.Name;
 				menuItem.Text = changeType.Description;
 				menuItem.Size = new System.Drawing.Size(189, 22);
-				menuItem.Click += (sender, e) => this.SyncContextMenu_LatestChangeType_Click(sender, e, changeType.Name);
+				menuItem.Click += (sender, e) => SyncContextMenu_LatestChangeType_Click(changeType.Name);
 
 				SyncContextMenu.Items.Add(menuItem);
 			}
@@ -5579,7 +5571,7 @@ namespace UnrealGameSync
 			SyncContextMenu.Show(StatusPanel, new Point(bounds.Left, bounds.Bottom), ToolStripDropDownDirection.BelowRight);
 		}
 
-		private void SyncContextMenu_LatestChangeType_Click(object? sender, EventArgs e, string syncTypeId)
+		private void SyncContextMenu_LatestChangeType_Click(string syncTypeId)
 		{
 			_settings.SyncTypeId = syncTypeId;
 			_settings.Save(_logger);
@@ -5603,10 +5595,7 @@ namespace UnrealGameSync
 		{
 			if (args.KeyCode == Keys.F5 || (args.Control && args.KeyCode == Keys.R))
 			{
-				if (_perforceMonitor != null)
-				{
-					_perforceMonitor.Refresh();
-				}
+				_perforceMonitor?.Refresh();
 			}
 
 			if (args.Control && args.KeyCode == Keys.C && BuildList.SelectedItems.Count > 0)
@@ -6018,7 +6007,7 @@ namespace UnrealGameSync
 			FilterContextMenu_Robomerge_Annotate.Checked = _settings.AnnotateRobmergeChanges;
 
 			// Set checks if an author filter string is set
-			FilterContextMenu_Author.Checked = !string.IsNullOrEmpty(this._authorFilterText) && this._authorFilterText != _authorFilterPlaceholderText;
+			FilterContextMenu_Author.Checked = !String.IsNullOrEmpty(_authorFilterText) && _authorFilterText != AuthorFilterPlaceholderText;
 
 			FilterContextMenu.Show(FilterButton, new Point(0, FilterButton.Height));
 		}
@@ -6096,10 +6085,10 @@ namespace UnrealGameSync
 		/// <param name="e"></param>
 		private void FilterContextMenu_Author_Name_GotFocus(object sender, EventArgs e)
 		{
-			if (this.FilterContextMenu_Author_Name.Text == _authorFilterPlaceholderText)
+			if (FilterContextMenu_Author_Name.Text == AuthorFilterPlaceholderText)
 			{
-				this.FilterContextMenu_Author_Name.Text = "";
-				this.FilterContextMenu_Author_Name.TextBox.ForeColor = Color.Black;
+				FilterContextMenu_Author_Name.Text = "";
+				FilterContextMenu_Author_Name.TextBox.ForeColor = Color.Black;
 			}
 		}
 
@@ -6111,10 +6100,10 @@ namespace UnrealGameSync
 		/// <param name="e"></param>
 		private void FilterContextMenu_Author_Name_LostFocus(object sender, EventArgs e)
 		{
-			if (string.IsNullOrEmpty(this.FilterContextMenu_Author_Name.Text))
+			if (String.IsNullOrEmpty(FilterContextMenu_Author_Name.Text))
 			{
-				this.FilterContextMenu_Author_Name.Text = _authorFilterPlaceholderText;
-				this.FilterContextMenu_Author_Name.TextBox.ForeColor = Color.DarkGray;
+				FilterContextMenu_Author_Name.Text = AuthorFilterPlaceholderText;
+				FilterContextMenu_Author_Name.TextBox.ForeColor = Color.DarkGray;
 			}
 		}
 
@@ -6126,11 +6115,10 @@ namespace UnrealGameSync
 		/// <param name="e"></param>
 		private void FilterContextMenu_Author_Name_Changed(object sender, EventArgs e)
 		{
-			_authorFilterText = this.FilterContextMenu_Author_Name.Text;
+			_authorFilterText = FilterContextMenu_Author_Name.Text;
 			UpdateBuildListFilter();
 		}
 
-
 		private void FilterContextMenu_Type_ShowAll_Click(object sender, EventArgs e)
 		{
 			_projectSettings.FilterType = FilterType.None;
@@ -6221,7 +6209,7 @@ namespace UnrealGameSync
 			ConfigObject overrideObject = new ConfigObject(backgroundColor);
 
 			float r, g, b;
-			if (!float.TryParse(overrideObject.GetValue("R", ""), out r) || !float.TryParse(overrideObject.GetValue("G", ""), out g) || !float.TryParse(overrideObject.GetValue("B", ""), out b))
+			if (!Single.TryParse(overrideObject.GetValue("R", ""), out r) || !Single.TryParse(overrideObject.GetValue("G", ""), out g) || !Single.TryParse(overrideObject.GetValue("B", ""), out b))
 			{
 				return null;
 			}
@@ -6275,7 +6263,7 @@ namespace UnrealGameSync
 			ConfigObject specifiedColorObject = new ConfigObject(specifiedColorValue);
 
 			float r, g, b;
-			if (!float.TryParse(specifiedColorObject.GetValue("R", ""), out r) || !float.TryParse(specifiedColorObject.GetValue("G", ""), out g) || !float.TryParse(specifiedColorObject.GetValue("B", ""), out b))
+			if (!Single.TryParse(specifiedColorObject.GetValue("R", ""), out r) || !Single.TryParse(specifiedColorObject.GetValue("G", ""), out g) || !Single.TryParse(specifiedColorObject.GetValue("B", ""), out b))
 			{
 				return null;
 			}
diff --git a/Engine/Source/Programs/UnrealGameSync/UnrealGameSync/DeploymentSettings.cs b/Engine/Source/Programs/UnrealGameSync/UnrealGameSync/DeploymentSettings.cs
--- a/Engine/Source/Programs/UnrealGameSync/UnrealGameSync/DeploymentSettings.cs
+++ b/Engine/Source/Programs/UnrealGameSync/UnrealGameSync/DeploymentSettings.cs
@@ -1,13 +1,11 @@
 // Copyright Epic Games, Inc. All Rights Reserved.
 
+#pragma warning disable IDE0005
+
 using Microsoft.Extensions.Logging;
-using System;
+using Microsoft.Extensions.Logging.Abstractions;
 using System.Collections.Generic;
 using System.Diagnostics.CodeAnalysis;
-using System.IO;
-using System.Linq;
-using System.Text;
-using System.Threading.Tasks;
 
 namespace UnrealGameSync
 {
diff --git a/Engine/Source/Programs/UnrealGameSync/UnrealGameSync/EventMonitor.cs b/Engine/Source/Programs/UnrealGameSync/UnrealGameSync/EventMonitor.cs
--- a/Engine/Source/Programs/UnrealGameSync/UnrealGameSync/EventMonitor.cs
+++ b/Engine/Source/Programs/UnrealGameSync/UnrealGameSync/EventMonitor.cs
@@ -3,7 +3,6 @@
 using System;
 using System.Collections.Concurrent;
 using System.Collections.Generic;
-using System.Globalization;
 using System.Diagnostics;
 using System.IO;
 using System.Linq;
@@ -81,15 +80,9 @@ namespace UnrealGameSync
 		public string Url { get; set; } = String.Empty;
 		public string Project { get; set; } = String.Empty;
 
-		public bool IsSuccess
-		{
-			get { return Result == BadgeResult.Success || Result == BadgeResult.Warning; }
-		}
+		public bool IsSuccess => Result == BadgeResult.Success || Result == BadgeResult.Warning;
 
-		public bool IsFailure
-		{
-			get { return Result == BadgeResult.Failure; }
-		}
+		public bool IsFailure => Result == BadgeResult.Failure;
 
 		public string BadgeName
 		{
@@ -212,33 +205,33 @@ namespace UnrealGameSync
 
 	class EventMonitor : IDisposable
 	{
-		string? _apiUrl;
+		readonly string? _apiUrl;
 		int _apiVersion;
-		string _project;
-		string _currentUserName;
-		SynchronizationContext _synchronizationContext;
-		CancellationTokenSource _cancellationSource;
+		readonly string _project;
+		readonly string _currentUserName;
+		readonly SynchronizationContext _synchronizationContext;
+		readonly CancellationTokenSource _cancellationSource;
 		Task? _workerTask;
-		AsyncEvent _refreshEvent = new AsyncEvent();
-		ConcurrentQueue<EventData> _outgoingEvents = new ConcurrentQueue<EventData>();
-		ConcurrentQueue<EventData> _incomingEvents = new ConcurrentQueue<EventData>();
-		ConcurrentQueue<CommentData> _outgoingComments = new ConcurrentQueue<CommentData>();
-		ConcurrentQueue<CommentData> _incomingComments = new ConcurrentQueue<CommentData>();
-		ConcurrentQueue<BadgeData> _incomingBadges = new ConcurrentQueue<BadgeData>();
-		SortedDictionary<int, EventSummary> _changeNumberToSummary = new SortedDictionary<int, EventSummary>();
-		Dictionary<string, EventData> _userNameToLastSyncEvent = new Dictionary<string, EventData>(StringComparer.InvariantCultureIgnoreCase);
-		Dictionary<string, BadgeData> _badgeNameToLatestData = new Dictionary<string, BadgeData>();
-		ILogger _logger;
-		IAsyncDisposer _asyncDisposer;
-		LatestData _latestIds;
+		readonly AsyncEvent _refreshEvent = new AsyncEvent();
+		readonly ConcurrentQueue<EventData> _outgoingEvents = new ConcurrentQueue<EventData>();
+		readonly ConcurrentQueue<EventData> _incomingEvents = new ConcurrentQueue<EventData>();
+		readonly ConcurrentQueue<CommentData> _outgoingComments = new ConcurrentQueue<CommentData>();
+		readonly ConcurrentQueue<CommentData> _incomingComments = new ConcurrentQueue<CommentData>();
+		readonly ConcurrentQueue<BadgeData> _incomingBadges = new ConcurrentQueue<BadgeData>();
+		readonly SortedDictionary<int, EventSummary> _changeNumberToSummary = new SortedDictionary<int, EventSummary>();
+		readonly Dictionary<string, EventData> _userNameToLastSyncEvent = new Dictionary<string, EventData>(StringComparer.InvariantCultureIgnoreCase);
+		readonly Dictionary<string, BadgeData> _badgeNameToLatestData = new Dictionary<string, BadgeData>();
+		readonly ILogger _logger;
+		readonly IAsyncDisposer _asyncDisposer;
+		readonly LatestData _latestIds;
 		HashSet<int> _filterChangeNumbers = new HashSet<int>();
-		List<EventData> _investigationEvents = new List<EventData>();
+		readonly List<EventData> _investigationEvents = new List<EventData>();
 		List<EventData>? _activeInvestigations;
 
 		// MetadataV2
-		string _metadataStream;
-		string _metadataProject;
-		ConcurrentQueue<GetMetadataResponseV2> _incomingMetadata = new ConcurrentQueue<GetMetadataResponseV2>();
+		readonly string _metadataStream;
+		readonly string _metadataProject;
+		readonly ConcurrentQueue<GetMetadataResponseV2> _incomingMetadata = new ConcurrentQueue<GetMetadataResponseV2>();
 		int _minChange;
 		int _newMinChange;
 		long _metadataSequenceNumber;
@@ -285,10 +278,7 @@ namespace UnrealGameSync
 
 		public void Start()
 		{
-			if (_workerTask == null)
-			{
-				_workerTask = Task.Run(() => PollForUpdatesAsync(_cancellationSource.Token));
-			}
+			_workerTask ??= Task.Run(() => PollForUpdatesAsync(_cancellationSource.Token));
 		}
 
 		public void Dispose()
@@ -491,7 +481,7 @@ namespace UnrealGameSync
 		void ApplyCommentUpdate(CommentData comment)
 		{
 			EventSummary summary = FindOrAddSummary(comment.ChangeNumber);
-			if(String.Compare(comment.UserName, _currentUserName, true) == 0 && summary.Comments.Count > 0 && summary.Comments.Last().Id == long.MaxValue)
+			if(String.Compare(comment.UserName, _currentUserName, true) == 0 && summary.Comments.Count > 0 && summary.Comments.Last().Id == Int64.MaxValue)
 			{
 				// This comment was added by PostComment(), to mask the latency of a round trip to the server. Remove it now we have the sorted comment.
 				summary.Comments.RemoveAt(summary.Comments.Count - 1);
@@ -560,7 +550,7 @@ namespace UnrealGameSync
 				newSummary.ProjectMetadata = summary.ProjectMetadata;
 			}
 
-			if (string.IsNullOrEmpty(metadata.Project))
+			if (String.IsNullOrEmpty(metadata.Project))
 			{
 				newSummary.SharedMetadata = metadata;
 			}
@@ -781,7 +771,7 @@ namespace UnrealGameSync
 				_logger.LogInformation("Posting event... ({Change}, {UserName}, {Type})", evt.Change, evt.UserName, evt.Type);
 				if (_apiVersion == 2)
 				{
-					await SendMetadataUpdateUpdateV2Async(evt.Change, evt.Project, evt.UserName, evt.Type, null, cancellationToken);
+					await SendMetadataUpdateUpdateV2Async(evt.Change, evt.UserName, evt.Type, null, cancellationToken);
 				}
 				else
 				{
@@ -804,7 +794,7 @@ namespace UnrealGameSync
 				_logger.LogInformation("Posting comment... ({Change}, {User}, {Text}, {Project})", comment.ChangeNumber, comment.UserName, comment.Text, comment.Project);
 				if (_apiVersion == 2)
 				{
-					await SendMetadataUpdateUpdateV2Async(comment.ChangeNumber, comment.Project, comment.UserName, null, comment.Text, cancellationToken);
+					await SendMetadataUpdateUpdateV2Async(comment.ChangeNumber, comment.UserName, null, comment.Text, cancellationToken);
 				}
 				else
 				{
@@ -820,7 +810,7 @@ namespace UnrealGameSync
 			}
 		}
 
-		async Task SendMetadataUpdateUpdateV2Async(int change, string project, string userName, EventType? evt, string? comment, CancellationToken cancellationToken)
+		async Task SendMetadataUpdateUpdateV2Async(int change, string userName, EventType? evt, string? comment, CancellationToken cancellationToken)
 		{
 			UpdateMetadataRequestV2 update = new UpdateMetadataRequestV2();
 			update.Stream = _metadataStream;
@@ -977,11 +967,6 @@ namespace UnrealGameSync
 			}
 		}
 
-		static bool MatchesWildcard(string wildcard, string project)
-		{
-			return wildcard.EndsWith("...") && project.StartsWith(wildcard.Substring(0, wildcard.Length - 3), StringComparison.InvariantCultureIgnoreCase);
-		}
-
 		public void PostEvent(int changeNumber, EventType type)
 		{
 			if(_apiUrl != null)
@@ -1005,7 +990,7 @@ namespace UnrealGameSync
 			if(_apiUrl != null)
 			{
 				CommentData comment = new CommentData();
-				comment.Id = long.MaxValue;
+				comment.Id = Int64.MaxValue;
 				comment.ChangeNumber = changeNumber;
 				comment.UserName = _currentUserName;
 				comment.Text = text;
diff --git a/Engine/Source/Programs/UnrealGameSync/UnrealGameSync/FindFoldersToCleanTask.cs b/Engine/Source/Programs/UnrealGameSync/UnrealGameSync/FindFoldersToCleanTask.cs
--- a/Engine/Source/Programs/UnrealGameSync/UnrealGameSync/FindFoldersToCleanTask.cs
+++ b/Engine/Source/Programs/UnrealGameSync/UnrealGameSync/FindFoldersToCleanTask.cs
@@ -7,7 +7,6 @@ using System.Collections.Generic;
 using System.Diagnostics;
 using System.IO;
 using System.Linq;
-using System.Text;
 using System.Threading;
 using System.Threading.Tasks;
 
@@ -15,43 +14,37 @@ namespace UnrealGameSync
 {
 	class FolderToClean
 	{
-		public DirectoryInfo Directory;
-		public Dictionary<string, FolderToClean> NameToSubFolder = new Dictionary<string, FolderToClean>(StringComparer.InvariantCultureIgnoreCase);
-		public Dictionary<string, FileInfo> NameToFile = new Dictionary<string, FileInfo>(StringComparer.InvariantCultureIgnoreCase);
-		public List<FileInfo> FilesToDelete = new List<FileInfo>();
-		public List<FileInfo> FilesToSync = new List<FileInfo>();
-		public bool EmptyLeaf = false;
-		public bool EmptyAfterClean = true;
+		public DirectoryInfo _directory;
+		public Dictionary<string, FolderToClean> _nameToSubFolder = new Dictionary<string, FolderToClean>(StringComparer.InvariantCultureIgnoreCase);
+		public Dictionary<string, FileInfo> _nameToFile = new Dictionary<string, FileInfo>(StringComparer.InvariantCultureIgnoreCase);
+		public List<FileInfo> _filesToDelete = new List<FileInfo>();
+		public List<FileInfo> _filesToSync = new List<FileInfo>();
+		public bool _emptyLeaf = false;
+		public bool _emptyAfterClean = true;
 
 		public FolderToClean(DirectoryInfo inDirectory)
 		{
-			Directory = inDirectory;
+			_directory = inDirectory;
 		}
 
-		public string Name
-		{
-			get { return Directory.Name; }
-		}
+		public string Name => _directory.Name;
 
-		public override string ToString()
-		{
-			return Directory.FullName;
-		}
+		public override string ToString() => _directory.FullName;
 	}
 
 	class FindFoldersToCleanTask : IDisposable
 	{
 		class PerforceHaveFolder
 		{
-			public Dictionary<string, PerforceHaveFolder> NameToSubFolder = new Dictionary<string,PerforceHaveFolder>(StringComparer.OrdinalIgnoreCase);
-			public Dictionary<string, FStatRecord> NameToFile = new Dictionary<string, FStatRecord>(StringComparer.OrdinalIgnoreCase);
+			public Dictionary<string, PerforceHaveFolder> _nameToSubFolder = new Dictionary<string,PerforceHaveFolder>(StringComparer.OrdinalIgnoreCase);
+			public Dictionary<string, FStatRecord> _nameToFile = new Dictionary<string, FStatRecord>(StringComparer.OrdinalIgnoreCase);
 		}
 
-		IPerforceSettings _perforceSettings;
-		string _clientRootPath;
-		IReadOnlyList<string> _syncPaths;
-		ILogger _logger;
-		FolderToClean _rootFolderToClean;
+		readonly IPerforceSettings _perforceSettings;
+		readonly string _clientRootPath;
+		readonly IReadOnlyList<string> _syncPaths;
+		readonly ILogger _logger;
+		readonly FolderToClean _rootFolderToClean;
 
 		int _remainingFoldersToScan;
 		ManualResetEvent _finishedScan = new ManualResetEvent(false);
@@ -65,7 +58,7 @@ namespace UnrealGameSync
 			_perforceSettings = inPerforceSettings;
 			_clientRootPath = inClientRootPath.TrimEnd('/') + "/";
 			_syncPaths = new List<string>(inSyncPaths);
-			this._logger = logger;
+			_logger = logger;
 			_rootFolderToClean = inRootFolderToClean;
 			_finishedScan = new ManualResetEvent(true);
 		}
@@ -85,24 +78,24 @@ namespace UnrealGameSync
 			{
 				try
 				{
-					if ((folder.Directory.Attributes & FileAttributes.ReparsePoint) == 0)
+					if ((folder._directory.Attributes & FileAttributes.ReparsePoint) == 0)
 					{
-						foreach (DirectoryInfo subDirectory in folder.Directory.EnumerateDirectories())
+						foreach (DirectoryInfo subDirectory in folder._directory.EnumerateDirectories())
 						{
 							FolderToClean subFolder = new FolderToClean(subDirectory);
-							folder.NameToSubFolder[subFolder.Name] = subFolder;
+							folder._nameToSubFolder[subFolder.Name] = subFolder;
 							QueueFolderToPopulate(subFolder);
 						}
-						foreach (FileInfo file in folder.Directory.EnumerateFiles())
+						foreach (FileInfo file in folder._directory.EnumerateFiles())
 						{
 							FileAttributes attributes = file.Attributes; // Force the value to be cached.
-							folder.NameToFile[file.Name] = file;
+							folder._nameToFile[file.Name] = file;
 						}
 					}
 				}
 				catch (Exception ex)
 				{
-					string newError = String.Format("Unable to enumerate contents of {0} due to an error:\n\n{1}", folder.Directory.FullName, ex);
+					string newError = String.Format("Unable to enumerate contents of {0} due to an error:\n\n{1}", folder._directory.FullName, ex);
 					Interlocked.CompareExchange(ref _scanError, newError, null);
 					_abortScan = true;
 				}
@@ -119,56 +112,56 @@ namespace UnrealGameSync
 			if(perforceFolder == null)
 			{
 				// Loop through all the local sub-folders
-				foreach(FolderToClean localSubFolder in localFolder.NameToSubFolder.Values)
+				foreach(FolderToClean localSubFolder in localFolder._nameToSubFolder.Values)
 				{
 					MergeTrees(localSubFolder, null, openClientPaths, perforceConfigFile);
 				}
 
 				// Delete everything
-				localFolder.FilesToDelete.AddRange(localFolder.NameToFile.Values);
+				localFolder._filesToDelete.AddRange(localFolder._nameToFile.Values);
 			}
 			else
 			{
 				// Loop through all the local sub-folders
-				foreach(FolderToClean localSubFolder in localFolder.NameToSubFolder.Values)
+				foreach(FolderToClean localSubFolder in localFolder._nameToSubFolder.Values)
 				{
 					PerforceHaveFolder? perforceSubFolder;
-					perforceFolder.NameToSubFolder.TryGetValue(localSubFolder.Name, out perforceSubFolder);
+					perforceFolder._nameToSubFolder.TryGetValue(localSubFolder.Name, out perforceSubFolder);
 					MergeTrees(localSubFolder, perforceSubFolder, openClientPaths, perforceConfigFile);
 				}
 
 				// Also merge all the Perforce folders that no longer exist
-				foreach(KeyValuePair<string, PerforceHaveFolder> perforceSubFolderPair in perforceFolder.NameToSubFolder)
+				foreach(KeyValuePair<string, PerforceHaveFolder> perforceSubFolderPair in perforceFolder._nameToSubFolder)
 				{
 					FolderToClean? localSubFolder;
-					if(!localFolder.NameToSubFolder.TryGetValue(perforceSubFolderPair.Key, out localSubFolder))
+					if(!localFolder._nameToSubFolder.TryGetValue(perforceSubFolderPair.Key, out localSubFolder))
 					{
-						localSubFolder = new FolderToClean(new DirectoryInfo(Path.Combine(localFolder.Directory.FullName, perforceSubFolderPair.Key)));
+						localSubFolder = new FolderToClean(new DirectoryInfo(Path.Combine(localFolder._directory.FullName, perforceSubFolderPair.Key)));
 						MergeTrees(localSubFolder, perforceSubFolderPair.Value, openClientPaths, perforceConfigFile);
-						localFolder.NameToSubFolder.Add(localSubFolder.Name, localSubFolder);
+						localFolder._nameToSubFolder.Add(localSubFolder.Name, localSubFolder);
 					}
 				}
 
 				// Find all the files that need to be re-synced
-				foreach(KeyValuePair<string, FStatRecord> filePair in perforceFolder.NameToFile)
+				foreach(KeyValuePair<string, FStatRecord> filePair in perforceFolder._nameToFile)
 				{
 					FileInfo? localFile;
-					if(!localFolder.NameToFile.TryGetValue(filePair.Key, out localFile))
+					if(!localFolder._nameToFile.TryGetValue(filePair.Key, out localFile))
 					{
-						localFolder.FilesToSync.Add(new FileInfo(Path.Combine(localFolder.Directory.FullName, filePair.Key)));
+						localFolder._filesToSync.Add(new FileInfo(Path.Combine(localFolder._directory.FullName, filePair.Key)));
 					}
 					else if(!IsFileTypeWritable(filePair.Value?.HeadType ?? filePair.Value?.Type) && (localFile.Attributes & FileAttributes.ReadOnly) == 0 && !openClientPaths.Contains(filePair.Value?.ClientFile ?? ""))
 					{
-						localFolder.FilesToSync.Add(localFile);
+						localFolder._filesToSync.Add(localFile);
 					}
 				}
 
 				// Find all the files that should be deleted
-				foreach(FileInfo localFileInfo in localFolder.NameToFile.Values)
+				foreach(FileInfo localFileInfo in localFolder._nameToFile.Values)
 				{
-					if(!perforceFolder.NameToFile.ContainsKey(localFileInfo.Name) && !openClientPaths.Contains(localFileInfo.FullName))
+					if(!perforceFolder._nameToFile.ContainsKey(localFileInfo.Name) && !openClientPaths.Contains(localFileInfo.FullName))
 					{
-						localFolder.FilesToDelete.Add(localFileInfo);
+						localFolder._filesToDelete.Add(localFileInfo);
 					}
 				}
 			}
@@ -176,14 +169,14 @@ namespace UnrealGameSync
 			// Remove any config files
 			if(perforceConfigFile != null)
 			{
-				localFolder.FilesToDelete.RemoveAll(x => String.Compare(x.Name, perforceConfigFile, StringComparison.OrdinalIgnoreCase) == 0);
+				localFolder._filesToDelete.RemoveAll(x => String.Compare(x.Name, perforceConfigFile, StringComparison.OrdinalIgnoreCase) == 0);
 			}
 
 			// Figure out if this folder is just an empty directory that needs to be removed
-			localFolder.EmptyLeaf = localFolder.NameToFile.Count == 0 && localFolder.NameToSubFolder.Count == 0 && localFolder.FilesToSync.Count == 0;
+			localFolder._emptyLeaf = localFolder._nameToFile.Count == 0 && localFolder._nameToSubFolder.Count == 0 && localFolder._filesToSync.Count == 0;
 
 			// Figure out if it the folder will be empty after the clean operation
-			localFolder.EmptyAfterClean = localFolder.NameToSubFolder.Values.All(x => x.EmptyAfterClean) && localFolder.FilesToDelete.Count == localFolder.NameToFile.Count && localFolder.FilesToSync.Count == 0;
+			localFolder._emptyAfterClean = localFolder._nameToSubFolder.Values.All(x => x._emptyAfterClean) && localFolder._filesToDelete.Count == localFolder._nameToFile.Count && localFolder._filesToSync.Count == 0;
 		}
 
 		static bool IsFileTypeWritable(string? type)
@@ -201,12 +194,12 @@ namespace UnrealGameSync
 
 		void RemoveEmptyFolders(FolderToClean folder)
 		{
-			foreach(FolderToClean subFolder in folder.NameToSubFolder.Values)
+			foreach(FolderToClean subFolder in folder._nameToSubFolder.Values)
 			{
 				RemoveEmptyFolders(subFolder);
 			}
 
-			folder.NameToSubFolder = folder.NameToSubFolder.Values.Where(x => x.NameToSubFolder.Count > 0 || x.FilesToSync.Count > 0 || x.FilesToDelete.Count > 0 || x.EmptyLeaf).ToDictionary(x => x.Name, x => x, StringComparer.InvariantCultureIgnoreCase);
+			folder._nameToSubFolder = folder._nameToSubFolder.Values.Where(x => x._nameToSubFolder.Count > 0 || x._filesToSync.Count > 0 || x._filesToDelete.Count > 0 || x._emptyLeaf).ToDictionary(x => x.Name, x => x, StringComparer.InvariantCultureIgnoreCase);
 		}
 
 		public void Dispose()
@@ -243,26 +236,26 @@ namespace UnrealGameSync
 					for(int idx = 0; idx < fragments.Length - 1; idx++)
 					{
 						FolderToClean? nextSyncFolder;
-						if(!syncFolder.NameToSubFolder.TryGetValue(fragments[idx], out nextSyncFolder))
+						if(!syncFolder._nameToSubFolder.TryGetValue(fragments[idx], out nextSyncFolder))
 						{
-							nextSyncFolder = new FolderToClean(new DirectoryInfo(Path.Combine(syncFolder.Directory.FullName, fragments[idx])));
-							syncFolder.NameToSubFolder[nextSyncFolder.Name] = nextSyncFolder;
+							nextSyncFolder = new FolderToClean(new DirectoryInfo(Path.Combine(syncFolder._directory.FullName, fragments[idx])));
+							syncFolder._nameToSubFolder[nextSyncFolder.Name] = nextSyncFolder;
 						}
 						syncFolder = nextSyncFolder;
 					}
 
-					string wildcard = fragments[fragments.Length - 1];
+					string wildcard = fragments[^1];
 					if(wildcard == "...")
 					{
 						QueueFolderToPopulate(syncFolder);
 					}
 					else
 					{
-						if(syncFolder.Directory.Exists)
+						if(syncFolder._directory.Exists)
 						{
-							foreach(FileInfo file in syncFolder.Directory.EnumerateFiles(wildcard))
+							foreach(FileInfo file in syncFolder._directory.EnumerateFiles(wildcard))
 							{
-								syncFolder.NameToFile[file.Name] = file;
+								syncFolder._nameToFile[file.Name] = file;
 							}
 						}
 					}
@@ -270,7 +263,7 @@ namespace UnrealGameSync
 			}
 
 			// Get the prefix for any local file
-			string localRootPrefix = _rootFolderToClean.Directory.FullName.Replace(Path.AltDirectorySeparatorChar, Path.DirectorySeparatorChar) + Path.DirectorySeparatorChar;
+			string localRootPrefix = _rootFolderToClean._directory.FullName.Replace(Path.AltDirectorySeparatorChar, Path.DirectorySeparatorChar) + Path.DirectorySeparatorChar;
 
 			// Query the have table and build a separate tree from it
 			PerforceHaveFolder rootHaveFolder = new PerforceHaveFolder();
@@ -281,7 +274,7 @@ namespace UnrealGameSync
 				{
 					if(fileRecord.ClientFile == null || !fileRecord.ClientFile.StartsWith(localRootPrefix, StringComparison.InvariantCultureIgnoreCase))
 					{
-						throw new UserErrorException($"Failed to get have table; file '{fileRecord.ClientFile}' doesn't start with root path ('{_rootFolderToClean.Directory.FullName}')");
+						throw new UserErrorException($"Failed to get have table; file '{fileRecord.ClientFile}' doesn't start with root path ('{_rootFolderToClean._directory.FullName}')");
 					}
 
 					string[] tokens = fileRecord.ClientFile.Substring(localRootPrefix.Length).Split('/', '\\');
@@ -290,14 +283,14 @@ namespace UnrealGameSync
 					for(int idx = 0; idx < tokens.Length - 1; idx++)
 					{
 						PerforceHaveFolder? nextFileFolder;
-						if(!fileFolder.NameToSubFolder.TryGetValue(tokens[idx], out nextFileFolder))
+						if(!fileFolder._nameToSubFolder.TryGetValue(tokens[idx], out nextFileFolder))
 						{
 							nextFileFolder = new PerforceHaveFolder();
-							fileFolder.NameToSubFolder.Add(tokens[idx], nextFileFolder);
+							fileFolder._nameToSubFolder.Add(tokens[idx], nextFileFolder);
 						}
 						fileFolder = nextFileFolder;
 					}
-					fileFolder.NameToFile[tokens[tokens.Length - 1]] = fileRecord;
+					fileFolder._nameToFile[tokens[^1]] = fileRecord;
 				}
 			}
 
diff --git a/Engine/Source/Programs/UnrealGameSync/UnrealGameSync/Forms/ApplicationSettingsWindow.cs b/Engine/Source/Programs/UnrealGameSync/UnrealGameSync/Forms/ApplicationSettingsWindow.cs
--- a/Engine/Source/Programs/UnrealGameSync/UnrealGameSync/Forms/ApplicationSettingsWindow.cs
+++ b/Engine/Source/Programs/UnrealGameSync/UnrealGameSync/Forms/ApplicationSettingsWindow.cs
@@ -5,17 +5,10 @@ using Microsoft.Extensions.Logging;
 using Microsoft.Win32;
 using System;
 using System.Collections.Generic;
-using System.ComponentModel;
-using System.Data;
-using System.Drawing;
-using System.IO;
 using System.Linq;
-using System.Runtime.InteropServices;
-using System.Text;
 using System.Threading;
 using System.Threading.Tasks;
 using System.Windows.Forms;
-using UnrealGameSync;
 
 namespace UnrealGameSync
 {
@@ -35,21 +28,21 @@ namespace UnrealGameSync
 			}
 		}
 
-		string _originalExecutableFileName;
-		IPerforceSettings _defaultPerforceSettings;
-		UserSettings _settings;
-		ILogger _logger;
+		readonly string _originalExecutableFileName;
+		readonly IPerforceSettings _defaultPerforceSettings;
+		readonly UserSettings _settings;
+		readonly ILogger _logger;
 
-		string? _initialServerAndPort;
-		string? _initialUserName;
-		string? _initialDepotPath;
-		bool _initialPreview;
-		int _initialAutomationPortNumber;
-		ProtocolHandlerState _initialProtocolHandlerState;
+		readonly string? _initialServerAndPort;
+		readonly string? _initialUserName;
+		readonly string? _initialDepotPath;
+		readonly bool _initialPreview;
+		readonly int _initialAutomationPortNumber;
+		readonly ProtocolHandlerState _initialProtocolHandlerState;
 
 		bool? _restartPreview;
 
-		ToolUpdateMonitor _toolUpdateMonitor;
+		readonly ToolUpdateMonitor _toolUpdateMonitor;
 
 		class ToolItem
 		{
@@ -57,7 +50,7 @@ namespace UnrealGameSync
 
 			public ToolItem(ToolDefinition definition)
 			{
-				this.Definition = definition;
+				Definition = definition;
 			}
 
 			public override string ToString()
@@ -70,11 +63,11 @@ namespace UnrealGameSync
 		{
 			InitializeComponent();
 
-			this._originalExecutableFileName = originalExecutableFileName;
-			this._defaultPerforceSettings = defaultPerforceSettings;
-			this._settings = settings;
-			this._toolUpdateMonitor = toolUpdateMonitor;
-			this._logger = logger;
+			_originalExecutableFileName = originalExecutableFileName;
+			_defaultPerforceSettings = defaultPerforceSettings;
+			_settings = settings;
+			_toolUpdateMonitor = toolUpdateMonitor;
+			_logger = logger;
 
 			GlobalPerforceSettings.ReadGlobalPerforceSettings(ref _initialServerAndPort, ref _initialUserName, ref _initialDepotPath, ref preview);
 			_initialPreview = preview;
@@ -82,62 +75,62 @@ namespace UnrealGameSync
 			_initialAutomationPortNumber = AutomationServer.GetPortNumber();
 			_initialProtocolHandlerState = ProtocolHandlerUtils.GetState();
 
-			this.AutomaticallyRunAtStartupCheckBox.Checked = IsAutomaticallyRunAtStartup();
-			this.KeepInTrayCheckBox.Checked = settings.KeepInTray;
-						
-			this.ServerTextBox.Text = _initialServerAndPort;
-			this.ServerTextBox.Select(ServerTextBox.TextLength, 0);
-			this.ServerTextBox.CueBanner = $"Default ({defaultPerforceSettings.ServerAndPort})";
+			AutomaticallyRunAtStartupCheckBox.Checked = IsAutomaticallyRunAtStartup();
+			KeepInTrayCheckBox.Checked = settings.KeepInTray;
+					
+			ServerTextBox.Text = _initialServerAndPort;
+			ServerTextBox.Select(ServerTextBox.TextLength, 0);
+			ServerTextBox.CueBanner = $"Default ({defaultPerforceSettings.ServerAndPort})";
 
-			this.UserNameTextBox.Text = _initialUserName;
-			this.UserNameTextBox.Select(UserNameTextBox.TextLength, 0);
-			this.UserNameTextBox.CueBanner = $"Default ({defaultPerforceSettings.UserName})";
+			UserNameTextBox.Text = _initialUserName;
+			UserNameTextBox.Select(UserNameTextBox.TextLength, 0);
+			UserNameTextBox.CueBanner = $"Default ({defaultPerforceSettings.UserName})";
 
-			this.ParallelSyncThreadsSpinner.Value = Math.Max(Math.Min(settings.SyncOptions.NumThreads ?? PerforceSyncOptions.DefaultNumThreads, ParallelSyncThreadsSpinner.Maximum), ParallelSyncThreadsSpinner.Minimum);
+			ParallelSyncThreadsSpinner.Value = Math.Max(Math.Min(settings.SyncOptions.NumThreads ?? PerforceSyncOptions.DefaultNumThreads, ParallelSyncThreadsSpinner.Maximum), ParallelSyncThreadsSpinner.Minimum);
 
-			this.DepotPathTextBox.Text = _initialDepotPath;
-			this.DepotPathTextBox.Select(DepotPathTextBox.TextLength, 0);
-			this.DepotPathTextBox.CueBanner = DeploymentSettings.DefaultDepotPath ?? String.Empty;
+			DepotPathTextBox.Text = _initialDepotPath;
+			DepotPathTextBox.Select(DepotPathTextBox.TextLength, 0);
+			DepotPathTextBox.CueBanner = DeploymentSettings.DefaultDepotPath ?? String.Empty;
 
-			this.UsePreviewBuildCheckBox.Checked = preview;
+			UsePreviewBuildCheckBox.Checked = preview;
 
 			if(_initialAutomationPortNumber > 0)
 			{
-				this.EnableAutomationCheckBox.Checked = true;
-				this.AutomationPortTextBox.Enabled = true;
-				this.AutomationPortTextBox.Text = _initialAutomationPortNumber.ToString();
+				EnableAutomationCheckBox.Checked = true;
+				AutomationPortTextBox.Enabled = true;
+				AutomationPortTextBox.Text = _initialAutomationPortNumber.ToString();
 			}
 			else
 			{
-				this.EnableAutomationCheckBox.Checked = false;
-				this.AutomationPortTextBox.Enabled = false;
-				this.AutomationPortTextBox.Text = AutomationServer.DefaultPortNumber.ToString();
+				EnableAutomationCheckBox.Checked = false;
+				AutomationPortTextBox.Enabled = false;
+				AutomationPortTextBox.Text = AutomationServer.DefaultPortNumber.ToString();
 			}
 
 			if(_initialProtocolHandlerState == ProtocolHandlerState.Installed)
 			{
-				this.EnableProtocolHandlerCheckBox.CheckState = CheckState.Checked;
+				EnableProtocolHandlerCheckBox.CheckState = CheckState.Checked;
 			}
 			else if (_initialProtocolHandlerState == ProtocolHandlerState.NotInstalled)
 			{
-				this.EnableProtocolHandlerCheckBox.CheckState = CheckState.Unchecked;
+				EnableProtocolHandlerCheckBox.CheckState = CheckState.Unchecked;
 			}
 			else
 			{
-				this.EnableProtocolHandlerCheckBox.CheckState = CheckState.Indeterminate;
+				EnableProtocolHandlerCheckBox.CheckState = CheckState.Indeterminate;
 			}
 
 			List<ToolDefinition> tools = toolUpdateMonitor.Tools;
 			foreach (ToolDefinition tool in tools)
 			{
-				this.CustomToolsListBox.Items.Add(new ToolItem(tool), settings.EnabledTools.Contains(tool.Id));
+				CustomToolsListBox.Items.Add(new ToolItem(tool), settings.EnabledTools.Contains(tool.Id));
 			}
 		}
 
 		public static bool? ShowModal(IWin32Window owner, IPerforceSettings defaultPerforceSettings, bool preview, string originalExecutableFileName, UserSettings settings, ToolUpdateMonitor toolUpdateMonitor, ILogger<ApplicationSettingsWindow> logger)
 		{
 			ApplicationSettingsWindow applicationSettings = new ApplicationSettingsWindow(defaultPerforceSettings, preview, originalExecutableFileName, settings, toolUpdateMonitor, logger);
-			if(applicationSettings.ShowDialog() == DialogResult.OK)
+			if(applicationSettings.ShowDialog(owner) == DialogResult.OK)
 			{
 				return applicationSettings._restartPreview;
 			}
@@ -176,9 +169,8 @@ namespace UnrealGameSync
 
 			bool preview = UsePreviewBuildCheckBox.Checked;
 
-
 			int automationPortNumber;
-			if(!EnableAutomationCheckBox.Checked || !int.TryParse(AutomationPortTextBox.Text, out automationPortNumber))
+			if(!EnableAutomationCheckBox.Checked || !Int32.TryParse(AutomationPortTextBox.Text, out automationPortNumber))
 			{
 				automationPortNumber = -1;
 			}
diff --git a/Engine/Source/Programs/UnrealGameSync/UnrealGameSync/Forms/ArgumentsWindow.cs b/Engine/Source/Programs/UnrealGameSync/UnrealGameSync/Forms/ArgumentsWindow.cs
--- a/Engine/Source/Programs/UnrealGameSync/UnrealGameSync/Forms/ArgumentsWindow.cs
+++ b/Engine/Source/Programs/UnrealGameSync/UnrealGameSync/Forms/ArgumentsWindow.cs
@@ -2,13 +2,8 @@
 
 using System;
 using System.Collections.Generic;
-using System.ComponentModel;
-using System.Data;
 using System.Drawing;
-using System.Linq;
 using System.Runtime.InteropServices;
-using System.Text;
-using System.Threading.Tasks;
 using System.Windows.Forms;
 using System.Windows.Forms.VisualStyles;
 
@@ -22,7 +17,7 @@ namespace UnrealGameSync
 		const int EmSetsel = 0xb1;
 
 		[DllImport("user32.dll", CharSet = CharSet.Auto)]
-		static extern IntPtr SendMessage(IntPtr hWnd, UInt32 msg, IntPtr wParam, IntPtr lParam);
+		static extern IntPtr SendMessage(IntPtr hWnd, uint msg, IntPtr wParam, IntPtr lParam);
 		
 		[DllImport("user32.dll")]
 		static extern bool SetWindowText(IntPtr hWnd, string text);
@@ -49,10 +44,7 @@ namespace UnrealGameSync
 			PromptBeforeLaunchCheckBox.Checked = promptBeforeLaunch;
 		}
 
-		public bool PromptBeforeLaunch
-		{
-			get { return PromptBeforeLaunchCheckBox.Checked; }
-		}
+		public bool PromptBeforeLaunch => PromptBeforeLaunchCheckBox.Checked;
 
 		public List<Tuple<string, bool>> GetItems()
 		{
@@ -225,11 +217,5 @@ namespace UnrealGameSync
 			MoveUpButton.Enabled = (ArgumentsList.SelectedIndices.Count == 1 && ArgumentsList.SelectedIndices[0] > 0);
 			MoveDownButton.Enabled = (ArgumentsList.SelectedIndices.Count == 1 && ArgumentsList.SelectedIndices[0] < ArgumentsList.Items.Count - 2);
 		}
-
-		private void checkBox1_CheckedChanged(object sender, EventArgs e)
-		{
-
-
-		}
 	}
 }
diff --git a/Engine/Source/Programs/UnrealGameSync/UnrealGameSync/Forms/AutomatedBuildWindow.cs b/Engine/Source/Programs/UnrealGameSync/UnrealGameSync/Forms/AutomatedBuildWindow.cs
--- a/Engine/Source/Programs/UnrealGameSync/UnrealGameSync/Forms/AutomatedBuildWindow.cs
+++ b/Engine/Source/Programs/UnrealGameSync/UnrealGameSync/Forms/AutomatedBuildWindow.cs
@@ -1,19 +1,10 @@
 // Copyright Epic Games, Inc. All Rights Reserved.
 
 using EpicGames.Perforce;
-using Microsoft.Extensions.Logging;
 using System;
-using System.Collections.Generic;
-using System.ComponentModel;
-using System.Data;
 using System.Diagnostics.CodeAnalysis;
 using System.Drawing;
-using System.IO;
-using System.Linq;
-using System.Text;
-using System.Threading.Tasks;
 using System.Windows.Forms;
-using UnrealGameSync.Properties;
 
 namespace UnrealGameSync.Forms
 {
@@ -21,34 +12,34 @@ namespace UnrealGameSync.Forms
 	{
 		public class BuildInfo
 		{
-			public AutomatedSyncWindow.WorkspaceInfo SelectedWorkspaceInfo;
-			public string ProjectPath;
-			public bool Sync;
-			public string ExecCommand;
+			public AutomatedSyncWindow.WorkspaceInfo SelectedWorkspaceInfo { get; }
+			public string ProjectPath { get; }
+			public bool Sync { get; }
+			public string ExecCommand { get; }
 
 			public BuildInfo(AutomatedSyncWindow.WorkspaceInfo selectedWorkspaceInfo, string projectPath, bool sync, string execCommand)
 			{
-				this.SelectedWorkspaceInfo = selectedWorkspaceInfo;
-				this.ProjectPath = projectPath;
-				this.Sync = sync;
-				this.ExecCommand = execCommand;
+				SelectedWorkspaceInfo = selectedWorkspaceInfo;
+				ProjectPath = projectPath;
+				Sync = sync;
+				ExecCommand = execCommand;
 			}
 		}
 
-		string _streamName;
-		IServiceProvider _serviceProvider;
+		readonly string _streamName;
+		readonly IServiceProvider _serviceProvider;
 
 		string? _serverAndPortOverride;
 		string? _userNameOverride;
-		IPerforceSettings _defaultPerforceSettings;
+		readonly IPerforceSettings _defaultPerforceSettings;
 
 		BuildInfo? _result;
 
 		private AutomatedBuildWindow(string streamName, int changelist, string command, IPerforceSettings defaultPerforceSettings, string? defaultWorkspaceName, string? defaultProjectPath, IServiceProvider serviceProvider)
 		{
-			this._streamName = streamName;
-			this._defaultPerforceSettings = defaultPerforceSettings;
-			this._serviceProvider = serviceProvider;
+			_streamName = streamName;
+			_defaultPerforceSettings = defaultPerforceSettings;
+			_serviceProvider = serviceProvider;
 
 			InitializeComponent();
 
@@ -77,10 +68,7 @@ namespace UnrealGameSync.Forms
 			UpdateWorkspacePathBrowseButton();
 		}
 
-		private IPerforceSettings Perforce
-		{
-			get => Utility.OverridePerforceSettings(_defaultPerforceSettings, _serverAndPortOverride, _userNameOverride);
-		}
+		private IPerforceSettings Perforce => Utility.OverridePerforceSettings(_defaultPerforceSettings, _serverAndPortOverride, _userNameOverride);
 
 		public static bool ShowModal(IWin32Window owner, IPerforceSettings defaultPerforceSettings, string streamName, string projectPath, int changelist, string command, UserSettings settings, IServiceProvider loggerFactory, [NotNullWhen(true)] out BuildInfo? buildInfo)
 		{
@@ -173,8 +161,7 @@ namespace UnrealGameSync.Forms
 
 		private void UpdateWorkspacePathBrowseButton()
 		{
-			string? workspaceName;
-			WorkspacePathBrowseBtn.Enabled = TryGetWorkspaceName(out workspaceName);
+			WorkspacePathBrowseBtn.Enabled = TryGetWorkspaceName(out _);
 		}
 
 		private void WorkspacePathBrowseBtn_Click(object sender, EventArgs e)
diff --git a/Engine/Source/Programs/UnrealGameSync/UnrealGameSync/Forms/AutomatedSyncWindow.cs b/Engine/Source/Programs/UnrealGameSync/UnrealGameSync/Forms/AutomatedSyncWindow.cs
--- a/Engine/Source/Programs/UnrealGameSync/UnrealGameSync/Forms/AutomatedSyncWindow.cs
+++ b/Engine/Source/Programs/UnrealGameSync/UnrealGameSync/Forms/AutomatedSyncWindow.cs
@@ -5,13 +5,10 @@ using Microsoft.Extensions.DependencyInjection;
 using Microsoft.Extensions.Logging;
 using System;
 using System.Collections.Generic;
-using System.ComponentModel;
 using System.Data;
 using System.Diagnostics.CodeAnalysis;
 using System.Drawing;
-using System.IO;
 using System.Linq;
-using System.Text;
 using System.Threading;
 using System.Threading.Tasks;
 using System.Windows.Forms;
@@ -53,15 +50,15 @@ namespace UnrealGameSync
 
 		class ValidateWorkspaceTask
 		{
-			public string WorkspaceName;
-			public string StreamName;
-			public bool RequiresStreamSwitch;
-			public bool HasOpenFiles;
+			public string WorkspaceName { get; }
+			public string StreamName { get; }
+			public bool RequiresStreamSwitch { get; private set; }
+			public bool HasOpenFiles { get; private set; }
 
 			public ValidateWorkspaceTask(string workspaceName, string streamName)
 			{
-				this.WorkspaceName = workspaceName;
-				this.StreamName = streamName;
+				WorkspaceName = workspaceName;
+				StreamName = streamName;
 			}
 
 			public async Task RunAsync(IPerforceConnection perforce, CancellationToken cancellationToken)
@@ -88,26 +85,26 @@ namespace UnrealGameSync
 
 			public WorkspaceInfo(string serverAndPort, string userName, string workspaceName, bool requiresStreamSwitch)
 			{
-				this.ServerAndPort = serverAndPort;
-				this.UserName = userName;
-				this.WorkspaceName = workspaceName;
-				this.RequiresStreamSwitch = requiresStreamSwitch;
+				ServerAndPort = serverAndPort;
+				UserName = userName;
+				WorkspaceName = workspaceName;
+				RequiresStreamSwitch = requiresStreamSwitch;
 			}
 		}
 
-		string _streamName;
-		IServiceProvider _serviceProvider;
+		readonly string _streamName;
+		readonly IServiceProvider _serviceProvider;
 
 		string? _serverAndPortOverride;
 		string? _userNameOverride;
-		IPerforceSettings _defaultPerforceSettings;
+		readonly IPerforceSettings _defaultPerforceSettings;
 		WorkspaceInfo? _selectedWorkspaceInfo;
 
 		private AutomatedSyncWindow(string streamName, string projectPath, string? workspaceName, IPerforceSettings defaultPerforceSettings, IServiceProvider serviceProvider)
 		{
-			this._streamName = streamName;
-			this._defaultPerforceSettings = defaultPerforceSettings;
-			this._serviceProvider = serviceProvider;
+			_streamName = streamName;
+			_defaultPerforceSettings = defaultPerforceSettings;
+			_serviceProvider = serviceProvider;
 
 			InitializeComponent();
 
diff --git a/Engine/Source/Programs/UnrealGameSync/UnrealGameSync/Forms/BuildStepWindow.cs b/Engine/Source/Programs/UnrealGameSync/UnrealGameSync/Forms/BuildStepWindow.cs
--- a/Engine/Source/Programs/UnrealGameSync/UnrealGameSync/Forms/BuildStepWindow.cs
+++ b/Engine/Source/Programs/UnrealGameSync/UnrealGameSync/Forms/BuildStepWindow.cs
@@ -3,23 +3,19 @@
 using EpicGames.Core;
 using System;
 using System.Collections.Generic;
-using System.ComponentModel;
-using System.Data;
 using System.Drawing;
 using System.IO;
 using System.Linq;
-using System.Text;
-using System.Threading.Tasks;
 using System.Windows.Forms;
 
 namespace UnrealGameSync
 {
 	partial class BuildStepWindow : Form
 	{
-		BuildStep _step;
-		List<string> _targetNames;
-		DirectoryReference _baseDirectory;
-		IReadOnlyDictionary<string, string> _variables;
+		readonly BuildStep _step;
+		readonly List<string> _targetNames;
+		readonly DirectoryReference _baseDirectory;
+		readonly IReadOnlyDictionary<string, string> _variables;
 		VariablesWindow? _variablesWindow;
 
 		public BuildStepWindow(BuildStep inTask, List<string> inTargetNames, DirectoryReference inBaseDirectory, IReadOnlyDictionary<string, string> inVariables)
@@ -31,8 +27,8 @@ namespace UnrealGameSync
 
 			InitializeComponent();
 
-			this.MaximumSize = new System.Drawing.Size(32768, Height);
-			this.MinimumSize = new System.Drawing.Size(Width, Height);
+			MaximumSize = new System.Drawing.Size(32768, Height);
+			MinimumSize = new System.Drawing.Size(Width, Height);
 		}
 
 		private void BuildTaskWindow_Load(object sender, EventArgs e)
@@ -122,7 +118,7 @@ namespace UnrealGameSync
 			_step.StatusText = StatusTextTextBox.Text;
 			_step.StatusPanelLink = StatusPanelLinkTextBox.Text;
 
-			if (!int.TryParse(DurationTextBox.Text, out _step.EstimatedDuration))
+			if (!Int32.TryParse(DurationTextBox.Text, out _step.EstimatedDuration))
 			{
 				_step.EstimatedDuration = 1;
 			}
