UGS: Add .editorconfig file for static analysis warnings in UGS (+lots of fixes).

diff --git a/Engine/Source/Programs/UnrealGameSync/UnrealGameSync/Forms/ChangelistWindow.cs b/Engine/Source/Programs/UnrealGameSync/UnrealGameSync/Forms/ChangelistWindow.cs
--- a/Engine/Source/Programs/UnrealGameSync/UnrealGameSync/Forms/ChangelistWindow.cs
+++ b/Engine/Source/Programs/UnrealGameSync/UnrealGameSync/Forms/ChangelistWindow.cs
@@ -1,36 +1,29 @@
 // Copyright Epic Games, Inc. All Rights Reserved.
 
 using System;
-using System.Collections.Generic;
-using System.ComponentModel;
-using System.Data;
-using System.Drawing;
-using System.Linq;
-using System.Text;
-using System.Threading.Tasks;
 using System.Windows.Forms;
 
 namespace UnrealGameSync
 {
 	public partial class ChangelistWindow : Form
 	{
-		public int ChangeNumber;
+		public int ChangeNumber { get; private set; }
 
-		public ChangelistWindow(int inChangeNumber)
+		public ChangelistWindow(int changeNumber)
 		{
 			InitializeComponent();
 
-			this.MaximumSize = new System.Drawing.Size(32768, Height);
-			this.MinimumSize = new System.Drawing.Size(Width, Height);
+			MaximumSize = new System.Drawing.Size(32768, Height);
+			MinimumSize = new System.Drawing.Size(Width, Height);
 
-			ChangeNumber = inChangeNumber;
+			ChangeNumber = changeNumber;
 			ChangeNumberTextBox.Text = (ChangeNumber > 0)? ChangeNumber.ToString() : "";
 		}
 
 		private void OkBtn_Click(object sender, EventArgs e)
 		{
 			int newChangeNumber;
-			if(int.TryParse(ChangeNumberTextBox.Text.Trim(), out newChangeNumber))
+			if(Int32.TryParse(ChangeNumberTextBox.Text.Trim(), out newChangeNumber))
 			{
 				int deltaChangeNumber = newChangeNumber - ChangeNumber;
 				if((deltaChangeNumber > -10000 && deltaChangeNumber < 1000) || ChangeNumber <= 0 || MessageBox.Show(String.Format("Changelist {0} is {1} changes away from your current sync. Was that intended?", newChangeNumber, Math.Abs(deltaChangeNumber)), "Changelist is far away", MessageBoxButtons.YesNo) == DialogResult.Yes)
diff --git a/Engine/Source/Programs/UnrealGameSync/UnrealGameSync/Forms/CleanWorkspaceWindow.cs b/Engine/Source/Programs/UnrealGameSync/UnrealGameSync/Forms/CleanWorkspaceWindow.cs
--- a/Engine/Source/Programs/UnrealGameSync/UnrealGameSync/Forms/CleanWorkspaceWindow.cs
+++ b/Engine/Source/Programs/UnrealGameSync/UnrealGameSync/Forms/CleanWorkspaceWindow.cs
@@ -5,16 +5,12 @@ using EpicGames.Perforce;
 using Microsoft.Extensions.Logging;
 using System;
 using System.Collections.Generic;
-using System.ComponentModel;
 using System.Data;
 using System.Diagnostics;
 using System.Drawing;
 using System.IO;
 using System.Linq;
 using System.Runtime.InteropServices;
-using System.Text;
-using System.Threading;
-using System.Threading.Tasks;
 using System.Windows.Forms;
 using System.Windows.Forms.VisualStyles;
 
@@ -30,14 +26,14 @@ namespace UnrealGameSync
 
 		class TreeNodeData
 		{
-			public TreeNodeAction Action;
-			public FileInfo? File;
-			public FolderToClean? Folder;
-			public int NumFiles;
-			public int NumSelectedFiles;
-			public int NumEmptySelectedFiles;
-			public int NumMissingSelectedFiles;
-			public int NumDefaultSelectedFiles;
+			public TreeNodeAction _action;
+			public FileInfo? _file;
+			public FolderToClean? _folder;
+			public int _numFiles;
+			public int _numSelectedFiles;
+			public int _numEmptySelectedFiles;
+			public int _numMissingSelectedFiles;
+			public int _numDefaultSelectedFiles;
 		}
 
 		enum SelectionType
@@ -49,17 +45,17 @@ namespace UnrealGameSync
 			None,
 		}
 
-		static readonly CheckBoxState[] CheckBoxStates = 
+		static readonly CheckBoxState[] s_checkBoxStates = 
 		{ 
 			CheckBoxState.UncheckedNormal, 
 			CheckBoxState.MixedNormal, 
 			CheckBoxState.CheckedNormal 
 		};
 
-		IPerforceSettings _perforceSettings;
-		FolderToClean _rootFolderToClean;
+		readonly IPerforceSettings _perforceSettings;
+		readonly FolderToClean _rootFolderToClean;
 
-		static readonly string[] SafeToDeleteFolders =
+		static readonly string[] s_safeToDeleteFolders =
 		{
 			"/binaries/",
 			"/intermediate/",
@@ -90,7 +86,7 @@ namespace UnrealGameSync
 			"/obj/x64/release/",
 		};
 
-		static readonly string[] SafeToDeleteExtensions =
+		static readonly string[] s_safeToDeleteExtensions =
 		{
 			".pdb",
 			".obj",
@@ -101,20 +97,20 @@ namespace UnrealGameSync
 			".csproj.references",
 		};
 
-		IReadOnlyList<string> _extraSafeToDeleteFolders;
-		IReadOnlyList<string> _extraSafeToDeleteExtensions;
-		ILogger _logger;
+		readonly IReadOnlyList<string> _extraSafeToDeleteFolders;
+		readonly IReadOnlyList<string> _extraSafeToDeleteExtensions;
+		readonly ILogger _logger;
 
 		[DllImport("Shell32.dll", EntryPoint = "ExtractIconExW", CharSet = CharSet.Unicode, ExactSpelling = true, CallingConvention = CallingConvention.StdCall)]
 		private static extern int ExtractIconEx(string sFile, int iIndex, IntPtr piLargeVersion, out IntPtr piSmallVersion, int amountIcons);
 
 		private CleanWorkspaceWindow(IPerforceSettings perforceSettings, FolderToClean rootFolderToClean, string[] extraSafeToDeleteFolders, string[] extraSafeToDeleteExtensions, ILogger<CleanWorkspaceWindow> logger)
 		{
-			this._perforceSettings = perforceSettings;
-			this._rootFolderToClean = rootFolderToClean;
-			this._extraSafeToDeleteFolders = extraSafeToDeleteFolders.Select(x => x.Trim().Replace('\\', '/').Trim('/')).Where(x => x.Length > 0).Select(x => String.Format("/{0}/", x.ToLowerInvariant())).ToArray();
-			this._extraSafeToDeleteExtensions = extraSafeToDeleteExtensions.Select(x => x.Trim().ToLowerInvariant()).Where(x => x.Length > 0).ToArray();
-			this._logger = logger;
+			_perforceSettings = perforceSettings;
+			_rootFolderToClean = rootFolderToClean;
+			_extraSafeToDeleteFolders = extraSafeToDeleteFolders.Select(x => x.Trim().Replace('\\', '/').Trim('/')).Where(x => x.Length > 0).Select(x => String.Format("/{0}/", x.ToLowerInvariant())).ToArray();
+			_extraSafeToDeleteExtensions = extraSafeToDeleteExtensions.Select(x => x.Trim().ToLowerInvariant()).Where(x => x.Length > 0).ToArray();
+			_logger = logger;
 
 			InitializeComponent();
 		}
@@ -133,7 +129,7 @@ namespace UnrealGameSync
 			}
 
 			// If there's nothing to delete, don't bother displaying the dialog at all
-			if(rootFolderToClean.FilesToDelete.Count == 0 && rootFolderToClean.NameToSubFolder.Count == 0)
+			if(rootFolderToClean._filesToDelete.Count == 0 && rootFolderToClean._nameToSubFolder.Count == 0)
 			{
 				MessageBox.Show("You have no local files which are not in Perforce.", "Workspace Clean", MessageBoxButtons.OK);
 				return;
@@ -163,7 +159,7 @@ namespace UnrealGameSync
 			Size largestCheckBoxSize = Size.Empty;
 			using(Graphics graphics = Graphics.FromHwnd(IntPtr.Zero))
 			{
-				foreach(CheckBoxState state in CheckBoxStates)
+				foreach(CheckBoxState state in s_checkBoxStates)
 				{
 					Size checkBoxSize = CheckBoxRenderer.GetGlyphSize(graphics, state);
 					largestCheckBoxSize = new Size(Math.Max(largestCheckBoxSize.Width, checkBoxSize.Width), Math.Max(largestCheckBoxSize.Height, checkBoxSize.Height));
@@ -180,8 +176,8 @@ namespace UnrealGameSync
 				{
 					for(int stateIdx = 0; stateIdx < 3; stateIdx++)
 					{
-						Size checkBoxSize = CheckBoxRenderer.GetGlyphSize(graphics, CheckBoxStates[stateIdx]);
-						CheckBoxRenderer.DrawCheckBox(graphics, new Point(minX + (largestCheckBoxSize.Width - checkBoxSize.Width) / 2, (largestCheckBoxSize.Height - checkBoxSize.Height) / 2), CheckBoxStates[stateIdx]);
+						Size checkBoxSize = CheckBoxRenderer.GetGlyphSize(graphics, s_checkBoxStates[stateIdx]);
+						CheckBoxRenderer.DrawCheckBox(graphics, new Point(minX + (largestCheckBoxSize.Width - checkBoxSize.Width) / 2, (largestCheckBoxSize.Height - checkBoxSize.Height) / 2), s_checkBoxStates[stateIdx]);
 
 						Size iconSize = icons[iconIdx].Size;
 						graphics.DrawIcon(icons[iconIdx], minX + largestCheckBoxSize.Width + (largestIconSize.Width - iconSize.Width) / 2, (largestIconSize.Height - iconSize.Height) / 2);
@@ -198,44 +194,44 @@ namespace UnrealGameSync
 			TreeView.ImageList = typeImageList;
 
 			TreeNode node = BuildTreeViewStructure(_rootFolderToClean, "/", false, 0);
-			node.Text = _rootFolderToClean.Directory.FullName;
+			node.Text = _rootFolderToClean._directory.FullName;
 			TreeView.Nodes.Add(node);
 		}
 
 		private TreeNode BuildTreeViewStructure(FolderToClean folder, string folderPath, bool parentFolderSelected, int depth)
 		{
-			bool selectFolder = parentFolderSelected || IsSafeToDeleteFolder(folderPath) || folder.EmptyLeaf;
+			bool selectFolder = parentFolderSelected || IsSafeToDeleteFolder(folderPath) || folder._emptyLeaf;
 
 			TreeNodeData folderNodeData = new TreeNodeData();
-			folderNodeData.Folder = folder;
+			folderNodeData._folder = folder;
 
 			TreeNode folderNode = new TreeNode();
 			folderNode.Text = folder.Name;
 			folderNode.Tag = folderNodeData;
 
-			foreach(FolderToClean subFolder in folder.NameToSubFolder.OrderBy(x => x.Key).Select(x => x.Value))
+			foreach(FolderToClean subFolder in folder._nameToSubFolder.OrderBy(x => x.Key).Select(x => x.Value))
 			{
 				TreeNode childNode = BuildTreeViewStructure(subFolder, folderPath + subFolder.Name.ToLowerInvariant() + "/", selectFolder, depth + 1);
 				folderNode.Nodes.Add(childNode);
 
 				TreeNodeData childNodeData = (TreeNodeData)childNode.Tag;
-				folderNodeData.NumFiles += childNodeData.NumFiles;
-				folderNodeData.NumSelectedFiles += childNodeData.NumSelectedFiles;
-				folderNodeData.NumEmptySelectedFiles += childNodeData.NumEmptySelectedFiles;
-				folderNodeData.NumMissingSelectedFiles += childNodeData.NumMissingSelectedFiles;
-				folderNodeData.NumDefaultSelectedFiles += childNodeData.NumDefaultSelectedFiles;
+				folderNodeData._numFiles += childNodeData._numFiles;
+				folderNodeData._numSelectedFiles += childNodeData._numSelectedFiles;
+				folderNodeData._numEmptySelectedFiles += childNodeData._numEmptySelectedFiles;
+				folderNodeData._numMissingSelectedFiles += childNodeData._numMissingSelectedFiles;
+				folderNodeData._numDefaultSelectedFiles += childNodeData._numDefaultSelectedFiles;
 			}
 
-			foreach(FileInfo file in folder.FilesToSync.OrderBy(x => x.Name))
+			foreach(FileInfo file in folder._filesToSync.OrderBy(x => x.Name))
 			{
 				TreeNodeData fileNodeData = new TreeNodeData();
-				fileNodeData.Action = TreeNodeAction.Sync;
-				fileNodeData.File = file;
-				fileNodeData.NumFiles = 1;
-				fileNodeData.NumSelectedFiles = 1;
-				fileNodeData.NumEmptySelectedFiles = 0;
-				fileNodeData.NumMissingSelectedFiles = 1;
-				fileNodeData.NumDefaultSelectedFiles = fileNodeData.NumSelectedFiles;
+				fileNodeData._action = TreeNodeAction.Sync;
+				fileNodeData._file = file;
+				fileNodeData._numFiles = 1;
+				fileNodeData._numSelectedFiles = 1;
+				fileNodeData._numEmptySelectedFiles = 0;
+				fileNodeData._numMissingSelectedFiles = 1;
+				fileNodeData._numDefaultSelectedFiles = fileNodeData._numSelectedFiles;
 
 				TreeNode fileNode = new TreeNode();
 				fileNode.Text = file.Name + " (sync)";
@@ -244,27 +240,27 @@ namespace UnrealGameSync
 
 				UpdateImage(fileNode);
 
-				folderNodeData.NumFiles++;
-				folderNodeData.NumSelectedFiles += fileNodeData.NumSelectedFiles;
-				folderNodeData.NumEmptySelectedFiles += fileNodeData.NumEmptySelectedFiles;
-				folderNodeData.NumMissingSelectedFiles += fileNodeData.NumMissingSelectedFiles;
-				folderNodeData.NumDefaultSelectedFiles += fileNodeData.NumDefaultSelectedFiles;
+				folderNodeData._numFiles++;
+				folderNodeData._numSelectedFiles += fileNodeData._numSelectedFiles;
+				folderNodeData._numEmptySelectedFiles += fileNodeData._numEmptySelectedFiles;
+				folderNodeData._numMissingSelectedFiles += fileNodeData._numMissingSelectedFiles;
+				folderNodeData._numDefaultSelectedFiles += fileNodeData._numDefaultSelectedFiles;
 			}
 
-			foreach(FileInfo file in folder.FilesToDelete.OrderBy(x => x.Name))
+			foreach(FileInfo file in folder._filesToDelete.OrderBy(x => x.Name))
 			{
 				string name = file.Name.ToLowerInvariant();
 
-				bool selectFile = selectFolder || IsSafeToDeleteFile(folderPath, file.Name.ToLowerInvariant());
+				bool selectFile = selectFolder || IsSafeToDeleteFile(file.Name.ToLowerInvariant());
 
 				TreeNodeData fileNodeData = new TreeNodeData();
-				fileNodeData.Action = TreeNodeAction.Delete;
-				fileNodeData.File = file;
-				fileNodeData.NumFiles = 1;
-				fileNodeData.NumSelectedFiles = selectFile? 1 : 0;
-				fileNodeData.NumEmptySelectedFiles = 0;
-				fileNodeData.NumMissingSelectedFiles = 0;
-				fileNodeData.NumDefaultSelectedFiles = fileNodeData.NumSelectedFiles;
+				fileNodeData._action = TreeNodeAction.Delete;
+				fileNodeData._file = file;
+				fileNodeData._numFiles = 1;
+				fileNodeData._numSelectedFiles = selectFile? 1 : 0;
+				fileNodeData._numEmptySelectedFiles = 0;
+				fileNodeData._numMissingSelectedFiles = 0;
+				fileNodeData._numDefaultSelectedFiles = fileNodeData._numSelectedFiles;
 
 				TreeNode fileNode = new TreeNode();
 				fileNode.Text = file.Name;
@@ -273,22 +269,22 @@ namespace UnrealGameSync
 
 				UpdateImage(fileNode);
 
-				folderNodeData.NumFiles++;
-				folderNodeData.NumSelectedFiles += fileNodeData.NumSelectedFiles;
-				folderNodeData.NumEmptySelectedFiles += fileNodeData.NumEmptySelectedFiles;
-				folderNodeData.NumMissingSelectedFiles += fileNodeData.NumMissingSelectedFiles;
-				folderNodeData.NumDefaultSelectedFiles += fileNodeData.NumDefaultSelectedFiles;
+				folderNodeData._numFiles++;
+				folderNodeData._numSelectedFiles += fileNodeData._numSelectedFiles;
+				folderNodeData._numEmptySelectedFiles += fileNodeData._numEmptySelectedFiles;
+				folderNodeData._numMissingSelectedFiles += fileNodeData._numMissingSelectedFiles;
+				folderNodeData._numDefaultSelectedFiles += fileNodeData._numDefaultSelectedFiles;
 			}
 
-			if(folderNodeData.Folder.EmptyLeaf)
+			if(folderNodeData._folder._emptyLeaf)
 			{
-				folderNodeData.NumFiles++;
-				folderNodeData.NumSelectedFiles++;
-				folderNodeData.NumEmptySelectedFiles++;
-				folderNodeData.NumDefaultSelectedFiles++;
+				folderNodeData._numFiles++;
+				folderNodeData._numSelectedFiles++;
+				folderNodeData._numEmptySelectedFiles++;
+				folderNodeData._numDefaultSelectedFiles++;
 			}
 
-			if(folderNodeData.NumSelectedFiles > 0 && !folderNodeData.Folder.EmptyAfterClean && depth < 2)
+			if(folderNodeData._numSelectedFiles > 0 && !folderNodeData._folder._emptyAfterClean && depth < 2)
 			{
 				folderNode.Expand();
 			}
@@ -303,12 +299,12 @@ namespace UnrealGameSync
 
 		private bool IsSafeToDeleteFolder(string folderPath)
 		{
-			return SafeToDeleteFolders.Any(x => folderPath.EndsWith(x)) || _extraSafeToDeleteFolders.Any(x => folderPath.EndsWith(x));
+			return s_safeToDeleteFolders.Any(x => folderPath.EndsWith(x)) || _extraSafeToDeleteFolders.Any(x => folderPath.EndsWith(x));
 		}
 
-		private bool IsSafeToDeleteFile(string folderPath, string name)
+		private bool IsSafeToDeleteFile(string name)
 		{
-			return SafeToDeleteExtensions.Any(x => name.EndsWith(x)) || _extraSafeToDeleteExtensions.Any(x => name.EndsWith(x));
+			return s_safeToDeleteExtensions.Any(x => name.EndsWith(x)) || _extraSafeToDeleteExtensions.Any(x => name.EndsWith(x));
 		}
 
 		private void TreeView_DrawNode(object sender, DrawTreeNodeEventArgs e)
@@ -328,7 +324,7 @@ namespace UnrealGameSync
 			else if(e.X >= e.Node.Bounds.Left - 32 && e.X < e.Node.Bounds.Left - 16)
 			{
 				TreeNodeData nodeData = (TreeNodeData)node.Tag;
-				SetSelected(node, (nodeData.NumSelectedFiles == 0)? SelectionType.All : SelectionType.None);
+				SetSelected(node, (nodeData._numSelectedFiles == 0)? SelectionType.All : SelectionType.None);
 			}
 		}
 
@@ -336,16 +332,16 @@ namespace UnrealGameSync
 		{
 			TreeNodeData parentNodeData = (TreeNodeData)parentNode.Tag;
 
-			int prevNumSelectedFiles = parentNodeData.NumSelectedFiles;
+			int prevNumSelectedFiles = parentNodeData._numSelectedFiles;
 			SetSelectedOnChildren(parentNode, type);
 
-			int deltaNumSelectedFiles = parentNodeData.NumSelectedFiles - prevNumSelectedFiles;
+			int deltaNumSelectedFiles = parentNodeData._numSelectedFiles - prevNumSelectedFiles;
 			if(deltaNumSelectedFiles != 0)
 			{
 				for(TreeNode nextParentNode = parentNode.Parent; nextParentNode != null; nextParentNode = nextParentNode.Parent)
 				{
 					TreeNodeData nextParentNodeData = (TreeNodeData)nextParentNode.Tag;
-					nextParentNodeData.NumSelectedFiles += deltaNumSelectedFiles;
+					nextParentNodeData._numSelectedFiles += deltaNumSelectedFiles;
 					UpdateImage(nextParentNode);
 				}
 			}
@@ -359,23 +355,23 @@ namespace UnrealGameSync
 			switch(type)
 			{
 				case SelectionType.All:
-					newNumSelectedFiles = parentNodeData.NumFiles;
+					newNumSelectedFiles = parentNodeData._numFiles;
 					break;
 				case SelectionType.Empty:
-					newNumSelectedFiles = parentNodeData.NumEmptySelectedFiles;
+					newNumSelectedFiles = parentNodeData._numEmptySelectedFiles;
 					break;
 				case SelectionType.Missing:
-					newNumSelectedFiles = parentNodeData.NumMissingSelectedFiles;
+					newNumSelectedFiles = parentNodeData._numMissingSelectedFiles;
 					break;
 				case SelectionType.SafeToDelete:
-					newNumSelectedFiles = parentNodeData.NumDefaultSelectedFiles;
+					newNumSelectedFiles = parentNodeData._numDefaultSelectedFiles;
 					break;
 				case SelectionType.None:
 					newNumSelectedFiles = 0;
 					break;
 			}
 
-			if(newNumSelectedFiles != parentNodeData.NumSelectedFiles)
+			if(newNumSelectedFiles != parentNodeData._numSelectedFiles)
 			{
 				foreach(TreeNode? childNode in parentNode.Nodes)
 				{
@@ -384,7 +380,7 @@ namespace UnrealGameSync
 						SetSelectedOnChildren(childNode, type);
 					}
 				}
-				parentNodeData.NumSelectedFiles = newNumSelectedFiles;
+				parentNodeData._numSelectedFiles = newNumSelectedFiles;
 				UpdateImage(parentNode);
 			}
 		}
@@ -392,8 +388,8 @@ namespace UnrealGameSync
 		private void UpdateImage(TreeNode node)
 		{
 			TreeNodeData nodeData = (TreeNodeData)node.Tag;
-			int imageIndex = (nodeData.Folder != null)? 0 : 3;
-			imageIndex += (nodeData.NumSelectedFiles == 0)? 0 : (nodeData.NumSelectedFiles < nodeData.NumFiles || (nodeData.Folder != null && !nodeData.Folder.EmptyAfterClean))? 1 : 2;
+			int imageIndex = (nodeData._folder != null)? 0 : 3;
+			imageIndex += (nodeData._numSelectedFiles == 0)? 0 : (nodeData._numSelectedFiles < nodeData._numFiles || (nodeData._folder != null && !nodeData._folder._emptyAfterClean))? 1 : 2;
 			node.ImageIndex = imageIndex;
 			node.SelectedImageIndex = imageIndex;
 		}
@@ -425,17 +421,17 @@ namespace UnrealGameSync
 		private void FindSelection(TreeNode node, List<FileInfo> filesToSync, List<FileInfo> filesToDelete, List<DirectoryInfo> directoriesToDelete)
 		{
 			TreeNodeData nodeData = (TreeNodeData)node.Tag;
-			if(nodeData.File != null)
+			if(nodeData._file != null)
 			{
-				if(nodeData.NumSelectedFiles > 0)
+				if(nodeData._numSelectedFiles > 0)
 				{
-					if(nodeData.Action == TreeNodeAction.Delete)
+					if(nodeData._action == TreeNodeAction.Delete)
 					{
-						filesToDelete.Add(nodeData.File);
+						filesToDelete.Add(nodeData._file);
 					}
 					else
 					{
-						filesToSync.Add(nodeData.File);
+						filesToSync.Add(nodeData._file);
 					}
 				}
 			}
@@ -448,9 +444,9 @@ namespace UnrealGameSync
 						FindSelection(childNode, filesToSync, filesToDelete, directoriesToDelete);
 					}
 				}
-				if(nodeData.Folder != null && nodeData.Folder.EmptyAfterClean && nodeData.NumSelectedFiles == nodeData.NumFiles)
+				if(nodeData._folder != null && nodeData._folder._emptyAfterClean && nodeData._numSelectedFiles == nodeData._numFiles)
 				{
-					directoriesToDelete.Add(nodeData.Folder.Directory);
+					directoriesToDelete.Add(nodeData._folder._directory);
 				}
 			}
 		}
@@ -506,13 +502,13 @@ namespace UnrealGameSync
 			TreeNode node = (TreeNode)FolderContextMenu.Tag;
 			TreeNodeData nodeData = (TreeNodeData)node.Tag;
 
-			if (nodeData.Folder != null)
+			if (nodeData._folder != null)
 			{
-				Process.Start("explorer.exe", String.Format("\"{0}\"", nodeData.Folder.Directory.FullName));
+				Process.Start("explorer.exe", String.Format("\"{0}\"", nodeData._folder._directory.FullName));
 			}
-			else if (nodeData.File != null)
+			else if (nodeData._file != null)
 			{
-				Process.Start("explorer.exe", String.Format("\"{0}\"", nodeData.File.Directory!.FullName));
+				Process.Start("explorer.exe", String.Format("\"{0}\"", nodeData._file.Directory!.FullName));
 			}
 		}
 	}
diff --git a/Engine/Source/Programs/UnrealGameSync/UnrealGameSync/Forms/ClobberWindow.cs b/Engine/Source/Programs/UnrealGameSync/UnrealGameSync/Forms/ClobberWindow.cs
--- a/Engine/Source/Programs/UnrealGameSync/UnrealGameSync/Forms/ClobberWindow.cs
+++ b/Engine/Source/Programs/UnrealGameSync/UnrealGameSync/Forms/ClobberWindow.cs
@@ -2,20 +2,15 @@
 
 using System;
 using System.Collections.Generic;
-using System.ComponentModel;
-using System.Data;
 using System.Drawing;
 using System.IO;
-using System.Linq;
-using System.Text;
-using System.Threading.Tasks;
 using System.Windows.Forms;
 
 namespace UnrealGameSync
 {
 	public partial class ClobberWindow : Form
 	{
-		Dictionary<string, bool> _filesToClobber;
+		readonly Dictionary<string, bool> _filesToClobber;
 
 		public ClobberWindow(Dictionary<string, bool> inFilesToClobber, HashSet<string> inUncontrolledFiles)
 		{
@@ -43,7 +38,7 @@ namespace UnrealGameSync
 			if (uncontrolledChangeFound)
 			{
 				// Updates the string to inform the user to take special care with Uncontrolled Changes
-				this.label1.Text = "The following files are writable in your workspace." + Environment.NewLine +
+				label1.Text = "The following files are writable in your workspace." + Environment.NewLine +
 	"Red files are Uncontrolled Changes and may contain modifications you made on purpose." + Environment.NewLine +
 	"Select which files you want to overwrite:";
 			}
diff --git a/Engine/Source/Programs/UnrealGameSync/UnrealGameSync/Forms/ConnectWindow.cs b/Engine/Source/Programs/UnrealGameSync/UnrealGameSync/Forms/ConnectWindow.cs
--- a/Engine/Source/Programs/UnrealGameSync/UnrealGameSync/Forms/ConnectWindow.cs
+++ b/Engine/Source/Programs/UnrealGameSync/UnrealGameSync/Forms/ConnectWindow.cs
@@ -1,48 +1,39 @@
 ﻿// Copyright Epic Games, Inc. All Rights Reserved.
 
 using EpicGames.Perforce;
-using Microsoft.Extensions.Logging;
 using System;
-using System.Collections.Generic;
-using System.ComponentModel;
-using System.Data;
-using System.Drawing;
-using System.IO;
-using System.Linq;
-using System.Text;
-using System.Threading.Tasks;
 using System.Windows.Forms;
 
 namespace UnrealGameSync
 {
 	partial class ConnectWindow : Form
 	{
-		IPerforceSettings _defaultPerforceSettings;
+		readonly IPerforceSettings _defaultPerforceSettings;
 		string? _serverAndPortOverride;
 		string? _userNameOverride;
-		IServiceProvider _serviceProvider;
+		readonly IServiceProvider _serviceProvider;
 
 		private ConnectWindow(IPerforceSettings defaultPerforceSettings, string? serverAndPortOverride, string? userNameOverride, IServiceProvider serviceProvider)
 		{
 			InitializeComponent();
 
-			this._defaultPerforceSettings = defaultPerforceSettings;
-			this._serviceProvider = serviceProvider;
+			_defaultPerforceSettings = defaultPerforceSettings;
+			_serviceProvider = serviceProvider;
 
 			if(!String.IsNullOrWhiteSpace(serverAndPortOverride))
 			{
-				this._serverAndPortOverride = serverAndPortOverride.Trim();
+				_serverAndPortOverride = serverAndPortOverride.Trim();
 			}
 			if(!String.IsNullOrEmpty(userNameOverride))
 			{
-				this._userNameOverride = userNameOverride.Trim();
+				_userNameOverride = userNameOverride.Trim();
 			}
 
 			ServerAndPortTextBox.CueBanner = defaultPerforceSettings.ServerAndPort;
-			ServerAndPortTextBox.Text = this._serverAndPortOverride ?? defaultPerforceSettings.ServerAndPort;
+			ServerAndPortTextBox.Text = _serverAndPortOverride ?? defaultPerforceSettings.ServerAndPort;
 			UserNameTextBox.CueBanner = defaultPerforceSettings.UserName;
-			UserNameTextBox.Text = this._userNameOverride ?? defaultPerforceSettings.UserName;
-			UseDefaultConnectionSettings.Checked = this._serverAndPortOverride == null && this._userNameOverride == null;
+			UserNameTextBox.Text = _userNameOverride ?? defaultPerforceSettings.UserName;
+			UseDefaultConnectionSettings.Checked = _serverAndPortOverride == null && _userNameOverride == null;
 
 			UpdateEnabledControls();
 		}
diff --git a/Engine/Source/Programs/UnrealGameSync/UnrealGameSync/Forms/DeleteWindow.cs b/Engine/Source/Programs/UnrealGameSync/UnrealGameSync/Forms/DeleteWindow.cs
--- a/Engine/Source/Programs/UnrealGameSync/UnrealGameSync/Forms/DeleteWindow.cs
+++ b/Engine/Source/Programs/UnrealGameSync/UnrealGameSync/Forms/DeleteWindow.cs
@@ -2,20 +2,13 @@
 
 using System;
 using System.Collections.Generic;
-using System.ComponentModel;
-using System.Data;
-using System.Drawing;
-using System.IO;
-using System.Linq;
-using System.Text;
-using System.Threading.Tasks;
 using System.Windows.Forms;
 
 namespace UnrealGameSync
 {
 	public partial class DeleteWindow : Form
 	{
-		Dictionary<string, bool> _filesToDelete;
+		readonly Dictionary<string, bool> _filesToDelete;
 
 		public DeleteWindow(Dictionary<string, bool> inFilesToDelete)
 		{
diff --git a/Engine/Source/Programs/UnrealGameSync/UnrealGameSync/Forms/DiagnosticsWindow.cs b/Engine/Source/Programs/UnrealGameSync/UnrealGameSync/Forms/DiagnosticsWindow.cs
--- a/Engine/Source/Programs/UnrealGameSync/UnrealGameSync/Forms/DiagnosticsWindow.cs
+++ b/Engine/Source/Programs/UnrealGameSync/UnrealGameSync/Forms/DiagnosticsWindow.cs
@@ -13,9 +13,9 @@ namespace UnrealGameSync
 {
 	public partial class DiagnosticsWindow : Form
 	{
-		DirectoryReference _appDataFolder;
-		DirectoryReference _workspaceDataFolder;
-		List<FileReference> _extraFiles;
+		readonly DirectoryReference _appDataFolder;
+		readonly DirectoryReference _workspaceDataFolder;
+		readonly List<FileReference> _extraFiles;
 
 		public DiagnosticsWindow(DirectoryReference inAppDataFolder, DirectoryReference inWorkspaceDataFolder, string inDiagnosticsText, IEnumerable<FileReference> inExtraFiles)
 		{
diff --git a/Engine/Source/Programs/UnrealGameSync/UnrealGameSync/Forms/IssueAlertWindow.cs b/Engine/Source/Programs/UnrealGameSync/UnrealGameSync/Forms/IssueAlertWindow.cs
--- a/Engine/Source/Programs/UnrealGameSync/UnrealGameSync/Forms/IssueAlertWindow.cs
+++ b/Engine/Source/Programs/UnrealGameSync/UnrealGameSync/Forms/IssueAlertWindow.cs
@@ -3,7 +3,6 @@
 using System;
 using System.Drawing;
 using System.Runtime.InteropServices;
-using System.Linq;
 using System.Text;
 using System.Windows.Forms;
 using System.Collections.Generic;
@@ -21,16 +20,16 @@ namespace UnrealGameSync
 		[DllImport("user32.dll", SetLastError=true)]
 		static extern bool SetWindowPos(IntPtr hWnd, IntPtr hWndInsertAfter, int x, int y, int width, int height, uint flags);
 
-		public readonly IssueMonitor IssueMonitor;
+		public IssueMonitor IssueMonitor { get; }
 
-		public bool isWarning;
+		public bool IsWarning { get; set; }
 
-		public bool? strongAlert;
+		public bool? StrongAlert { get; set; }
 
 		public IssueAlertWindow(IssueMonitor issueMonitor, IssueData issue, IssueAlertReason reason)
 		{
-			this.IssueMonitor = issueMonitor;
-			this.Issue = issue;
+			IssueMonitor = issueMonitor;
+			Issue = issue;
 
 			InitializeComponent();
 
@@ -53,9 +52,9 @@ namespace UnrealGameSync
 		{
 			base.OnPaintBackground(e);
 
-			if(strongAlert ?? false)
+			if(StrongAlert ?? false)
 			{
-				Color stripeColor = isWarning? Color.FromArgb(216, 167, 64) : Color.FromArgb(200, 74, 49);//214, 69, 64);
+				Color stripeColor = IsWarning? Color.FromArgb(216, 167, 64) : Color.FromArgb(200, 74, 49);//214, 69, 64);
 				using (Brush stripeBrush = new SolidBrush(stripeColor))
 				{
 					e.Graphics.FillRectangle(stripeBrush, 0, 0, Bounds.Width, Bounds.Height);
@@ -67,7 +66,7 @@ namespace UnrealGameSync
 			}
 			else
 			{
-				Color stripeColor = isWarning? Color.FromArgb(216, 167, 64) : Color.FromArgb(214, 69, 64);
+				Color stripeColor = IsWarning? Color.FromArgb(216, 167, 64) : Color.FromArgb(214, 69, 64);
 
 				Color backgroundColor = Color.FromArgb(241, 236, 236);
 				using(Brush solidBrush = new SolidBrush(backgroundColor))
@@ -75,7 +74,6 @@ namespace UnrealGameSync
 	//				e.Graphics.FillRectangle(SolidBrush, 0, 0, Bounds.Width - 1, Bounds.Height - 1);
 				}
 
-				Color borderColor = stripeColor; 
 				using(Pen pen = new Pen(Color.FromArgb(128, 128, 128)))
 				{
 					e.Graphics.DrawRectangle(pen, 0, 0, Bounds.Width - 1, Bounds.Height - 1);
@@ -140,7 +138,7 @@ namespace UnrealGameSync
 				summary = summary.Substring(0, maxLength).TrimEnd() + "...";
 			}
 
-			if(summary != SummaryLabel.Text || ownerText != OwnerLabel.Text || Reason != newReason || isWarning != newIsWarning || strongAlert != newStrongAlert)
+			if(summary != SummaryLabel.Text || ownerText != OwnerLabel.Text || Reason != newReason || IsWarning != newIsWarning || StrongAlert != newStrongAlert)
 			{
 				Rectangle prevBounds = Bounds;
 				SuspendLayout();
@@ -149,9 +147,9 @@ namespace UnrealGameSync
 				OwnerLabel.Text = ownerText;
 
 				bool forceUpdateButtons = false;
-				if(strongAlert != newStrongAlert)
+				if(StrongAlert != newStrongAlert)
 				{
-					strongAlert = newStrongAlert;
+					StrongAlert = newStrongAlert;
 
 					if(newStrongAlert)
 					{
@@ -175,9 +173,9 @@ namespace UnrealGameSync
 					forceUpdateButtons = true;
 				}
 
-				if (isWarning != newIsWarning)
+				if (IsWarning != newIsWarning)
 				{
-					isWarning = newIsWarning;
+					IsWarning = newIsWarning;
 				}
 
 				if(Reason != newReason || forceUpdateButtons)
@@ -226,10 +224,7 @@ namespace UnrealGameSync
 			SetWindowPos(Handle, HwndTopmost, 0, 0, 0, 0, SwpNomove | SwpNosize | SwpNoactivate);
 		}
 
-		protected override bool ShowWithoutActivation 
-		{
-			get { return true; }
-		}
+		protected override bool ShowWithoutActivation => true;
 
 		private void SummaryLabel_LinkClicked(object sender, System.Windows.Forms.LinkLabelLinkClickedEventArgs e)
 		{
diff --git a/Engine/Source/Programs/UnrealGameSync/UnrealGameSync/Forms/IssueBrowserWindow.cs b/Engine/Source/Programs/UnrealGameSync/UnrealGameSync/Forms/IssueBrowserWindow.cs
--- a/Engine/Source/Programs/UnrealGameSync/UnrealGameSync/Forms/IssueBrowserWindow.cs
+++ b/Engine/Source/Programs/UnrealGameSync/UnrealGameSync/Forms/IssueBrowserWindow.cs
@@ -1,18 +1,11 @@
 // Copyright Epic Games, Inc. All Rights Reserved.
 
 using EpicGames.Perforce;
-using Microsoft.Extensions.DependencyInjection;
-using Microsoft.Extensions.Logging;
 using System;
 using System.Collections.Generic;
-using System.ComponentModel;
 using System.Data;
 using System.Drawing;
-using System.IO;
 using System.Linq;
-using System.Runtime.CompilerServices;
-using System.Text;
-using System.Text.RegularExpressions;
 using System.Threading;
 using System.Threading.Tasks;
 using System.Windows.Forms;
@@ -21,32 +14,30 @@ namespace UnrealGameSync
 {
 	partial class IssueBrowserWindow : Form
 	{
-		IssueMonitor _issueMonitor;
-		IPerforceSettings _perforceSettings;
-		TimeSpan? _serverTimeOffset;
-		IServiceProvider _serviceProvider;
-		ILogger _logger;
-		string? _currentStream;
+		readonly IssueMonitor _issueMonitor;
+		readonly IPerforceSettings _perforceSettings;
+		readonly TimeSpan? _serverTimeOffset;
+		readonly IServiceProvider _serviceProvider;
+		readonly string? _currentStream;
 		int _maxResults = 0;
 		int _pendingMaxResults = 0;
 		string? _filterName;
-		Dictionary<string, Func<IssueData, bool>> _customFilters = new Dictionary<string, Func<IssueData, bool>>();
+		readonly Dictionary<string, Func<IssueData, bool>> _customFilters = new Dictionary<string, Func<IssueData, bool>>();
 		List<IssueData> _issues = new List<IssueData>();
-		SynchronizationContext _mainThreadSynchronizationContext;
+		readonly SynchronizationContext _mainThreadSynchronizationContext;
 		Task? _backgroundTask;
 		bool _disposed;
 
 		public IssueBrowserWindow(IssueMonitor issueMonitor, IPerforceSettings perforceSettings, TimeSpan? serverTimeOffset, IServiceProvider serviceProvider, string? currentStream, Dictionary<string, Func<IssueData, bool>> customFilters, string? filterName)
 		{
-			this._issueMonitor = issueMonitor;
-			this._perforceSettings = perforceSettings;
-			this._serverTimeOffset = serverTimeOffset;
-			this._serviceProvider = serviceProvider;
-			this._logger = serviceProvider.GetRequiredService<ILogger<IssueBrowserWindow>>();
-			this._currentStream = currentStream;
-			this._filterName = filterName;
-			this._customFilters = customFilters;
-			this._mainThreadSynchronizationContext = SynchronizationContext.Current!;
+			_issueMonitor = issueMonitor;
+			_perforceSettings = perforceSettings;
+			_serverTimeOffset = serverTimeOffset;
+			_serviceProvider = serviceProvider;
+			_currentStream = currentStream;
+			_filterName = filterName;
+			_customFilters = customFilters;
+			_mainThreadSynchronizationContext = SynchronizationContext.Current!;
 
 			issueMonitor.AddRef();
 
@@ -112,11 +103,23 @@ namespace UnrealGameSync
 				}
 
 				List<IssueData> newIssues = newSortedIssues.Values.Reverse().ToList();
-				_mainThreadSynchronizationContext.Post((o) => { if (!_disposed) { FetchIssuesSuccess(newMaxResults, newIssues); } }, null);
+				_mainThreadSynchronizationContext.Post((o) => 
+				{ 
+					if (!_disposed) 
+					{ 
+						FetchIssuesSuccess(newMaxResults, newIssues); 
+					} 
+				}, null);
 			}
 			catch(Exception ex)
 			{
-				_mainThreadSynchronizationContext.Post((o) => { if (!_disposed) { FetchIssuesFailure(ex); } }, null);
+				_mainThreadSynchronizationContext.Post((o) => 
+				{ 
+					if (!_disposed) 
+					{ 
+						FetchIssuesFailure(ex); 
+					} 
+				}, null);
 			}
 		}
 
@@ -258,11 +261,11 @@ namespace UnrealGameSync
 			}
 		}
 
-		static List<IssueBrowserWindow> _existingWindows = new List<IssueBrowserWindow>();
+		static readonly List<IssueBrowserWindow> s_existingWindows = new List<IssueBrowserWindow>();
 
 		public static void Show(Form owner, IssueMonitor issueMonitor, IPerforceSettings perforceSettings, TimeSpan? serverTimeOffset, IServiceProvider serviceProvider, string? currentStream, Dictionary<string, Func<IssueData, bool>> customFilters, string? defaultFilter)
 		{
-			IssueBrowserWindow? window = _existingWindows.FirstOrDefault(x => x._issueMonitor == issueMonitor);
+			IssueBrowserWindow? window = s_existingWindows.FirstOrDefault(x => x._issueMonitor == issueMonitor);
 			if(window == null)
 			{
 				window = new IssueBrowserWindow(issueMonitor, perforceSettings, serverTimeOffset, serviceProvider, currentStream, customFilters, defaultFilter);
@@ -271,8 +274,8 @@ namespace UnrealGameSync
 				window.Location = new Point(owner.Location.X + (owner.Width - window.Width) / 2, owner.Location.Y + (owner.Height - window.Height) / 2);
 				window.Show(owner);
 
-				_existingWindows.Add(window);
-				window.FormClosed += (e, s) => _existingWindows.Remove(window);
+				s_existingWindows.Add(window);
+				window.FormClosed += (e, s) => s_existingWindows.Remove(window);
 			}
 			else
 			{
@@ -317,7 +320,11 @@ namespace UnrealGameSync
 				{
 					ToolStripMenuItem item = new ToolStripMenuItem(customFilter.Key);
 					item.Checked = (_filterName == customFilter.Key);
-					item.Click += (s, e) => { _filterName = customFilter.Key; UpdateIssueList(); };
+					item.Click += (s, e) => 
+					{ 
+						_filterName = customFilter.Key; 
+						UpdateIssueList(); 
+					};
 					FilterMenu.Items.Add(item);
 				}
 			}
@@ -338,7 +345,11 @@ namespace UnrealGameSync
 				{
 					ToolStripMenuItem item = new ToolStripMenuItem(stream);
 					item.Checked = (_filterName == stream);
-					item.Click += (s, e) => { _filterName = stream; UpdateIssueList(); };
+					item.Click += (s, e) => 
+					{ 
+						_filterName = stream; 
+						UpdateIssueList(); 
+					};
 					FilterMenu.Items.Add(item);
 				}
 			}
@@ -372,10 +383,8 @@ namespace UnrealGameSync
 			{
 				IssueListView.DrawTrackedBackground(e.Graphics, e.Bounds);
 			}
-			else if (e.Item.Tag is IssueData)
+			else if (e.Item.Tag is IssueData issue)
 			{
-				IssueData issue = (IssueData)e.Item.Tag;
-
 				Color backgroundColor;
 				if (issue.ResolvedAt.HasValue)
 				{
diff --git a/Engine/Source/Programs/UnrealGameSync/UnrealGameSync/Forms/IssueDetailsWindow.cs b/Engine/Source/Programs/UnrealGameSync/UnrealGameSync/Forms/IssueDetailsWindow.cs
--- a/Engine/Source/Programs/UnrealGameSync/UnrealGameSync/Forms/IssueDetailsWindow.cs
+++ b/Engine/Source/Programs/UnrealGameSync/UnrealGameSync/Forms/IssueDetailsWindow.cs
@@ -7,10 +7,8 @@ using Microsoft.Extensions.Logging;
 using System;
 using System.Collections.Generic;
 using System.Data;
-using System.Diagnostics;
 using System.Diagnostics.CodeAnalysis;
 using System.Drawing;
-using System.Drawing.Drawing2D;
 using System.IO;
 using System.Linq;
 using System.Text;
@@ -33,28 +31,28 @@ namespace UnrealGameSync
 
 			public BuildGroup(string jobName, string jobUrl, int change, IssueBuildOutcome outcome, IReadOnlyList<IssueBuildData> builds)
 			{
-				this.JobName = jobName;
-				this.JobUrl = jobUrl;
-				this.Change = change;
-				this.Outcome = outcome;
-				this.Builds = builds;
+				JobName = jobName;
+				JobUrl = jobUrl;
+				Change = change;
+				Outcome = outcome;
+				Builds = builds;
 			}
 		}
 
 		class PerforceChangeRange
 		{
-			public bool Expanded;
+			public bool Expanded { get; set; }
 			public int MinChange { get; }
 			public int MaxChange { get; }
-			public List<ChangesRecord>? Changes;
-			public string? ErrorMessage;
+			public List<ChangesRecord>? Changes { get; set; }
+			public string? ErrorMessage { get; set; }
 			public BuildGroup BuildGroup { get; }
 
 			public PerforceChangeRange(int minChange, int maxChange, BuildGroup buildGroup)
 			{
-				this.MinChange = minChange;
-				this.MaxChange = maxChange;
-				this.BuildGroup = buildGroup;
+				MinChange = minChange;
+				MaxChange = maxChange;
+				BuildGroup = buildGroup;
 			}
 		}
 
@@ -65,7 +63,7 @@ namespace UnrealGameSync
 			public PerforceChangeDetailsWithDescribeRecord(DescribeRecord describeRecord)
 				: base(describeRecord)
 			{
-				this.DescribeRecord = describeRecord;
+				DescribeRecord = describeRecord;
 			}
 		}
 
@@ -79,19 +77,19 @@ namespace UnrealGameSync
 			readonly object _lockObject = new object();
 			readonly SynchronizationContext _synchronizationContext;
 			Task? _workerTask;
-			CancellationTokenSource _cancellationSource;
-			AsyncEvent _refreshEvent;
-			List<PerforceChangeRange> _requests = new List<PerforceChangeRange>();
-			Dictionary<int, PerforceChangeDetailsWithDescribeRecord> _changeNumberToDetails = new Dictionary<int, PerforceChangeDetailsWithDescribeRecord>();
+			readonly CancellationTokenSource _cancellationSource;
+			readonly AsyncEvent _refreshEvent;
+			readonly List<PerforceChangeRange> _requests = new List<PerforceChangeRange>();
+			readonly Dictionary<int, PerforceChangeDetailsWithDescribeRecord> _changeNumberToDetails = new Dictionary<int, PerforceChangeDetailsWithDescribeRecord>();
 
 			public PerforceWorkerTask(IPerforceSettings perforceSettings, string filter, Action<PerforceChangeRange>? onUpdateChanges, Action<ChangesRecord>? onUpdateChangeMetadata, ILogger logger)
 			{
-				this._perforceSettings = perforceSettings;
-				this._filter = filter;
-				this._onUpdateChanges = onUpdateChanges;
-				this._onUpdateChangeMetadata = onUpdateChangeMetadata;
-				this._logger = logger;
-				this._synchronizationContext = SynchronizationContext.Current!;
+				_perforceSettings = perforceSettings;
+				_filter = filter;
+				_onUpdateChanges = onUpdateChanges;
+				_onUpdateChangeMetadata = onUpdateChangeMetadata;
+				_logger = logger;
+				_synchronizationContext = SynchronizationContext.Current!;
 
 				_refreshEvent = new AsyncEvent();
 				_cancellationSource = new CancellationTokenSource();
@@ -211,25 +209,25 @@ namespace UnrealGameSync
 
 		class BadgeInfo
 		{
-			public string Label;
-			public Color Color;
+			public string _label;
+			public Color _color;
 
 			public BadgeInfo(string label, Color color)
 			{
-				this.Label = label;
-				this.Color = color;
+				_label = label;
+				_color = color;
 			}
 		}
 
 		class ExpandRangeStatusElement : StatusElement
 		{
-			string _text;
-			Action _linkAction;
+			readonly string _text;
+			readonly Action _linkAction;
 
 			public ExpandRangeStatusElement(string text, Action inLinkAction)
 			{
-				this._text = text;
-				this._linkAction = inLinkAction;
+				_text = text;
+				_linkAction = inLinkAction;
 				Cursor = NativeCursors.Hand;
 			}
 
@@ -240,7 +238,7 @@ namespace UnrealGameSync
 
 			public override Size Measure(Graphics graphics, StatusElementResources resources)
 			{
-				return TextRenderer.MeasureText(graphics, _text, resources.FindOrAddFont(FontStyle.Regular), new Size(int.MaxValue, int.MaxValue), TextFormatFlags.NoPadding);
+				return TextRenderer.MeasureText(graphics, _text, resources.FindOrAddFont(FontStyle.Regular), new Size(Int32.MaxValue, Int32.MaxValue), TextFormatFlags.NoPadding);
 			}
 
 			public override void Draw(Graphics graphics, StatusElementResources resources)
@@ -263,13 +261,13 @@ namespace UnrealGameSync
 
 		IssueMonitor _issueMonitor;
 		IssueData _issue;
-		List<IssueBuildData> _issueBuilds;
-		List<IssueDiagnosticData> _diagnostics;
-		IPerforceSettings _perforceSettings;
-		TimeSpan? _serverTimeOffset;
+		readonly List<IssueBuildData> _issueBuilds;
+		readonly List<IssueDiagnosticData> _diagnostics;
+		readonly IPerforceSettings _perforceSettings;
+		readonly TimeSpan? _serverTimeOffset;
 		PerforceWorkerTask? _perforceWorker;
-		IServiceProvider _serviceProvider;
-		SynchronizationContext _mainThreadSynchronizationContext;
+		readonly IServiceProvider _serviceProvider;
+		readonly SynchronizationContext _mainThreadSynchronizationContext;
 		string? _selectedStream;
 		List<PerforceChangeRange> _selectedStreamRanges = new List<PerforceChangeRange>();
 		bool _isDisposing;
@@ -282,13 +280,13 @@ namespace UnrealGameSync
 
 		IssueDetailsWindow(IssueMonitor issueMonitor, IssueData issue, List<IssueBuildData> issueBuilds, List<IssueDiagnosticData> diagnostics, IPerforceSettings perforceSettings, TimeSpan? serverTimeOffset, IServiceProvider serviceProvider, string? currentStream)
 		{
-			this._issueMonitor = issueMonitor;
-			this._issue = issue;
-			this._issueBuilds = issueBuilds;
-			this._diagnostics = diagnostics;
-			this._perforceSettings = perforceSettings;
-			this._serverTimeOffset = serverTimeOffset;
-			this._serviceProvider = serviceProvider;
+			_issueMonitor = issueMonitor;
+			_issue = issue;
+			_issueBuilds = issueBuilds;
+			_diagnostics = diagnostics;
+			_perforceSettings = perforceSettings;
+			_serverTimeOffset = serverTimeOffset;
+			_serviceProvider = serviceProvider;
 
 			issueMonitor.AddRef();
 
@@ -296,9 +294,9 @@ namespace UnrealGameSync
 
 			InitializeComponent();
 
-			this.Text = String.Format("Issue {0}", issue.Id);
-			this._statusElementResources = new StatusElementResources(BuildListView.Font);
-			base.Disposed += IssueDetailsWindow_Disposed;
+			Text = String.Format("Issue {0}", issue.Id);
+			_statusElementResources = new StatusElementResources(BuildListView.Font);
+			Disposed += IssueDetailsWindow_Disposed;
 
 			issueMonitor.OnIssuesChanged += OnUpdateIssuesAsync;
 			issueMonitor.StartTracking(issue.Id);
@@ -397,16 +395,6 @@ namespace UnrealGameSync
 			}
 		}
 
-		void UpdateSummaryTextIfChanged(TextBox textBox, string newText)
-		{
-			if(textBox.Text != newText)
-			{
-				textBox.Text = newText;
-				textBox.SelectionLength = 0;
-				textBox.SelectionStart = newText.Length;
-			}
-		}
-
 		static void AppendEscapedRtf(StringBuilder result, string text)
 		{
 			for (int idx = 0; idx < text.Length; idx++)
@@ -443,7 +431,7 @@ namespace UnrealGameSync
 			richText.Append(@"}");
 		}
 
-		static string CreateRichTextErrors(IssueData issue, List<IssueBuildData> issueBuilds, List<IssueDiagnosticData> diagnostics)
+		static string CreateRichTextErrors(List<IssueBuildData> issueBuilds, List<IssueDiagnosticData> diagnostics)
 		{
 			StringBuilder richText = new StringBuilder();
 
@@ -604,7 +592,7 @@ namespace UnrealGameSync
 			UpdateSummaryTextIfChanged(StepNamesTextBox, String.Join(", ", _issueBuilds.Select(x => x.JobStepName).Distinct().OrderBy(x => x)));
 			UpdateSummaryTextIfChanged(StreamNamesTextBox, String.Join(", ", _issueBuilds.Select(x => x.Stream).Distinct().OrderBy(x => x)));
 
-			string rtfText = CreateRichTextErrors(_issue, _issueBuilds, _diagnostics);
+			string rtfText = CreateRichTextErrors(_issueBuilds, _diagnostics);
 			if (_lastDetailsText != rtfText)
 			{
 				using (MemoryStream stream = new MemoryStream(Encoding.UTF8.GetBytes(rtfText), false))
@@ -631,14 +619,6 @@ namespace UnrealGameSync
 			DestroyWorker();
 		}
 
-		void FetchAllBuildChanges()
-		{
-			foreach(PerforceChangeRange range in _selectedStreamRanges)
-			{
-				FetchBuildChanges(range);
-			}
-		}
-
 		void FetchBuildChanges(PerforceChangeRange range)
 		{
 			if(!range.Expanded)
@@ -945,7 +925,7 @@ namespace UnrealGameSync
 							FetchBuildChanges(_selectedStreamRanges[idx + 1]);
 						}
 					}
-					FetchBuildChanges(_selectedStreamRanges[_selectedStreamRanges.Count - 1]);
+					FetchBuildChanges(_selectedStreamRanges[^1]);
 				}
 
 				BuildListView.EndUpdate();
@@ -969,7 +949,13 @@ namespace UnrealGameSync
 
 		private void OnUpdateIssuesAsync()
 		{
-			_mainThreadSynchronizationContext.Post((o) => { if(!_isDisposing){ OnUpdateIssues(); } }, null);
+			_mainThreadSynchronizationContext.Post((o) => 
+			{ 
+				if(!_isDisposing)
+				{ 
+					OnUpdateIssues(); 
+				} 
+			}, null);
 		}
 
 		void DestroyWorker()
@@ -1011,19 +997,19 @@ namespace UnrealGameSync
 			}
 		}
 
-		static List<IssueDetailsWindow> _existingWindows = new List<IssueDetailsWindow>();
+		static readonly List<IssueDetailsWindow> s_existingWindows = new List<IssueDetailsWindow>();
 
 		class UpdateIssueDetailsTask
 		{
-			string _apiUrl;
-			long _issueId;
-			List<IssueDiagnosticData> _diagnostics;
+			readonly string _apiUrl;
+			readonly long _issueId;
+			readonly List<IssueDiagnosticData> _diagnostics;
 
 			public UpdateIssueDetailsTask(string apiUrl, long issueId, List<IssueDiagnosticData> diagnostics)
 			{
-				this._apiUrl = apiUrl;
-				this._issueId = issueId;
-				this._diagnostics = diagnostics;
+				_apiUrl = apiUrl;
+				_issueId = issueId;
+				_diagnostics = diagnostics;
 			}
 
 			public async Task RunAsync(CancellationToken cancellationToken)
@@ -1034,8 +1020,9 @@ namespace UnrealGameSync
 
 		public static void Show(Form owner, IssueMonitor issueMonitor, IPerforceSettings perforceSettings, TimeSpan? serverTimeOffset, IssueData issue, IServiceProvider serviceProvider, string? currentStream)
 		{
-			Func<CancellationToken, Task<List<IssueBuildData>>> func = x => RestApi.GetAsync<List<IssueBuildData>>($"{issueMonitor.ApiUrl}/api/issues/{issue.Id}/builds", x);
-			ModalTask<List<IssueBuildData>>? issueBuilds = ModalTask.Execute(owner, "Querying issue", "Querying issue, please wait...", func);
+			Task<List<IssueBuildData>> Func(CancellationToken cancellationToken) => RestApi.GetAsync<List<IssueBuildData>>($"{issueMonitor.ApiUrl}/api/issues/{issue.Id}/builds", cancellationToken);
+
+			ModalTask<List<IssueBuildData>>? issueBuilds = ModalTask.Execute(owner, "Querying issue", "Querying issue, please wait...", Func);
 			if (issueBuilds != null && issueBuilds.Succeeded)
 			{
 				Show(owner, issueMonitor, perforceSettings, serverTimeOffset, issue, issueBuilds.Result, serviceProvider, currentStream);
@@ -1044,7 +1031,7 @@ namespace UnrealGameSync
 
 		public static void Show(Form owner, IssueMonitor issueMonitor, IPerforceSettings perforceSettings, TimeSpan? serverTimeOffset, IssueData issue, List<IssueBuildData> issueBuilds, IServiceProvider serviceProvider, string? currentStream)
 		{
-			IssueDetailsWindow? window = _existingWindows.FirstOrDefault(x => x._issueMonitor == issueMonitor && x._issue.Id == issue.Id);
+			IssueDetailsWindow? window = s_existingWindows.FirstOrDefault(x => x._issueMonitor == issueMonitor && x._issue.Id == issue.Id);
 			if(window == null)
 			{
 				List<IssueDiagnosticData> diagnostics = new List<IssueDiagnosticData>();
@@ -1069,8 +1056,8 @@ namespace UnrealGameSync
 				}
 				window.Show();
 
-				_existingWindows.Add(window);
-				window.FormClosed += (s, e) => _existingWindows.Remove(window);
+				s_existingWindows.Add(window);
+				window.FormClosed += (s, e) => s_existingWindows.Remove(window);
 			}
 			else
 			{
@@ -1106,10 +1093,8 @@ namespace UnrealGameSync
 			{
 				BuildListView.DrawTrackedBackground(e.Graphics, e.Bounds);
 			}
-			else if(e.Item.Tag is BuildGroup)
+			else if(e.Item.Tag is BuildGroup buildGroup)
 			{
-				BuildGroup buildGroup = (BuildGroup)e.Item.Tag;
-
 				Color backgroundColor;
 				if(buildGroup.Outcome == IssueBuildOutcome.Error)
 				{
@@ -1146,7 +1131,7 @@ namespace UnrealGameSync
 				return;
 			}
 
-			Font currentFont = this.Font;//BuildFont;(Change.Number == Workspace.PendingChangeNumber || Change.Number == Workspace.CurrentChangeNumber)? SelectedBuildFont : BuildFont;
+			Font currentFont = Font;//BuildFont;(Change.Number == Workspace.PendingChangeNumber || Change.Number == Workspace.CurrentChangeNumber)? SelectedBuildFont : BuildFont;
 
 			Color textColor = SystemColors.WindowText;
 
@@ -1154,10 +1139,8 @@ namespace UnrealGameSync
 			{
 				BuildListView.DrawCustomSubItem(e.Graphics, e.SubItem);
 			}
-			else if(e.Item.Tag is ChangesRecord)
+			else if(e.Item.Tag is ChangesRecord change)
 			{
-				ChangesRecord change = (ChangesRecord)e.Item.Tag;
-
 				Font changeFont = BuildListView.Font;
 				if(_issue.Owner != null && String.Compare(change.User, _issue.Owner, StringComparison.OrdinalIgnoreCase) == 0)
 				{
@@ -1181,12 +1164,11 @@ namespace UnrealGameSync
 					TextRenderer.DrawText(e.Graphics, e.SubItem.Text, changeFont, e.Bounds, textColor, TextFormatFlags.EndEllipsis | TextFormatFlags.SingleLine | TextFormatFlags.HorizontalCenter | TextFormatFlags.VerticalCenter | TextFormatFlags.NoPrefix);
 				}
 			}
-			else if(e.Item.Tag is IssueBuildData)
+			else if(e.Item.Tag is IssueBuildData buildData)
 			{
 				Font boldFont = BuildListView.Font;
 
 //				TextColor = SystemColors.Window;
-				IssueBuildData buildData = (IssueBuildData)e.Item.Tag;
 				if(e.ColumnIndex == IconHeader.Index)
 				{
 					if(buildData.Outcome == IssueBuildOutcome.Success)
@@ -1258,16 +1240,10 @@ namespace UnrealGameSync
 
 		private void BuildListView_FontChanged(object sender, EventArgs e)
 		{
-			if(_boldFont != null)
-			{
-				_boldFont.Dispose();
-			}
+			_boldFont?.Dispose();
 			_boldFont = new Font(BuildListView.Font.FontFamily, BuildListView.Font.SizeInPoints, FontStyle.Bold);
 
-			if(_statusElementResources != null)
-			{
-				_statusElementResources.Dispose();
-			}
+			_statusElementResources?.Dispose();
 			_statusElementResources = new StatusElementResources(BuildListView.Font);
 		}
 
diff --git a/Engine/Source/Programs/UnrealGameSync/UnrealGameSync/Forms/IssueFixedWindow.cs b/Engine/Source/Programs/UnrealGameSync/UnrealGameSync/Forms/IssueFixedWindow.cs
--- a/Engine/Source/Programs/UnrealGameSync/UnrealGameSync/Forms/IssueFixedWindow.cs
+++ b/Engine/Source/Programs/UnrealGameSync/UnrealGameSync/Forms/IssueFixedWindow.cs
@@ -8,9 +8,7 @@ using System;
 using System.Collections.Generic;
 using System.ComponentModel;
 using System.Data;
-using System.Drawing;
 using System.Linq;
-using System.Text;
 using System.Threading;
 using System.Threading.Tasks;
 using System.Windows.Forms;
@@ -25,22 +23,22 @@ namespace UnrealGameSync
 		{
 			public delegate void OnCompleteDelegate(string userName, List<DescribeRecord> changes);
 
-			IPerforceSettings _perforceSettings;
-			SynchronizationContext _mainThreadSyncContext;
-			AsyncEvent _wakeEvent;
+			readonly IPerforceSettings _perforceSettings;
+			readonly SynchronizationContext _mainThreadSyncContext;
+			readonly AsyncEvent _wakeEvent;
 			string? _requestedUserName;
-			CancellationTokenSource _cancellationSource;
+			readonly CancellationTokenSource _cancellationSource;
 			Task _backgroundTask;
 			OnCompleteDelegate? _onComplete;
-			ILogger _logger;
+			readonly ILogger _logger;
 
 			public FindChangesWorker(IPerforceSettings perforceSettings, OnCompleteDelegate? onComplete, ILogger logger)
 			{
-				this._perforceSettings = perforceSettings;
-				this._mainThreadSyncContext = SynchronizationContext.Current!;
+				_perforceSettings = perforceSettings;
+				_mainThreadSyncContext = SynchronizationContext.Current!;
 				_wakeEvent = new AsyncEvent();
-				this._onComplete = onComplete;
-				this._logger = logger;
+				_onComplete = onComplete;
+				_logger = logger;
 
 				_cancellationSource = new CancellationTokenSource();
 				_backgroundTask = Task.Run(() => DoWork(_cancellationSource.Token));
@@ -108,18 +106,18 @@ namespace UnrealGameSync
 			}
 		}
 
-		IPerforceSettings _perforceSettings;
+		readonly IPerforceSettings _perforceSettings;
 		int _changeNumber;
-		FindChangesWorker _worker;
-		IServiceProvider _serviceProvider;
+		readonly FindChangesWorker _worker;
+		readonly IServiceProvider _serviceProvider;
 	
 		public IssueFixedWindow(IPerforceSettings perforceSettings, int initialChangeNumber, IServiceProvider serviceProvider)
 		{
 			InitializeComponent();
 
-			this._perforceSettings = perforceSettings;
-			this._worker = new FindChangesWorker(perforceSettings, PopulateChanges, serviceProvider.GetRequiredService<ILogger<FindChangesWorker>>());
-			this._serviceProvider = serviceProvider;
+			_perforceSettings = perforceSettings;
+			_worker = new FindChangesWorker(perforceSettings, PopulateChanges, serviceProvider.GetRequiredService<ILogger<FindChangesWorker>>());
+			_serviceProvider = serviceProvider;
 			components!.Add(_worker);
 
 			UserNameTextBox.Text = perforceSettings.UserName;
@@ -301,7 +299,7 @@ namespace UnrealGameSync
 		{
 			if(SpecifyChangeRadioButton.Checked)
 			{
-				return int.TryParse(ChangeNumberTextBox.Text, out changeNumber);
+				return Int32.TryParse(ChangeNumberTextBox.Text, out changeNumber);
 			}
 			else if(SystemicFixRadioButton.Checked)
 			{
@@ -326,8 +324,7 @@ namespace UnrealGameSync
 
 		private void UpdateOkButton()
 		{
-			int changeNumber;
-			OkBtn.Enabled = TryGetSelectedChange(out changeNumber);
+			OkBtn.Enabled = TryGetSelectedChange(out _);
 		}
 
 		private void SpecifyChangeRadioButton_CheckedChanged(object sender, EventArgs e)
diff --git a/Engine/Source/Programs/UnrealGameSync/UnrealGameSync/Forms/IssueSettingsWindow.cs b/Engine/Source/Programs/UnrealGameSync/UnrealGameSync/Forms/IssueSettingsWindow.cs
--- a/Engine/Source/Programs/UnrealGameSync/UnrealGameSync/Forms/IssueSettingsWindow.cs
+++ b/Engine/Source/Programs/UnrealGameSync/UnrealGameSync/Forms/IssueSettingsWindow.cs
@@ -3,25 +3,19 @@
 using Microsoft.Extensions.Logging;
 using System;
 using System.Collections.Generic;
-using System.ComponentModel;
-using System.Data;
-using System.Drawing;
-using System.Linq;
-using System.Text;
-using System.Threading.Tasks;
 using System.Windows.Forms;
 
 namespace UnrealGameSync
 {
 	partial class IssueSettingsWindow : Form
 	{
-		UserSettings _settings;
-		ILogger _logger;
+		readonly UserSettings _settings;
+		readonly ILogger _logger;
 
 		public IssueSettingsWindow(UserSettings settings, string currentProject, ILogger logger)
 		{
-			this._settings = settings;
-			this._logger = logger;
+			_settings = settings;
+			_logger = logger;
 
 			InitializeComponent();
 
@@ -112,7 +106,7 @@ namespace UnrealGameSync
 			if(NotifyUnresolvedCheckBox.Checked)
 			{
 				ushort newNotifyUnresolvedMinutesValue;
-				if(!ushort.TryParse(NotifyUnresolvedTextBox.Text, out newNotifyUnresolvedMinutesValue))
+				if(!UInt16.TryParse(NotifyUnresolvedTextBox.Text, out newNotifyUnresolvedMinutesValue))
 				{
 					MessageBox.Show("Invalid time");
 					return;
@@ -124,7 +118,7 @@ namespace UnrealGameSync
 			if (NotifyUnacknowledgedCheckBox.Checked)
 			{
 				ushort newNotifyUnacknowledgedMinutesValue;
-				if (!ushort.TryParse(NotifyUnacknowledgedTextBox.Text, out newNotifyUnacknowledgedMinutesValue))
+				if (!UInt16.TryParse(NotifyUnacknowledgedTextBox.Text, out newNotifyUnacknowledgedMinutesValue))
 				{
 					MessageBox.Show("Invalid time");
 					return;
@@ -136,7 +130,7 @@ namespace UnrealGameSync
 			if(NotifyUnassignedCheckBox.Checked)
 			{
 				ushort newNotifyUnassignedMinutesValue;
-				if(!ushort.TryParse(NotifyUnassignedTextBox.Text, out newNotifyUnassignedMinutesValue))
+				if(!UInt16.TryParse(NotifyUnassignedTextBox.Text, out newNotifyUnassignedMinutesValue))
 				{
 					MessageBox.Show("Invalid time");
 					return;
diff --git a/Engine/Source/Programs/UnrealGameSync/UnrealGameSync/Forms/LauncherUpdateWindow.cs b/Engine/Source/Programs/UnrealGameSync/UnrealGameSync/Forms/LauncherUpdateWindow.cs
--- a/Engine/Source/Programs/UnrealGameSync/UnrealGameSync/Forms/LauncherUpdateWindow.cs
+++ b/Engine/Source/Programs/UnrealGameSync/UnrealGameSync/Forms/LauncherUpdateWindow.cs
@@ -1,11 +1,5 @@
 // Copyright Epic Games, Inc. All Rights Reserved.
 
-using System;
-using System.Collections.Generic;
-using System.ComponentModel;
-using System.Data;
-using System.Drawing;
-using System.Text;
 using System.Windows.Forms;
 
 namespace UnrealGameSync.Forms
@@ -16,10 +10,5 @@ namespace UnrealGameSync.Forms
 		{
 			InitializeComponent();
 		}
-
-		private void button1_Click(object sender, EventArgs e)
-		{
-
-		}
 	}
 }
diff --git a/Engine/Source/Programs/UnrealGameSync/UnrealGameSync/Forms/LeaveCommentWindow.cs b/Engine/Source/Programs/UnrealGameSync/UnrealGameSync/Forms/LeaveCommentWindow.cs
--- a/Engine/Source/Programs/UnrealGameSync/UnrealGameSync/Forms/LeaveCommentWindow.cs
+++ b/Engine/Source/Programs/UnrealGameSync/UnrealGameSync/Forms/LeaveCommentWindow.cs
@@ -1,13 +1,5 @@
-﻿// Copyright Epic Games, Inc. All Rights Reserved.
+// Copyright Epic Games, Inc. All Rights Reserved.
 
-using System;
-using System.Collections.Generic;
-using System.ComponentModel;
-using System.Data;
-using System.Drawing;
-using System.Linq;
-using System.Text;
-using System.Threading.Tasks;
 using System.Windows.Forms;
 
 namespace UnrealGameSync
diff --git a/Engine/Source/Programs/UnrealGameSync/UnrealGameSync/Forms/MainWindow.cs b/Engine/Source/Programs/UnrealGameSync/UnrealGameSync/Forms/MainWindow.cs
--- a/Engine/Source/Programs/UnrealGameSync/UnrealGameSync/Forms/MainWindow.cs
+++ b/Engine/Source/Programs/UnrealGameSync/UnrealGameSync/Forms/MainWindow.cs
@@ -58,12 +58,12 @@ namespace UnrealGameSync
 	{
 		class WorkspaceIssueMonitor
 		{
-			public IssueMonitor IssueMonitor;
-			public int RefCount;
+			public IssueMonitor _issueMonitor;
+			public int _refCount;
 
 			public WorkspaceIssueMonitor(IssueMonitor issueMonitor)
 			{
-				this.IssueMonitor = issueMonitor;
+				_issueMonitor = issueMonitor;
 			}
 		}
 
@@ -83,49 +83,46 @@ namespace UnrealGameSync
 		static extern int SetWindowTheme(IntPtr hWnd, string pszSubAppName, string pszSubIdList);
 
 		[DllImport("user32.dll")]
-		public static extern int SendMessage(IntPtr hWnd, Int32 wMsg, Int32 wParam, Int32 lParam);
-
-		private const int WmSetredraw = 11; 
-
-		UpdateMonitor _updateMonitor;
-		SynchronizationContext _mainThreadSynchronizationContext;
-		List<IssueMonitor> _defaultIssueMonitors = new List<IssueMonitor>();
-		List<WorkspaceIssueMonitor> _workspaceIssueMonitors = new List<WorkspaceIssueMonitor>();
-
-		string? _apiUrl;
-		DirectoryReference _dataFolder;
-		DirectoryReference _cacheFolder;
-		IPerforceSettings _defaultPerforceSettings;
-		IServiceProvider _serviceProvider;
-		ILogger _logger;
-		UserSettings _settings;
+		public static extern int SendMessage(IntPtr hWnd, int wMsg, int wParam, int lParam);
+
+		private const int WmSetredraw = 11;
+
+		readonly UpdateMonitor _updateMonitor;
+		readonly SynchronizationContext _mainThreadSynchronizationContext;
+		readonly List<IssueMonitor> _defaultIssueMonitors = new List<IssueMonitor>();
+		readonly List<WorkspaceIssueMonitor> _workspaceIssueMonitors = new List<WorkspaceIssueMonitor>();
+
+		readonly string? _apiUrl;
+		readonly DirectoryReference _dataFolder;
+		readonly IPerforceSettings _defaultPerforceSettings;
+		readonly IServiceProvider _serviceProvider;
+		readonly ILogger _logger;
+		readonly UserSettings _settings;
 		int _tabMenuTabIdx = -1;
 		int _changingWorkspacesRefCount;
 
 		bool _allowClose = false;
 
-		bool _restoreStateOnLoad;
-
-		OidcTokenManager? _oidcTokenManager;
+		readonly OidcTokenManager? _oidcTokenManager;
 
 		System.Threading.Timer? _scheduleTimer;
 		System.Threading.Timer? _scheduleSettledTimer;
 
-		string _originalExecutableFileName;
-		bool _preview;
+		readonly string _originalExecutableFileName;
+		readonly bool _preview;
 
 		IMainWindowTabPanel? _currentTabPanel;
 
 		AutomationServer _automationServer;
 
-		bool _allowCreatingHandle;
+		readonly bool _allowCreatingHandle;
 
 		Rectangle _primaryWorkArea;
-		List<IssueAlertWindow> _alertWindows = new List<IssueAlertWindow>();
+		readonly List<IssueAlertWindow> _alertWindows = new List<IssueAlertWindow>();
 
 		public ToolUpdateMonitor ToolUpdateMonitor { get; private set; }
 
-		public MainWindow(UpdateMonitor inUpdateMonitor, string? inApiUrl, DirectoryReference inDataFolder, DirectoryReference inCacheFolder, bool inRestoreStateOnLoad, string inOriginalExecutableFileName, bool inPreview, List<(UserSelectedProjectSettings, ModalTask<OpenProjectInfo>)> startupTasks, IPerforceSettings inDefaultPerforceSettings, IServiceProvider inServiceProvider, UserSettings inSettings, string? inUri, OidcTokenManager? inOidcTokenManager)
+		public MainWindow(UpdateMonitor inUpdateMonitor, string? inApiUrl, DirectoryReference inDataFolder, string inOriginalExecutableFileName, bool inPreview, List<(UserSelectedProjectSettings, ModalTask<OpenProjectInfo>)> startupTasks, IPerforceSettings inDefaultPerforceSettings, IServiceProvider inServiceProvider, UserSettings inSettings, string? inUri, OidcTokenManager? inOidcTokenManager)
 		{
 			_serviceProvider = inServiceProvider;
 			_logger = _serviceProvider.GetRequiredService<ILogger<MainWindow>>();
@@ -138,8 +135,6 @@ namespace UnrealGameSync
 			_mainThreadSynchronizationContext = SynchronizationContext.Current!;
 			_apiUrl = inApiUrl;
 			_dataFolder = inDataFolder;
-			_cacheFolder = inCacheFolder;
-			_restoreStateOnLoad = inRestoreStateOnLoad;
 			_originalExecutableFileName = inOriginalExecutableFileName;
 			_preview = inPreview;
 			_defaultPerforceSettings = inDefaultPerforceSettings;
@@ -214,7 +209,7 @@ namespace UnrealGameSync
 
 			foreach(WorkspaceIssueMonitor workspaceIssueMonitor in _workspaceIssueMonitors)
 			{
-				workspaceIssueMonitor.IssueMonitor.Start();
+				workspaceIssueMonitor._issueMonitor.Start();
 			}
 
 			ToolUpdateMonitor.Start();
@@ -364,15 +359,15 @@ namespace UnrealGameSync
 				}
 
 				// Switch the stream
-				Func<IPerforceConnection, CancellationToken, Task> switchTask = async (connection, cancellationToken) =>
+				async Task SwitchTask(IPerforceConnection perforce, CancellationToken cancellationToken)
 				{
-					await connection.SwitchClientToStreamAsync(streamName, SwitchClientOptions.None, cancellationToken);
-				};
+					await perforce.SwitchClientToStreamAsync(streamName, SwitchClientOptions.None, cancellationToken);
+				}
 
 				PerforceSettings settings = new PerforceSettings(workspaceInfo.ServerAndPort, workspaceInfo.UserName);
 				settings.ClientName = workspaceInfo.WorkspaceName;
 
-				ModalTask? result = PerforceModalTask.Execute(Owner, "Please wait", "Switching streams, please wait...", settings, switchTask, _logger, ModalTaskFlags.Quiet);
+				ModalTask? result = PerforceModalTask.Execute(Owner, "Please wait", "Switching streams, please wait...", settings, SwitchTask, _logger, ModalTaskFlags.Quiet);
 				if (result == null || !result.Succeeded)
 				{
 					_logger.LogError("Unable to switch stream ({Message})", result?.Error ?? "Operation cancelled");
@@ -478,8 +473,9 @@ namespace UnrealGameSync
 						issueMonitor.AddRef();
 						try
 						{
-							Func<CancellationToken, Task<IssueData>> func = x => RestApi.GetAsync<IssueData>($"{issueMonitor.ApiUrl}/api/issues/{issueId}", x);
-							ModalTask<IssueData>? issueTask = ModalTask.Execute(this, "Finding Issue", "Querying issue data, please wait...", func);
+							Task<IssueData> Func(CancellationToken cancellationToken) => RestApi.GetAsync<IssueData>($"{issueMonitor.ApiUrl}/api/issues/{issueId}", cancellationToken);
+							
+							ModalTask<IssueData>? issueTask = ModalTask.Execute(this, "Finding Issue", "Querying issue data, please wait...", Func);
 							if (issueTask == null)
 							{
 								_logger.LogInformation("Operation cancelled");
@@ -740,11 +736,20 @@ namespace UnrealGameSync
 			lines.Add(new StatusLine(){ LineHeight = 0.5f });
 
 			StatusLine actionLine = new StatusLine();
-			actionLine.AddLink("Retry", FontStyle.Bold | FontStyle.Underline, () => { BeginInvoke(new MethodInvoker(() => { TryOpenProject(project, TabControl.FindTabIndex(errorPanel)); })); });
+			actionLine.AddLink("Retry", FontStyle.Bold | FontStyle.Underline, () => 
+			{ 
+				BeginInvoke(new MethodInvoker(() => TryOpenProject(project, TabControl.FindTabIndex(errorPanel)))); 
+			});
 			actionLine.AddText(" | ");
-			actionLine.AddLink("Settings", FontStyle.Bold | FontStyle.Underline, () => { BeginInvoke(new MethodInvoker(() => { EditSelectedProject(errorPanel); })); });
+			actionLine.AddLink("Settings", FontStyle.Bold | FontStyle.Underline, () => 
+			{ 
+				BeginInvoke(new MethodInvoker(() => EditSelectedProject(errorPanel))); 
+			});
 			actionLine.AddText(" | ");
-			actionLine.AddLink("Close", FontStyle.Bold | FontStyle.Underline, () => { BeginInvoke(new MethodInvoker(() => { TabControl.RemoveTab(TabControl.FindTabIndex(errorPanel)); })); });
+			actionLine.AddLink("Close", FontStyle.Bold | FontStyle.Underline, () => 
+			{ 
+				BeginInvoke(new MethodInvoker(() => TabControl.RemoveTab(TabControl.FindTabIndex(errorPanel)))); 
+			});
 			lines.Add(actionLine);
 
 			errorPanel.Set(lines, null, null, null);
@@ -827,18 +832,12 @@ namespace UnrealGameSync
 
 		public void SyncLatestChange()
 		{
-			if(_currentTabPanel != null)
-			{
-				_currentTabPanel.SyncLatestChange();
-			}
+			_currentTabPanel?.SyncLatestChange();
 		}
 
 		public void LaunchEditor()
 		{
-			if(_currentTabPanel != null)
-			{
-				_currentTabPanel.LaunchEditor();
-			}
+			_currentTabPanel?.LaunchEditor();
 		}
 
 		public void ForceClose()
@@ -849,18 +848,12 @@ namespace UnrealGameSync
 
 		private void MainWindow_Activated(object sender, EventArgs e)
 		{
-			if(_currentTabPanel != null)
-			{
-				_currentTabPanel.Activate();
-			}
+			_currentTabPanel?.Activate();
 		}
 
 		private void MainWindow_Deactivate(object sender, EventArgs e)
 		{
-			if(_currentTabPanel != null)
-			{
-				_currentTabPanel.Deactivate();
-			}
+			_currentTabPanel?.Deactivate();
 		}
 
 		public void SetupScheduledSync()
@@ -905,8 +898,8 @@ namespace UnrealGameSync
 
 				// add or subtract from the schedule time to distribute scheduled syncs over a little bit more time
 				// this avoids everyone hitting the p4 server at exactly the same time.
-				const int fudgeMinutes = 10;
-				TimeSpan fudgeTime = TimeSpan.FromMinutes(rnd.Next(fudgeMinutes * -100, fudgeMinutes * 100) / 100.0);
+				const int FudgeMinutes = 10;
+				TimeSpan fudgeTime = TimeSpan.FromMinutes(rnd.Next(FudgeMinutes * -100, FudgeMinutes * 100) / 100.0);
 				DateTime nextScheduleTime = new DateTime(currentTime.Year, currentTime.Month, currentTime.Day, _settings.ScheduleTime.Hours, _settings.ScheduleTime.Minutes, _settings.ScheduleTime.Seconds);
 				nextScheduleTime += fudgeTime;
 
@@ -916,7 +909,13 @@ namespace UnrealGameSync
 				}
 
 				TimeSpan intervalToFirstTick = nextScheduleTime - currentTime;
-				_scheduleTimer = new System.Threading.Timer(x => _mainThreadSynchronizationContext.Post((o) => { if(!IsDisposed){ ScheduleTimerElapsed(); } }, null), null, intervalToFirstTick, TimeSpan.FromDays(1));
+				_scheduleTimer = new System.Threading.Timer(x => _mainThreadSynchronizationContext.Post((o) => 
+				{ 
+					if(!IsDisposed)
+					{ 
+						ScheduleTimerElapsed(); 
+					} 
+				}, null), null, intervalToFirstTick, TimeSpan.FromDays(1));
 
 				_logger.LogInformation("Schedule: Started ScheduleTimer for {Time} ({Time} remaining)", nextScheduleTime, intervalToFirstTick);
 			}
@@ -1035,7 +1034,13 @@ namespace UnrealGameSync
 		private void StartScheduleSettledTimer()
 		{
 			StopScheduleSettledTimer();
-			_scheduleSettledTimer = new System.Threading.Timer(x => _mainThreadSynchronizationContext.Post((o) => { if(!IsDisposed){ ScheduleSettledTimerElapsed(); } }, null), null, TimeSpan.FromSeconds(20.0), TimeSpan.FromMilliseconds(-1.0));
+			_scheduleSettledTimer = new System.Threading.Timer(x => _mainThreadSynchronizationContext.Post((o) => 
+			{ 
+				if(!IsDisposed)
+				{ 
+					ScheduleSettledTimerElapsed(); 
+				} 
+			}, null), null, TimeSpan.FromSeconds(20.0), TimeSpan.FromMilliseconds(-1.0));
 			_logger.LogInformation("Schedule: Started ScheduleSettledTimer");
 		}
 
@@ -1132,7 +1137,7 @@ namespace UnrealGameSync
 
 		public void OpenNewProject()
 		{
-			OpenProjectInfo? openProjectInfo = OpenProjectWindow.ShowModal(this, null, _settings, _dataFolder, _cacheFolder, _defaultPerforceSettings, _serviceProvider, _logger);
+			OpenProjectInfo? openProjectInfo = OpenProjectWindow.ShowModal(this, null, _settings, _defaultPerforceSettings, _serviceProvider, _logger);
 			if(openProjectInfo != null)
 			{
 				int newTabIdx = TryOpenProject(openProjectInfo, -1, OpenProjectOptions.None);
@@ -1150,10 +1155,10 @@ namespace UnrealGameSync
 			_settings.RecentProjects.RemoveAll(x => x.LocalPath != null && x.LocalPath.Equals(detectedProjectSettings.LocalPath, StringComparison.OrdinalIgnoreCase));
 			_settings.RecentProjects.Insert(0, detectedProjectSettings);
 
-			const int maxRecentProjects = 10;
-			if (_settings.RecentProjects.Count > maxRecentProjects)
+			const int MaxRecentProjects = 10;
+			if (_settings.RecentProjects.Count > MaxRecentProjects)
 			{
-				_settings.RecentProjects.RemoveRange(maxRecentProjects, _settings.RecentProjects.Count - maxRecentProjects);
+				_settings.RecentProjects.RemoveRange(MaxRecentProjects, _settings.RecentProjects.Count - MaxRecentProjects);
 			}
 
 			_settings.Save(_logger);
@@ -1162,14 +1167,12 @@ namespace UnrealGameSync
 		public void EditSelectedProject(int tabIdx)
 		{
 			object tabData = TabControl.GetTabData(tabIdx);
-			if(tabData is WorkspaceControl)
+			if(tabData is WorkspaceControl workspace)
 			{
-				WorkspaceControl workspace = (WorkspaceControl)tabData;
 				EditSelectedProject(tabIdx, workspace.SelectedProject);
 			}
-			else if(tabData is ErrorPanel)
+			else if(tabData is ErrorPanel error)
 			{
-				ErrorPanel error = (ErrorPanel)tabData;
 				EditSelectedProject(tabIdx, error.SelectedProject);
 			}
 		}
@@ -1194,7 +1197,7 @@ namespace UnrealGameSync
 
 		public void EditSelectedProject(int tabIdx, UserSelectedProjectSettings selectedProject)
 		{
-			OpenProjectInfo? openProjectInfo = OpenProjectWindow.ShowModal(this, selectedProject, _settings, _dataFolder, _cacheFolder, _defaultPerforceSettings, _serviceProvider, _logger);
+			OpenProjectInfo? openProjectInfo = OpenProjectWindow.ShowModal(this, selectedProject, _settings, _defaultPerforceSettings, _serviceProvider, _logger);
 			if(openProjectInfo != null)
 			{
 				int newTabIdx = TryOpenProject(openProjectInfo, tabIdx, OpenProjectOptions.None);
@@ -1222,7 +1225,10 @@ namespace UnrealGameSync
 			ModalTask<OpenProjectInfo>? settingsTask = PerforceModalTask.Execute(this, "Opening Project", "Opening project, please wait...", perforceSettings, (p, c) => OpenProjectWindow.DetectSettingsAsync(p, project, _settings, projectLogger, c), projectLogger, taskFlags);
 			if (settingsTask == null || settingsTask.Failed)
 			{
-				if(settingsTask != null) CreateErrorPanel(replaceTabIdx, project, settingsTask.Error);
+				if (settingsTask != null)
+				{
+					CreateErrorPanel(replaceTabIdx, project, settingsTask.Error);
+				}
 				return -1;
 			}
 
@@ -1308,7 +1314,13 @@ namespace UnrealGameSync
 
 		public void StreamChanged(WorkspaceControl workspace)
 		{
-			_mainThreadSynchronizationContext.Post((o) => { if(!IsDisposed) { StreamChangedCallback(workspace); } }, null);
+			_mainThreadSynchronizationContext.Post((o) => 
+			{ 
+				if(!IsDisposed) 
+				{ 
+					StreamChangedCallback(workspace); 
+				} 
+			}, null);
 		}
 
 		public void StreamChangedCallback(WorkspaceControl workspace)
@@ -1375,7 +1387,7 @@ namespace UnrealGameSync
 			}
 
 			// if this failes, return something sensible to avoid blank tabs
-			if (string.IsNullOrEmpty(tabName))
+			if (String.IsNullOrEmpty(tabName))
 			{
 				_logger.LogInformation("No TabName for {ClientName} for setting {TabSetting}. Defaulting to client name", workspace.ClientName, _settings.TabLabels.ToString());
 				tabName = workspace.ClientName;
@@ -1564,7 +1576,7 @@ namespace UnrealGameSync
 			}
 
 			HashSet<IssueData> allIssues = new HashSet<IssueData>();
-			foreach(IssueMonitor issueMonitor in _workspaceIssueMonitors.Select(x => x.IssueMonitor))
+			foreach(IssueMonitor issueMonitor in _workspaceIssueMonitors.Select(x => x._issueMonitor))
 			{
 				List<IssueData> issues = issueMonitor.GetIssues();
 				foreach(IssueData issue in issues)
@@ -1712,8 +1724,6 @@ namespace UnrealGameSync
 
 		void CloseAlertWindow(IssueAlertWindow alert)
 		{
-			IssueData issue = alert.Issue;
-
 			alert.Close();
 			alert.Dispose();
 
@@ -1758,7 +1768,7 @@ namespace UnrealGameSync
 
 		public IssueMonitor CreateIssueMonitor(string? apiUrl, string userName)
 		{
-			WorkspaceIssueMonitor? workspaceIssueMonitor = _workspaceIssueMonitors.FirstOrDefault(x => String.Compare(x.IssueMonitor.ApiUrl, apiUrl, StringComparison.OrdinalIgnoreCase) == 0 && String.Compare(x.IssueMonitor.UserName, userName, StringComparison.OrdinalIgnoreCase) == 0);
+			WorkspaceIssueMonitor? workspaceIssueMonitor = _workspaceIssueMonitors.FirstOrDefault(x => String.Compare(x._issueMonitor.ApiUrl, apiUrl, StringComparison.OrdinalIgnoreCase) == 0 && String.Compare(x._issueMonitor.UserName, userName, StringComparison.OrdinalIgnoreCase) == 0);
 			if (workspaceIssueMonitor == null)
 			{
 				string serverId = apiUrl != null ? Regex.Replace(apiUrl, @"^.*://", "") : "noserver";
@@ -1773,25 +1783,25 @@ namespace UnrealGameSync
 				_workspaceIssueMonitors.Add(workspaceIssueMonitor);
 			}
 
-			workspaceIssueMonitor.RefCount++;
+			workspaceIssueMonitor._refCount++;
 
-			if (workspaceIssueMonitor.RefCount == 1 && _allowCreatingHandle)
+			if (workspaceIssueMonitor._refCount == 1 && _allowCreatingHandle)
 			{
-				workspaceIssueMonitor.IssueMonitor.Start();
+				workspaceIssueMonitor._issueMonitor.Start();
 			}
 
-			return workspaceIssueMonitor.IssueMonitor;
+			return workspaceIssueMonitor._issueMonitor;
 		}
 
 		public void ReleaseIssueMonitor(IssueMonitor issueMonitor)
 		{
-			int index = _workspaceIssueMonitors.FindIndex(x => x.IssueMonitor == issueMonitor);
+			int index = _workspaceIssueMonitors.FindIndex(x => x._issueMonitor == issueMonitor);
 			if(index != -1)
 			{
 				WorkspaceIssueMonitor workspaceIssueMonitor = _workspaceIssueMonitors[index];
-				workspaceIssueMonitor.RefCount--;
+				workspaceIssueMonitor._refCount--;
 
-				if (workspaceIssueMonitor.RefCount == 0)
+				if (workspaceIssueMonitor._refCount == 0)
 				{
 					for (int idx = _alertWindows.Count - 1; idx >= 0; idx--)
 					{
diff --git a/Engine/Source/Programs/UnrealGameSync/UnrealGameSync/Forms/ModalTaskWindow.cs b/Engine/Source/Programs/UnrealGameSync/UnrealGameSync/Forms/ModalTaskWindow.cs
--- a/Engine/Source/Programs/UnrealGameSync/UnrealGameSync/Forms/ModalTaskWindow.cs
+++ b/Engine/Source/Programs/UnrealGameSync/UnrealGameSync/Forms/ModalTaskWindow.cs
@@ -1,13 +1,6 @@
 // Copyright Epic Games, Inc. All Rights Reserved.
 
 using System;
-using System.Collections.Generic;
-using System.ComponentModel;
-using System.Data;
-using System.Drawing;
-using System.IO;
-using System.Linq;
-using System.Text;
 using System.Threading;
 using System.Threading.Tasks;
 using System.Windows.Forms;
@@ -16,8 +9,8 @@ namespace UnrealGameSync
 {
 	partial class ModalTaskWindow : Form
 	{
-		Task _task;
-		CancellationTokenSource _cancellationSource;
+		readonly Task _task;
+		readonly CancellationTokenSource _cancellationSource;
 
 		public ModalTaskWindow(string inTitle, string inMessage, FormStartPosition inStartPosition, Task inTask, CancellationTokenSource inCancellationSource)//, Func<CancellationToken, Task> InTaskFunc)
 		{
diff --git a/Engine/Source/Programs/UnrealGameSync/UnrealGameSync/Forms/ModifyBuildStepsWindow.cs b/Engine/Source/Programs/UnrealGameSync/UnrealGameSync/Forms/ModifyBuildStepsWindow.cs
--- a/Engine/Source/Programs/UnrealGameSync/UnrealGameSync/Forms/ModifyBuildStepsWindow.cs
+++ b/Engine/Source/Programs/UnrealGameSync/UnrealGameSync/Forms/ModifyBuildStepsWindow.cs
@@ -3,12 +3,7 @@
 using EpicGames.Core;
 using System;
 using System.Collections.Generic;
-using System.ComponentModel;
-using System.Data;
 using System.Drawing;
-using System.Linq;
-using System.Text;
-using System.Threading.Tasks;
 using System.Windows.Forms;
 using System.Windows.Forms.VisualStyles;
 
@@ -16,12 +11,12 @@ namespace UnrealGameSync
 {
 	partial class ModifyBuildStepsWindow : Form
 	{
-		List<string> _targetNames;
-		List<BuildStep> _steps;
-		HashSet<Guid> _projectSteps;
-		DirectoryReference _baseDirectory;
+		readonly List<string> _targetNames;
+		readonly List<BuildStep> _steps;
+		readonly HashSet<Guid> _projectSteps;
+		readonly DirectoryReference _baseDirectory;
 		ListViewItem.ListViewSubItem? _mouseDownSubItem = null;
-		IReadOnlyDictionary<string, string> _variables;
+		readonly IReadOnlyDictionary<string, string> _variables;
 
 		public ModifyBuildStepsWindow(List<string> inTargetNames, List<BuildStep> inSteps, HashSet<Guid> inProjectSteps, DirectoryReference inBaseDirectory, IReadOnlyDictionary<string, string> inVariables)
 		{
diff --git a/Engine/Source/Programs/UnrealGameSync/UnrealGameSync/Forms/NewWorkspaceWindow.cs b/Engine/Source/Programs/UnrealGameSync/UnrealGameSync/Forms/NewWorkspaceWindow.cs
--- a/Engine/Source/Programs/UnrealGameSync/UnrealGameSync/Forms/NewWorkspaceWindow.cs
+++ b/Engine/Source/Programs/UnrealGameSync/UnrealGameSync/Forms/NewWorkspaceWindow.cs
@@ -6,10 +6,8 @@ using Microsoft.Extensions.DependencyInjection;
 using Microsoft.Extensions.Logging;
 using System;
 using System.Collections.Generic;
-using System.ComponentModel;
 using System.Data;
 using System.Diagnostics.CodeAnalysis;
-using System.Drawing;
 using System.IO;
 using System.Linq;
 using System.Text;
@@ -29,23 +27,23 @@ namespace UnrealGameSync
 
 			public NewWorkspaceSettings(string name, string stream, DirectoryReference rootDir)
 			{
-				this.Name = name;
-				this.Stream = stream;
-				this.RootDir = rootDir;
+				Name = name;
+				Stream = stream;
+				RootDir = rootDir;
 			}
 		}
 
 		class FindWorkspaceSettingsTask
 		{
-			public InfoRecord Info;
-			public List<ClientsRecord> Clients;
-			public string? CurrentStream;
+			public InfoRecord Info { get; }
+			public List<ClientsRecord> Clients { get; }
+			public string? CurrentStream { get; }
 
 			public FindWorkspaceSettingsTask(InfoRecord info, List<ClientsRecord> clients, string? currentStream)
 			{
-				this.Info = info;
-				this.Clients = clients;
-				this.CurrentStream = currentStream;
+				Info = info;
+				Clients = clients;
+				CurrentStream = currentStream;
 			}
 
 			public static async Task<FindWorkspaceSettingsTask> RunAsync(IPerforceConnection perforce, string currentWorkspaceName, CancellationToken cancellationToken)
@@ -85,21 +83,21 @@ namespace UnrealGameSync
 			}
 		}
 
-		IPerforceSettings _perforceSettings;
-		InfoRecord _info;
-		List<ClientsRecord> _clients;
-		IServiceProvider _serviceProvider;
+		readonly IPerforceSettings _perforceSettings;
+		readonly InfoRecord _info;
+		readonly List<ClientsRecord> _clients;
+		readonly IServiceProvider _serviceProvider;
 		NewWorkspaceSettings? _settings;
-		DirectoryReference? _defaultRootPath;
+		readonly DirectoryReference? _defaultRootPath;
 
 		private NewWorkspaceWindow(IPerforceSettings perforceSettings, string? forceStream, string? defaultStream, InfoRecord info, List<ClientsRecord> clients, IServiceProvider serviceProvider)
 		{
 			InitializeComponent();
 
-			this._perforceSettings = perforceSettings;
-			this._info = info;
-			this._clients = clients;
-			this._serviceProvider = serviceProvider;
+			_perforceSettings = perforceSettings;
+			_info = info;
+			_clients = clients;
+			_serviceProvider = serviceProvider;
 
 			Dictionary<DirectoryReference, int> rootPathToCount = new Dictionary<DirectoryReference, int>();
 			foreach(ClientsRecord client in clients)
@@ -244,8 +242,7 @@ namespace UnrealGameSync
 
 		private void UpdateOkButton()
 		{
-			NewWorkspaceSettings? settings;
-			OkBtn.Enabled = TryGetWorkspaceSettings(out settings);
+			OkBtn.Enabled = TryGetWorkspaceSettings(out _);
 		}
 
 		private bool TryGetWorkspaceSettings([NotNullWhen(true)] out NewWorkspaceSettings? settings)
diff --git a/Engine/Source/Programs/UnrealGameSync/UnrealGameSync/Forms/NotificationWindow.cs b/Engine/Source/Programs/UnrealGameSync/UnrealGameSync/Forms/NotificationWindow.cs
--- a/Engine/Source/Programs/UnrealGameSync/UnrealGameSync/Forms/NotificationWindow.cs
+++ b/Engine/Source/Programs/UnrealGameSync/UnrealGameSync/Forms/NotificationWindow.cs
@@ -1,16 +1,9 @@
 // Copyright Epic Games, Inc. All Rights Reserved.
 
 using System;
-using System.Collections.Generic;
-using System.ComponentModel;
-using System.Data;
 using System.Drawing;
 using System.Drawing.Drawing2D;
-using System.Linq;
-using System.Text;
-using System.Threading.Tasks;
 using System.Windows.Forms;
-using System.Windows.Forms.VisualStyles;
 
 namespace UnrealGameSync
 {
@@ -135,11 +128,8 @@ namespace UnrealGameSync
 		{
 			base.OnFontChanged(e);
 
-			if(_captionFont != null)
-			{
-				_captionFont.Dispose();
-			}
-			_captionFont = new Font(this.Font.Name, this.Font.Size * (12.0f / 9.0f), FontStyle.Regular);
+			_captionFont?.Dispose();
+			_captionFont = new Font(Font.Name, Font.Size * (12.0f / 9.0f), FontStyle.Regular);
 
 			CalculateBounds();
 			Invalidate();
diff --git a/Engine/Source/Programs/UnrealGameSync/UnrealGameSync/Forms/OpenProjectWindow.cs b/Engine/Source/Programs/UnrealGameSync/UnrealGameSync/Forms/OpenProjectWindow.cs
--- a/Engine/Source/Programs/UnrealGameSync/UnrealGameSync/Forms/OpenProjectWindow.cs
+++ b/Engine/Source/Programs/UnrealGameSync/UnrealGameSync/Forms/OpenProjectWindow.cs
@@ -1,17 +1,12 @@
 // Copyright Epic Games, Inc. All Rights Reserved.
 
-using EpicGames.Core;
 using EpicGames.Perforce;
 using Microsoft.Extensions.DependencyInjection;
 using Microsoft.Extensions.Logging;
 using System;
-using System.Collections.Generic;
-using System.ComponentModel;
-using System.Data;
 using System.Diagnostics.CodeAnalysis;
 using System.Drawing;
 using System.IO;
-using System.Linq;
 using System.Text;
 using System.Threading;
 using System.Threading.Tasks;
@@ -26,20 +21,20 @@ namespace UnrealGameSync
 		string? _serverAndPortOverride;
 		string? _userNameOverride;
 		OpenProjectInfo? _openProjectInfo;
-		IPerforceSettings _defaultSettings;
-		IServiceProvider _serviceProvider;
-		UserSettings _settings;
-		ILogger _logger;
+		readonly IPerforceSettings _defaultSettings;
+		readonly IServiceProvider _serviceProvider;
+		readonly UserSettings _settings;
+		readonly ILogger _logger;
 
 		private OpenProjectWindow(UserSelectedProjectSettings? project, UserSettings settings, IPerforceSettings defaultSettings, IServiceProvider serviceProvider, ILogger logger)
 		{
 			InitializeComponent();
 
-			this._settings = settings;
-			this._openProjectInfo = null;
-			this._defaultSettings = defaultSettings;
-			this._serviceProvider = serviceProvider;
-			this._logger = logger;
+			_settings = settings;
+			_openProjectInfo = null;
+			_defaultSettings = defaultSettings;
+			_serviceProvider = serviceProvider;
+			_logger = logger;
 
 			if (project == null)
 			{
@@ -87,12 +82,9 @@ namespace UnrealGameSync
 			UpdateOkButton();
 		}
 
-		private IPerforceSettings Perforce
-		{
-			get => Utility.OverridePerforceSettings(_defaultSettings, _serverAndPortOverride, _userNameOverride);
-		}
+		private IPerforceSettings Perforce => Utility.OverridePerforceSettings(_defaultSettings, _serverAndPortOverride, _userNameOverride);
 
-		public static OpenProjectInfo? ShowModal(IWin32Window owner, UserSelectedProjectSettings? project, UserSettings settings, DirectoryReference dataFolder, DirectoryReference cacheFolder, IPerforceSettings defaultPerforceSettings, IServiceProvider serviceProvider, ILogger logger)
+		public static OpenProjectInfo? ShowModal(IWin32Window owner, UserSelectedProjectSettings? project, UserSettings settings, IPerforceSettings defaultPerforceSettings, IServiceProvider serviceProvider, ILogger logger)
 		{
 			OpenProjectWindow window = new OpenProjectWindow(project, settings, defaultPerforceSettings, serviceProvider, logger);
 			if(window.ShowDialog(owner) == DialogResult.OK)
@@ -161,14 +153,12 @@ namespace UnrealGameSync
 
 		private void UpdateWorkspacePathBrowseButton()
 		{
-			string? workspaceName;
-			WorkspacePathBrowseBtn.Enabled = TryGetWorkspaceName(out workspaceName);
+			WorkspacePathBrowseBtn.Enabled = TryGetWorkspaceName(out _);
 		}
 
 		private void UpdateOkButton()
 		{
-			string? projectPath;
-			OkBtn.Enabled = WorkspaceRadioBtn.Checked? TryGetClientPath(out projectPath) : TryGetLocalPath(out projectPath);
+			OkBtn.Enabled = WorkspaceRadioBtn.Checked? TryGetClientPath(out _) : TryGetLocalPath(out _);
 		}
 
 		private void WorkspaceNewBtn_Click(object sender, EventArgs e)
diff --git a/Engine/Source/Programs/UnrealGameSync/UnrealGameSync/Forms/PasswordWindow.cs b/Engine/Source/Programs/UnrealGameSync/UnrealGameSync/Forms/PasswordWindow.cs
--- a/Engine/Source/Programs/UnrealGameSync/UnrealGameSync/Forms/PasswordWindow.cs
+++ b/Engine/Source/Programs/UnrealGameSync/UnrealGameSync/Forms/PasswordWindow.cs
@@ -1,26 +1,19 @@
 ﻿// Copyright Epic Games, Inc. All Rights Reserved.
 
 using System;
-using System.Collections.Generic;
-using System.ComponentModel;
-using System.Data;
-using System.Drawing;
-using System.Linq;
-using System.Text;
-using System.Threading.Tasks;
 using System.Windows.Forms;
 
 namespace UnrealGameSync
 {
 	public partial class PasswordWindow : Form
 	{
-		public string Password;
+		public string Password { get; private set; }
 
 		public PasswordWindow(string prompt, string password)
 		{
 			InitializeComponent();
 
-			this.Password = password;
+			Password = password;
 
 			PromptLabel.Text = prompt;
 			PasswordTextBox.Text = password;
diff --git a/Engine/Source/Programs/UnrealGameSync/UnrealGameSync/Forms/PerforceSyncSettingsWindow.cs b/Engine/Source/Programs/UnrealGameSync/UnrealGameSync/Forms/PerforceSyncSettingsWindow.cs
--- a/Engine/Source/Programs/UnrealGameSync/UnrealGameSync/Forms/PerforceSyncSettingsWindow.cs
+++ b/Engine/Source/Programs/UnrealGameSync/UnrealGameSync/Forms/PerforceSyncSettingsWindow.cs
@@ -8,13 +8,13 @@ namespace UnrealGameSync
 {
 	partial class PerforceSyncSettingsWindow : Form
 	{
-		UserSettings _settings;
-		ILogger _logger;
+		readonly UserSettings _settings;
+		readonly ILogger _logger;
 
 		public PerforceSyncSettingsWindow(UserSettings settings, ILogger logger)
 		{
-			this._settings = settings;
-			this._logger = logger;
+			_settings = settings;
+			_logger = logger;
 
 			InitializeComponent();
 		}
@@ -51,7 +51,6 @@ namespace UnrealGameSync
 
 		private void ResetButton_Click(object sender, EventArgs e)
 		{
-			PerforceSyncOptions syncOptions = _settings.SyncOptions;
 			numericUpDownMaxCommandsPerBatch.Value = PerforceSyncOptions.DefaultMaxCommandsPerBatch;
 			numericUpDownMaxSizePerBatch.Value = PerforceSyncOptions.DefaultMaxSizePerBatch / 1024 / 1024;
 			numericUpDownRetriesOnSyncError.Value = PerforceSyncOptions.DefaultNumSyncErrorRetries;
diff --git a/Engine/Source/Programs/UnrealGameSync/UnrealGameSync/Forms/ProgramsRunningWindow.cs b/Engine/Source/Programs/UnrealGameSync/UnrealGameSync/Forms/ProgramsRunningWindow.cs
--- a/Engine/Source/Programs/UnrealGameSync/UnrealGameSync/Forms/ProgramsRunningWindow.cs
+++ b/Engine/Source/Programs/UnrealGameSync/UnrealGameSync/Forms/ProgramsRunningWindow.cs
@@ -2,23 +2,18 @@
 
 using EpicGames.Core;
 using System;
-using System.Collections.Generic;
-using System.ComponentModel;
 using System.Data;
-using System.Drawing;
 using System.Linq;
-using System.Text;
 using System.Threading;
-using System.Threading.Tasks;
 using System.Windows.Forms;
 
 namespace UnrealGameSync
 {
 	public partial class ProgramsRunningWindow : Form
 	{
-		object _syncObject = new object();
+		readonly object _syncObject = new object();
 		FileReference[] _programs;
-		Func<FileReference[]> _enumeratePrograms;
+		readonly Func<FileReference[]> _enumeratePrograms;
 		ManualResetEvent? _terminateEvent;
 		Thread? _backgroundThread;
 
@@ -26,9 +21,9 @@ namespace UnrealGameSync
 		{
 			InitializeComponent();
 
-			this._programs = programs.OrderBy(x => x).ToArray();
-			this._enumeratePrograms = enumeratePrograms;
-			this.ProgramListBox.Items.AddRange(programs);
+			_programs = programs.OrderBy(x => x).ToArray();
+			_enumeratePrograms = enumeratePrograms;
+			ProgramListBox.Items.AddRange(programs);
 		}
 
 		private void ProgramsRunningWindow_Load(object sender, EventArgs e)
diff --git a/Engine/Source/Programs/UnrealGameSync/UnrealGameSync/Forms/ScheduleWindow.cs b/Engine/Source/Programs/UnrealGameSync/UnrealGameSync/Forms/ScheduleWindow.cs
--- a/Engine/Source/Programs/UnrealGameSync/UnrealGameSync/Forms/ScheduleWindow.cs
+++ b/Engine/Source/Programs/UnrealGameSync/UnrealGameSync/Forms/ScheduleWindow.cs
@@ -2,22 +2,16 @@
 
 using System;
 using System.Collections.Generic;
-using System.ComponentModel;
 using System.Data;
-using System.Drawing;
 using System.Globalization;
-using System.IO;
 using System.Linq;
-using System.Text;
-using System.Threading.Tasks;
 using System.Windows.Forms;
 
 namespace UnrealGameSync
 {
 	partial class ScheduleWindow : Form
 	{
-
-		Dictionary<UserSelectedProjectSettings, List<LatestChangeType>> _projectToLatestChangeTypes;
+		readonly Dictionary<UserSelectedProjectSettings, List<LatestChangeType>> _projectToLatestChangeTypes;
 
 		public ScheduleWindow(
 			bool inEnabled,
@@ -69,7 +63,7 @@ namespace UnrealGameSync
 					menuItem.Name = changeType.Name;
 					menuItem.Text = changeType.Description;
 					menuItem.Size = new System.Drawing.Size(189, 22);
-					menuItem.Click += (sender, e) => SyncTypeDropDown_Click(sender, e, changeType.Name);
+					menuItem.Click += (sender, e) => SyncTypeDropDown_Click(changeType.Name);
 
 					SyncTypeDropDown.Items.Add(menuItem);
 				}
@@ -98,7 +92,7 @@ namespace UnrealGameSync
 			}
 		}
 
-		private void SyncTypeDropDown_Click(object? sender, EventArgs e, string syncTypeId)
+		private void SyncTypeDropDown_Click(string syncTypeId)
 		{
 			UserSelectedProjectSettings? projectSetting = ProjectListBox.Items[ProjectListBox.SelectedIndex] as UserSelectedProjectSettings;
 			if (projectSetting != null)
diff --git a/Engine/Source/Programs/UnrealGameSync/UnrealGameSync/Forms/SdkInfoWindow.cs b/Engine/Source/Programs/UnrealGameSync/UnrealGameSync/Forms/SdkInfoWindow.cs
--- a/Engine/Source/Programs/UnrealGameSync/UnrealGameSync/Forms/SdkInfoWindow.cs
+++ b/Engine/Source/Programs/UnrealGameSync/UnrealGameSync/Forms/SdkInfoWindow.cs
@@ -2,15 +2,12 @@
 
 using System;
 using System.Collections.Generic;
-using System.ComponentModel;
 using System.Data;
 using System.Diagnostics;
 using System.Drawing;
 using System.Drawing.Drawing2D;
 using System.IO;
 using System.Linq;
-using System.Text;
-using System.Threading.Tasks;
 using System.Windows.Forms;
 
 #nullable enable
@@ -23,13 +20,13 @@ namespace UnrealGameSync
 		{
 			public string Category { get; }
 			public string Description { get; }
-			public string? Install;
-			public string? Browse;
+			public string? Install { get; set; }
+			public string? Browse { get; set; }
 
 			public SdkItem(string category, string description)
 			{
-				this.Category = category;
-				this.Description = description;
+				Category = category;
+				Description = description;
 			}
 		}
 
@@ -37,24 +34,24 @@ namespace UnrealGameSync
 		{
 			public string UniqueId { get; }
 			public string Label { get; }
-			public Rectangle Rectangle;
-			public Action? OnClick;
+			public Rectangle Rectangle { get; set; }
+			public Action? OnClick { get; set; }
 
 			public BadgeInfo(string uniqueId, string label)
 			{
-				this.UniqueId = uniqueId;
-				this.Label = label;
+				UniqueId = uniqueId;
+				Label = label;
 			}
 		}
 
-		Font _badgeFont;
+		readonly Font _badgeFont;
 		string? _hoverBadgeUniqueId;
 
 		public SdkInfoWindow(string[] sdkInfoEntries, Dictionary<string, string> variables, Font badgeFont)
 		{
 			InitializeComponent();
 
-			this._badgeFont = badgeFont;
+			_badgeFont = badgeFont;
 
 			Dictionary<string, ConfigObject> uniqueIdToObject = new Dictionary<string, ConfigObject>(StringComparer.InvariantCultureIgnoreCase);
 			foreach(string sdkInfoEntry in sdkInfoEntries)
@@ -127,7 +124,7 @@ namespace UnrealGameSync
 
 			if(SdkListView.Items.Count > 0)
 			{
-				int itemsHeight = SdkListView.Items[SdkListView.Items.Count - 1].Bounds.Bottom + 20;
+				int itemsHeight = SdkListView.Items[^1].Bounds.Bottom + 20;
 				Height = SdkListView.Top + itemsHeight + (Height - SdkListView.Bottom);
 			}
 		}
diff --git a/Engine/Source/Programs/UnrealGameSync/UnrealGameSync/Forms/SelectProjectFromWorkspaceWindow.cs b/Engine/Source/Programs/UnrealGameSync/UnrealGameSync/Forms/SelectProjectFromWorkspaceWindow.cs
--- a/Engine/Source/Programs/UnrealGameSync/UnrealGameSync/Forms/SelectProjectFromWorkspaceWindow.cs
+++ b/Engine/Source/Programs/UnrealGameSync/UnrealGameSync/Forms/SelectProjectFromWorkspaceWindow.cs
@@ -5,14 +5,12 @@ using Microsoft.Extensions.DependencyInjection;
 using Microsoft.Extensions.Logging;
 using System;
 using System.Collections.Generic;
-using System.ComponentModel;
 using System.Data;
 using System.Diagnostics.CodeAnalysis;
 using System.Drawing;
 using System.IO;
 using System.Linq;
 using System.Runtime.InteropServices;
-using System.Text;
 using System.Threading;
 using System.Threading.Tasks;
 using System.Windows.Forms;
@@ -74,7 +72,7 @@ namespace UnrealGameSync
 		[DllImport("Shell32.dll", EntryPoint = "ExtractIconExW", CharSet = CharSet.Unicode, ExactSpelling = true, CallingConvention = CallingConvention.StdCall)]
 		private static extern int ExtractIconEx(string sFile, int iIndex, IntPtr piLargeVersion, out IntPtr piSmallVersion, int amountIcons);
 
-		List<string> _projectFiles;
+		readonly List<string> _projectFiles;
 		string _selectedProjectFile;
 
 		class ProjectNode
@@ -85,7 +83,7 @@ namespace UnrealGameSync
 
 			public ProjectNode(string fullName)
 			{
-				this.FullName = fullName;
+				FullName = fullName;
 
 				int slashIdx = fullName.LastIndexOf('/');
 				Folder = fullName.Substring(0, slashIdx);
@@ -97,8 +95,8 @@ namespace UnrealGameSync
 		{
 			InitializeComponent();
 			
-			this._projectFiles = projectFiles;
-			this._selectedProjectFile = selectedProjectFile;
+			_projectFiles = projectFiles;
+			_selectedProjectFile = selectedProjectFile;
 
 			// Make the image strip containing icons for nodes in the tree
 			IntPtr folderIconPtr;
diff --git a/Engine/Source/Programs/UnrealGameSync/UnrealGameSync/Forms/SelectStreamWindow.cs b/Engine/Source/Programs/UnrealGameSync/UnrealGameSync/Forms/SelectStreamWindow.cs
--- a/Engine/Source/Programs/UnrealGameSync/UnrealGameSync/Forms/SelectStreamWindow.cs
+++ b/Engine/Source/Programs/UnrealGameSync/UnrealGameSync/Forms/SelectStreamWindow.cs
@@ -5,13 +5,10 @@ using Microsoft.Extensions.DependencyInjection;
 using Microsoft.Extensions.Logging;
 using System;
 using System.Collections.Generic;
-using System.ComponentModel;
 using System.Data;
 using System.Diagnostics.CodeAnalysis;
 using System.Drawing;
-using System.IO;
 using System.Linq;
-using System.Text;
 using System.Threading;
 using System.Threading.Tasks;
 using System.Windows.Forms;
@@ -37,7 +34,7 @@ namespace UnrealGameSync
 
 			public StreamNode(StreamsRecord record)
 			{
-				this.Record = record;
+				Record = record;
 			}
 
 			public void Sort()
@@ -57,7 +54,7 @@ namespace UnrealGameSync
 
 			public StreamDepot(string name)
 			{
-				this.Name = name;
+				Name = name;
 			}
 
 			public void Sort()
@@ -71,13 +68,13 @@ namespace UnrealGameSync
 		}
 
 		private string? _selectedStream;
-		private List<StreamDepot> _depots = new List<StreamDepot>();
+		private readonly List<StreamDepot> _depots = new List<StreamDepot>();
 
 		private SelectStreamWindow(List<StreamsRecord> streams, string? streamName)
 		{
 			InitializeComponent();
 
-			this._selectedStream = streamName;
+			_selectedStream = streamName;
 
 			// Set up the image list
 			ImageList perforceImageList = new ImageList();
diff --git a/Engine/Source/Programs/UnrealGameSync/UnrealGameSync/Forms/SelectUserWindow.cs b/Engine/Source/Programs/UnrealGameSync/UnrealGameSync/Forms/SelectUserWindow.cs
--- a/Engine/Source/Programs/UnrealGameSync/UnrealGameSync/Forms/SelectUserWindow.cs
+++ b/Engine/Source/Programs/UnrealGameSync/UnrealGameSync/Forms/SelectUserWindow.cs
@@ -5,13 +5,7 @@ using Microsoft.Extensions.DependencyInjection;
 using Microsoft.Extensions.Logging;
 using System;
 using System.Collections.Generic;
-using System.ComponentModel;
-using System.Data;
 using System.Diagnostics.CodeAnalysis;
-using System.Drawing;
-using System.IO;
-using System.Linq;
-using System.Text;
 using System.Threading;
 using System.Threading.Tasks;
 using System.Windows.Forms;
@@ -31,14 +25,14 @@ namespace UnrealGameSync
 		}
 
 		private int _selectedUserIndex;
-		private List<UsersRecord> _users;
+		private readonly List<UsersRecord> _users;
 		
 		private SelectUserWindow(List<UsersRecord> users, int selectedUserIndex)
 		{
 			InitializeComponent();
 
-			this._selectedUserIndex = selectedUserIndex;
-			this._users = users;
+			_selectedUserIndex = selectedUserIndex;
+			_users = users;
 
 			PopulateList();
 			UpdateOkButton();
