UGS: Add .editorconfig file for static analysis warnings in UGS (+lots of fixes).

diff --git a/Engine/Source/Programs/UnrealGameSync/UnrealGameSync/Forms/SelectWorkspaceWindow.cs b/Engine/Source/Programs/UnrealGameSync/UnrealGameSync/Forms/SelectWorkspaceWindow.cs
--- a/Engine/Source/Programs/UnrealGameSync/UnrealGameSync/Forms/SelectWorkspaceWindow.cs
+++ b/Engine/Source/Programs/UnrealGameSync/UnrealGameSync/Forms/SelectWorkspaceWindow.cs
@@ -5,13 +5,8 @@ using Microsoft.Extensions.DependencyInjection;
 using Microsoft.Extensions.Logging;
 using System;
 using System.Collections.Generic;
-using System.ComponentModel;
 using System.Data;
-using System.Drawing;
-using System.IO;
 using System.Linq;
-using System.Net;
-using System.Text;
 using System.Threading;
 using System.Threading.Tasks;
 using System.Windows.Forms;
@@ -29,8 +24,8 @@ namespace UnrealGameSync
 
 			public EnumerateWorkspaces(InfoRecord info, List<ClientsRecord> clients)
 			{
-				this.Info = info;
-				this.Clients = clients;
+				Info = info;
+				Clients = clients;
 			}
 
 			public static async Task<EnumerateWorkspaces> RunAsync(IPerforceConnection perforce, CancellationToken cancellationToken)
@@ -41,17 +36,17 @@ namespace UnrealGameSync
 			}
 		}
 
-		InfoRecord _info;
-		List<ClientsRecord> _clients;
+		readonly InfoRecord _info;
+		readonly List<ClientsRecord> _clients;
 		string? _workspaceName;
 
 		private SelectWorkspaceWindow(InfoRecord info, List<ClientsRecord> clients, string? workspaceName)
 		{
 			InitializeComponent();
 
-			this._info = info;
-			this._clients = clients;
-			this._workspaceName = workspaceName;
+			_info = info;
+			_clients = clients;
+			_workspaceName = workspaceName;
 
 			UpdateListView();
 			UpdateOkButton();
diff --git a/Engine/Source/Programs/UnrealGameSync/UnrealGameSync/Forms/SyncFilter.cs b/Engine/Source/Programs/UnrealGameSync/UnrealGameSync/Forms/SyncFilter.cs
--- a/Engine/Source/Programs/UnrealGameSync/UnrealGameSync/Forms/SyncFilter.cs
+++ b/Engine/Source/Programs/UnrealGameSync/UnrealGameSync/Forms/SyncFilter.cs
@@ -2,22 +2,17 @@
 
 using System;
 using System.Collections.Generic;
-using System.ComponentModel;
-using System.Data;
-using System.Drawing;
 using System.Linq;
-using System.Text;
-using System.Threading.Tasks;
 using System.Windows.Forms;
 
 namespace UnrealGameSync
 {
 	partial class SyncFilter : Form
 	{
-		Dictionary<Guid, WorkspaceSyncCategory> _uniqueIdToCategory;
+		readonly Dictionary<Guid, WorkspaceSyncCategory> _uniqueIdToCategory;
 		public FilterSettings GlobalFilter;
 		public FilterSettings WorkspaceFilter;
-		ConfigSection? _perforceSection;
+		readonly ConfigSection? _perforceSection;
 
 		public SyncFilter(Dictionary<Guid, WorkspaceSyncCategory> uniqueIdToCategory, FilterSettings globalFilter, FilterSettings workspaceFilter, ConfigSection? perforceSection)
 		{
@@ -121,16 +116,6 @@ namespace UnrealGameSync
 			return result;
 		}
 
-		private static string[] GetView(TextBox filterText)
-		{
-			List<string> newLines = new List<string>(filterText.Lines);
-			while (newLines.Count > 0 && newLines.Last().Trim().Length == 0)
-			{
-				newLines.RemoveAt(newLines.Count - 1);
-			}
-			return newLines.Count > 0 ? filterText.Lines : new string[0];
-		}
-
 		private void OkButton_Click(object sender, EventArgs e)
 		{
 			GetSettings(out FilterSettings newGlobalFilter, out FilterSettings newWorkspaceFilter);
diff --git a/Engine/Source/Programs/UnrealGameSync/UnrealGameSync/Forms/VariablesWindow.cs b/Engine/Source/Programs/UnrealGameSync/UnrealGameSync/Forms/VariablesWindow.cs
--- a/Engine/Source/Programs/UnrealGameSync/UnrealGameSync/Forms/VariablesWindow.cs
+++ b/Engine/Source/Programs/UnrealGameSync/UnrealGameSync/Forms/VariablesWindow.cs
@@ -3,12 +3,7 @@
 using System;
 using System.Collections;
 using System.Collections.Generic;
-using System.ComponentModel;
-using System.Data;
-using System.Drawing;
 using System.Linq;
-using System.Text;
-using System.Threading.Tasks;
 using System.Windows.Forms;
 
 #nullable enable
@@ -17,7 +12,7 @@ namespace UnrealGameSync
 {
 	public partial class VariablesWindow : Form
 	{
-		static HashSet<string> _legacyVariables = new HashSet<string>()
+		static readonly HashSet<string> s_legacyVariables = new HashSet<string>()
 		{
 			"UE4EditorConfig",
 			"UE4EditorDebugArg",
@@ -42,7 +37,7 @@ namespace UnrealGameSync
 
 			foreach(KeyValuePair<string, string> pair in variables)
 			{
-				if (!_legacyVariables.Contains(pair.Key))
+				if (!s_legacyVariables.Contains(pair.Key))
 				{
 					ListViewItem item = new ListViewItem(String.Format("$({0})", pair.Key));
 					item.SubItems.Add(pair.Value);
@@ -64,10 +59,7 @@ namespace UnrealGameSync
 			}
 		}
 
-		protected override bool ShowWithoutActivation
-		{
-			get { return true; }
-		}
+		protected override bool ShowWithoutActivation => true;
 
 		private void OkButton_Click(object sender, EventArgs e)
 		{
diff --git a/Engine/Source/Programs/UnrealGameSync/UnrealGameSync/GlobalSettings.cs b/Engine/Source/Programs/UnrealGameSync/UnrealGameSync/GlobalSettings.cs
--- a/Engine/Source/Programs/UnrealGameSync/UnrealGameSync/GlobalSettings.cs
+++ b/Engine/Source/Programs/UnrealGameSync/UnrealGameSync/GlobalSettings.cs
@@ -2,9 +2,7 @@
 
 using Microsoft.Win32;
 using System;
-using System.Collections.Generic;
 using System.Linq;
-using System.Text;
 using System.Text.RegularExpressions;
 using System.Windows.Forms;
 
@@ -74,7 +72,7 @@ namespace UnrealGameSync
 
 					if (String.IsNullOrEmpty(serverAndPort))
 					{
-						try { key.DeleteValue("ServerAndPort"); } catch (Exception) { }
+						DeleteValueGuarded(key, "ServerAndPort");
 					}
 					else
 					{
@@ -83,7 +81,7 @@ namespace UnrealGameSync
 
 					if (String.IsNullOrEmpty(userName))
 					{
-						try { key.DeleteValue("UserName"); } catch (Exception) { }
+						DeleteValueGuarded(key, "UserName");
 					}
 					else
 					{
@@ -105,7 +103,7 @@ namespace UnrealGameSync
 					}
 					else
 					{
-						try { key.DeleteValue("Preview"); } catch (Exception) { }
+						DeleteValueGuarded(key, "Preview"); 
 					}
 				}
 			}
@@ -114,5 +112,16 @@ namespace UnrealGameSync
 				MessageBox.Show("Unable to save settings.\n\n" + ex.ToString());
 			}
 		}
+
+		static void DeleteValueGuarded(RegistryKey key, string name)
+		{
+			try
+			{
+				key.DeleteValue(name);
+			}
+			catch (Exception)
+			{
+			}
+		}
 	}
 }
diff --git a/Engine/Source/Programs/UnrealGameSync/UnrealGameSync/IssueMonitor.cs b/Engine/Source/Programs/UnrealGameSync/UnrealGameSync/IssueMonitor.cs
--- a/Engine/Source/Programs/UnrealGameSync/UnrealGameSync/IssueMonitor.cs
+++ b/Engine/Source/Programs/UnrealGameSync/UnrealGameSync/IssueMonitor.cs
@@ -1,23 +1,17 @@
 // Copyright Epic Games, Inc. All Rights Reserved.
 
-using EnvDTE;
 using EpicGames.Core;
 using Microsoft.Extensions.DependencyInjection;
 using Microsoft.Extensions.Logging;
 using System;
-using System.Collections.Concurrent;
 using System.Collections.Generic;
 using System.Diagnostics;
 using System.Globalization;
 using System.Linq;
-using System.Security.Policy;
-using System.Text;
 using System.Text.RegularExpressions;
 using System.Threading;
 using System.Threading.Tasks;
 
-using Thread = System.Threading.Thread;
-
 #nullable enable
 
 namespace UnrealGameSync
@@ -136,33 +130,33 @@ namespace UnrealGameSync
 
 	class IssueMonitor : IDisposable
 	{
-		public readonly string? ApiUrl;
-		public readonly string UserName;
+		public string? ApiUrl { get; }
+		public string UserName { get; }
 		int _refCount = 1;
 		Task? _workerTask;
-		ILogger _logger;
-		CancellationTokenSource _cancellationSource;
-		AsyncEvent _refreshEvent;
+		readonly ILogger _logger;
+		readonly CancellationTokenSource _cancellationSource;
+		readonly AsyncEvent _refreshEvent;
 		int _updateIntervalMs;
-		List<long> _trackingIssueIds = new List<long>();
+		readonly List<long> _trackingIssueIds = new List<long>();
 		List<IssueData> _issues = new List<IssueData>();
-		object _lockObject = new object();
-		List<IssueUpdateData> _pendingUpdates = new List<IssueUpdateData>();
-		IAsyncDisposer _asyncDisposer;
+		readonly object _lockObject = new object();
+		readonly List<IssueUpdateData> _pendingUpdates = new List<IssueUpdateData>();
+		readonly IAsyncDisposer _asyncDisposer;
 
 		public Action? OnIssuesChanged { get; set; }
 
 		// Only used by MainWindow, but easier to just store here
-		public Dictionary<long, IssueAlertReason> IssueIdToAlertReason = new Dictionary<long, IssueAlertReason>();
+		public Dictionary<long, IssueAlertReason> IssueIdToAlertReason { get; } = new Dictionary<long, IssueAlertReason>();
 
 		public IssueMonitor(string? apiUrl, string userName, TimeSpan updateInterval, IServiceProvider serviceProvider)
 		{
-			this.ApiUrl = apiUrl;
-			this.UserName = userName;
-			this._updateIntervalMs = (int)updateInterval.TotalMilliseconds;
-			this._logger = serviceProvider.GetRequiredService<ILogger<IssueMonitor>>();
+			ApiUrl = apiUrl;
+			UserName = userName;
+			_updateIntervalMs = (int)updateInterval.TotalMilliseconds;
+			_logger = serviceProvider.GetRequiredService<ILogger<IssueMonitor>>();
 			_cancellationSource = new CancellationTokenSource();
-			this._asyncDisposer = serviceProvider.GetRequiredService<IAsyncDisposer>();
+			_asyncDisposer = serviceProvider.GetRequiredService<IAsyncDisposer>();
 
 			if (apiUrl == null)
 			{
@@ -170,7 +164,7 @@ namespace UnrealGameSync
 			}
 			else
 			{
-				_logger.LogInformation("Using connection string: {ApiUrl}", this.ApiUrl);
+				_logger.LogInformation("Using connection string: {ApiUrl}", ApiUrl);
 			}
 
 			_refreshEvent = new AsyncEvent();
@@ -347,7 +341,10 @@ namespace UnrealGameSync
 				{
 					if (await SendUpdateAsync(pendingUpdate, cancellationToken))
 					{
-						lock (_lockObject) { _pendingUpdates.RemoveAt(0); }
+						lock (_lockObject)
+						{ 
+							_pendingUpdates.RemoveAt(0);
+						}
 					}
 					else
 					{
@@ -449,10 +446,7 @@ namespace UnrealGameSync
 				// Update the main thread
 				if(initialNumIssues > 0 || _issues.Count > 0)
 				{
-					if(OnIssuesChanged != null)
-					{
-						OnIssuesChanged();
-					}
+					OnIssuesChanged?.Invoke();
 				}
 
 				// Update the stats
diff --git a/Engine/Source/Programs/UnrealGameSync/UnrealGameSync/Jupiter/JupiterMonitor.cs b/Engine/Source/Programs/UnrealGameSync/UnrealGameSync/Jupiter/JupiterMonitor.cs
--- a/Engine/Source/Programs/UnrealGameSync/UnrealGameSync/Jupiter/JupiterMonitor.cs
+++ b/Engine/Source/Programs/UnrealGameSync/UnrealGameSync/Jupiter/JupiterMonitor.cs
@@ -71,7 +71,6 @@ namespace UnrealGameSync
 				{
 					_archives = newArchives;
 				}
-
 			}
 			catch (Exception exception)
 			{
@@ -83,30 +82,40 @@ namespace UnrealGameSync
 		{
 			ConfigSection jupiterConfigSection = configFile.FindSection("Jupiter");
 			if (jupiterConfigSection == null)
+			{
 				return null;
+			}
 
 			string jupiterUrl = jupiterConfigSection.GetValue("JupiterUrl");
 			string oidcProviderIdentifier = jupiterConfigSection.GetValue("OIDCProviderIdentifier");
 
 			ConfigSection projectConfigSection = configFile.FindSection(selectedProjectIdentifier);
 			if (projectConfigSection == null)
+			{
 				return null;
+			}
 
 			string jupiterNamespace = projectConfigSection.GetValue("JupiterNamespace");
 			// Is no namespace has been specified we are unable to fetch builds
 			if (jupiterNamespace == null)
+			{
 				return null;
+			}
 
 			string expectedBranch = projectConfigSection.GetValue("ExpectedBranch");
 			// If we do not know which branch to fetch we can not list builds, as it would risk getting binaries from other branches
 			if (expectedBranch == null)
+			{
 				return null;
+			}
 
 			oidcProviderIdentifier = projectConfigSection.GetValue("OIDCProviderIdentifier") ?? oidcProviderIdentifier;
 
 			// with no oidc provider we are unable to login, thus it is required
 			if (oidcProviderIdentifier == null)
+			{
 				return null;
+			}
 
 			// project specific overrides
 			jupiterUrl = projectConfigSection.GetValue("JupiterUrl") ?? jupiterUrl;
@@ -148,22 +157,24 @@ namespace UnrealGameSync
 
 					// skip if we do not have metadata as we need that to be able to determine if this root is of the type we want
 					if (branch == null || project == null || archiveType == null || changelistString == null)
+					{
 						continue;
+					}
 
-					if (!string.Equals(_expectedBranch, branch, StringComparison.InvariantCultureIgnoreCase))
+					if (!String.Equals(_expectedBranch, branch, StringComparison.InvariantCultureIgnoreCase))
 					{
 						continue;
 					}
 
 					int changelist;
-					if (!int.TryParse(changelistString, out changelist))
+					if (!Int32.TryParse(changelistString, out changelist))
 					{
 						continue; // invalid changelist format
 					}
 
 					JupiterArchiveInfo existingArchive = newArchives.FirstOrDefault(info =>
-						string.Equals(info.Name, project, StringComparison.InvariantCultureIgnoreCase) &&
-						string.Equals(info.Type, archiveType, StringComparison.InvariantCultureIgnoreCase));
+						String.Equals(info.Name, project, StringComparison.InvariantCultureIgnoreCase) &&
+						String.Equals(info.Type, archiveType, StringComparison.InvariantCultureIgnoreCase));
 					if (existingArchive == null)
 					{
 						JupiterArchiveInfo archive = new JupiterArchiveInfo(project, archiveType, _jupiterUrl.ToString(), _jupiterNamespace);
@@ -211,10 +222,7 @@ namespace UnrealGameSync
 			public string Name { get; }
 			public string Type { get; }
 
-			public string BasePath
-			{
-				get { return null; } 
-			}
+			public string BasePath => null;
 
 			public string Target => throw new NotSupportedException();
 
diff --git a/Engine/Source/Programs/UnrealGameSync/UnrealGameSync/ModalTask.cs b/Engine/Source/Programs/UnrealGameSync/UnrealGameSync/ModalTask.cs
--- a/Engine/Source/Programs/UnrealGameSync/UnrealGameSync/ModalTask.cs
+++ b/Engine/Source/Programs/UnrealGameSync/UnrealGameSync/ModalTask.cs
@@ -1,9 +1,6 @@
 // Copyright Epic Games, Inc. All Rights Reserved.
 
 using System;
-using System.Collections.Generic;
-using System.Linq;
-using System.Text;
 using System.Threading;
 using System.Threading.Tasks;
 using System.Windows.Forms;
@@ -28,7 +25,7 @@ namespace UnrealGameSync
 
 		public ModalTask(Task task)
 		{
-			this.Task = task;
+			Task = task;
 		}
 
 		public string Error
@@ -58,8 +55,12 @@ namespace UnrealGameSync
 
 		public static ModalTask? Execute(IWin32Window? owner, string title, string message, Func<CancellationToken, Task> taskFunc, ModalTaskFlags flags = ModalTaskFlags.None)
 		{
-			Func<CancellationToken, Task<int>> typedTaskFunc = async x => { await taskFunc(x); return 0; };
-			return Execute(owner, title, message, typedTaskFunc, flags);
+			async Task<int> TypedTaskFunc(CancellationToken cancellationToken)
+			{
+				await taskFunc(cancellationToken);
+				return 0;
+			}
+			return Execute(owner, title, message, TypedTaskFunc, flags);
 		}
 
 		public static ModalTask<T>? Execute<T>(IWin32Window? owner, string title, string message, Func<CancellationToken, Task<T>> taskFunc, ModalTaskFlags flags = ModalTaskFlags.None)
diff --git a/Engine/Source/Programs/UnrealGameSync/UnrealGameSync/PerforceModalTask.cs b/Engine/Source/Programs/UnrealGameSync/UnrealGameSync/PerforceModalTask.cs
--- a/Engine/Source/Programs/UnrealGameSync/UnrealGameSync/PerforceModalTask.cs
+++ b/Engine/Source/Programs/UnrealGameSync/UnrealGameSync/PerforceModalTask.cs
@@ -3,10 +3,6 @@
 using EpicGames.Perforce;
 using Microsoft.Extensions.Logging;
 using System;
-using System.Collections.Generic;
-using System.IO;
-using System.Linq;
-using System.Text;
 using System.Threading;
 using System.Threading.Tasks;
 using System.Windows.Forms;
@@ -19,7 +15,7 @@ namespace UnrealGameSync
 
 		public PerforceLoginException(PerforceResponse<LoginRecord> response) : base($"Login failed: {response}")
 		{
-			this.Response = response;
+			Response = response;
 		}
 	}
 
@@ -27,8 +23,13 @@ namespace UnrealGameSync
 	{
 		public static ModalTask? Execute(IWin32Window? owner, string title, string message, IPerforceSettings perforceSettings, Func<IPerforceConnection, CancellationToken, Task> executeAsync, ILogger logger, ModalTaskFlags flags = ModalTaskFlags.None)
 		{
-			Func<IPerforceConnection, CancellationToken, Task<int>> executeTypedAsync = async (p, c) => { await executeAsync(p, c); return 0; };
-			return Execute(owner, title, message, perforceSettings, executeTypedAsync, logger, flags);
+			async Task<int> ExecuteTypedAsync(IPerforceConnection perforce, CancellationToken cancellationToken)
+			{
+				await executeAsync(perforce, cancellationToken); 
+				return 0; 
+			}
+
+			return Execute(owner, title, message, perforceSettings, ExecuteTypedAsync, logger, flags);
 		}
 
 		public static ModalTask<T>? Execute<T>(IWin32Window? owner, string title, string message, IPerforceSettings perforceSettings, Func<IPerforceConnection, CancellationToken, Task<T>> executeAsync, ILogger logger, ModalTaskFlags flags = ModalTaskFlags.None)
@@ -36,12 +37,12 @@ namespace UnrealGameSync
 			IPerforceConnection? connection = null;
 			try
 			{
-				Func<Task<IPerforceConnection>> connectAsync = async () =>
+				async Task<IPerforceConnection> ConnectAsync()
 				{
 					connection ??= await PerforceConnection.CreateAsync(perforceSettings, logger);
 					return connection;
-				};
-				return ExecuteInternal(owner, title, message, perforceSettings, connectAsync, executeAsync, flags);
+				}
+				return ExecuteInternal(owner, title, message, perforceSettings, ConnectAsync, executeAsync, flags);
 			}
 			finally
 			{
@@ -51,11 +52,15 @@ namespace UnrealGameSync
 
 		public static ModalTask? Execute(IWin32Window? owner, string title, string message, IPerforceConnection perforceConnection, Func<IPerforceConnection, CancellationToken, Task> executeAsync, ILogger logger, ModalTaskFlags flags = ModalTaskFlags.None)
 		{
-			Func<IPerforceConnection, CancellationToken, Task<int>> executeTypedAsync = async (p, c) => { await executeAsync(p, c); return 0; };
-			return Execute(owner, title, message, perforceConnection, executeTypedAsync, logger, flags);
+			async Task<int> ExecuteTypedAsync(IPerforceConnection perforce, CancellationToken cancellationToken)
+			{
+				await executeAsync(perforce, cancellationToken);
+				return 0;
+			}
+			return Execute(owner, title, message, perforceConnection, ExecuteTypedAsync, logger, flags);
 		}
 
-		public static ModalTask<T>? Execute<T>(IWin32Window? owner, string title, string message, IPerforceConnection perforceConnection, Func<IPerforceConnection, CancellationToken, Task<T>> executeAsync, ILogger logger, ModalTaskFlags flags = ModalTaskFlags.None)
+		public static ModalTask<T>? Execute<T>(IWin32Window? owner, string title, string message, IPerforceConnection perforceConnection, Func<IPerforceConnection, CancellationToken, Task<T>> executeAsync, ModalTaskFlags flags = ModalTaskFlags.None)
 		{
 			return ExecuteInternal(owner, title, message, perforceConnection.Settings, () => Task.FromResult(perforceConnection), executeAsync, flags);
 		}
@@ -65,9 +70,9 @@ namespace UnrealGameSync
 			string? password = perforceSettings.Password;
 			for(;;)
 			{
-				Func<CancellationToken, Task<T>> runAsync = cancellationToken => LoginAndExecuteAsync(password, connectAsync, executeAsync, cancellationToken);
+				Task<T> RunAsync(CancellationToken cancellationToken) => LoginAndExecuteAsync(password, connectAsync, executeAsync, cancellationToken);
 
-				ModalTask<T>? result = ModalTask.Execute(owner, title, message, runAsync, ModalTaskFlags.Quiet);
+				ModalTask<T>? result = ModalTask.Execute(owner, title, message, RunAsync, ModalTaskFlags.Quiet);
 				if (result != null && result.Failed && (flags & ModalTaskFlags.Quiet) == 0)
 				{
 					if (result.Exception is PerforceLoginException)
diff --git a/Engine/Source/Programs/UnrealGameSync/UnrealGameSync/PerforceMonitor.cs b/Engine/Source/Programs/UnrealGameSync/UnrealGameSync/PerforceMonitor.cs
--- a/Engine/Source/Programs/UnrealGameSync/UnrealGameSync/PerforceMonitor.cs
+++ b/Engine/Source/Programs/UnrealGameSync/UnrealGameSync/PerforceMonitor.cs
@@ -5,14 +5,10 @@ using EpicGames.Perforce;
 using Microsoft.Extensions.DependencyInjection;
 using Microsoft.Extensions.Logging;
 using System;
-using System.Collections.Concurrent;
 using System.Collections.Generic;
 using System.Diagnostics;
 using System.Diagnostics.CodeAnalysis;
-using System.IO;
 using System.Linq;
-using System.Runtime.CompilerServices;
-using System.Text;
 using System.Threading;
 using System.Threading.Tasks;
 
@@ -33,29 +29,29 @@ namespace UnrealGameSync
 			}
 		}
 
-		public int InitialMaxChangesValue = 100;
-		public bool ShowChangesForAllProjects = true;
+		public const int InitialMaxChangesValue = 100;
+		public bool ShowChangesForAllProjects { get; set; } = true;
 
-		IPerforceSettings _perforceSettings;
+		readonly IPerforceSettings _perforceSettings;
 		readonly string _branchClientPath;
 		readonly string _selectedClientFileName;
 		readonly string _selectedProjectIdentifier;
 		Task? _workerTask;
-		CancellationTokenSource _cancellationSource;
+		readonly CancellationTokenSource _cancellationSource;
 		int _pendingMaxChangesValue;
 		SortedSet<ChangesRecord> _changes = new SortedSet<ChangesRecord>(new PerforceChangeSorter());
-		SortedDictionary<int, PerforceChangeDetails> _changeDetails = new SortedDictionary<int,PerforceChangeDetails>();
-		SortedSet<int> _promotedChangeNumbers = new SortedSet<int>();
+		readonly SortedDictionary<int, PerforceChangeDetails> _changeDetails = new SortedDictionary<int,PerforceChangeDetails>();
+		readonly SortedSet<int> _promotedChangeNumbers = new SortedSet<int>();
 		List<PerforceArchiveInfo> _archives = new List<PerforceArchiveInfo>();
-		AsyncEvent _refreshEvent = new AsyncEvent();
-		ILogger _logger;
-		bool _isEnterpriseProject;
-		DirectoryReference _cacheFolder;
-		List<KeyValuePair<FileReference, DateTime>> _localConfigFiles;
-		IAsyncDisposer _asyncDisposeTasks;
+		readonly AsyncEvent _refreshEvent = new AsyncEvent();
+		readonly ILogger _logger;
+		readonly bool _isEnterpriseProject;
+		readonly DirectoryReference _cacheFolder;
+		readonly List<KeyValuePair<FileReference, DateTime>> _localConfigFiles;
+		readonly IAsyncDisposer _asyncDisposeTasks;
 		string[] prevCodeRules = Array.Empty<string>();
 
-		SynchronizationContext _synchronizationContext;
+		readonly SynchronizationContext _synchronizationContext;
 		public event Action? OnUpdate;
 		public event Action? OnUpdateMetadata;
 		public event Action? OnStreamChange;
@@ -90,10 +86,7 @@ namespace UnrealGameSync
 
 		public void Start()
 		{
-			if (_workerTask == null)
-			{
-				_workerTask = Task.Run(() => PollForUpdates(_cancellationSource.Token));
-			}
+			_workerTask ??= Task.Run(() => PollForUpdates(_cancellationSource.Token));
 		}
 
 		public void Dispose()
@@ -131,8 +124,18 @@ namespace UnrealGameSync
 
 		public int PendingMaxChanges
 		{
-			get { return _pendingMaxChangesValue; }
-			set { lock(this){ if(value != _pendingMaxChangesValue){ _pendingMaxChangesValue = value; _refreshEvent.Set(); } } }
+			get => _pendingMaxChangesValue;
+			set 
+			{ 
+				lock(this)
+				{ 
+					if(value != _pendingMaxChangesValue)
+					{ 
+						_pendingMaxChangesValue = value; 
+						_refreshEvent.Set(); 
+					} 
+				} 
+			}
 		}
 
 		async Task PollForUpdates(CancellationToken cancellationToken)
@@ -279,7 +282,7 @@ namespace UnrealGameSync
 					IEnumerable<string> additionalPaths = projectConfigSection.GetValues("AdditionalPathsToSync", new string[0]);
 
 					// turn into //ws/path
-					depotPaths.AddRange(additionalPaths.Select(p => string.Format("{0}/{1}", _branchClientPath, p.TrimStart('/'))));
+					depotPaths.AddRange(additionalPaths.Select(p => String.Format("{0}/{1}", _branchClientPath, p.TrimStart('/'))));
 				}
 			}
 
@@ -397,29 +400,6 @@ namespace UnrealGameSync
 			return true;
 		}
 
-		Func<string, bool>? CreateCodeFilterFunc()
-		{
-			ConfigSection? projectConfigSection = LatestProjectConfigFile.FindSection("Perforce");
-			if (projectConfigSection == null)
-			{
-				return null;
-			}
-
-			string[] rules = projectConfigSection.GetValues("CodeFilter", new string[0]);
-			if (rules.Length == 0)
-			{
-				return null;
-			}
-
-			FileFilter filter = new FileFilter(PerforceUtils.CodeExtensions.Select(x => $"*.{x}"));
-			foreach (string rule in rules)
-			{
-				filter.AddRule(rule);
-			}
-
-			return filter.Matches;
-		}
-
 		public async Task<bool> UpdateChangeTypesAsync(IPerforceConnection perforce, CancellationToken cancellationToken)
 		{
 			// Get the filter for code changes
diff --git a/Engine/Source/Programs/UnrealGameSync/UnrealGameSync/ProgramApplicationContext.cs b/Engine/Source/Programs/UnrealGameSync/UnrealGameSync/ProgramApplicationContext.cs
--- a/Engine/Source/Programs/UnrealGameSync/UnrealGameSync/ProgramApplicationContext.cs
+++ b/Engine/Source/Programs/UnrealGameSync/UnrealGameSync/ProgramApplicationContext.cs
@@ -11,7 +11,6 @@ using System.Drawing;
 using System.IO;
 using System.Linq;
 using System.Reflection;
-using System.Text;
 using System.Threading;
 using System.Threading.Tasks;
 using System.Windows.Forms;
@@ -24,36 +23,36 @@ namespace UnrealGameSync
 {
 	class ProgramApplicationContext : ApplicationContext
 	{
-		SynchronizationContext _mainThreadSynchronizationContext;
+		readonly SynchronizationContext _mainThreadSynchronizationContext;
 
-		IPerforceSettings _defaultPerforceSettings;
+		readonly IPerforceSettings _defaultPerforceSettings;
 		UpdateMonitor _updateMonitor;
-		string? _apiUrl;
-		DirectoryReference _dataFolder;
-		DirectoryReference _cacheFolder;
-		bool _restoreState;
-		string? _updateSpawn;
-		bool _preview;
+		readonly string? _apiUrl;
+		readonly DirectoryReference _dataFolder;
+		readonly DirectoryReference _cacheFolder;
+		readonly bool _restoreState;
+		readonly string? _updateSpawn;
+		readonly bool _preview;
 		bool _isClosing;
-		string? _uri;
+		readonly string? _uri;
 
-		IServiceProvider _serviceProvider;
-		ILogger _logger;
-		UserSettings _settings;
+		readonly IServiceProvider _serviceProvider;
+		readonly ILogger _logger;
+		readonly UserSettings _settings;
 		ActivationListener _activationListener;
 
 		Container _components = new Container();
 		NotifyIcon _notifyIcon;
-		ContextMenuStrip _notifyMenu;
-		ToolStripMenuItem _notifyMenuOpenUnrealGameSync;
-		ToolStripSeparator _notifyMenuOpenUnrealGameSyncSeparator;
-		ToolStripMenuItem _notifyMenuSyncNow;
-		ToolStripMenuItem _notifyMenuLaunchEditor;
-		ToolStripSeparator _notifyMenuExitSeparator;
-		ToolStripMenuItem _notifyMenuExit;
+		readonly ContextMenuStrip _notifyMenu;
+		readonly ToolStripMenuItem _notifyMenuOpenUnrealGameSync;
+		readonly ToolStripSeparator _notifyMenuOpenUnrealGameSyncSeparator;
+		readonly ToolStripMenuItem _notifyMenuSyncNow;
+		readonly ToolStripMenuItem _notifyMenuLaunchEditor;
+		readonly ToolStripSeparator _notifyMenuExitSeparator;
+		readonly ToolStripMenuItem _notifyMenuExit;
 
 		CancellationTokenSource _startupCancellationSource = new CancellationTokenSource();
-		Task _startupTask;
+		readonly Task _startupTask;
 		ModalTaskWindow? _startupWindow;
 		MainWindow? _mainWindowInstance;
 
@@ -62,17 +61,17 @@ namespace UnrealGameSync
 
 		public ProgramApplicationContext(IPerforceSettings defaultPerforceSettings, UpdateMonitor updateMonitor, string? apiUrl, DirectoryReference dataFolder, EventWaitHandle activateEvent, bool restoreState, string? updateSpawn, string? projectFileName, bool preview, IServiceProvider serviceProvider, string? uri)
 		{
-			this._defaultPerforceSettings = defaultPerforceSettings;
-			this._updateMonitor = updateMonitor;
-			this._apiUrl = apiUrl;
-			this._dataFolder = dataFolder;
-			this._cacheFolder = DirectoryReference.Combine(dataFolder, "Cache");
-			this._restoreState = restoreState;
-			this._updateSpawn = updateSpawn;
-			this._preview = preview;
-			this._serviceProvider = serviceProvider;
-			this._logger = serviceProvider.GetRequiredService<ILogger<ProgramApplicationContext>>();
-			this._uri = uri;
+			_defaultPerforceSettings = defaultPerforceSettings;
+			_updateMonitor = updateMonitor;
+			_apiUrl = apiUrl;
+			_dataFolder = dataFolder;
+			_cacheFolder = DirectoryReference.Combine(dataFolder, "Cache");
+			_restoreState = restoreState;
+			_updateSpawn = updateSpawn;
+			_preview = preview;
+			_serviceProvider = serviceProvider;
+			_logger = serviceProvider.GetRequiredService<ILogger<ProgramApplicationContext>>();
+			_uri = uri;
 
 			// Create the directories
 			DirectoryReference.CreateDirectory(dataFolder);
@@ -321,7 +320,7 @@ namespace UnrealGameSync
 			}
 
 			// Create the main window 
-			_mainWindowInstance = new MainWindow(_updateMonitor, _apiUrl, _dataFolder, _cacheFolder, _restoreState, _updateSpawn ?? originalExe, _preview, startupTasks, _defaultPerforceSettings, _serviceProvider, _settings, _uri, _oidcTokenManager);
+			_mainWindowInstance = new MainWindow(_updateMonitor, _apiUrl, _dataFolder, _updateSpawn ?? originalExe, _preview, startupTasks, _defaultPerforceSettings, _serviceProvider, _settings, _uri, _oidcTokenManager);
 			if(visible)
 			{
 				_mainWindowInstance.Show();
@@ -428,46 +427,28 @@ namespace UnrealGameSync
 			_notifyMenuExitSeparator.Visible = canSyncNow || canLaunchEditor;
 
 			// Show the startup window, if not already visible
-			if(_startupWindow != null)
-			{
-				_startupWindow.Show();
-			}
+			_startupWindow?.Show();
 		}
 
 		private void NotifyIcon_DoubleClick(object? sender, EventArgs e)
 		{
-			if(_mainWindowInstance != null)
-			{
-				_mainWindowInstance.ShowAndActivate();
-			}
+			_mainWindowInstance?.ShowAndActivate();
 		}
 
 		private void NotifyMenu_OpenUnrealGameSync_Click(object? sender, EventArgs e)
 		{
-			if(_startupWindow != null)
-			{
-				_startupWindow.ShowAndActivate();
-			}
-			if(_mainWindowInstance != null)
-			{
-				_mainWindowInstance.ShowAndActivate();
-			}
+			_startupWindow?.ShowAndActivate();
+			_mainWindowInstance?.ShowAndActivate();
 		}
 
 		private void NotifyMenu_SyncNow_Click(object? sender, EventArgs e)
 		{
-			if(_mainWindowInstance != null)
-			{
-				_mainWindowInstance.SyncLatestChange();
-			}
+			_mainWindowInstance?.SyncLatestChange();
 		}
 
 		private void NotifyMenu_LaunchEditor_Click(object? sender, EventArgs e)
 		{
-			if(_mainWindowInstance != null)
-			{
-				_mainWindowInstance.LaunchEditor();
-			}
+			_mainWindowInstance?.LaunchEditor();
 		}
 
 		private void NotifyMenu_Exit_Click(object? sender, EventArgs e)
diff --git a/Engine/Source/Programs/UnrealGameSync/UnrealGameSync/RESTApi.cs b/Engine/Source/Programs/UnrealGameSync/UnrealGameSync/RESTApi.cs
--- a/Engine/Source/Programs/UnrealGameSync/UnrealGameSync/RESTApi.cs
+++ b/Engine/Source/Programs/UnrealGameSync/UnrealGameSync/RESTApi.cs
@@ -1,12 +1,8 @@
 ﻿// Copyright Epic Games, Inc. All Rights Reserved.
 
 using System;
-using System.Collections.Generic;
-using System.IO;
-using System.Net;
 using System.Net.Http;
 using System.Net.Http.Headers;
-using System.Net.Mime;
 using System.Text;
 using System.Text.Json;
 using System.Threading;
@@ -29,14 +25,14 @@ namespace UnrealGameSync
 
 	public static class RestApi
 	{
-		static HttpClient s_httpClient = new HttpClient();
+		static readonly HttpClient s_httpClient = new HttpClient();
 
 		private static async Task<string> SendRequestInternalAsync(string url, HttpMethod method, string? requestBody, CancellationToken cancellationToken)
 		{
 			using (HttpRequestMessage request = new HttpRequestMessage(method, url))
 			{
 				// Add json to request body
-				if (!string.IsNullOrEmpty(requestBody))
+				if (!String.IsNullOrEmpty(requestBody))
 				{
 					if (method == HttpMethod.Post || method == HttpMethod.Put)
 					{
@@ -47,7 +43,7 @@ namespace UnrealGameSync
 
 				try
 				{
-					using (HttpResponseMessage response = await s_httpClient.SendAsync(request))
+					using (HttpResponseMessage response = await s_httpClient.SendAsync(request, cancellationToken))
 					{
 						return await response.Content.ReadAsStringAsync();
 					}
diff --git a/Engine/Source/Programs/UnrealGameSync/UnrealGameSync/Taskbar.cs b/Engine/Source/Programs/UnrealGameSync/UnrealGameSync/Taskbar.cs
--- a/Engine/Source/Programs/UnrealGameSync/UnrealGameSync/Taskbar.cs
+++ b/Engine/Source/Programs/UnrealGameSync/UnrealGameSync/Taskbar.cs
@@ -1,11 +1,7 @@
 // Copyright Epic Games, Inc. All Rights Reserved.
 
 using System;
-using System.Collections.Generic;
-using System.Linq;
 using System.Runtime.InteropServices;
-using System.Text;
-using System.Threading.Tasks;
 
 namespace UnrealGameSync
 {
@@ -29,7 +25,7 @@ namespace UnrealGameSync
 			void ActivateTab(IntPtr hWnd);
 			void SetActiveAlt(IntPtr hWnd);
 			void MarkFullscreenWindow(IntPtr hWnd, int fFullscreen);
-			void SetProgressValue(IntPtr hWnd, UInt64 ullCompleted, UInt64 ullTotal);
+			void SetProgressValue(IntPtr hWnd, ulong ullCompleted, ulong ullTotal);
 			void SetProgressState(IntPtr hWnd, TaskbarState state);
 		}
 
@@ -38,7 +34,7 @@ namespace UnrealGameSync
 		{
 		}
 
-		static ITaskbarList3? _interface;
+		static readonly ITaskbarList3? _interface;
 
 		static Taskbar()
 		{
diff --git a/Engine/Source/Programs/UnrealGameSync/UnrealGameSync/ToolUpdateMonitor.cs b/Engine/Source/Programs/UnrealGameSync/UnrealGameSync/ToolUpdateMonitor.cs
--- a/Engine/Source/Programs/UnrealGameSync/UnrealGameSync/ToolUpdateMonitor.cs
+++ b/Engine/Source/Programs/UnrealGameSync/UnrealGameSync/ToolUpdateMonitor.cs
@@ -9,10 +9,8 @@ using System.Collections.Generic;
 using System.Diagnostics;
 using System.IO;
 using System.Linq;
-using System.Text;
 using System.Threading;
 using System.Threading.Tasks;
-using System.Windows.Forms;
 
 #nullable enable
 
@@ -27,8 +25,8 @@ namespace UnrealGameSync
 
 		public ToolLink(string label, string fileName)
 		{
-			this.Label = label;
-			this.FileName = fileName;
+			Label = label;
+			FileName = fileName;
 		}
 	}
 
@@ -48,26 +46,26 @@ namespace UnrealGameSync
 
 		public ToolDefinition(Guid id, string name, string description, string zipPath, int zipChange, string configPath, int configChange)
 		{
-			this.Id = id;
-			this.Name = name;
-			this.Description = description;
-			this.ZipPath = zipPath;
-			this.ZipChange = zipChange;
-			this.ConfigPath = configPath;
-			this.ConfigChange = configChange;
+			Id = id;
+			Name = name;
+			Description = description;
+			ZipPath = zipPath;
+			ZipChange = zipChange;
+			ConfigPath = configPath;
+			ConfigChange = configChange;
 		}
 	}
 
 	class ToolUpdateMonitor : IDisposable
 	{
-		CancellationTokenSource _cancellationSource;
-		SynchronizationContext _synchronizationContext;
+		readonly CancellationTokenSource _cancellationSource;
+		readonly SynchronizationContext _synchronizationContext;
 		Task? _workerTask;
-		AsyncEvent _wakeEvent;
-		ILogger _logger;
+		readonly AsyncEvent _wakeEvent;
+		readonly ILogger _logger;
 		public List<ToolDefinition> Tools { get; private set; } = new List<ToolDefinition>();
-		int _lastChange = -1;
-		IAsyncDisposer _asyncDisposer;
+		readonly int _lastChange = -1;
+		readonly IAsyncDisposer _asyncDisposer;
 
 		IPerforceSettings PerforceSettings { get; }
 		DirectoryReference ToolsDir { get; }
@@ -79,11 +77,11 @@ namespace UnrealGameSync
 		{
 			_cancellationSource = new CancellationTokenSource();
 			_synchronizationContext = SynchronizationContext.Current!;
-			this.ToolsDir = DirectoryReference.Combine(dataDir, "Tools");
-			this.PerforceSettings = perforceSettings;
-			this.Settings = settings;
-			this._logger = serviceProvider.GetRequiredService<ILogger<ToolUpdateMonitor>>();
-			this._asyncDisposer = serviceProvider.GetRequiredService<IAsyncDisposer>();
+			ToolsDir = DirectoryReference.Combine(dataDir, "Tools");
+			PerforceSettings = perforceSettings;
+			Settings = settings;
+			_logger = serviceProvider.GetRequiredService<ILogger<ToolUpdateMonitor>>();
+			_asyncDisposer = serviceProvider.GetRequiredService<IAsyncDisposer>();
 
 			DirectoryReference.CreateDirectory(ToolsDir);
 
@@ -185,7 +183,7 @@ namespace UnrealGameSync
 					ToolDefinition? tool = Tools.FirstOrDefault(x => x.ConfigPath.Equals(fileRecord.DepotFile, StringComparison.Ordinal));
 					if (tool == null || tool.ConfigChange != fileRecord.HeadChange)
 					{
-						tool = await ReadToolDefinitionAsync(perforce, fileRecord.DepotFile, fileRecord.HeadChange, logger, cancellationToken);
+						tool = await ReadToolDefinitionAsync(perforce, fileRecord.DepotFile, fileRecord.HeadChange, cancellationToken);
 					}
 					if (tool != null)
 					{
@@ -245,7 +243,7 @@ namespace UnrealGameSync
 			_synchronizationContext.Post(_ => OnChange?.Invoke(), null);
 		}
 
-		async Task<ToolDefinition?> ReadToolDefinitionAsync(IPerforceConnection perforce, string depotPath, int change, ILogger logger, CancellationToken cancellationToken)
+		async Task<ToolDefinition?> ReadToolDefinitionAsync(IPerforceConnection perforce, string depotPath, int change, CancellationToken cancellationToken)
 		{
 			PerforceResponse<PrintRecord<string[]>> response = await perforce.TryPrintLinesAsync($"{depotPath}@{change}", cancellationToken);
 			if (!response.Succeeded || response.Data.Contents == null)
@@ -319,7 +317,7 @@ namespace UnrealGameSync
 			DirectoryReference toolPath = GetToolPathInternal(toolName);
 
 			string commandExe = command;
-			string commandArgs = string.Empty;
+			string commandArgs = String.Empty;
 
 			int spaceIdx = command.IndexOf(' ');
 			if (spaceIdx != -1)
diff --git a/Engine/Source/Programs/UnrealGameSync/UnrealGameSync/UpdateMonitor.cs b/Engine/Source/Programs/UnrealGameSync/UnrealGameSync/UpdateMonitor.cs
--- a/Engine/Source/Programs/UnrealGameSync/UnrealGameSync/UpdateMonitor.cs
+++ b/Engine/Source/Programs/UnrealGameSync/UnrealGameSync/UpdateMonitor.cs
@@ -1,14 +1,9 @@
 // Copyright Epic Games, Inc. All Rights Reserved.
 
-using EpicGames.Core;
 using EpicGames.Perforce;
 using Microsoft.Extensions.DependencyInjection;
 using Microsoft.Extensions.Logging;
 using System;
-using System.Collections.Generic;
-using System.IO;
-using System.Linq;
-using System.Text;
 using System.Threading;
 using System.Threading.Tasks;
 
@@ -23,9 +18,9 @@ namespace UnrealGameSync
 	class UpdateMonitor : IDisposable
 	{
 		Task? _workerTask;
-		CancellationTokenSource _cancellationSource = new CancellationTokenSource();
-		ILogger _logger;
-		IAsyncDisposer _asyncDisposer;
+		readonly CancellationTokenSource _cancellationSource = new CancellationTokenSource();
+		readonly ILogger _logger;
+		readonly IAsyncDisposer _asyncDisposer;
 
 		public Action<UpdateType>? OnUpdateAvailable;
 
@@ -37,8 +32,8 @@ namespace UnrealGameSync
 
 		public UpdateMonitor(IPerforceSettings perforceSettings, string? watchPath, IServiceProvider serviceProvider)
 		{
-			this._logger = serviceProvider.GetRequiredService<ILogger<UpdateMonitor>>();
-			this._asyncDisposer = serviceProvider.GetRequiredService<IAsyncDisposer>();
+			_logger = serviceProvider.GetRequiredService<ILogger<UpdateMonitor>>();
+			_asyncDisposer = serviceProvider.GetRequiredService<IAsyncDisposer>();
 
 			if(watchPath != null)
 			{
@@ -100,7 +95,7 @@ namespace UnrealGameSync
 
 		public void TriggerUpdate(UpdateType updateType, bool? relaunchPreview)
 		{
-			this.RelaunchPreview = relaunchPreview;
+			RelaunchPreview = relaunchPreview;
 			IsUpdateAvailable = true;
 			if(OnUpdateAvailable != null)
 			{
diff --git a/Engine/Source/Programs/UnrealGameSync/UnrealGameSync/UriHandlers/RdpHandler.cs b/Engine/Source/Programs/UnrealGameSync/UnrealGameSync/UriHandlers/RdpHandler.cs
--- a/Engine/Source/Programs/UnrealGameSync/UnrealGameSync/UriHandlers/RdpHandler.cs
+++ b/Engine/Source/Programs/UnrealGameSync/UnrealGameSync/UriHandlers/RdpHandler.cs
@@ -1,21 +1,10 @@
 // Copyright Epic Games, Inc. All Rights Reserved.
 
 using System;
-using System.Collections.Generic;
-using Microsoft.Win32;
-using System.Reflection;
-using System.Threading;
-using System.Collections.Specialized;
-using System.Web;
-using System.Linq;
 using System.Diagnostics;
-using System.Windows.Forms;
-using System.Configuration;
 using System.IO;
 using System.Runtime.InteropServices;
 using System.Runtime.InteropServices.ComTypes;
-using System.Text;
-using System.Security;
 using System.Diagnostics.CodeAnalysis;
 
 namespace UnrealGameSync
@@ -26,6 +15,8 @@ namespace UnrealGameSync
 		const uint CredTypeDomainPassword = 2;
 		const uint CredPersistLocalMachine = 2;
 
+#pragma warning disable IDE0044
+#pragma warning disable IDE1006
 		[StructLayout(LayoutKind.Sequential, CharSet = CharSet.Auto)]
 		struct CredentialAttribute
 		{
@@ -51,6 +42,8 @@ namespace UnrealGameSync
 			public string? TargetAlias;
 			public string? UserName;
 		}
+#pragma warning restore IDE0044
+#pragma warning restore IDE1006
 
 		[DllImport("Advapi32.dll", EntryPoint = "CredReadW", CharSet = CharSet.Unicode, SetLastError = true)]
 		static extern bool CredRead(string target, uint type, int reservedFlag, out IntPtr buffer);
diff --git a/Engine/Source/Programs/UnrealGameSync/UnrealGameSync/UriHandlers/UGSHandler.cs b/Engine/Source/Programs/UnrealGameSync/UnrealGameSync/UriHandlers/UGSHandler.cs
--- a/Engine/Source/Programs/UnrealGameSync/UnrealGameSync/UriHandlers/UGSHandler.cs
+++ b/Engine/Source/Programs/UnrealGameSync/UnrealGameSync/UriHandlers/UGSHandler.cs
@@ -1,19 +1,16 @@
 // Copyright Epic Games, Inc. All Rights Reserved.
 
-using System.Collections.Generic;
+using System;
 using System.IO;
-using System.Text;
 using System.Windows.Forms;
 
 namespace UnrealGameSync
 {
-
 	/// <summary>
 	/// UGS handler 
 	/// </summary>
 	static class UgsHandler
 	{
-
 		[UriHandler]
 		public static UriResult OpenProject(string stream, string project, bool sync = false)
 		{
@@ -34,7 +31,7 @@ namespace UnrealGameSync
 		[UriHandler]
 		public static UriResult BuildStep(string project, string stream, string step, string changelist, string arguments)
 		{
-			MessageBox.Show(string.Format("Project: {0}\nStream: {1}\nStep: {2}\nChange: {3}\nArguments: {4}", project, stream, step, changelist, arguments), "UGS Build Step Handler");
+			MessageBox.Show(String.Format("Project: {0}\nStream: {1}\nStep: {2}\nChange: {3}\nArguments: {4}", project, stream, step, changelist, arguments), "UGS Build Step Handler");
 
 			return new UriResult() { Success = true };
 		}
diff --git a/Engine/Source/Programs/UnrealGameSync/UnrealGameSync/UriHandlers/UriHandler.cs b/Engine/Source/Programs/UnrealGameSync/UnrealGameSync/UriHandlers/UriHandler.cs
--- a/Engine/Source/Programs/UnrealGameSync/UnrealGameSync/UriHandlers/UriHandler.cs
+++ b/Engine/Source/Programs/UnrealGameSync/UnrealGameSync/UriHandlers/UriHandler.cs
@@ -10,7 +10,6 @@ using System.Web;
 using System.Linq;
 using System.Diagnostics;
 using System.Windows.Forms;
-using System.Configuration;
 using System.IO;
 using Microsoft.Extensions.Logging;
 
@@ -41,7 +40,7 @@ namespace UnrealGameSync
 				// Check if this is a registered uri request
 				if (!Handlers.TryGetValue(uri.Host, out MethodInfo? info))
 				{
-					return new UriResult() { Error = string.Format("Unknown Uri {0}", uri.Host) };
+					return new UriResult() { Error = String.Format("Unknown Uri {0}", uri.Host) };
 				}
 
 				NameValueCollection query = HttpUtility.ParseQueryString(uri.Query);
@@ -56,7 +55,7 @@ namespace UnrealGameSync
 					{
 						if (!param.HasDefaultValue)
 						{
-							return new UriResult() { Error = string.Format("Uri {0} is missing required parameter {1}", uri.Host, param.Name) };
+							return new UriResult() { Error = String.Format("Uri {0} is missing required parameter {1}", uri.Host, param.Name) };
 						}
 
 						parameters.Add(param.DefaultValue!);
@@ -79,24 +78,24 @@ namespace UnrealGameSync
 						}
 						else
 						{
-							return new UriResult() { Error = string.Format("Uri {0} bool parameter {1} must be true or false", uri.Host, param.Name) };
+							return new UriResult() { Error = String.Format("Uri {0} bool parameter {1} must be true or false", uri.Host, param.Name) };
 						}
 					}
 					else if (param.ParameterType == typeof(int))
 					{
 						int numberValue;
-						if (!int.TryParse(value, out numberValue))
+						if (!Int32.TryParse(value, out numberValue))
 						{
-							return new UriResult() { Error = string.Format("Uri {0} invalid integer parameter {1} : {2}", uri.Host, param.Name, value) };
+							return new UriResult() { Error = String.Format("Uri {0} invalid integer parameter {1} : {2}", uri.Host, param.Name, value) };
 						}
 						parameters.Add(numberValue);
 					}
 					else if (param.ParameterType == typeof(float))
 					{
 						float numberValue;
-						if (!float.TryParse(value, out numberValue))
+						if (!Single.TryParse(value, out numberValue))
 						{
-							return new UriResult() { Error = string.Format("Uri {0} invalid number parameter {1} : {2}", uri.Host, param.Name, value) };
+							return new UriResult() { Error = String.Format("Uri {0} invalid number parameter {1} : {2}", uri.Host, param.Name, value) };
 						}
 						parameters.Add(numberValue);
 					}
@@ -145,7 +144,7 @@ namespace UnrealGameSync
 			}
 			else
 			{
-				string uriIn = string.Empty;
+				string uriIn = String.Empty;
 				for (int idx = 0; idx < args.Length; idx++)
 				{
 					const string prefix = "-uri=";
@@ -155,7 +154,7 @@ namespace UnrealGameSync
 					}
 				}
 
-				if (uriIn == string.Empty)
+				if (uriIn == String.Empty)
 				{
 					return false;
 				}
@@ -208,7 +207,6 @@ namespace UnrealGameSync
 			}
 		}
 
-
 		static UriHandler()
 		{
 			foreach (Type type in Assembly.GetExecutingAssembly().GetTypes())
@@ -218,17 +216,17 @@ namespace UnrealGameSync
 				{
 					if (!methodInfo.IsStatic)
 					{
-						throw new Exception(string.Format("UriHandler method {0} must be static", methodInfo.Name));
+						throw new Exception(String.Format("UriHandler method {0} must be static", methodInfo.Name));
 					}
 
 					if (methodInfo.ReturnType != typeof(UriResult))
 					{
-						throw new Exception(string.Format("UriHandler method {0} must return UriResult type", methodInfo.Name));
+						throw new Exception(String.Format("UriHandler method {0} must return UriResult type", methodInfo.Name));
 					}
 
 					if (Handlers.ContainsKey(methodInfo.Name))
 					{
-						throw new Exception(string.Format("UriHandler method {0} clashes with another handler", methodInfo.Name));
+						throw new Exception(String.Format("UriHandler method {0} clashes with another handler", methodInfo.Name));
 					}
 
 					foreach (ParameterInfo parameterInfo in methodInfo.GetParameters())
@@ -236,7 +234,7 @@ namespace UnrealGameSync
 						Type parameterType = parameterInfo.ParameterType;
 						if (parameterType != typeof(bool) && parameterType != typeof(string) && parameterType != typeof(float) && parameterType != typeof(int))
 						{
-							throw new Exception(string.Format("UriHandler method parameter {0} must be bool, string, int, or float", parameterInfo.Name));
+							throw new Exception(String.Format("UriHandler method parameter {0} must be bool, string, int, or float", parameterInfo.Name));
 						}
 					}
 
@@ -258,11 +256,11 @@ namespace UnrealGameSync
 	[AttributeUsage(AttributeTargets.Method)]
 	class UriHandlerAttribute : Attribute
 	{
-		public bool Terminate;
+		public bool Terminate { get; set; }
 
 		public UriHandlerAttribute(bool terminate = false)
 		{
-			this.Terminate = terminate;			
+			Terminate = terminate;			
 		}
 	}
 
@@ -288,17 +286,14 @@ namespace UnrealGameSync
 
 			public RegistrySetting(RegistryKey rootKey, string keyName, string? valueName, string value)
 			{
-				this.RootKey = rootKey;
-				this.KeyName = keyName;
-				this.ValueName = valueName;
-				this.Value = value;
+				RootKey = rootKey;
+				KeyName = keyName;
+				ValueName = valueName;
+				Value = value;
 			}
 		}
 
-		public static string CurrentProcessFilePath
-		{
-			get { return Path.ChangeExtension(Path.GetFullPath(Process.GetCurrentProcess().MainModule!.FileName!), ".exe"); }
-		}
+		public static string CurrentProcessFilePath => Path.ChangeExtension(Path.GetFullPath(Process.GetCurrentProcess().MainModule!.FileName!), ".exe");
 
 		static List<RegistrySetting> GetGlobalRegistrySettings()
 		{
@@ -441,6 +436,3 @@ namespace UnrealGameSync
 		}
 	}
 }
-
-
-
diff --git a/Engine/Source/Programs/UnrealGameSync/UnrealGameSync/UriHandlers/VisualStudioHandler.cs b/Engine/Source/Programs/UnrealGameSync/UnrealGameSync/UriHandlers/VisualStudioHandler.cs
--- a/Engine/Source/Programs/UnrealGameSync/UnrealGameSync/UriHandlers/VisualStudioHandler.cs
+++ b/Engine/Source/Programs/UnrealGameSync/UnrealGameSync/UriHandlers/VisualStudioHandler.cs
@@ -1,14 +1,9 @@
 // Copyright Epic Games, Inc. All Rights Reserved.
 
-using System.IO;
-using System.Threading;
-using System.Threading.Tasks;
-using EnvDTE;
 using Microsoft.Extensions.Logging.Abstractions;
 
 namespace UnrealGameSync
 {
-
 	/// <summary>
 	/// VisualStudio Uri Handler
 	/// </summary>
@@ -27,7 +22,5 @@ namespace UnrealGameSync
 
 			return new UriResult() { Success = true };
 		}
-
 	}
-
 }
\ No newline at end of file
diff --git a/Engine/Source/Programs/UnrealGameSync/UnrealGameSync/Workspace.cs b/Engine/Source/Programs/UnrealGameSync/UnrealGameSync/Workspace.cs
--- a/Engine/Source/Programs/UnrealGameSync/UnrealGameSync/Workspace.cs
+++ b/Engine/Source/Programs/UnrealGameSync/UnrealGameSync/Workspace.cs
@@ -5,18 +5,7 @@ using EpicGames.Perforce;
 using Microsoft.Extensions.DependencyInjection;
 using Microsoft.Extensions.Logging;
 using System;
-using System.Collections.Concurrent;
 using System.Collections.Generic;
-using System.Diagnostics;
-using System.Diagnostics.CodeAnalysis;
-using System.Globalization;
-using System.IO;
-using System.IO.Compression;
-using System.Linq;
-using System.Security.Cryptography;
-using System.Text;
-using System.Text.Encodings.Web;
-using System.Text.Json;
 using System.Threading;
 using System.Threading.Tasks;
 
@@ -25,12 +14,12 @@ namespace UnrealGameSync
 	public sealed class Workspace : IDisposable
 	{
 		public IPerforceSettings PerforceSettings { get; }
-		WorkspaceStateWrapper _stateWrapper;
+		readonly WorkspaceStateWrapper _stateWrapper;
 		public ReadOnlyWorkspaceState State { get; private set; }
 		public ProjectInfo Project { get; }
 		public SynchronizationContext SynchronizationContext { get; }
 		public WorkspaceLock Lock { get; }
-		ILogger _logger;
+		readonly ILogger _logger;
 
 		bool Syncing => _currentUpdate != null;
 
@@ -40,7 +29,7 @@ namespace UnrealGameSync
 
 		public event Action<ReadOnlyWorkspaceState>? OnStateChanged;
 
-		IAsyncDisposer _asyncDisposer;
+		readonly IAsyncDisposer _asyncDisposer;
 
 		public Workspace(IPerforceSettings perforceSettings, ProjectInfo project, WorkspaceStateWrapper stateWrapper, ConfigFile projectConfigFile, IReadOnlyList<string>? projectStreamFilter, ILogger logger, IServiceProvider serviceProvider)
 		{
@@ -53,12 +42,12 @@ namespace UnrealGameSync
 			Lock = new WorkspaceLock(project.LocalRootPath);
 			Lock.OnChange += OnLockChangedInternal;
 
-			this.SynchronizationContext = SynchronizationContext.Current!;
-			this._logger = logger;
-			this._asyncDisposer = serviceProvider.GetRequiredService<IAsyncDisposer>();
+			SynchronizationContext = SynchronizationContext.Current!;
+			_logger = logger;
+			_asyncDisposer = serviceProvider.GetRequiredService<IAsyncDisposer>();
 
-			this.ProjectConfigFile = projectConfigFile;
-			this.ProjectStreamFilter = projectStreamFilter;
+			ProjectConfigFile = projectConfigFile;
+			ProjectStreamFilter = projectStreamFilter;
 		}
 
 		public bool IsExternalSyncActive() => Lock.IsLockedByOtherProcess();
@@ -247,10 +236,7 @@ namespace UnrealGameSync
 
 		public int PendingChangeNumber => _currentUpdate?.Context?.ChangeNumber ?? CurrentChangeNumber;
 
-		public string ClientName
-		{
-			get { return PerforceSettings.ClientName!; }
-		}
+		public string ClientName => PerforceSettings.ClientName!;
 
 		public Tuple<string, float> CurrentProgress => _currentUpdate?.CurrentProgress ?? new Tuple<string, float>("", 0.0f);
 	}
diff --git a/Engine/Source/Programs/UnrealGameSync/UnrealGameSyncCmd/Program.cs b/Engine/Source/Programs/UnrealGameSync/UnrealGameSyncCmd/Program.cs
--- a/Engine/Source/Programs/UnrealGameSync/UnrealGameSyncCmd/Program.cs
+++ b/Engine/Source/Programs/UnrealGameSync/UnrealGameSyncCmd/Program.cs
@@ -7,13 +7,11 @@ using Serilog;
 using System;
 using System.Collections.Generic;
 using System.Diagnostics;
-using System.Globalization;
 using System.IO;
 using System.Linq;
 using System.Net;
 using System.Reflection;
 using System.Runtime.InteropServices;
-using System.Text;
 using System.Text.RegularExpressions;
 using System.Threading;
 using System.Threading.Tasks;
@@ -33,8 +31,8 @@ namespace UnrealGameSyncCmd
 		public UserErrorException(LogEvent evt)
 			: base(evt.ToString())
 		{
-			this.Event = evt;
-			this.Code = 1;
+			Event = evt;
+			Code = 1;
 		}
 
 		public UserErrorException(string message, params object[] args)
@@ -56,14 +54,14 @@ namespace UnrealGameSyncCmd
 
 			public CommandInfo(string name, Type type, string usage, string brief)
 			{
-				this.Name = name;
-				this.Type = type;
-				this.Usage = usage;
-				this.Brief = brief;
+				Name = name;
+				Type = type;
+				Usage = usage;
+				Brief = brief;
 			}
 		}
 
-		static CommandInfo[] _commands =
+		static readonly CommandInfo[] _commands =
 		{
 			new CommandInfo("init", typeof(InitCommand),
 				"ugs init [stream-path] [-client=..] [-server=..] [-user=..] [-branch=..] [-project=..]",
@@ -124,10 +122,10 @@ namespace UnrealGameSyncCmd
 
 			public CommandContext(CommandLineArguments arguments, ILogger logger, ILoggerFactory loggerFactory, GlobalSettingsFile userSettings)
 			{
-				this.Arguments = arguments;
-				this.Logger = logger;
-				this.LoggerFactory = loggerFactory;
-				this.UserSettings = userSettings;
+				Arguments = arguments;
+				Logger = logger;
+				LoggerFactory = loggerFactory;
+				UserSettings = userSettings;
 			}
 		}
 
@@ -301,7 +299,7 @@ namespace UnrealGameSyncCmd
 
 		public static async Task<WorkspaceStateWrapper> ReadWorkspaceState(IPerforceConnection perforceClient, UserWorkspaceSettings settings, GlobalSettingsFile userSettings, ILogger logger)
 		{
-			WorkspaceStateWrapper state = userSettings.FindOrAddWorkspaceState(settings, logger);
+			WorkspaceStateWrapper state = userSettings.FindOrAddWorkspaceState(settings);
 			if (state.Current.SettingsTimeUtc != settings.LastModifiedTimeUtc)
 			{
 				logger.LogDebug("Updating state due to modified settings timestamp");
@@ -407,15 +405,15 @@ namespace UnrealGameSyncCmd
 				// Create the perforce connection
 				if (initName != null)
 				{
-					await InitNewClientAsync(perforce, context, initName, hostName, options, logger);
+					await InitNewClientAsync(perforce, initName, hostName, options, logger);
 				}
 				else
 				{
-					await InitExistingClientAsync(perforce, context, hostName, options, logger);
+					await InitExistingClientAsync(perforce, hostName, options, logger);
 				}
 			}
 
-			async Task InitNewClientAsync(IPerforceConnection perforce, CommandContext context, string streamName, string hostName, ProjectInitOptions options, ILogger logger)
+			async Task InitNewClientAsync(IPerforceConnection perforce, string streamName, string hostName, ProjectInitOptions options, ILogger logger)
 			{
 				logger.LogInformation("Checking stream...");
 
@@ -481,7 +479,7 @@ namespace UnrealGameSyncCmd
 				}
 			}
 
-			async Task InitExistingClientAsync(IPerforceConnection perforce, CommandContext context, string hostName, ProjectInitOptions options, ILogger logger)
+			async Task InitExistingClientAsync(IPerforceConnection perforce, string hostName, ProjectInitOptions options, ILogger logger)
 			{
 				DirectoryReference currentDir = DirectoryReference.GetCurrentDirectory();
 
@@ -624,7 +622,7 @@ namespace UnrealGameSyncCmd
 				bool syncLatest = String.Equals(changeString, "latest", StringComparison.OrdinalIgnoreCase);
 
 				int change;
-				if (!int.TryParse(changeString, out change))
+				if (!Int32.TryParse(changeString, out change))
 				{
 					if (syncLatest)
 					{
diff --git a/Engine/Source/Programs/UnrealGameSync/UnrealGameSyncLauncher/Forms/LogWindow.cs b/Engine/Source/Programs/UnrealGameSync/UnrealGameSyncLauncher/Forms/LogWindow.cs
--- a/Engine/Source/Programs/UnrealGameSync/UnrealGameSyncLauncher/Forms/LogWindow.cs
+++ b/Engine/Source/Programs/UnrealGameSync/UnrealGameSyncLauncher/Forms/LogWindow.cs
@@ -1,13 +1,5 @@
 // Copyright Epic Games, Inc. All Rights Reserved.
 
-using System;
-using System.Collections.Generic;
-using System.ComponentModel;
-using System.Data;
-using System.Drawing;
-using System.Linq;
-using System.Text;
-using System.Threading.Tasks;
 using System.Windows.Forms;
 
 namespace UnrealGameSyncLauncher
diff --git a/Engine/Source/Programs/UnrealGameSync/UnrealGameSyncLauncher/Forms/SettingsWindow.cs b/Engine/Source/Programs/UnrealGameSync/UnrealGameSyncLauncher/Forms/SettingsWindow.cs
--- a/Engine/Source/Programs/UnrealGameSync/UnrealGameSyncLauncher/Forms/SettingsWindow.cs
+++ b/Engine/Source/Programs/UnrealGameSync/UnrealGameSyncLauncher/Forms/SettingsWindow.cs
@@ -4,14 +4,7 @@ using EpicGames.Core;
 using EpicGames.Perforce;
 using Microsoft.Extensions.Logging;
 using System;
-using System.Collections.Generic;
-using System.ComponentModel;
-using System.Data;
-using System.Drawing;
-using System.IO;
-using System.Linq;
 using System.Runtime.InteropServices;
-using System.Text;
 using System.Threading;
 using System.Threading.Tasks;
 using System.Windows.Forms;
@@ -29,7 +22,7 @@ namespace UnrealGameSyncLauncher
 		const int EmSetcuebanner = 0x1501;
 
 		string? _logText;
-		SyncAndRunDelegate _syncAndRun;
+		readonly SyncAndRunDelegate _syncAndRun;
 
 		public SettingsWindow(string? prompt, string? logText, string? serverAndPort, string? userName, string? depotPath, bool preview, SyncAndRunDelegate syncAndRun)
 		{
@@ -37,15 +30,15 @@ namespace UnrealGameSyncLauncher
 
 			if(prompt != null)
 			{
-				this.PromptLabel.Text = prompt;
+				PromptLabel.Text = prompt;
 			}
 
-			this._logText = logText;
-			this.ServerTextBox.Text = serverAndPort ?? String.Empty;
-			this.UserNameTextBox.Text = userName ?? String.Empty;
-			this.DepotPathTextBox.Text = depotPath ?? String.Empty;
-			this.UsePreviewBuildCheckBox.Checked = preview;
-			this._syncAndRun = syncAndRun;
+			_logText = logText;
+			ServerTextBox.Text = serverAndPort ?? String.Empty;
+			UserNameTextBox.Text = userName ?? String.Empty;
+			DepotPathTextBox.Text = depotPath ?? String.Empty;
+			UsePreviewBuildCheckBox.Checked = preview;
+			_syncAndRun = syncAndRun;
 
 			ViewLogBtn.Visible = logText != null;
 		}
diff --git a/Engine/Source/Programs/UnrealGameSync/UnrealGameSyncLauncher/Program.cs b/Engine/Source/Programs/UnrealGameSync/UnrealGameSyncLauncher/Program.cs
--- a/Engine/Source/Programs/UnrealGameSync/UnrealGameSyncLauncher/Program.cs
+++ b/Engine/Source/Programs/UnrealGameSync/UnrealGameSyncLauncher/Program.cs
@@ -3,7 +3,6 @@
 using EpicGames.Core;
 using EpicGames.Perforce;
 using Microsoft.Extensions.Logging;
-using Microsoft.Win32;
 using System;
 using System.Collections.Generic;
 using System.Diagnostics;
@@ -50,11 +49,11 @@ namespace UnrealGameSyncLauncher
 				GlobalPerforceSettings.ReadGlobalPerforceSettings(ref serverAndPort, ref userName, ref depotPath, ref preview);
 
 				// If the shift key is held down, immediately show the settings window
-				SettingsWindow.SyncAndRunDelegate syncAndRunWrapper = (perforce, depotParam, previewParam, logWriter, cancellationToken) => SyncAndRun(perforce, depotParam, previewParam, args, instanceMutex, logWriter, cancellationToken);
+				Task SyncAndRunWrapper(IPerforceConnection perforce, string? depotPath, bool preview, ILogger logWriter, CancellationToken cancellationToken) => SyncAndRun(perforce, depotPath, preview, args, instanceMutex, logWriter, cancellationToken);
 				if ((Control.ModifierKeys & Keys.Shift) != 0)
 				{
 					// Show the settings window immediately
-					SettingsWindow updateError = new SettingsWindow(null, null, serverAndPort, userName, depotPath, preview, syncAndRunWrapper);
+					SettingsWindow updateError = new SettingsWindow(null, null, serverAndPort, userName, depotPath, preview, SyncAndRunWrapper);
 					if(updateError.ShowDialog() == DialogResult.OK)
 					{
 						return 0;
@@ -77,7 +76,7 @@ namespace UnrealGameSyncLauncher
 						return 0;
 					}
 
-					SettingsWindow updateError = new SettingsWindow("Unable to update UnrealGameSync from Perforce. Verify that your connection settings are correct.", logger.Render(Environment.NewLine), serverAndPort, userName, depotPath, preview, syncAndRunWrapper);
+					SettingsWindow updateError = new SettingsWindow("Unable to update UnrealGameSync from Perforce. Verify that your connection settings are correct.", logger.Render(Environment.NewLine), serverAndPort, userName, depotPath, preview, SyncAndRunWrapper);
 					if(updateError.ShowDialog() == DialogResult.OK)
 					{
 						return 0;
@@ -310,7 +309,7 @@ namespace UnrealGameSyncLauncher
 				DirectoryInfo directory = new DirectoryInfo(directoryName);
 				foreach(FileInfo childFile in directory.EnumerateFiles("*", SearchOption.AllDirectories))
 				{
-					childFile.Attributes = childFile.Attributes & ~FileAttributes.ReadOnly;
+					childFile.Attributes &= ~FileAttributes.ReadOnly;
 					childFile.Delete();
 				}
 				foreach(DirectoryInfo childDirectory in directory.EnumerateDirectories())
diff --git a/Engine/Source/Programs/UnrealGameSync/UnrealGameSyncShared/ArchiveInfo.cs b/Engine/Source/Programs/UnrealGameSync/UnrealGameSyncShared/ArchiveInfo.cs
--- a/Engine/Source/Programs/UnrealGameSync/UnrealGameSyncShared/ArchiveInfo.cs
+++ b/Engine/Source/Programs/UnrealGameSync/UnrealGameSyncShared/ArchiveInfo.cs
@@ -2,14 +2,12 @@
 
 using EpicGames.Core;
 using EpicGames.Perforce;
-using JetBrains.Annotations;
 using Microsoft.Extensions.Logging;
 using System;
 using System.Collections.Generic;
 using System.Diagnostics.CodeAnalysis;
 using System.IO;
 using System.Linq;
-using System.Text;
 using System.Threading;
 using System.Threading.Tasks;
 
@@ -36,20 +34,17 @@ namespace UnrealGameSync
 		public string DepotPath { get; }
 		public string? Target { get; }
 
-		public string BasePath
-		{
-			get { return DepotPath; }
-		}
+		public string BasePath => DepotPath;
 
 		// TODO: executable/configuration?
 		public SortedList<int, string> ChangeNumberToFileRevision = new SortedList<int, string>();
 
 		public PerforceArchiveInfo(string name, string type, string depotPath, string? target)
 		{
-			this.Name = name;
-			this.Type = type;
-			this.DepotPath = depotPath;
-			this.Target = target;
+			Name = name;
+			Type = type;
+			DepotPath = depotPath;
+			Target = target;
 		}
 
 		public override bool Equals(object? other)
@@ -188,7 +183,7 @@ namespace UnrealGameSync
 									if (tokens[0].StartsWith("[CL") && tokens[1].EndsWith("]"))
 									{
 										int originalChangeNumber;
-										if (int.TryParse(tokens[1].Substring(0, tokens[1].Length - 1), out originalChangeNumber) && !newArchive.ChangeNumberToFileRevision.ContainsKey(originalChangeNumber))
+										if (Int32.TryParse(tokens[1].Substring(0, tokens[1].Length - 1), out originalChangeNumber) && !newArchive.ChangeNumberToFileRevision.ContainsKey(originalChangeNumber))
 										{
 											newArchive.ChangeNumberToFileRevision[originalChangeNumber] = $"{newArchive.DepotPath}#{revision.RevisionNumber}";
 										}
diff --git a/Engine/Source/Programs/UnrealGameSync/UnrealGameSyncShared/ArchiveManifest.cs b/Engine/Source/Programs/UnrealGameSync/UnrealGameSyncShared/ArchiveManifest.cs
--- a/Engine/Source/Programs/UnrealGameSync/UnrealGameSyncShared/ArchiveManifest.cs
+++ b/Engine/Source/Programs/UnrealGameSync/UnrealGameSyncShared/ArchiveManifest.cs
@@ -7,9 +7,6 @@ using System.Collections.Generic;
 using System.Diagnostics;
 using System.IO;
 using System.IO.Compression;
-using System.Linq;
-using System.Text;
-using System.Threading.Tasks;
 
 namespace UnrealGameSync
 {
diff --git a/Engine/Source/Programs/UnrealGameSync/UnrealGameSyncShared/BuildStep.cs b/Engine/Source/Programs/UnrealGameSync/UnrealGameSyncShared/BuildStep.cs
--- a/Engine/Source/Programs/UnrealGameSync/UnrealGameSyncShared/BuildStep.cs
+++ b/Engine/Source/Programs/UnrealGameSync/UnrealGameSyncShared/BuildStep.cs
@@ -3,9 +3,6 @@
 using System;
 using System.Collections.Generic;
 using System.Diagnostics;
-using System.Linq;
-using System.Text;
-using System.Threading.Tasks;
 
 namespace UnrealGameSync
 {
@@ -85,7 +82,7 @@ namespace UnrealGameSync
 			Description = obj.GetValue("Description", "Untitled");
 			StatusText = obj.GetValue("StatusText", "Untitled");
 
-			if(!int.TryParse(obj.GetValue("EstimatedDuration", ""), out EstimatedDuration) || EstimatedDuration < 1)
+			if(!Int32.TryParse(obj.GetValue("EstimatedDuration", ""), out EstimatedDuration) || EstimatedDuration < 1)
 			{
 				EstimatedDuration = 1;
 			}
@@ -133,7 +130,7 @@ namespace UnrealGameSync
 					requires.Add(require);
 				}
 			}
-			this.Requires = requires.ToArray();
+			Requires = requires.ToArray();
 	
 			if (!Guid.TryParse(obj.GetValue("Tool", ""), out ToolId))
 			{
diff --git a/Engine/Source/Programs/UnrealGameSync/UnrealGameSyncShared/ConfigFile.cs b/Engine/Source/Programs/UnrealGameSync/UnrealGameSyncShared/ConfigFile.cs
--- a/Engine/Source/Programs/UnrealGameSync/UnrealGameSyncShared/ConfigFile.cs
+++ b/Engine/Source/Programs/UnrealGameSync/UnrealGameSyncShared/ConfigFile.cs
@@ -11,7 +11,6 @@ using System.Linq;
 using System.Text;
 using System.Text.Json;
 using System.Text.Json.Serialization;
-using System.Threading.Tasks;
 
 namespace UnrealGameSync
 {
@@ -182,7 +181,7 @@ namespace UnrealGameSync
 			if(stringValue != null)
 			{
 				int value;
-				if(int.TryParse(stringValue, out value))
+				if(Int32.TryParse(stringValue, out value))
 				{
 					return value;
 				}
@@ -196,7 +195,7 @@ namespace UnrealGameSync
 			if(stringValue != null)
 			{
 				bool value;
-				if(bool.TryParse(stringValue, out value))
+				if(Boolean.TryParse(stringValue, out value))
 				{
 					return value;
 				}
@@ -241,8 +240,8 @@ namespace UnrealGameSync
 
 		public string? this[string key]
 		{
-			get { return GetValue(key); }
-			set { SetValue(key, value); }
+			get => GetValue(key);
+			set => SetValue(key, value);
 		}
 
 		public void SetDefaults(ConfigObject other)
@@ -427,7 +426,7 @@ namespace UnrealGameSync
 			if(valueString != null)
 			{
 				int value;
-				if(int.TryParse(valueString, out value))
+				if(Int32.TryParse(valueString, out value))
 				{
 					return value;
 				}
@@ -441,7 +440,7 @@ namespace UnrealGameSync
 			if (valueString != null)
 			{
 				long value;
-				if (long.TryParse(valueString, out value))
+				if (Int64.TryParse(valueString, out value))
 				{
 					return value;
 				}
@@ -504,7 +503,7 @@ namespace UnrealGameSync
 
 	public class ConfigFile
 	{
-		List<ConfigSection> _sections = new List<ConfigSection>();
+		readonly List<ConfigSection> _sections = new List<ConfigSection>();
 
 		public ConfigFile()
 		{
diff --git a/Engine/Source/Programs/UnrealGameSync/UnrealGameSyncShared/ConfigUtils.cs b/Engine/Source/Programs/UnrealGameSync/UnrealGameSyncShared/ConfigUtils.cs
--- a/Engine/Source/Programs/UnrealGameSync/UnrealGameSyncShared/ConfigUtils.cs
+++ b/Engine/Source/Programs/UnrealGameSync/UnrealGameSyncShared/ConfigUtils.cs
@@ -9,8 +9,6 @@ using System.Diagnostics.CodeAnalysis;
 using System.IO;
 using System.Linq;
 using System.Runtime.InteropServices;
-using System.Text;
-using System.Text.Json;
 using System.Threading;
 using System.Threading.Tasks;
 
@@ -249,7 +247,7 @@ namespace UnrealGameSync
 			if (TryGetProjectSetting(projectConfig, projectInfo.ProjectIdentifier, "UseSharedEditor", out setting))
 			{
 				bool value;
-				if (bool.TryParse(setting, out value))
+				if (Boolean.TryParse(setting, out value))
 				{
 					return value;
 				}
diff --git a/Engine/Source/Programs/UnrealGameSync/UnrealGameSyncShared/FileFilter.cs b/Engine/Source/Programs/UnrealGameSync/UnrealGameSyncShared/FileFilter.cs
--- a/Engine/Source/Programs/UnrealGameSync/UnrealGameSyncShared/FileFilter.cs
+++ b/Engine/Source/Programs/UnrealGameSync/UnrealGameSyncShared/FileFilter.cs
@@ -4,8 +4,6 @@ using System;
 using System.Collections.Generic;
 using System.IO;
 using System.Linq;
-using System.Text;
-using System.Threading.Tasks;
 
 namespace UnrealGameSync
 {
@@ -28,17 +26,17 @@ namespace UnrealGameSync
 		/// <summary>
 		/// Root node for the tree.
 		/// </summary>
-		FileFilterNode _rootNode;
+		readonly FileFilterNode _rootNode;
 
 		/// <summary>
 		/// The default node, which will match any path
 		/// </summary>
-		FileFilterNode _defaultNode;
+		readonly FileFilterNode _defaultNode;
 
 		/// <summary>
 		/// Terminating nodes for each rule added to the filter
 		/// </summary>
-		List<FileFilterNode> _rules = new List<FileFilterNode>();
+		readonly List<FileFilterNode> _rules = new List<FileFilterNode>();
 
 		/// <summary>
 		/// Default constructor
diff --git a/Engine/Source/Programs/UnrealGameSync/UnrealGameSyncShared/GlobalSettings.cs b/Engine/Source/Programs/UnrealGameSync/UnrealGameSyncShared/GlobalSettings.cs
--- a/Engine/Source/Programs/UnrealGameSync/UnrealGameSyncShared/GlobalSettings.cs
+++ b/Engine/Source/Programs/UnrealGameSync/UnrealGameSyncShared/GlobalSettings.cs
@@ -5,11 +5,6 @@ using Microsoft.Extensions.Logging;
 using System;
 using System.Collections.Generic;
 using System.Linq;
-using System.Text;
-using System.Text.Json;
-using System.Text.Json.Serialization;
-using System.Threading;
-using System.Threading.Tasks;
 
 namespace UnrealGameSync
 {
@@ -100,8 +95,8 @@ namespace UnrealGameSync
 
 		public GlobalSettingsFile(FileReference file, GlobalSettings global)
 		{
-			this.File = file;
-			this.Global = global;
+			File = file;
+			Global = global;
 		}
 
 		public static GlobalSettingsFile Create(FileReference file)
@@ -128,12 +123,12 @@ namespace UnrealGameSync
 			}
 		}
 
-		public WorkspaceStateWrapper FindOrAddWorkspaceState(UserWorkspaceSettings settings, ILogger logger)
+		public WorkspaceStateWrapper FindOrAddWorkspaceState(UserWorkspaceSettings settings)
 		{
-			return FindOrAddWorkspaceState(settings.RootDir, settings.ClientName, settings.BranchPath, logger);
+			return FindOrAddWorkspaceState(settings.RootDir, settings.ClientName, settings.BranchPath);
 		}
 
-		public WorkspaceStateWrapper FindOrAddWorkspaceState(DirectoryReference rootDir, string clientName, string branchPath, ILogger logger)
+		public WorkspaceStateWrapper FindOrAddWorkspaceState(DirectoryReference rootDir, string clientName, string branchPath)
 		{
 			return new WorkspaceStateWrapper(rootDir, () =>
 			{
@@ -143,9 +138,9 @@ namespace UnrealGameSync
 			});
 		}
 
-		public WorkspaceStateWrapper FindOrAddWorkspaceState(ProjectInfo projectInfo, UserWorkspaceSettings settings, ILogger logger)
+		public WorkspaceStateWrapper FindOrAddWorkspaceState(ProjectInfo projectInfo, UserWorkspaceSettings settings)
 		{
-			WorkspaceStateWrapper wrapper = FindOrAddWorkspaceState(projectInfo.LocalRootPath, projectInfo.ClientName, projectInfo.BranchPath, logger);
+			WorkspaceStateWrapper wrapper = FindOrAddWorkspaceState(projectInfo.LocalRootPath, projectInfo.ClientName, projectInfo.BranchPath);
 			wrapper.Modify(x => x.UpdateCachedProjectInfo(projectInfo, settings.LastModifiedTimeUtc));
 			return wrapper;
 		}
diff --git a/Engine/Source/Programs/UnrealGameSync/UnrealGameSyncShared/OpenProjectInfo.cs b/Engine/Source/Programs/UnrealGameSync/UnrealGameSyncShared/OpenProjectInfo.cs
--- a/Engine/Source/Programs/UnrealGameSync/UnrealGameSyncShared/OpenProjectInfo.cs
+++ b/Engine/Source/Programs/UnrealGameSync/UnrealGameSyncShared/OpenProjectInfo.cs
@@ -5,10 +5,8 @@ using EpicGames.Perforce;
 using Microsoft.Extensions.Logging;
 using System;
 using System.Collections.Generic;
-using System.Drawing;
 using System.IO;
 using System.Linq;
-using System.Text.RegularExpressions;
 using System.Threading;
 using System.Threading.Tasks;
 
@@ -29,16 +27,16 @@ namespace UnrealGameSync
 
 		public OpenProjectInfo(UserSelectedProjectSettings selectedProject, IPerforceSettings perforceSettings, ProjectInfo projectInfo, UserWorkspaceSettings workspaceSettings, WorkspaceStateWrapper workspaceStateWrapper, ConfigFile latestProjectConfigFile, ConfigFile workspaceProjectConfigFile, IReadOnlyList<string>? workspaceProjectStreamFilter, List<KeyValuePair<FileReference, DateTime>> localConfigFiles)
 		{
-			this.SelectedProject = selectedProject;
-
-			this.PerforceSettings = perforceSettings;
-			this.ProjectInfo = projectInfo;
-			this.WorkspaceSettings = workspaceSettings;
-			this.WorkspaceStateWrapper = workspaceStateWrapper;
-			this.LatestProjectConfigFile = latestProjectConfigFile;
-			this.WorkspaceProjectConfigFile = workspaceProjectConfigFile;
-			this.WorkspaceProjectStreamFilter = workspaceProjectStreamFilter;
-			this.LocalConfigFiles = localConfigFiles;
+			SelectedProject = selectedProject;
+
+			PerforceSettings = perforceSettings;
+			ProjectInfo = projectInfo;
+			WorkspaceSettings = workspaceSettings;
+			WorkspaceStateWrapper = workspaceStateWrapper;
+			LatestProjectConfigFile = latestProjectConfigFile;
+			WorkspaceProjectConfigFile = workspaceProjectConfigFile;
+			WorkspaceProjectStreamFilter = workspaceProjectStreamFilter;
+			LocalConfigFiles = localConfigFiles;
 		}
 
 		public static async Task<OpenProjectInfo> CreateAsync(IPerforceSettings defaultPerforceSettings, UserSelectedProjectSettings selectedProject, UserSettings userSettings, ILogger<OpenProjectInfo> logger, CancellationToken cancellationToken)
@@ -218,7 +216,7 @@ namespace UnrealGameSync
 				ProjectInfo projectInfo = await ProjectInfo.CreateAsync(perforceClient, userWorkspaceSettings, cancellationToken);
 
 				// Update the cached workspace state
-				WorkspaceStateWrapper workspaceStateWrapper = userSettings.FindOrAddWorkspaceState(projectInfo, userWorkspaceSettings, logger);
+				WorkspaceStateWrapper workspaceStateWrapper = userSettings.FindOrAddWorkspaceState(projectInfo, userWorkspaceSettings);
 
 				// Read the initial config file
 				List<KeyValuePair<FileReference, DateTime>> localConfigFiles = new List<KeyValuePair<FileReference, DateTime>>();
@@ -226,7 +224,7 @@ namespace UnrealGameSync
 
 				// Get the local config file and stream filter
 				ConfigFile workspaceProjectConfigFile = await WorkspaceUpdate.ReadProjectConfigFile(branchDirectoryName, newSelectedFileName, logger);
-				IReadOnlyList<string>? workspaceProjectStreamFilter = await WorkspaceUpdate.ReadProjectStreamFilter(perforceClient, workspaceProjectConfigFile, logger, cancellationToken);
+				IReadOnlyList<string>? workspaceProjectStreamFilter = await WorkspaceUpdate.ReadProjectStreamFilter(perforceClient, workspaceProjectConfigFile, cancellationToken);
 
 				OpenProjectInfo workspaceSettings = new OpenProjectInfo(selectedProject, perforceSettings, projectInfo, userWorkspaceSettings, workspaceStateWrapper, latestProjectConfigFile, workspaceProjectConfigFile, workspaceProjectStreamFilter, localConfigFiles);
 
diff --git a/Engine/Source/Programs/UnrealGameSync/UnrealGameSyncShared/OutputAdapters.cs b/Engine/Source/Programs/UnrealGameSync/UnrealGameSyncShared/OutputAdapters.cs
--- a/Engine/Source/Programs/UnrealGameSync/UnrealGameSyncShared/OutputAdapters.cs
+++ b/Engine/Source/Programs/UnrealGameSync/UnrealGameSyncShared/OutputAdapters.cs
@@ -1,13 +1,9 @@
 // Copyright Epic Games, Inc. All Rights Reserved.
 
-using EpicGames.Core;
 using Microsoft.Extensions.Logging;
 using System;
 using System.Collections.Generic;
-using System.IO;
 using System.Linq;
-using System.Text;
-using System.Threading.Tasks;
 
 #nullable enable
 
@@ -15,8 +11,8 @@ namespace UnrealGameSync
 {
 	class PrefixedTextWriter : ILogger
 	{
-		string _prefix;
-		ILogger _inner;
+		readonly string _prefix;
+		readonly ILogger _inner;
 
 		public PrefixedTextWriter(string inPrefix, ILogger inInner)
 		{
@@ -37,7 +33,7 @@ namespace UnrealGameSync
 	public class ProgressValue
 	{
 		Tuple<string, float> _state = null!;
-		Stack<Tuple<float, float>> _ranges = new Stack<Tuple<float,float>>();
+		readonly Stack<Tuple<float, float>> _ranges = new Stack<Tuple<float,float>>();
 
 		public ProgressValue()
 		{
@@ -52,10 +48,7 @@ namespace UnrealGameSync
 			_ranges.Push(new Tuple<float, float>(0.0f, 1.0f));
 		}
 
-		public Tuple<string, float> Current
-		{
-			get { return _state; }
-		}
+		public Tuple<string, float> Current => _state;
 
 		public void Set(string message)
 		{
@@ -251,7 +244,7 @@ namespace UnrealGameSync
 			if(tokenIdx + 2 <= tokens.Count && tokens[tokenIdx + 1] == "%")
 			{
 				int numerator;
-				if(int.TryParse(tokens[tokenIdx], out numerator))
+				if(Int32.TryParse(tokens[tokenIdx], out numerator))
 				{
 					fraction = (float)numerator / 100.0f;
 					tokenIdx += 2;
@@ -263,7 +256,7 @@ namespace UnrealGameSync
 			if(tokenIdx + 3 <= tokens.Count && tokens[tokenIdx + 1] == "/")
 			{
 				int numerator, denominator;
-				if(int.TryParse(tokens[tokenIdx], out numerator) && int.TryParse(tokens[tokenIdx + 2], out denominator))
+				if(Int32.TryParse(tokens[tokenIdx], out numerator) && Int32.TryParse(tokens[tokenIdx + 2], out denominator))
 				{
 					fraction = (float)numerator / (float)denominator;
 					tokenIdx += 3;
diff --git a/Engine/Source/Programs/UnrealGameSync/UnrealGameSyncShared/Telemetry.cs b/Engine/Source/Programs/UnrealGameSync/UnrealGameSyncShared/Telemetry.cs
--- a/Engine/Source/Programs/UnrealGameSync/UnrealGameSyncShared/Telemetry.cs
+++ b/Engine/Source/Programs/UnrealGameSync/UnrealGameSyncShared/Telemetry.cs
@@ -3,16 +3,12 @@
 using Microsoft.Extensions.Logging;
 using System;
 using System.Collections.Generic;
-using System.IO;
 using System.Linq;
-using System.Net;
 using System.Net.Http;
 using System.Net.Http.Headers;
 using System.Reflection;
-using System.Runtime.CompilerServices;
 using System.Text;
 using System.Text.Json;
-using System.Threading;
 using System.Threading.Channels;
 using System.Threading.Tasks;
 using System.Web;
@@ -64,8 +60,8 @@ namespace UnrealGameSync
 		/// </summary>
 		public EpicTelemetrySink(string url, ILogger logger)
 		{
-			this._url = url;
-			this._logger = logger;
+			_url = url;
+			_logger = logger;
 
 			logger.LogInformation("Posting to URL: {Url}", url);
 
diff --git a/Engine/Source/Programs/UnrealGameSync/UnrealGameSyncShared/TelemetryStopwatch.cs b/Engine/Source/Programs/UnrealGameSync/UnrealGameSyncShared/TelemetryStopwatch.cs
--- a/Engine/Source/Programs/UnrealGameSync/UnrealGameSyncShared/TelemetryStopwatch.cs
+++ b/Engine/Source/Programs/UnrealGameSync/UnrealGameSyncShared/TelemetryStopwatch.cs
@@ -15,7 +15,7 @@ namespace UnrealGameSync
 
 		public TelemetryStopwatch(string eventName, string project)
 		{
-			this._eventName = eventName;
+			_eventName = eventName;
 
 			_eventData = new Dictionary<string, object?>();
 			_eventData["Project"] = project;
@@ -52,9 +52,6 @@ namespace UnrealGameSync
 			Telemetry.SendEvent(_eventName, _eventData);
 		}
 
-		public TimeSpan Elapsed
-		{
-			get { return _timer.Elapsed; }
-		}
+		public TimeSpan Elapsed => _timer.Elapsed;
 	}
 }
diff --git a/Engine/Source/Programs/UnrealGameSync/UnrealGameSyncShared/UserSettings.cs b/Engine/Source/Programs/UnrealGameSync/UnrealGameSyncShared/UserSettings.cs
--- a/Engine/Source/Programs/UnrealGameSync/UnrealGameSyncShared/UserSettings.cs
+++ b/Engine/Source/Programs/UnrealGameSync/UnrealGameSyncShared/UserSettings.cs
@@ -11,7 +11,6 @@ using System.Linq;
 using System.Runtime.InteropServices;
 using System.Text;
 using System.Text.Json.Serialization;
-using System.Text.RegularExpressions;
 
 namespace UnrealGameSync
 {
@@ -150,9 +149,9 @@ namespace UnrealGameSync
 
 		public ArchiveSettings(bool enabled, string type, IEnumerable<string> order)
 		{
-			this.Enabled = enabled;
-			this.Type = type;
-			this.Order = new List<string>(order);
+			Enabled = enabled;
+			Type = type;
+			Order = new List<string>(order);
 		}
 
 		public static bool TryParseConfigEntry(string text, [NotNullWhen(true)] out ArchiveSettings? settings)
@@ -203,11 +202,11 @@ namespace UnrealGameSync
 
 		public UserSelectedProjectSettings(string? serverAndPort, string? userName, UserSelectedProjectType type, string? clientPath, string? localPath)
 		{
-			this.ServerAndPort = serverAndPort;
-			this.UserName = userName;
-			this.Type = type;
-			this.ClientPath = clientPath;
-			this.LocalPath = localPath;
+			ServerAndPort = serverAndPort;
+			UserName = userName;
+			Type = type;
+			ClientPath = clientPath;
+			LocalPath = localPath;
 		}
 
 		public static bool TryParseConfigEntry(string text, [NotNullWhen(true)] out UserSelectedProjectSettings? project)
@@ -350,11 +349,11 @@ namespace UnrealGameSync
 			ProjectInfo.ValidateBranchPath(branchPath);
 			ProjectInfo.ValidateProjectPath(projectPath);
 
-			this.ServerAndPort = serverAndPort;
-			this.UserName = userName;
-			this.ClientName = clientName;
-			this.BranchPath = branchPath;
-			this.ProjectPath = projectPath;
+			ServerAndPort = serverAndPort;
+			UserName = userName;
+			ClientName = clientName;
+			BranchPath = branchPath;
+			ProjectPath = projectPath;
 		}
 
 		public static bool TryLoad(DirectoryReference rootDir, [NotNullWhen(true)] out UserWorkspaceSettings? settings)
@@ -373,7 +372,7 @@ namespace UnrealGameSync
 			}
 		}
 
-		static object _syncRoot = new object();
+		static readonly object _syncRoot = new object();
 
 		public bool Save(ILogger logger)
 		{
@@ -414,7 +413,7 @@ namespace UnrealGameSync
 		public FilterType FilterType { get; set; }
 		public HashSet<string> FilterBadges { get; set; } = new HashSet<string>(StringComparer.OrdinalIgnoreCase);
 
-		static object _syncRoot = new object();
+		static readonly object _syncRoot = new object();
 
 		private UserProjectSettings()
 		{
@@ -422,7 +421,7 @@ namespace UnrealGameSync
 
 		public UserProjectSettings(FileReference configFile)
 		{
-			this.ConfigFile = configFile;
+			ConfigFile = configFile;
 		}
 
 		public static bool TryLoad(FileReference configFile, [NotNullWhen(true)] out UserProjectSettings? settings)
@@ -474,8 +473,8 @@ namespace UnrealGameSync
 			None		// Show no robomerge changes
 		};
 
-		FileReference _fileName;
-		ConfigFile _configFile;
+		readonly FileReference _fileName;
+		readonly ConfigFile _configFile;
 
 		// General settings
 		public UserSettingsVersion Version = UserSettingsVersion.Latest;
@@ -528,8 +527,8 @@ namespace UnrealGameSync
 		public int NotifyUnresolvedMinutes;
 
 		// Project settings
-		Dictionary<DirectoryReference, UserWorkspaceSettings> _workspaceDirToSettings = new Dictionary<DirectoryReference, UserWorkspaceSettings>();
-		Dictionary<FileReference, UserProjectSettings> _projectKeyToSettings = new Dictionary<FileReference, UserProjectSettings>();
+		readonly Dictionary<DirectoryReference, UserWorkspaceSettings> _workspaceDirToSettings = new Dictionary<DirectoryReference, UserWorkspaceSettings>();
+		readonly Dictionary<FileReference, UserProjectSettings> _projectKeyToSettings = new Dictionary<FileReference, UserProjectSettings>();
 
 		// Perforce settings
 		public PerforceSyncOptions SyncOptions => base.Global.Perforce;
@@ -613,8 +612,8 @@ namespace UnrealGameSync
 		public UserSettings(FileReference inFileName, ConfigFile inConfigFile, FileReference inCoreFileName, GlobalSettings inCoreSettingsData)
 			: base(inCoreFileName, inCoreSettingsData)
 		{
-			this._fileName = inFileName;
-			this._configFile = inConfigFile;
+			_fileName = inFileName;
+			_configFile = inConfigFile;
 
 			// General settings
 			Version = (UserSettingsVersion)_configFile.GetValue("General.Version", (int)UserSettingsVersion.Initial);
@@ -645,9 +644,9 @@ namespace UnrealGameSync
 			{
 				enabledTools.Add(new Guid("922EED87-E732-464C-92DC-5A8F7ED955E2"));
 			}
-			this.EnabledTools = enabledTools.ToArray();
+			EnabledTools = enabledTools.ToArray();
 
-			int.TryParse(_configFile.GetValue("General.FilterIndex", "0"), out FilterIndex);
+			Int32.TryParse(_configFile.GetValue("General.FilterIndex", "0"), out FilterIndex);
 
 			string? lastProjectString = _configFile.GetValue("General.LastProject", null);
 			if(lastProjectString != null)
@@ -775,7 +774,7 @@ namespace UnrealGameSync
 			if (numThreadsStr != null)
 			{
 				int numThreads;
-				if (int.TryParse(numThreadsStr, out numThreads) && numThreads > 0)
+				if (Int32.TryParse(numThreadsStr, out numThreads) && numThreads > 0)
 				{
 					if (Version >= UserSettingsVersion.DefaultNumberOfThreads || numThreads > 1)
 					{
@@ -806,21 +805,6 @@ namespace UnrealGameSync
 			return result;
 		}
 
-		static void SetCategorySettings(ConfigSection section, string includedKey, string excludedKey, Dictionary<Guid, bool> categories)
-		{
-			Guid[] includedCategories = categories.Where(x => x.Value).Select(x => x.Key).ToArray();
-			if (includedCategories.Length > 0)
-			{
-				section.SetValues(includedKey, includedCategories);
-			}
-
-			Guid[] excludedCategories = categories.Where(x => !x.Value).Select(x => x.Key).ToArray();
-			if (excludedCategories.Length > 0)
-			{
-				section.SetValues(excludedKey, excludedCategories);
-			}
-		}
-
 		static Rectangle? ParseRectangleValue(string text)
 		{
 			ConfigObject obj = new ConfigObject(text);
@@ -892,7 +876,7 @@ namespace UnrealGameSync
 					if(atIdx != -1)
 					{
 						int changeNumber;
-						if(int.TryParse(currentSync.Substring(atIdx + 1), out changeNumber))
+						if(Int32.TryParse(currentSync.Substring(atIdx + 1), out changeNumber))
 						{
 							currentWorkspace.ProjectIdentifier = currentSync.Substring(0, atIdx);
 							currentWorkspace.CurrentChangeNumber = changeNumber;
@@ -907,7 +891,7 @@ namespace UnrealGameSync
 					if(colonIdx != -1)
 					{
 						int changeNumber;
-						if(int.TryParse(lastUpdateResultText.Substring(0, colonIdx), out changeNumber))
+						if(Int32.TryParse(lastUpdateResultText.Substring(0, colonIdx), out changeNumber))
 						{
 							WorkspaceUpdateResult result;
 							if(Enum.TryParse(lastUpdateResultText.Substring(colonIdx + 1), out result))
@@ -927,7 +911,7 @@ namespace UnrealGameSync
 				foreach(string additionalChangeNumberString in workspaceSection.GetValues("AdditionalChangeNumbers", new string[0]))
 				{
 					int additionalChangeNumber;
-					if(int.TryParse(additionalChangeNumberString, out additionalChangeNumber))
+					if(Int32.TryParse(additionalChangeNumberString, out additionalChangeNumber))
 					{
 						currentWorkspace.AdditionalChangeNumbers.Add(additionalChangeNumber);
 					}
@@ -1077,7 +1061,6 @@ namespace UnrealGameSync
 			}
 			oidcSection.SetValues("Tokens", tokenObjects.Select(x => x.ToString()).ToArray());
 
-
 			// Window settings
 			ConfigSection windowSection = _configFile.FindOrAddSection("Window");
 			windowSection.Clear();
@@ -1124,45 +1107,6 @@ namespace UnrealGameSync
 			}
 		}
 
-		[return: NotNullIfNotNull("text")]
-		static string? EscapeText(string? text)
-		{
-			if(text == null)
-			{
-				return null;
-			}
-
-			StringBuilder result = new StringBuilder();
-			for(int idx = 0; idx < text.Length; idx++)
-			{
-				switch(text[idx])
-				{
-					case '\\':
-						result.Append("\\\\");
-						break;
-					case '\t':
-						result.Append("\\t");
-						break;
-					case '\r':
-						result.Append("\\r");
-						break;
-					case '\n':
-						result.Append("\\n");
-						break;
-					case '\'':
-						result.Append("\\\'");
-						break;
-					case '\"':
-						result.Append("\\\"");
-						break;
-					default:
-						result.Append(text[idx]);
-						break;
-				}
-			}
-			return result.ToString();
-		}
-
 		[return: NotNullIfNotNull("text")]
 		static string? UnescapeText(string? text)
 		{
diff --git a/Engine/Source/Programs/UnrealGameSync/UnrealGameSyncShared/Utility.cs b/Engine/Source/Programs/UnrealGameSync/UnrealGameSyncShared/Utility.cs
--- a/Engine/Source/Programs/UnrealGameSync/UnrealGameSyncShared/Utility.cs
+++ b/Engine/Source/Programs/UnrealGameSync/UnrealGameSyncShared/Utility.cs
@@ -3,7 +3,6 @@
 using EpicGames.Core;
 using EpicGames.Perforce;
 using Microsoft.Extensions.Logging;
-using Microsoft.Win32;
 using System;
 using System.Collections.Generic;
 using System.Diagnostics;
@@ -14,7 +13,6 @@ using System.Linq;
 using System.Text;
 using System.Text.Json;
 using System.Text.Json.Serialization;
-using System.Text.RegularExpressions;
 using System.Threading;
 using System.Threading.Tasks;
 
@@ -26,15 +24,15 @@ namespace UnrealGameSync
 
 		public UserErrorException(string message, int code = 1) : base(message)
 		{
-			this.Code = code;
+			Code = code;
 		}
 	}
 
 	public class PerforceChangeDetails
 	{
-		public string Description;
-		public bool ContainsCode;
-		public bool ContainsContent;
+		public string Description { get; }
+		public bool ContainsCode { get; }
+		public bool ContainsContent { get; }
 
 		public PerforceChangeDetails(DescribeRecord describeRecord, Func<string, bool>? isCodeFile = null)
 		{
@@ -125,7 +123,7 @@ namespace UnrealGameSync
 
 		public static byte[] SerializeJson<T>(T obj)
 		{
-			JsonSerializerOptions options = new JsonSerializerOptions { WriteIndented = true, DefaultIgnoreCondition = JsonIgnoreCondition.WhenWritingDefault };
+			JsonSerializerOptions options = new JsonSerializerOptions { IgnoreNullValues = true, WriteIndented = true };
 			options.Converters.Add(new JsonStringEnumConverter());
 
 			byte[] buffer;
@@ -520,9 +518,9 @@ namespace UnrealGameSync
 			}
 		}
 
-		public static Color Blend(Color first, Color second, float T)
+		public static Color Blend(Color first, Color second, float t)
 		{
-			return Color.FromArgb((int)(first.R + (second.R - first.R) * T), (int)(first.G + (second.G - first.G) * T), (int)(first.B + (second.B - first.B) * T));
+			return Color.FromArgb((int)(first.R + (second.R - first.R) * t), (int)(first.G + (second.G - first.G) * t), (int)(first.B + (second.B - first.B) * t));
 		}
 
 		public static PerforceSettings OverridePerforceSettings(IPerforceSettings defaultConnection, string? serverAndPort, string? userName)
diff --git a/Engine/Source/Programs/UnrealGameSync/UnrealGameSyncShared/WorkspaceLock.cs b/Engine/Source/Programs/UnrealGameSync/UnrealGameSyncShared/WorkspaceLock.cs
--- a/Engine/Source/Programs/UnrealGameSync/UnrealGameSyncShared/WorkspaceLock.cs
+++ b/Engine/Source/Programs/UnrealGameSync/UnrealGameSyncShared/WorkspaceLock.cs
@@ -2,12 +2,9 @@
 
 using EpicGames.Core;
 using System;
-using System.Collections.Generic;
-using System.Runtime.CompilerServices;
 using System.Security.Cryptography;
 using System.Text;
 using System.Threading;
-using System.Threading.Tasks;
 
 namespace UnrealGameSync
 {
@@ -28,7 +25,7 @@ namespace UnrealGameSync
 		EventWaitHandle? _lockedEvent;
 
 		Thread? _monitorThread;
-		ManualResetEvent _cancelMonitorEvent = new ManualResetEvent(false);
+		readonly ManualResetEvent _cancelMonitorEvent = new ManualResetEvent(false);
 
 		/// <summary>
 		/// Callback for the lock state changing
diff --git a/Engine/Source/Programs/UnrealGameSync/UnrealGameSyncShared/WorkspaceState.cs b/Engine/Source/Programs/UnrealGameSync/UnrealGameSyncShared/WorkspaceState.cs
--- a/Engine/Source/Programs/UnrealGameSync/UnrealGameSyncShared/WorkspaceState.cs
+++ b/Engine/Source/Programs/UnrealGameSync/UnrealGameSyncShared/WorkspaceState.cs
@@ -1,13 +1,11 @@
 // Copyright Epic Games, Inc. All Rights Reserved.
 
 using EpicGames.Core;
-using JetBrains.Annotations;
 using System;
 using System.Collections.Generic;
 using System.IO;
 using System.Linq;
 using System.Reflection;
-using System.Text;
 using System.Threading;
 
 namespace UnrealGameSync
@@ -50,7 +48,7 @@ namespace UnrealGameSync
 
 		public void UpdateCachedProjectInfo(ProjectInfo projectInfo, long settingsTimeUtc)
 		{
-			this.SettingsTimeUtc = settingsTimeUtc;
+			SettingsTimeUtc = settingsTimeUtc;
 
 			ClientName = projectInfo.ClientName;
 			BranchPath = projectInfo.BranchPath;
