UGS: Add .editorconfig file for static analysis warnings in UGS (+lots of fixes).

diff --git a/Engine/Source/Programs/UnrealGameSync/UnrealGameSyncShared/WorkspaceUpdate.cs b/Engine/Source/Programs/UnrealGameSync/UnrealGameSyncShared/WorkspaceUpdate.cs
--- a/Engine/Source/Programs/UnrealGameSync/UnrealGameSyncShared/WorkspaceUpdate.cs
+++ b/Engine/Source/Programs/UnrealGameSync/UnrealGameSyncShared/WorkspaceUpdate.cs
@@ -2,16 +2,11 @@
 
 using EpicGames.Core;
 using EpicGames.Perforce;
-using Microsoft.Extensions.DependencyInjection;
 using Microsoft.Extensions.Logging;
 using System;
-using System.Collections.Concurrent;
 using System.Collections.Generic;
 using System.Diagnostics;
-using System.Diagnostics.CodeAnalysis;
-using System.Globalization;
 using System.IO;
-using System.IO.Compression;
 using System.Linq;
 using System.Security.Cryptography;
 using System.Text;
@@ -144,11 +139,11 @@ namespace UnrealGameSync
 
 		public WorkspaceSyncCategory(Guid uniqueId, string name, params string[] paths)
 		{
-			this.UniqueId = uniqueId;
-			this.Enable = true;
-			this.Name = name;
-			this.Paths = paths;
-			this.Requires = new Guid[0];
+			UniqueId = uniqueId;
+			Enable = true;
+			Name = name;
+			Paths = paths;
+			Requires = new Guid[0];
 		}
 
 		public static Dictionary<Guid, bool> GetDefault(IEnumerable<WorkspaceSyncCategory> categories)
@@ -220,13 +215,13 @@ namespace UnrealGameSync
 			ValidateBranchPath(branchPath);
 			ValidateProjectPath(projectPath);
 
-			this.LocalRootPath = localRootPath;
-			this.ClientName = clientName;
-			this.BranchPath = branchPath;
-			this.ProjectPath = projectPath;
-			this.StreamName = streamName;
-			this.ProjectIdentifier = projectIdentifier;
-			this.IsEnterpriseProject = isEnterpriseProject;
+			LocalRootPath = localRootPath;
+			ClientName = clientName;
+			BranchPath = branchPath;
+			ProjectPath = projectPath;
+			StreamName = streamName;
+			ProjectIdentifier = projectIdentifier;
+			IsEnterpriseProject = isEnterpriseProject;
 		}
 
 		public static async Task<ProjectInfo> CreateAsync(IPerforceConnection perforceClient, UserWorkspaceSettings settings, CancellationToken cancellationToken)
@@ -346,19 +341,19 @@ namespace UnrealGameSync
 		static readonly string LocalVersionHeaderFileName = VersionHeaderFileName.Replace('/', Path.DirectorySeparatorChar);
 		static readonly string LocalObjectVersionFileName = ObjectVersionFileName.Replace('/', Path.DirectorySeparatorChar);
 
-		static SemaphoreSlim _updateSemaphore = new SemaphoreSlim(1);
+		static readonly SemaphoreSlim _updateSemaphore = new SemaphoreSlim(1);
 
 		class RecordCounter : IDisposable
 		{
-			ProgressValue _progress;
-			string _message;
+			readonly ProgressValue _progress;
+			readonly string _message;
 			int _count;
-			Stopwatch _timer = Stopwatch.StartNew();
+			readonly Stopwatch _timer = Stopwatch.StartNew();
 
 			public RecordCounter(ProgressValue progress, string message)
 			{
-				this._progress = progress;
-				this._message = message;
+				_progress = progress;
+				_message = message;
 
 				progress.Set(message);
 			}
@@ -396,9 +391,9 @@ namespace UnrealGameSync
 
 			public SyncBatchBuilder(int? maxCommandsPerList, long? maxSizePerList)
 			{
-				this.MaxCommandsPerList = maxCommandsPerList ?? PerforceSyncOptions.DefaultMaxCommandsPerBatch;
-				this.MaxSizePerList = maxSizePerList ?? PerforceSyncOptions.DefaultMaxSizePerBatch;
-				this.Batches = new Queue<List<string>>();
+				MaxCommandsPerList = maxCommandsPerList ?? PerforceSyncOptions.DefaultMaxCommandsPerBatch;
+				MaxSizePerList = maxSizePerList ?? PerforceSyncOptions.DefaultMaxSizePerBatch;
+				Batches = new Queue<List<string>>();
 			}
 
 			public void Add(string newCommand, long newSize)
@@ -439,7 +434,7 @@ namespace UnrealGameSync
 
 			public SyncTree(bool canUseWildcard)
 			{
-				this.CanUseWildcard = canUseWildcard;
+				CanUseWildcard = canUseWildcard;
 			}
 
 			public SyncTree FindOrAddSubTree(string name)
@@ -510,7 +505,7 @@ namespace UnrealGameSync
 
 		public WorkspaceUpdate(WorkspaceUpdateContext context)
 		{
-			this.Context = context;
+			Context = context;
 		}
 
 		class SyncFile
@@ -521,20 +516,20 @@ namespace UnrealGameSync
 
 			public SyncFile(string depotFile, string relativePath, long size)
 			{
-				this.DepotFile = depotFile;
-				this.RelativePath = relativePath;
-				this.Size = size;
+				DepotFile = depotFile;
+				RelativePath = relativePath;
+				Size = size;
 			}
 		};
 
 		class SemaphoreScope : IDisposable
 		{
-			SemaphoreSlim _semaphore;
+			readonly SemaphoreSlim _semaphore;
 			public bool HasLock { get; private set; }
 
 			public SemaphoreScope(SemaphoreSlim semaphore)
 			{
-				this._semaphore = semaphore;
+				_semaphore = semaphore;
 			}
 
 			public bool TryAcquire()
@@ -563,7 +558,7 @@ namespace UnrealGameSync
 
 		public static string ShellScriptExt { get; }= RuntimeInformation.IsOSPlatform(OSPlatform.Windows)? "bat" : "sh";
 
-		public Task<int> ExecuteShellCommandAsync(string commandLine, string? workingDir, Action<string> processOutput, CancellationToken cancellationToken)
+		public Task<int> ExecuteShellCommandAsync(string commandLine, Action<string> processOutput, CancellationToken cancellationToken)
 		{
 			if(RuntimeInformation.IsOSPlatform(OSPlatform.Windows))
 			{
@@ -903,7 +898,7 @@ namespace UnrealGameSync
 					{
 						// Read the new config file
 						Context.ProjectConfigFile = await ReadProjectConfigFile(project.LocalRootPath, project.LocalFileName, logger);
-						Context.ProjectStreamFilter = await ReadProjectStreamFilter(perforce, Context.ProjectConfigFile, logger, cancellationToken);
+						Context.ProjectStreamFilter = await ReadProjectStreamFilter(perforce, Context.ProjectConfigFile, cancellationToken);
 
 						// Get the branch name
 						string? branchOrStreamName = await perforce.GetCurrentStreamAsync(cancellationToken);
@@ -983,7 +978,7 @@ namespace UnrealGameSync
 						// Update the version files
 						if (Context.ProjectConfigFile.GetValue("Options.UseFastModularVersioningV2", false))
 						{
-							bool isLicenseeVersion = await IsLicenseeVersion(perforce, project, logger, cancellationToken);
+							bool isLicenseeVersion = await IsLicenseeVersion(perforce, project, cancellationToken);
 							if (!await UpdateVersionFile(perforce, project.ClientRootPath + BuildVersionFileName, Context.ChangeNumber, text => UpdateBuildVersion(text, Context.ChangeNumber, versionChangeNumber, branchOrStreamName, isLicenseeVersion), logger, cancellationToken))
 							{
 								return (WorkspaceUpdateResult.FailedToSync, $"Failed to update {BuildVersionFileName}.");
@@ -991,7 +986,7 @@ namespace UnrealGameSync
 						}
 						else if (Context.ProjectConfigFile.GetValue("Options.UseFastModularVersioning", false))
 						{
-							bool isLicenseeVersion = await IsLicenseeVersion(perforce, project, logger, cancellationToken);
+							bool isLicenseeVersion = await IsLicenseeVersion(perforce, project, cancellationToken);
 							if (!await UpdateVersionFile(perforce, project.ClientRootPath + BuildVersionFileName, Context.ChangeNumber, text => UpdateBuildVersion(text, Context.ChangeNumber, versionChangeNumber, branchOrStreamName, isLicenseeVersion), logger, cancellationToken))
 							{
 								return (WorkspaceUpdateResult.FailedToSync, $"Failed to update {BuildVersionFileName}");
@@ -1205,7 +1200,7 @@ namespace UnrealGameSync
 					logger.LogInformation("Generating project files...");
 					logger.LogInformation("gpf> Running {Arguments}", commandLine);
 
-					int generateProjectFilesResult = await ExecuteShellCommandAsync(commandLine.ToString(), null, line => ProcessOutput(line, "gpf> ", Progress, logger), cancellationToken);
+					int generateProjectFilesResult = await ExecuteShellCommandAsync(commandLine.ToString(), line => ProcessOutput(line, "gpf> ", Progress, logger), cancellationToken);
 					if (generateProjectFilesResult != 0)
 					{
 						return (WorkspaceUpdateResult.FailedToCompile, $"Failed to generate project files (exit code {generateProjectFilesResult}).");
@@ -1256,7 +1251,7 @@ namespace UnrealGameSync
 					foreach (string cleanBuildChange in Context.ProjectConfigFile.GetValues("ForceClean.Changelist", new string[0]))
 					{
 						int changeNumber;
-						if (int.TryParse(cleanBuildChange, out changeNumber))
+						if (Int32.TryParse(cleanBuildChange, out changeNumber))
 						{
 							if ((state.LastBuiltChangeNumber >= changeNumber) != (state.CurrentChangeNumber >= changeNumber))
 							{
@@ -1319,12 +1314,12 @@ namespace UnrealGameSync
 										if (Context.Options.HasFlag(WorkspaceUpdateOptions.Clean) || forceClean)
 										{
 											logger.LogInformation("ubt> Running {Arguments}", commandLine + " -clean");
-											await ExecuteShellCommandAsync(commandLine + " -clean", null, line => ProcessOutput(line, "ubt> ", Progress, logger), cancellationToken);
+											await ExecuteShellCommandAsync(commandLine + " -clean", line => ProcessOutput(line, "ubt> ", Progress, logger), cancellationToken);
 										}
 
 										logger.LogInformation("ubt> Running {FileName} {Arguments}", buildBat, commandLine + " -progress");
 
-										int resultFromBuild = await ExecuteShellCommandAsync(commandLine + " -progress", null, line => ProcessOutput(line, "ubt> ", Progress, logger), cancellationToken);
+										int resultFromBuild = await ExecuteShellCommandAsync(commandLine + " -progress", line => ProcessOutput(line, "ubt> ", Progress, logger), cancellationToken);
 										if (resultFromBuild != 0)
 										{
 											stepStopwatch.Stop("Failed");
@@ -1354,7 +1349,7 @@ namespace UnrealGameSync
 										string arguments = String.Format("\"{0}\" -profile=\"{1}\"", localRunUat, FileReference.Combine(project.LocalRootPath, step.FileName ?? "unknown"));
 										logger.LogInformation("uat> Running {FileName} {Argument}", localRunUat, arguments);
 
-										int resultFromUat = await ExecuteShellCommandAsync(arguments, null, line => ProcessOutput(line, "uat> ", Progress, logger), cancellationToken);
+										int resultFromUat = await ExecuteShellCommandAsync(arguments, line => ProcessOutput(line, "uat> ", Progress, logger), cancellationToken);
 										if (resultFromUat != 0)
 										{
 											stepStopwatch.Stop("Failed");
@@ -1446,7 +1441,7 @@ namespace UnrealGameSync
 			}
 		}
 
-		static async Task<bool> IsLicenseeVersion(IPerforceConnection perforce, ProjectInfo project, ILogger logger, CancellationToken cancellationToken)
+		static async Task<bool> IsLicenseeVersion(IPerforceConnection perforce, ProjectInfo project, CancellationToken cancellationToken)
 		{
 			string[] files =
 			{
@@ -1553,10 +1548,10 @@ namespace UnrealGameSync
 
 			public SyncState(HashSet<string> remainingDepotPaths, Queue<List<string>> syncCommandLists)
 			{
-				this.TotalDepotPaths = remainingDepotPaths.Count;
-				this.RemainingDepotPaths = remainingDepotPaths;
-				this.SyncCommandLists = syncCommandLists;
-				this.StatusMessage = "Succeeded.";
+				TotalDepotPaths = remainingDepotPaths.Count;
+				RemainingDepotPaths = remainingDepotPaths;
+				SyncCommandLists = syncCommandLists;
+				StatusMessage = "Succeeded.";
 			}
 		}
 
@@ -1798,7 +1793,7 @@ namespace UnrealGameSync
 			return projectConfig;
 		}
 
-		public static async Task<IReadOnlyList<string>?> ReadProjectStreamFilter(IPerforceConnection perforce, ConfigFile projectConfigFile, ILogger logger, CancellationToken cancellationToken)
+		public static async Task<IReadOnlyList<string>?> ReadProjectStreamFilter(IPerforceConnection perforce, ConfigFile projectConfigFile, CancellationToken cancellationToken)
 		{
 			string? streamListDepotPath = projectConfigFile.GetValue("Options.QuickSelectStreamList", null);
 			if (streamListDepotPath == null)
@@ -1938,13 +1933,13 @@ namespace UnrealGameSync
 			int prevCompatibleChangelist = 0;
 			if (obj.TryGetValue("CompatibleChangelist", out object? prevCompatibleChangelistObj))
 			{
-				int.TryParse(prevCompatibleChangelistObj?.ToString(), out prevCompatibleChangelist);
+				Int32.TryParse(prevCompatibleChangelistObj?.ToString(), out prevCompatibleChangelist);
 			}
 
 			int prevIsLicenseeVersion = 0;
 			if (obj.TryGetValue("IsLicenseeVersion", out object? prevIsLicenseeVersionObj))
 			{
-				int.TryParse(prevIsLicenseeVersionObj?.ToString(), out prevIsLicenseeVersion);
+				Int32.TryParse(prevIsLicenseeVersionObj?.ToString(), out prevIsLicenseeVersion);
 			}
 
 			obj["Changelist"] = changelist;
@@ -2041,9 +2036,6 @@ namespace UnrealGameSync
 			return line.Substring(startIdx, lineIdx - startIdx);
 		}
 
-		public Tuple<string, float> CurrentProgress
-		{
-			get { return Progress.Current; }
-		}
+		public Tuple<string, float> CurrentProgress => Progress.Current;
 	}
 }
