UGS: Add empty editorconfig file for MetadataServer to suppress warnings.

diff --git a/Engine/Source/Programs/UnrealGameSync/MetadataServer/.editorconfig b/Engine/Source/Programs/UnrealGameSync/MetadataServer/.editorconfig
new file mode 100644
--- /dev/null
+++ b/Engine/Source/Programs/UnrealGameSync/MetadataServer/.editorconfig
@@ -0,0 +1,5 @@
+# NOTE: Requires **VS2019 16.3** or later
+
+# EditorConfig is awesome:http://EditorConfig.org
+
+root=True
