UGS: Allow specifying fixed changelist numbers to sync to for certain latest change types. This could be used to specify older changelists that users may need which no longer populate in the workspace control by default.

diff --git a/Engine/Source/Programs/UnrealGameSync/UnrealGameSync/Controls/WorkspaceControl.cs b/Engine/Source/Programs/UnrealGameSync/UnrealGameSync/Controls/WorkspaceControl.cs
--- a/Engine/Source/Programs/UnrealGameSync/UnrealGameSync/Controls/WorkspaceControl.cs
+++ b/Engine/Source/Programs/UnrealGameSync/UnrealGameSync/Controls/WorkspaceControl.cs
@@ -5095,6 +5095,12 @@ namespace UnrealGameSync
 
 		bool FindChangeToSync(LatestChangeType changeType, out int changeNumber)
 		{
+			if (changeType.Number != -1)
+			{
+				changeNumber = changeType.Number;
+				return true;
+			}
+
 			for (int idx = 0; idx < BuildList.Items.Count; idx++)
 			{
 				ChangesRecord change = (ChangesRecord)BuildList.Items[idx].Tag;
diff --git a/Engine/Source/Programs/UnrealGameSync/UnrealGameSyncShared/UserSettings.cs b/Engine/Source/Programs/UnrealGameSync/UnrealGameSyncShared/UserSettings.cs
--- a/Engine/Source/Programs/UnrealGameSync/UnrealGameSyncShared/UserSettings.cs
+++ b/Engine/Source/Programs/UnrealGameSync/UnrealGameSyncShared/UserSettings.cs
@@ -66,6 +66,7 @@ namespace UnrealGameSync
 		public bool Starred = false;
 		public bool FindNewestGoodContent = false;
 		public List<string> RequiredBadges = new List<string>();
+		public int Number = -1;
 
 		public static bool TryParseConfigEntry(string text, [NotNullWhen(true)] out LatestChangeType? changeType)
 		{
@@ -83,6 +84,7 @@ namespace UnrealGameSync
 				changeType.Starred = definitionObject.GetValue("bStarred", false);
 				changeType.FindNewestGoodContent = definitionObject.GetValue("bFindNewestGoodContent", false);
 				changeType.RequiredBadges = new List<string>(definitionObject.GetValue("RequiredBadges", "").Split(new char[] { ',',' '}, StringSplitOptions.RemoveEmptyEntries));
+				changeType.Number = definitionObject.GetValue("Number", -1);
 				return true;
 			}
 			else
