UGS: Fix build script

diff --git a/Engine/Source/Programs/UnrealGameSync/BuildUGS.xml b/Engine/Source/Programs/UnrealGameSync/BuildUGS.xml
--- a/Engine/Source/Programs/UnrealGameSync/BuildUGS.xml
+++ b/Engine/Source/Programs/UnrealGameSync/BuildUGS.xml
@@ -40,7 +40,7 @@
 		</Node>
 		
 		<Property Name="InstallerDependency" Value=""/>
-		<Property Name="InstallerDependency" Value="Build UnrealGameSync Installer" If="IncludeInstaller"/>
+		<Property Name="InstallerDependency" Value="Build UnrealGameSync Installer" If="$(IncludeInstaller)"/>
 		
 		<Node Name="Build UnrealGameSync" Requires="$(InstallerDependency)">
 			<Delete Files="$(UgsOutputDir)/..."/>
