UGS: Fix workspace lock potentially being released on a different thread. Use a dedicated background thread to acquire and release the mutex.

diff --git a/Engine/Source/Programs/UnrealGameSync/UnrealGameSync/Workspace.cs b/Engine/Source/Programs/UnrealGameSync/UnrealGameSync/Workspace.cs
--- a/Engine/Source/Programs/UnrealGameSync/UnrealGameSync/Workspace.cs
+++ b/Engine/Source/Programs/UnrealGameSync/UnrealGameSync/Workspace.cs
@@ -151,7 +151,7 @@ namespace UnrealGameSync
 			string statusMessage;
 			WorkspaceUpdateResult result = WorkspaceUpdateResult.FailedToSync;
 
-			if (!Lock.TryAcquire())
+			if (!await Lock.TryAcquireAsync())
 			{
 				statusMessage = "Command line sync already in progress";
 				_logger.LogError("Another process is already syncing this workspace.");
@@ -178,7 +178,7 @@ namespace UnrealGameSync
 				}
 				finally
 				{
-					Lock.Release();
+					await Lock.ReleaseAsync();
 				}
 			}
 
diff --git a/Engine/Source/Programs/UnrealGameSync/UnrealGameSyncCmd/Program.cs b/Engine/Source/Programs/UnrealGameSync/UnrealGameSyncCmd/Program.cs
--- a/Engine/Source/Programs/UnrealGameSync/UnrealGameSyncCmd/Program.cs
+++ b/Engine/Source/Programs/UnrealGameSync/UnrealGameSyncCmd/Program.cs
@@ -658,7 +658,7 @@ namespace UnrealGameSyncCmd
 				string[] syncFilter = ReadSyncFilter(settings, context.UserSettings, projectConfig);
 
 				using WorkspaceLock workspaceLock = new WorkspaceLock(settings.RootDir);
-				if (!workspaceLock.TryAcquire())
+				if (!await workspaceLock.TryAcquireAsync())
 				{
 					logger.LogError("Another process is already syncing this workspace.");
 					return;
diff --git a/Engine/Source/Programs/UnrealGameSync/UnrealGameSyncShared/WorkspaceLock.cs b/Engine/Source/Programs/UnrealGameSync/UnrealGameSyncShared/WorkspaceLock.cs
--- a/Engine/Source/Programs/UnrealGameSync/UnrealGameSyncShared/WorkspaceLock.cs
+++ b/Engine/Source/Programs/UnrealGameSync/UnrealGameSyncShared/WorkspaceLock.cs
@@ -2,9 +2,12 @@
 
 using EpicGames.Core;
 using System;
+using System.Collections.Concurrent;
 using System.Security.Cryptography;
 using System.Text;
 using System.Threading;
+using System.Threading.Channels;
+using System.Threading.Tasks;
 
 namespace UnrealGameSync
 {
@@ -24,6 +27,10 @@ namespace UnrealGameSync
 		bool _locked;
 		EventWaitHandle? _lockedEvent;
 
+		Thread? _acquireThread;
+		readonly BlockingCollection<Action> _acquireActions = new BlockingCollection<Action>();
+		readonly CancellationTokenSource _acquireCancellationSource = new CancellationTokenSource();
+
 		Thread? _monitorThread;
 		readonly ManualResetEvent _cancelMonitorEvent = new ManualResetEvent(false);
 
@@ -46,6 +53,10 @@ namespace UnrealGameSync
 
 			_mutex = new Mutex(false, $"{Prefix}.{_objectName}.mutex");
 
+			_acquireThread = new Thread(AcquireThread);
+			_acquireThread.IsBackground = true;
+			_acquireThread.Start();
+
 			_monitorThread = new Thread(MonitorThread);
 			_monitorThread.IsBackground = true;
 			_monitorThread.Start();
@@ -54,9 +65,12 @@ namespace UnrealGameSync
 		/// <inheritdoc/>
 		public void Dispose()
 		{
-			while (_acquireCount > 0)
+			if (_acquireThread != null)
 			{
-				Release();
+				_acquireCancellationSource.Cancel();
+
+				_acquireThread.Join();
+				_acquireThread = null;
 			}
 
 			if (_monitorThread != null)
@@ -65,10 +79,10 @@ namespace UnrealGameSync
 
 				_monitorThread.Join();
 				_monitorThread = null;
-
-				_cancelMonitorEvent.Dispose();
 			}
 
+			_acquireCancellationSource.Dispose();
+			_cancelMonitorEvent.Dispose();
 			_mutex.Dispose();
 		}
 
@@ -88,35 +102,47 @@ namespace UnrealGameSync
 		/// Attempt to acquire the mutext
 		/// </summary>
 		/// <returns></returns>
-		public bool TryAcquire()
+		public async Task<bool> TryAcquireAsync()
+		{
+			TaskCompletionSource<bool> result = new TaskCompletionSource<bool>();
+			_acquireActions.Add(() => TryAcquireInternal(result));
+			return await result.Task;
+		}
+
+		void TryAcquireInternal(TaskCompletionSource<bool> resultTcs)
 		{
+			bool result;
 			lock (_lockObject)
 			{
 				try
 				{
-					if (!_mutex.WaitOne(0))
-					{
-						return false;
-					}
+					result = _mutex.WaitOne(0);
 				}
 				catch (AbandonedMutexException)
 				{
+					result = true;
 				}
 
-				if (++_acquireCount == 1)
+				if (result && ++_acquireCount == 1)
 				{
 					_lockedEvent = CreateLockedEvent();
 					_lockedEvent.Set();
 				}
-
-				return true;
 			}
+			Task.Run(() => resultTcs.TrySetResult(result));
 		}
 
 		/// <summary>
 		/// Release the current mutext
 		/// </summary>
-		public void Release()
+		public async Task ReleaseAsync()
+		{
+			TaskCompletionSource<bool> resultTcs = new TaskCompletionSource<bool>();
+			_acquireActions.Add(() => ReleaseInternal(resultTcs));
+			await resultTcs.Task;
+		}
+
+		private void ReleaseInternal(TaskCompletionSource<bool> resultTcs)
 		{
 			lock (_lockObject)
 			{
@@ -131,6 +157,27 @@ namespace UnrealGameSync
 					}
 				}
 			}
+			Task.Run(() => resultTcs.TrySetResult(true));
+		}
+
+		void AcquireThread()
+		{
+			for (; ; )
+			{
+				try
+				{
+					_acquireActions.Take(_acquireCancellationSource.Token)();
+				}
+				catch (OperationCanceledException)
+				{
+					break;
+				}
+			}
+
+			for(; _acquireCount > 0; _acquireCount--)
+			{
+				_mutex.ReleaseMutex();
+			}
 		}
 
 		void MonitorThread()
