Horde: Allow specifying latest change types in stream-specific sections of the config file.

diff --git a/Engine/Source/Programs/UnrealGameSync/UnrealGameSync/Controls/WorkspaceControl.cs b/Engine/Source/Programs/UnrealGameSync/UnrealGameSync/Controls/WorkspaceControl.cs
--- a/Engine/Source/Programs/UnrealGameSync/UnrealGameSync/Controls/WorkspaceControl.cs
+++ b/Engine/Source/Programs/UnrealGameSync/UnrealGameSync/Controls/WorkspaceControl.cs
@@ -2612,7 +2612,12 @@ namespace UnrealGameSync
 			foundLatestChangeTypes.Add(LatestChangeType.LatestGoodChange());
 			foundLatestChangeTypes.Add(LatestChangeType.LatestStarredChange());
 
-			string[] latestChangeTypeDefinitions = _perforceMonitor.LatestProjectConfigFile.GetValues("Sync.LatestChangeType", new string[0]);
+			ConfigFile configFile = _perforceMonitor.LatestProjectConfigFile;
+
+			List<string> latestChangeTypeDefinitions = new List<string>();
+			latestChangeTypeDefinitions.AddRange(configFile.GetValues("Sync.LatestChangeType", Array.Empty<string>()));
+			ConfigUtils.GetProjectSettings(configFile, SelectedProjectIdentifier, "LatestChangeType", latestChangeTypeDefinitions);
+
 			foreach (string latestChangeDefinition in latestChangeTypeDefinitions.Distinct())
 			{
 				LatestChangeType? newType;
diff --git a/Engine/Source/Programs/UnrealGameSync/UnrealGameSyncShared/ConfigUtils.cs b/Engine/Source/Programs/UnrealGameSync/UnrealGameSyncShared/ConfigUtils.cs
--- a/Engine/Source/Programs/UnrealGameSync/UnrealGameSyncShared/ConfigUtils.cs
+++ b/Engine/Source/Programs/UnrealGameSync/UnrealGameSyncShared/ConfigUtils.cs
@@ -517,6 +517,33 @@ namespace UnrealGameSync
 			return false;
 		}
 
+		public static void GetProjectSettings(ConfigFile projectConfigFile, string selectedProjectIdentifier, string name, List<string> values)
+		{
+			string path = selectedProjectIdentifier;
+			for (; ; )
+			{
+				ConfigSection? projectSection = projectConfigFile.FindSection(path);
+				if (projectSection != null)
+				{
+					values.AddRange(projectSection.GetValues(name, Array.Empty<string>()));
+				}
+
+				int lastSlash = path.LastIndexOf('/');
+				if (lastSlash < 2)
+				{
+					break;
+				}
+
+				path = path.Substring(0, lastSlash);
+			}
+
+			ConfigSection? defaultSection = projectConfigFile.FindSection("Default");
+			if (defaultSection != null)
+			{
+				values.AddRange(defaultSection.GetValues(name, Array.Empty<string>()));
+			}
+		}
+
 		public static Dictionary<Guid, WorkspaceSyncCategory> GetSyncCategories(ConfigFile projectConfigFile)
 		{
 			Dictionary<Guid, WorkspaceSyncCategory> uniqueIdToCategory = new Dictionary<Guid, WorkspaceSyncCategory>();
