UGS: Fix double slashes in OIDC config files.

diff --git a/Engine/Source/Programs/Shared/EpicGames.OIDC/ProviderConfigurationFactory.cs b/Engine/Source/Programs/Shared/EpicGames.OIDC/ProviderConfigurationFactory.cs
--- a/Engine/Source/Programs/Shared/EpicGames.OIDC/ProviderConfigurationFactory.cs
+++ b/Engine/Source/Programs/Shared/EpicGames.OIDC/ProviderConfigurationFactory.cs
@@ -3,8 +3,6 @@
 using System;
 using System.Collections.Generic;
 using System.IO;
-using System.Linq;
-using System.Net;
 using Microsoft.Extensions.Configuration;
 
 namespace EpicGames.OIDC
@@ -15,9 +13,9 @@ namespace EpicGames.OIDC
 
 		public static IReadOnlyList<string> ConfigPaths { get; } = new string[]
 		{
-			$"/Programs/OidcToken/{ConfigFileName}",
-			$"/Restricted/NoRedist/Programs/OidcToken/{ConfigFileName}",
-			$"/Restricted/NotForLicensees/Programs/OidcToken/{ConfigFileName}"
+			$"Programs/OidcToken/{ConfigFileName}",
+			$"Restricted/NoRedist/Programs/OidcToken/{ConfigFileName}",
+			$"Restricted/NotForLicensees/Programs/OidcToken/{ConfigFileName}"
 		};
 
 		public static IConfiguration ReadConfiguration(DirectoryInfo engineDir, DirectoryInfo? gameDir)
diff --git a/Engine/Source/Programs/UnrealGameSync/UnrealGameSyncShared/ConfigUtils.cs b/Engine/Source/Programs/UnrealGameSync/UnrealGameSyncShared/ConfigUtils.cs
--- a/Engine/Source/Programs/UnrealGameSync/UnrealGameSyncShared/ConfigUtils.cs
+++ b/Engine/Source/Programs/UnrealGameSync/UnrealGameSyncShared/ConfigUtils.cs
@@ -204,12 +204,12 @@ namespace UnrealGameSync
 			if (TryGetProjectSetting(projectConfigFile, projectIdentifier, "OidcProvider", out oidcProvider))
 			{
 				List<string> configFiles = new List<string>();
-				configFiles.AddRange(ProviderConfigurationFactory.ConfigPaths.Select(x => $"{clientRoot}/Engine/{x}"));
+				configFiles.AddRange(ProviderConfigurationFactory.ConfigPaths.Select(x => $"{clientRoot.TrimEnd('/')}/Engine/{x}"));
 
 				if (clientProjectFile != null && clientProjectFile.EndsWith(".uproject", StringComparison.OrdinalIgnoreCase))
 				{
 					string clientProjectPath = clientProjectFile.Substring(0, clientProjectFile.LastIndexOf('/'));
-					configFiles.AddRange(ProviderConfigurationFactory.ConfigPaths.Select(x => $"{clientProjectPath}/{x}"));
+					configFiles.AddRange(ProviderConfigurationFactory.ConfigPaths.Select(x => $"{clientProjectPath.TrimEnd('/')}/{x}"));
 				}
 
 				List<string[]> latestOidcConfigFiles = await ConfigUtils.ReadConfigFiles(perforce, configFiles, localFiles, cacheFolder, logger, cancellationToken);
