UGS: Allow reading the change to sync from a text file submitted to the depot.

diff --git a/Engine/Source/Programs/UnrealGameSync/UnrealGameSync/Controls/WorkspaceControl.cs b/Engine/Source/Programs/UnrealGameSync/UnrealGameSync/Controls/WorkspaceControl.cs
--- a/Engine/Source/Programs/UnrealGameSync/UnrealGameSync/Controls/WorkspaceControl.cs
+++ b/Engine/Source/Programs/UnrealGameSync/UnrealGameSync/Controls/WorkspaceControl.cs
@@ -5137,10 +5137,28 @@ namespace UnrealGameSync
 
 		bool FindChangeToSync(LatestChangeType changeType, out int changeNumber)
 		{
-			if (changeType.Number != -1)
+			if (!String.IsNullOrEmpty(changeType.ReadFrom))
 			{
-				changeNumber = changeType.Number;
-				return true;
+				async Task<int> ReadChangeToSync(IPerforceConnection perforceConnection, CancellationToken cancellationToken)
+				{
+					PerforceResponse<PrintRecord<byte[]>> record = await perforceConnection.TryPrintDataAsync(changeType.ReadFrom, cancellationToken);
+					record.EnsureSuccess();
+
+					string contents = Encoding.UTF8.GetString(record.Data.Contents ?? Array.Empty<byte>());
+					return Int32.Parse(contents.Trim());
+				}
+
+				ModalTask<int>? task = PerforceModalTask.Execute(this, "Reading change", "Reading latest change...", _perforceSettings, ReadChangeToSync, _logger);
+				if (task == null || !task.Succeeded || task.Result <= 0)
+				{
+					changeNumber = 0;
+					return false;
+				}
+				else
+				{
+					changeNumber = task.Result;
+					return true;
+				}
 			}
 
 			for (int idx = 0; idx < BuildList.Items.Count; idx++)
diff --git a/Engine/Source/Programs/UnrealGameSync/UnrealGameSyncShared/UserSettings.cs b/Engine/Source/Programs/UnrealGameSync/UnrealGameSyncShared/UserSettings.cs
--- a/Engine/Source/Programs/UnrealGameSync/UnrealGameSyncShared/UserSettings.cs
+++ b/Engine/Source/Programs/UnrealGameSync/UnrealGameSyncShared/UserSettings.cs
@@ -66,7 +66,9 @@ namespace UnrealGameSync
 		public bool Starred = false;
 		public bool FindNewestGoodContent = false;
 		public List<string> RequiredBadges = new List<string>();
-		public int Number = -1;
+
+		// Depot path to read for the latest change number
+		public string? ReadFrom = "";
 
 		public static bool TryParseConfigEntry(string text, [NotNullWhen(true)] out LatestChangeType? changeType)
 		{
@@ -84,7 +86,8 @@ namespace UnrealGameSync
 				changeType.Starred = definitionObject.GetValue("bStarred", false);
 				changeType.FindNewestGoodContent = definitionObject.GetValue("bFindNewestGoodContent", false);
 				changeType.RequiredBadges = new List<string>(definitionObject.GetValue("RequiredBadges", "").Split(new char[] { ',',' '}, StringSplitOptions.RemoveEmptyEntries));
-				changeType.Number = definitionObject.GetValue("Number", -1);
+
+				changeType.ReadFrom = definitionObject.GetValue("ReadFrom", null);
 				return true;
 			}
 			else
