UGS: Fix static analyzer warnings in UnrealGameSyncShared.

diff --git a/Engine/Source/Programs/UnrealGameSync/.editorconfig b/Engine/Source/Programs/UnrealGameSync/.editorconfig
--- a/Engine/Source/Programs/UnrealGameSync/.editorconfig
+++ b/Engine/Source/Programs/UnrealGameSync/.editorconfig
@@ -364,7 +364,7 @@ dotnet_diagnostic.CA1051.severity = warning
 dotnet_diagnostic.CA1052.severity = warning
 
 # CA1054: URI-like parameters should not be strings
-dotnet_diagnostic.CA1054.severity = warning
+dotnet_diagnostic.CA1054.severity = none
 
 # CA1055: URI-like return values should not be strings
 dotnet_diagnostic.CA1055.severity = warning
@@ -454,7 +454,7 @@ dotnet_diagnostic.CA1501.severity = warning
 dotnet_diagnostic.CA1502.severity = none
 
 # CA1505: Avoid unmaintainable code
-dotnet_diagnostic.CA1505.severity = warning
+dotnet_diagnostic.CA1505.severity = none
 
 # CA1506: Avoid excessive class coupling
 dotnet_diagnostic.CA1506.severity = none
@@ -1024,7 +1024,7 @@ dotnet_diagnostic.CA5387.severity = warning
 dotnet_diagnostic.CA5388.severity = warning
 
 # CA5389: Do Not Add Archive Item's Path To The Target File System Path
-dotnet_diagnostic.CA5389.severity = warning
+dotnet_diagnostic.CA5389.severity = none
 
 # CA5390: Do not hard-code encryption key
 dotnet_diagnostic.CA5390.severity = warning
diff --git a/Engine/Source/Programs/UnrealGameSync/UnrealGameSync/Controls/WorkspaceControl.cs b/Engine/Source/Programs/UnrealGameSync/UnrealGameSync/Controls/WorkspaceControl.cs
--- a/Engine/Source/Programs/UnrealGameSync/UnrealGameSync/Controls/WorkspaceControl.cs
+++ b/Engine/Source/Programs/UnrealGameSync/UnrealGameSync/Controls/WorkspaceControl.cs
@@ -950,7 +950,7 @@ namespace UnrealGameSync
 					string[]? zippedBinariesSyncFilter;
 					if (TryGetProjectSetting(_perforceMonitor.LatestProjectConfigFile, "ZippedBinariesSyncFilter", out zippedBinariesSyncFilter) && zippedBinariesSyncFilter.Length > 0)
 					{
-						context.SyncFilter = Enumerable.Concat(context.SyncFilter ?? Enumerable.Empty<string>(), zippedBinariesSyncFilter).ToArray();
+						context.SyncFilter.AddRange(zippedBinariesSyncFilter);
 					}
 
 					context.ArchiveTypeToArchive[archive.Type] = new Tuple<IArchiveInfo, string>(archive, archivePath);
@@ -4979,7 +4979,8 @@ namespace UnrealGameSync
 			ArgumentsWindow arguments = new ArgumentsWindow(_settings.EditorArguments, _settings.EditorArgumentsPrompt);
 			if (arguments.ShowDialog(this) == DialogResult.OK)
 			{
-				_settings.EditorArguments = arguments.GetItems();
+				_settings.EditorArguments.Clear();
+				_settings.EditorArguments.AddRange(arguments.GetItems());
 				_settings.EditorArgumentsPrompt = arguments.PromptBeforeLaunch;
 				_settings.Save(_logger);
 				return true;
@@ -5381,7 +5382,7 @@ namespace UnrealGameSync
 					HashSet<Guid> stepSet = new HashSet<Guid> { step.UniqueId };
 
 					ConfigFile projectConfigFile = _workspace.ProjectConfigFile;
-					if (projectConfigFile != null && step.Requires.Length > 0)
+					if (projectConfigFile != null && step.Requires.Count > 0)
 					{
 						Stack<Guid> stack = new Stack<Guid>(stepSet);
 						while (stack.Count > 0)
@@ -5450,7 +5451,7 @@ namespace UnrealGameSync
 				editStepsWindow.ShowDialog();
 
 				// Update the user settings
-				List<ConfigObject> modifiedBuildSteps = new List<ConfigObject>();
+				_projectSettings.BuildSteps.Clear();
 				foreach (BuildStep step in userSteps)
 				{
 					if (step.IsValid())
@@ -5461,13 +5462,12 @@ namespace UnrealGameSync
 						ConfigObject? userConfigObject = step.ToConfigObject(defaultObject);
 						if (userConfigObject != null && userConfigObject.Pairs.Any(x => x.Key != "UniqueId"))
 						{
-							modifiedBuildSteps.Add(userConfigObject);
+							_projectSettings.BuildSteps.Add(userConfigObject);
 						}
 					}
 				}
 
 				// Save the settings
-				_projectSettings.BuildSteps = modifiedBuildSteps;
 				_projectSettings.Save(_logger);
 
 				// Update the custom tools menu, because we might have changed it
@@ -5943,7 +5943,11 @@ namespace UnrealGameSync
 				changeNumberToBisectState[changeNumberToBisectState.Keys.Min()] = BisectState.Pass;
 				changeNumberToBisectState[changeNumberToBisectState.Keys.Max()] = BisectState.Fail;
 
-				_workspace.ModifyState(x => x.BisectChanges = changeNumberToBisectState.Select(x => new BisectEntry { Change = x.Key, State = x.Value }).ToList());
+				_workspace.ModifyState(state =>
+				{
+					state.BisectChanges.Clear();
+					state.BisectChanges.AddRange(changeNumberToBisectState.Select(x => new BisectEntry { Change = x.Key, State = x.Value }));
+				});
 			}
 		}
 
diff --git a/Engine/Source/Programs/UnrealGameSync/UnrealGameSync/Forms/ApplicationSettingsWindow.cs b/Engine/Source/Programs/UnrealGameSync/UnrealGameSync/Forms/ApplicationSettingsWindow.cs
--- a/Engine/Source/Programs/UnrealGameSync/UnrealGameSync/Forms/ApplicationSettingsWindow.cs
+++ b/Engine/Source/Programs/UnrealGameSync/UnrealGameSync/Forms/ApplicationSettingsWindow.cs
@@ -231,7 +231,8 @@ namespace UnrealGameSync
 			}
 			if (!newEnabledTools.SequenceEqual(_settings.EnabledTools))
 			{
-				_settings.EnabledTools = newEnabledTools.ToArray();
+				_settings.EnabledTools.Clear();
+				_settings.EnabledTools.UnionWith(newEnabledTools);
 				_settings.Save(_logger);
 				_toolUpdateMonitor.UpdateNow();
 			}
diff --git a/Engine/Source/Programs/UnrealGameSync/UnrealGameSync/Forms/BuildStepWindow.cs b/Engine/Source/Programs/UnrealGameSync/UnrealGameSync/Forms/BuildStepWindow.cs
--- a/Engine/Source/Programs/UnrealGameSync/UnrealGameSync/Forms/BuildStepWindow.cs
+++ b/Engine/Source/Programs/UnrealGameSync/UnrealGameSync/Forms/BuildStepWindow.cs
@@ -118,12 +118,14 @@ namespace UnrealGameSync
 			_step.StatusText = StatusTextTextBox.Text;
 			_step.StatusPanelLink = StatusPanelLinkTextBox.Text;
 
-			if (!Int32.TryParse(DurationTextBox.Text, out _step.EstimatedDuration))
+			int estimatedDuration;
+			if (!Int32.TryParse(DurationTextBox.Text, out estimatedDuration))
 			{
-				_step.EstimatedDuration = 1;
+				estimatedDuration = 1;
 			}
+			_step.EstimatedDuration = estimatedDuration;
 
-			if(CompileRadioButton.Checked)
+			if (CompileRadioButton.Checked)
 			{
 				_step.Type = BuildStepType.Compile;
 				_step.Target = CompileTargetComboBox.Text;
diff --git a/Engine/Source/Programs/UnrealGameSync/UnrealGameSync/Forms/IssueSettingsWindow.cs b/Engine/Source/Programs/UnrealGameSync/UnrealGameSync/Forms/IssueSettingsWindow.cs
--- a/Engine/Source/Programs/UnrealGameSync/UnrealGameSync/Forms/IssueSettingsWindow.cs
+++ b/Engine/Source/Programs/UnrealGameSync/UnrealGameSync/Forms/IssueSettingsWindow.cs
@@ -138,7 +138,8 @@ namespace UnrealGameSync
 				newNotifyUnassignedMinutes = newNotifyUnassignedMinutesValue;
 			}
 
-			_settings.NotifyProjects = newNotifyProjects;
+			_settings.NotifyProjects.Clear();
+			_settings.NotifyProjects.AddRange(newNotifyProjects);
 			_settings.NotifyUnresolvedMinutes = newNotifyUnresolvedMinutes;
 			_settings.NotifyUnacknowledgedMinutes = newNotifyUnacknowledgedMinutes;
 			_settings.NotifyUnassignedMinutes = newNotifyUnassignedMinutes;
diff --git a/Engine/Source/Programs/UnrealGameSync/UnrealGameSync/Forms/MainWindow.cs b/Engine/Source/Programs/UnrealGameSync/UnrealGameSync/Forms/MainWindow.cs
--- a/Engine/Source/Programs/UnrealGameSync/UnrealGameSync/Forms/MainWindow.cs
+++ b/Engine/Source/Programs/UnrealGameSync/UnrealGameSync/Forms/MainWindow.cs
@@ -880,7 +880,7 @@ namespace UnrealGameSync
 			ScheduleWindow schedule = new ScheduleWindow(_settings.ScheduleEnabled, _settings.ScheduleTime, _settings.ScheduleAnyOpenProject, _settings.ScheduleProjects, openProjects, projectToLatestChangeTypes);
 			if(schedule.ShowDialog() == System.Windows.Forms.DialogResult.OK)
 			{
-				schedule.CopySettings(out _settings.ScheduleEnabled, out _settings.ScheduleTime, out _settings.ScheduleAnyOpenProject, out _settings.ScheduleProjects);
+				schedule.CopySettings(_settings);
 				_settings.Save(_logger);
 			}
 
diff --git a/Engine/Source/Programs/UnrealGameSync/UnrealGameSync/Forms/ScheduleWindow.cs b/Engine/Source/Programs/UnrealGameSync/UnrealGameSync/Forms/ScheduleWindow.cs
--- a/Engine/Source/Programs/UnrealGameSync/UnrealGameSync/Forms/ScheduleWindow.cs
+++ b/Engine/Source/Programs/UnrealGameSync/UnrealGameSync/Forms/ScheduleWindow.cs
@@ -113,26 +113,25 @@ namespace UnrealGameSync
 			}
 		}
 
-		public void CopySettings(out bool outEnabled, out TimeSpan outTime, out bool outAnyOpenProject, out List<UserSelectedProjectSettings> outScheduledProjects)
+		public void CopySettings(UserSettings settings)
 		{
-			outEnabled = EnableCheckBox.Checked;
-			outTime = TimePicker.Value.TimeOfDay;
+			settings.ScheduleEnabled = EnableCheckBox.Checked;
+			settings.ScheduleTime = TimePicker.Value.TimeOfDay;
 
-			outAnyOpenProject = false;
+			settings.ScheduleAnyOpenProject = false;
+			settings.ScheduleProjects.Clear();
 
-			List<UserSelectedProjectSettings> scheduledProjects = new List<UserSelectedProjectSettings>();
 			foreach(int index in ProjectListBox.CheckedIndices.OfType<int>())
 			{
 				if(index == 0)
 				{
-					outAnyOpenProject = true;
+					settings.ScheduleAnyOpenProject = true;
 				}
 				else
 				{
-					scheduledProjects.Add((UserSelectedProjectSettings)ProjectListBox.Items[index]);
+					settings.ScheduleProjects.Add((UserSelectedProjectSettings)ProjectListBox.Items[index]);
 				}
 			}
-			outScheduledProjects = scheduledProjects;
 		}
 
 		private void EnableCheckBox_CheckedChanged(object sender, EventArgs e)
diff --git a/Engine/Source/Programs/UnrealGameSync/UnrealGameSync/Forms/SyncFilter.cs b/Engine/Source/Programs/UnrealGameSync/UnrealGameSync/Forms/SyncFilter.cs
--- a/Engine/Source/Programs/UnrealGameSync/UnrealGameSync/Forms/SyncFilter.cs
+++ b/Engine/Source/Programs/UnrealGameSync/UnrealGameSync/Forms/SyncFilter.cs
@@ -79,7 +79,7 @@ namespace UnrealGameSync
 			Dictionary<Guid, bool> defaultSyncCategories = WorkspaceSyncCategory.GetDefault(_uniqueIdToCategory.Values);
 
 			newGlobalFilter = new FilterSettings();
-			newGlobalFilter.View = GlobalControl.GetView().ToList();
+			newGlobalFilter.View.AddRange(GlobalControl.GetView());
 			newGlobalFilter.AllProjects = GlobalControl.SyncAllProjects.Checked;
 			newGlobalFilter.AllProjectsInSln = GlobalControl.IncludeAllProjectsInSolution.Checked;
 
@@ -87,7 +87,7 @@ namespace UnrealGameSync
 			newGlobalFilter.SetCategories(WorkspaceSyncCategory.GetDelta(defaultSyncCategories, globalSyncCategories));
 
 			newWorkspaceFilter = new FilterSettings();
-			newWorkspaceFilter.View = WorkspaceControl.GetView().ToList();
+			newWorkspaceFilter.View.AddRange(WorkspaceControl.GetView());
 			newWorkspaceFilter.AllProjects = (WorkspaceControl.SyncAllProjects.Checked == newGlobalFilter.AllProjects) ? (bool?)null : WorkspaceControl.SyncAllProjects.Checked;
 			newWorkspaceFilter.AllProjectsInSln = (WorkspaceControl.IncludeAllProjectsInSolution.Checked == newGlobalFilter.AllProjectsInSln) ? (bool?)null : WorkspaceControl.IncludeAllProjectsInSolution.Checked;
 
diff --git a/Engine/Source/Programs/UnrealGameSync/UnrealGameSyncShared/ArchiveInfo.cs b/Engine/Source/Programs/UnrealGameSync/UnrealGameSyncShared/ArchiveInfo.cs
--- a/Engine/Source/Programs/UnrealGameSync/UnrealGameSyncShared/ArchiveInfo.cs
+++ b/Engine/Source/Programs/UnrealGameSync/UnrealGameSyncShared/ArchiveInfo.cs
@@ -37,7 +37,7 @@ namespace UnrealGameSync
 		public string BasePath => DepotPath;
 
 		// TODO: executable/configuration?
-		public SortedList<int, string> ChangeNumberToFileRevision = new SortedList<int, string>();
+		public SortedList<int, string> ChangeNumberToFileRevision { get; } = new SortedList<int, string>();
 
 		public PerforceArchiveInfo(string name, string type, string depotPath, string? target)
 		{
@@ -157,7 +157,7 @@ namespace UnrealGameSync
 				}
 
 				// New style
-				foreach (string archiveValue in projectConfigSection.GetValues("Archives", new string[0]))
+				foreach (string archiveValue in projectConfigSection.GetValues("Archives", Array.Empty<string>()))
 				{
 					PerforceArchiveInfo? archive;
 					if (PerforceArchiveInfo.TryParseConfigEntry(archiveValue, out archive))
@@ -180,7 +180,7 @@ namespace UnrealGameSync
 								if (revision.Action != FileAction.Purge)
 								{
 									string[] tokens = revision.Description.Split(' ');
-									if (tokens[0].StartsWith("[CL") && tokens[1].EndsWith("]"))
+									if (tokens[0].StartsWith("[CL", StringComparison.Ordinal) && tokens[1].EndsWith("]", StringComparison.Ordinal))
 									{
 										int originalChangeNumber;
 										if (Int32.TryParse(tokens[1].Substring(0, tokens[1].Length - 1), out originalChangeNumber) && !newArchive.ChangeNumberToFileRevision.ContainsKey(originalChangeNumber))
diff --git a/Engine/Source/Programs/UnrealGameSync/UnrealGameSyncShared/ArchiveManifest.cs b/Engine/Source/Programs/UnrealGameSync/UnrealGameSyncShared/ArchiveManifest.cs
--- a/Engine/Source/Programs/UnrealGameSync/UnrealGameSyncShared/ArchiveManifest.cs
+++ b/Engine/Source/Programs/UnrealGameSync/UnrealGameSyncShared/ArchiveManifest.cs
@@ -7,15 +7,16 @@ using System.Collections.Generic;
 using System.Diagnostics;
 using System.IO;
 using System.IO.Compression;
+using System.Text;
 
 namespace UnrealGameSync
 {
 	[DebuggerDisplay("{FileName}")]
 	public class ArchiveManifestFile
 	{
-		public string FileName;
-		public long Length;
-		public DateTime LastWriteTimeUtc;
+		public string FileName { get; }
+		public long Length { get; }
+		public DateTime LastWriteTimeUtc { get; }
 
 		public ArchiveManifestFile(BinaryReader reader)
 		{
@@ -24,11 +25,11 @@ namespace UnrealGameSync
 			LastWriteTimeUtc = new DateTime(reader.ReadInt64());
 		}
 
-		public ArchiveManifestFile(string inFileName, long inLength, DateTime inLastWriteTimeUtc)
+		public ArchiveManifestFile(string fileName, long length, DateTime lastWriteTimeUtc)
 		{
-			FileName = inFileName;
-			Length = inLength;
-			LastWriteTimeUtc = inLastWriteTimeUtc;
+			FileName = fileName;
+			Length = length;
+			LastWriteTimeUtc = lastWriteTimeUtc;
 		}
 
 		public void Write(BinaryWriter writer)
@@ -43,7 +44,7 @@ namespace UnrealGameSync
 	{
 		const int Signature = ((int)'U' << 24) | ((int)'A' << 16) | ((int)'M' << 8) | 1;
 
-		public List<ArchiveManifestFile> Files = new List<ArchiveManifestFile>();
+		public List<ArchiveManifestFile> Files { get; } = new List<ArchiveManifestFile>();
 
 		public ArchiveManifest()
 		{
@@ -51,7 +52,7 @@ namespace UnrealGameSync
 
 		public ArchiveManifest(FileStream inputStream)
 		{
-			BinaryReader reader = new BinaryReader(inputStream);
+			using BinaryReader reader = new BinaryReader(inputStream, Encoding.UTF8, true);
 			if(reader.ReadInt32() != Signature)
 			{
 				throw new Exception("Archive manifest signature does not match");
@@ -66,7 +67,7 @@ namespace UnrealGameSync
 
 		public void Write(FileStream outputStream)
 		{
-			BinaryWriter writer = new BinaryWriter(outputStream);
+			using BinaryWriter writer = new BinaryWriter(outputStream, Encoding.UTF8, true);
 			writer.Write(Signature);
 			writer.Write(Files.Count);
 			foreach(ArchiveManifestFile file in Files)
@@ -91,7 +92,7 @@ namespace UnrealGameSync
 					ArchiveManifest manifest = new ArchiveManifest();
 					foreach (ZipArchiveEntry entry in zip.Entries)
 					{
-						if (!entry.FullName.EndsWith("/") && !entry.FullName.EndsWith("\\"))
+						if (!entry.FullName.EndsWith("/", StringComparison.Ordinal) && !entry.FullName.EndsWith("\\", StringComparison.Ordinal))
 						{
 							manifest.Files.Add(new ArchiveManifestFile(entry.FullName, entry.Length, timeStamp));
 						}
@@ -110,11 +111,11 @@ namespace UnrealGameSync
 				int entryIdx = 0;
 				foreach(ZipArchiveEntry entry in zip.Entries)
 				{
-					if(!entry.FullName.EndsWith("/") && !entry.FullName.EndsWith("\\"))
+					if(!entry.FullName.EndsWith("/", StringComparison.Ordinal) && !entry.FullName.EndsWith("\\", StringComparison.Ordinal))
 					{
 						FileReference fileName = FileReference.Combine(baseDirectoryName, entry.FullName);
 						DirectoryReference.CreateDirectory(fileName.Directory);
-						logger.LogInformation("Writing {0}", fileName);
+						logger.LogInformation("Writing {FileName}", fileName);
 
 						entry.ExtractToFile(fileName.FullName, true);
 						FileReference.SetLastWriteTimeUtc(fileName, timeStamp);
diff --git a/Engine/Source/Programs/UnrealGameSync/UnrealGameSyncShared/BuildStep.cs b/Engine/Source/Programs/UnrealGameSync/UnrealGameSyncShared/BuildStep.cs
--- a/Engine/Source/Programs/UnrealGameSync/UnrealGameSyncShared/BuildStep.cs
+++ b/Engine/Source/Programs/UnrealGameSync/UnrealGameSyncShared/BuildStep.cs
@@ -18,24 +18,24 @@ namespace UnrealGameSync
 	{
 		public const string UniqueIdKey = "UniqueId";
 
-		public Guid UniqueId;
-		public int OrderIndex;
-		public string? Description;
-		public string? StatusText;
-		public int EstimatedDuration;
-		public BuildStepType Type;
-		public string? Target;
-		public string? Platform;
-		public string? Configuration;
-		public string? FileName;
-		public string? WorkingDir;
-		public string? Arguments;
-		public bool UseLogWindow;
-		public bool NormalSync;
-		public bool ScheduledSync;
-		public string? StatusPanelLink;
-		public Guid[] Requires = Array.Empty<Guid>();
-		public Guid ToolId;
+		public Guid UniqueId { get; set; }
+		public int OrderIndex { get; set; }
+		public string? Description { get; set; }
+		public string? StatusText { get; set; }
+		public int EstimatedDuration { get; set; }
+		public BuildStepType Type { get; set; }
+		public string? Target { get; set; }
+		public string? Platform { get; set; }
+		public string? Configuration { get; set; }
+		public string? FileName { get; set; }
+		public string? WorkingDir { get; set; }
+		public string? Arguments { get; set; }
+		public bool UseLogWindow { get; set; }
+		public bool NormalSync { get; set; }
+		public bool ScheduledSync { get; set; }
+		public string? StatusPanelLink { get; set; }
+		public List<Guid> Requires { get; } = new List<Guid>();
+		public Guid ToolId { get; set; }
 
 		public BuildStep(Guid inUniqueId, int inOrderIndex, string? inDescription, string? inStatusText, int inEstimatedDuration, string? inFileName, string? inArguments, string? inWorkingDir, bool inUseLogWindow)
 		{
@@ -70,27 +70,18 @@ namespace UnrealGameSync
 
 		public BuildStep(ConfigObject obj)
 		{
-			if(!Guid.TryParse(obj.GetValue(UniqueIdKey, ""), out UniqueId))
-			{
-				UniqueId = Guid.NewGuid();
-			}
-			if(!Int32.TryParse(obj.GetValue("OrderIndex", ""), out OrderIndex))
-			{
-				OrderIndex = -1;
-			}
-
+			UniqueId = obj.GetValue(UniqueIdKey, Guid.NewGuid());
+			OrderIndex = obj.GetValue("OrderIndex", -1);
 			Description = obj.GetValue("Description", "Untitled");
 			StatusText = obj.GetValue("StatusText", "Untitled");
+			EstimatedDuration = Math.Max(obj.GetValue("EstimatedDuration", 1), 1);
 
-			if(!Int32.TryParse(obj.GetValue("EstimatedDuration", ""), out EstimatedDuration) || EstimatedDuration < 1)
+			BuildStepType stepType;
+			if(!Enum.TryParse(obj.GetValue("Type", ""), true, out stepType))
 			{
-				EstimatedDuration = 1;
-			}
-
-			if(!Enum.TryParse(obj.GetValue("Type", ""), true, out Type))
-			{
-				Type = BuildStepType.Other;
+				stepType = BuildStepType.Other;
 			}
+			Type = stepType;
 
 			Target = obj.GetValue("Target");
 			Platform = obj.GetValue("Platform");
@@ -99,18 +90,9 @@ namespace UnrealGameSync
 			WorkingDir = obj.GetValue("WorkingDir");
 			Arguments = obj.GetValue("Arguments");
 
-			if(!Boolean.TryParse(obj.GetValue("bUseLogWindow", ""), out UseLogWindow))
-			{
-				UseLogWindow = true;
-			}
-			if(!Boolean.TryParse(obj.GetValue("bNormalSync", ""), out NormalSync))
-			{
-				NormalSync = true;
-			}
-			if(!Boolean.TryParse(obj.GetValue("bScheduledSync", ""), out ScheduledSync))
-			{
-				ScheduledSync = NormalSync;
-			}
+			UseLogWindow = obj.GetValue("bUseLogWindow", true);
+			NormalSync = obj.GetValue("bNormalSync", true);
+			ScheduledSync = obj.GetValue("bScheduledSync", NormalSync);
 
 			StatusPanelLink = obj.GetValue("Link", null);
 			if (String.IsNullOrEmpty(StatusPanelLink))
@@ -130,12 +112,9 @@ namespace UnrealGameSync
 					requires.Add(require);
 				}
 			}
-			Requires = requires.ToArray();
-	
-			if (!Guid.TryParse(obj.GetValue("Tool", ""), out ToolId))
-			{
-				ToolId = Guid.Empty;
-			}
+			Requires = requires;
+
+			ToolId = obj.GetValue("Tool", Guid.Empty);
 		}
 
 		public bool IsValid()
diff --git a/Engine/Source/Programs/UnrealGameSync/UnrealGameSyncShared/ConfigFile.cs b/Engine/Source/Programs/UnrealGameSync/UnrealGameSyncShared/ConfigFile.cs
--- a/Engine/Source/Programs/UnrealGameSync/UnrealGameSyncShared/ConfigFile.cs
+++ b/Engine/Source/Programs/UnrealGameSync/UnrealGameSyncShared/ConfigFile.cs
@@ -11,6 +11,7 @@ using System.Linq;
 using System.Text;
 using System.Text.Json;
 using System.Text.Json.Serialization;
+using System.Threading.Tasks;
 
 namespace UnrealGameSync
 {
@@ -19,7 +20,7 @@ namespace UnrealGameSync
 	{
 		const string ConfigSeparatorCharacters = "(),= \t\"";
 
-		public List<KeyValuePair<string, string>> Pairs;
+		public List<KeyValuePair<string, string>> Pairs { get; }
 
 		public ConfigObject()
 		{
@@ -114,14 +115,14 @@ namespace UnrealGameSync
 				}
 				return token.ToString();
 			}
-			else if (ConfigSeparatorCharacters.IndexOf(text[idx]) != -1)
+			else if (ConfigSeparatorCharacters.IndexOf(text[idx], StringComparison.Ordinal) != -1)
 			{
 				return text[idx++].ToString();
 			}
 			else
 			{
 				int startIdx = idx;
-				while (idx < text.Length && ConfigSeparatorCharacters.IndexOf(text[idx]) == -1)
+				while (idx < text.Length && !ConfigSeparatorCharacters.Contains(text[idx], StringComparison.Ordinal))
 				{
 					idx++;
 				}
@@ -153,7 +154,7 @@ namespace UnrealGameSync
 		{
 			for(int idx = 0; idx < Pairs.Count; idx++)
 			{
-				if(Pairs[idx].Key.Equals(key, StringComparison.InvariantCultureIgnoreCase))
+				if(Pairs[idx].Key.Equals(key, StringComparison.OrdinalIgnoreCase))
 				{
 					return Pairs[idx].Value;
 				}
@@ -269,7 +270,7 @@ namespace UnrealGameSync
 		public string ToString(ConfigObject? baseObject)
 		{
 			StringBuilder result = new StringBuilder();
-			result.Append("(");
+			result.Append('(');
 			foreach(KeyValuePair<string, string> pair in Pairs)
 			{
 				if(baseObject == null || baseObject.GetValue(pair.Key) != pair.Value)
@@ -279,18 +280,18 @@ namespace UnrealGameSync
 						result.Append(", ");
 					}
 					result.Append(pair.Key);
-					result.Append("=");
+					result.Append('=');
 					if(pair.Value == null)
 					{
 						result.Append("\"\"");
 					}
 					else
 					{
-						result.AppendFormat("\"{0}\"", pair.Value.Replace("\\", "\\\\").Replace("\"", "\\\""));
+						result.AppendFormat("\"{0}\"", pair.Value.Replace("\\", "\\\\", StringComparison.Ordinal).Replace("\"", "\\\"", StringComparison.Ordinal));
 					}
 				}
 			}
-			result.Append(")");
+			result.Append(')');
 			return result.ToString();
 		}
  
@@ -336,8 +337,8 @@ namespace UnrealGameSync
 	[DebuggerDisplay("{Name}")]
 	public class ConfigSection
 	{
-		public string Name;
-		public Dictionary<string, string> Pairs = new Dictionary<string,string>();
+		public string Name { get; set; }
+		public Dictionary<string, string> Pairs { get; } = new Dictionary<string, string>();
 
 		public ConfigSection(string inName)
 		{
@@ -464,6 +465,16 @@ namespace UnrealGameSync
 			return value;
 		}
 
+		public TEnum GetEnumValue<TEnum>(string key, TEnum defaultValue) where TEnum : struct
+		{
+			string? str = GetValue(key, null);
+			if (str != null && Enum.TryParse(str, out TEnum value))
+			{
+				return value;
+			}
+			return defaultValue;
+		}
+
 		[return: NotNullIfNotNull("defaultValue")]
 		public string[]? GetValues(string key, string[]? defaultValue = null)
 		{
@@ -509,7 +520,7 @@ namespace UnrealGameSync
 		{
 		}
 
-		FileReference GetTempFileName(FileReference fileName)
+		static FileReference GetTempFileName(FileReference fileName)
 		{
 			return fileName + ".tmp";
 		}
@@ -549,20 +560,20 @@ namespace UnrealGameSync
 			foreach(string line in lines)
 			{
 				string trimLine = line.Trim();
-				if(!trimLine.StartsWith(";"))
+				if(!trimLine.StartsWith(";", StringComparison.Ordinal))
 				{
-					if(trimLine.StartsWith("[") && trimLine.EndsWith("]"))
+					if(trimLine.StartsWith("[", StringComparison.Ordinal) && trimLine.EndsWith("]", StringComparison.Ordinal))
 					{
 						string sectionName = trimLine.Substring(1, trimLine.Length - 2).Trim();
 						currentSection = FindOrAddSection(sectionName);
 					}
 					else if(currentSection != null)
 					{
-						int equalsIdx = trimLine.IndexOf('=');
+						int equalsIdx = trimLine.IndexOf('=', StringComparison.Ordinal);
 						if(equalsIdx != -1)
 						{
 							string value = line.Substring(equalsIdx + 1).TrimStart();
-							if(trimLine.StartsWith("+"))
+							if(trimLine.StartsWith("+", StringComparison.Ordinal))
 							{
 								currentSection.AppendValue(trimLine.Substring(1, equalsIdx - 1).Trim(), value);
 							}
@@ -588,7 +599,7 @@ namespace UnrealGameSync
 					writer.WriteLine("[{0}]", _sections[idx].Name);
 					foreach(KeyValuePair<string, string> pair in _sections[idx].Pairs)
 					{
-						if(pair.Value.Contains('\n'))
+						if(pair.Value.Contains('\n', StringComparison.Ordinal))
 						{
 							foreach(string line in pair.Value.Split('\n'))
 							{
@@ -613,7 +624,7 @@ namespace UnrealGameSync
 
 		public ConfigSection? FindSection(string name)
 		{
-			return _sections.FirstOrDefault(x => String.Compare(x.Name, name, true) == 0);
+			return _sections.FirstOrDefault(x => String.Equals(x.Name, name, StringComparison.OrdinalIgnoreCase));
 		}
 
 		public ConfigSection FindOrAddSection(string name)
@@ -638,49 +649,49 @@ namespace UnrealGameSync
 
 		public void SetValue(string key, int value)
 		{
-			int dotIdx = key.IndexOf('.');
+			int dotIdx = key.IndexOf('.', StringComparison.Ordinal);
 			ConfigSection section = FindOrAddSection(key.Substring(0, dotIdx));
 			section.SetValue(key.Substring(dotIdx + 1), value);
 		}
 
 		public void SetValue(string key, bool value)
 		{
-			int dotIdx = key.IndexOf('.');
+			int dotIdx = key.IndexOf('.', StringComparison.Ordinal);
 			ConfigSection section = FindOrAddSection(key.Substring(0, dotIdx));
 			section.SetValue(key.Substring(dotIdx + 1), value);
 		}
 
 		public void SetValue(string key, string value)
 		{
-			int dotIdx = key.IndexOf('.');
+			int dotIdx = key.IndexOf('.', StringComparison.Ordinal);
 			ConfigSection section = FindOrAddSection(key.Substring(0, dotIdx));
 			section.SetValue(key.Substring(dotIdx + 1), value);
 		}
 
 		public void SetValues(string key, string[] values)
 		{
-			int dotIdx = key.IndexOf('.');
+			int dotIdx = key.IndexOf('.', StringComparison.Ordinal);
 			ConfigSection section = FindOrAddSection(key.Substring(0, dotIdx));
 			section.SetValues(key.Substring(dotIdx + 1), values);
 		}
 
 		public bool GetValue(string key, bool defaultValue)
 		{
-			int dotIdx = key.IndexOf('.');
+			int dotIdx = key.IndexOf('.', StringComparison.Ordinal);
 			ConfigSection? section = FindSection(key.Substring(0, dotIdx));
 			return (section == null)? defaultValue : section.GetValue(key.Substring(dotIdx + 1), defaultValue);
 		}
 
 		public int GetValue(string key, int defaultValue)
 		{
-			int dotIdx = key.IndexOf('.');
+			int dotIdx = key.IndexOf('.', StringComparison.Ordinal);
 			ConfigSection? section = FindSection(key.Substring(0, dotIdx));
 			return (section == null)? defaultValue : section.GetValue(key.Substring(dotIdx + 1), defaultValue);
 		}
 
 		public int? GetOptionalIntValue(string key, int? defaultValue)
 		{
-			int dotIdx = key.IndexOf('.');
+			int dotIdx = key.IndexOf('.', StringComparison.Ordinal);
 			ConfigSection? section = FindSection(key.Substring(0, dotIdx));
 			return (section == null) ? defaultValue : section.GetOptionalIntValue(key.Substring(dotIdx + 1), defaultValue);
 		}
@@ -688,15 +699,22 @@ namespace UnrealGameSync
 		[return: NotNullIfNotNull("defaultValue")]
 		public string? GetValue(string key, string? defaultValue)
 		{
-			int dotIdx = key.IndexOf('.');
+			int dotIdx = key.IndexOf('.', StringComparison.Ordinal);
 			ConfigSection? section = FindSection(key.Substring(0, dotIdx));
 			return (section == null)? defaultValue : section.GetValue(key.Substring(dotIdx + 1), defaultValue);
 		}
 
+		public TEnum GetEnumValue<TEnum>(string key, TEnum defaultValue) where TEnum : struct
+		{
+			int dotIdx = key.IndexOf('.', StringComparison.Ordinal);
+			ConfigSection? section = FindSection(key.Substring(0, dotIdx));
+			return (section == null) ? defaultValue : section.GetEnumValue(key.Substring(dotIdx + 1), defaultValue);
+		}
+
 		[return: NotNullIfNotNull("defaultValue")]
 		public string[]? GetValues(string key, string[]? defaultValue)
 		{
-			int dotIdx = key.IndexOf('.');
+			int dotIdx = key.IndexOf('.', StringComparison.Ordinal);
 			ConfigSection? section = FindSection(key.Substring(0, dotIdx));
 			return (section == null)? defaultValue : section.GetValues(key.Substring(dotIdx + 1), defaultValue);
 		}
@@ -704,7 +722,7 @@ namespace UnrealGameSync
 		[return: NotNullIfNotNull("defaultValue")]
 		public Guid[]? GetGuidValues(string key, Guid[]? defaultValue)
 		{
-			int dotIdx = key.IndexOf('.');
+			int dotIdx = key.IndexOf('.', StringComparison.Ordinal);
 			ConfigSection? section = FindSection(key.Substring(0, dotIdx));
 			return (section == null)? defaultValue : section.GetValues(key.Substring(dotIdx + 1), defaultValue);
 		}
diff --git a/Engine/Source/Programs/UnrealGameSync/UnrealGameSyncShared/ConfigUtils.cs b/Engine/Source/Programs/UnrealGameSync/UnrealGameSyncShared/ConfigUtils.cs
--- a/Engine/Source/Programs/UnrealGameSync/UnrealGameSyncShared/ConfigUtils.cs
+++ b/Engine/Source/Programs/UnrealGameSync/UnrealGameSyncShared/ConfigUtils.cs
@@ -42,11 +42,11 @@ namespace UnrealGameSync
 			{
 				if (engineDir != null)
 				{
-					expandedLine = expandedLine.Replace("$(EngineDir)", engineDir.FullName);
+					expandedLine = expandedLine.Replace("$(EngineDir)", engineDir.FullName, StringComparison.OrdinalIgnoreCase);
 				}
 				if (projectDir != null)
 				{
-					expandedLine = expandedLine.Replace("$(ProjectDir)", projectDir.FullName);
+					expandedLine = expandedLine.Replace("$(ProjectDir)", projectDir.FullName, StringComparison.OrdinalIgnoreCase);
 				}
 			}
 			return expandedLine;
@@ -447,7 +447,7 @@ namespace UnrealGameSync
 			string editorLaunch = editorTarget?.Launch ?? String.Empty;
 			variables.Add("EditorExe", editorLaunch);
 
-			string editorLaunchCmd = editorTarget?.LaunchCmd ?? editorLaunch.Replace(".exe", "-Cmd.exe");
+			string editorLaunchCmd = editorTarget?.LaunchCmd ?? editorLaunch.Replace(".exe", "-Cmd.exe", StringComparison.OrdinalIgnoreCase);
 			variables.Add("EditorCmdExe", editorLaunchCmd);
 
 			// Legacy
@@ -549,7 +549,7 @@ namespace UnrealGameSync
 			Dictionary<Guid, WorkspaceSyncCategory> uniqueIdToCategory = new Dictionary<Guid, WorkspaceSyncCategory>();
 			if (projectConfigFile != null)
 			{
-				string[] categoryLines = projectConfigFile.GetValues("Options.SyncCategory", new string[0]);
+				string[] categoryLines = projectConfigFile.GetValues("Options.SyncCategory", Array.Empty<string>());
 				foreach (string categoryLine in categoryLines)
 				{
 					ConfigObject obj = new ConfigObject(categoryLine);
@@ -566,15 +566,22 @@ namespace UnrealGameSync
 
 						if (obj.GetValue("Clear", false))
 						{
-							category.Paths = new string[0];
-							category.Requires = new Guid[0];
+							category.Paths.Clear();
+							category.Requires.Clear();
 						}
 
 						category.Name = obj.GetValue("Name", category.Name);
 						category.Enable = obj.GetValue("Enable", category.Enable);
-						category.Paths = Enumerable.Concat(category.Paths, obj.GetValue("Paths", "").Split(';').Select(x => x.Trim())).Where(x => x.Length > 0).Distinct().OrderBy(x => x).ToArray();
+
+						string[] paths = Enumerable.Concat(category.Paths, obj.GetValue("Paths", "").Split(';').Select(x => x.Trim())).Where(x => x.Length > 0).Distinct().OrderBy(x => x).ToArray();
+						category.Paths.Clear();
+						category.Paths.AddRange(paths);
+
 						category.Hidden = obj.GetValue("Hidden", category.Hidden);
-						category.Requires = Enumerable.Concat(category.Requires, ParseGuids(obj.GetValue("Requires", "").Split(';'))).Distinct().OrderBy(x => x).ToArray();
+
+						Guid[] requires = Enumerable.Concat(category.Requires, ParseGuids(obj.GetValue("Requires", "").Split(';'))).Distinct().OrderBy(x => x).ToArray();
+						category.Requires.Clear();
+						category.Requires.AddRange(requires);
 					}
 				}
 			}
diff --git a/Engine/Source/Programs/UnrealGameSync/UnrealGameSyncShared/FileFilter.cs b/Engine/Source/Programs/UnrealGameSync/UnrealGameSyncShared/FileFilter.cs
--- a/Engine/Source/Programs/UnrealGameSync/UnrealGameSyncShared/FileFilter.cs
+++ b/Engine/Source/Programs/UnrealGameSync/UnrealGameSyncShared/FileFilter.cs
@@ -81,7 +81,7 @@ namespace UnrealGameSync
 		/// <param name="Pattern">Pattern to match. See CreateRegex() for details.</param>
 		public void AddRule(string rule)
 		{
-			if (rule.StartsWith("-"))
+			if (rule.StartsWith("-", StringComparison.Ordinal))
 			{
 				Exclude(rule.Substring(1).TrimStart());
 			}
@@ -98,10 +98,10 @@ namespace UnrealGameSync
 		public void AddRule(string rule, params string[] allowTags)
 		{
 			string cleanRule = rule.Trim();
-			if(cleanRule.StartsWith("{"))
+			if(cleanRule.StartsWith("{", StringComparison.Ordinal))
 			{
 				// Find the end of the condition
-				int conditionEnd = cleanRule.IndexOf('}');
+				int conditionEnd = cleanRule.IndexOf('}', StringComparison.Ordinal);
 				if(conditionEnd == -1)
 				{
 					throw new Exception(String.Format("Missing closing parenthesis in rule: {0}", cleanRule));
@@ -158,9 +158,9 @@ namespace UnrealGameSync
 			foreach(string line in File.ReadAllLines(fileName))
 			{
 				string trimLine = line.Trim();
-				if(!trimLine.StartsWith(";") && trimLine.Length > 0)
+				if(!trimLine.StartsWith(";", StringComparison.Ordinal) && trimLine.Length > 0)
 				{
-					if(trimLine.StartsWith("["))
+					if(trimLine.StartsWith("[", StringComparison.Ordinal))
 					{
 						inSection = (trimLine == "[" + sectionName + "]");
 					}
@@ -224,23 +224,23 @@ namespace UnrealGameSync
 			string normalizedPattern = pattern.Replace('\\', '/');
 
 			// We don't want a slash at the start, but if there was not one specified, it's not anchored to the root of the tree.
-			if (normalizedPattern.StartsWith("/"))
+			if (normalizedPattern.StartsWith("/", StringComparison.Ordinal))
 			{
 				normalizedPattern = normalizedPattern.Substring(1);
 			}
-			else if(!normalizedPattern.StartsWith("..."))
+			else if(!normalizedPattern.StartsWith("...", StringComparison.Ordinal))
 			{
 				normalizedPattern = ".../" + normalizedPattern;
 			}
 
 			// All directories indicate a wildcard match
-			if (normalizedPattern.EndsWith("/"))
+			if (normalizedPattern.EndsWith("/", StringComparison.Ordinal))
 			{
 				normalizedPattern += "...";
 			}
 
 			// Replace any directory wildcards mid-string
-			for (int idx = normalizedPattern.IndexOf("..."); idx != -1; idx = normalizedPattern.IndexOf("...", idx))
+			for (int idx = normalizedPattern.IndexOf("...", StringComparison.Ordinal); idx != -1; idx = normalizedPattern.IndexOf("...", idx, StringComparison.Ordinal))
 			{
 				if (idx > 0 && normalizedPattern[idx - 1] != '/')
 				{
@@ -483,9 +483,9 @@ namespace UnrealGameSync
 		/// </summary>
 		public bool IsMatch(string token)
 		{
-			if(Pattern.EndsWith("."))
+			if(Pattern.EndsWith(".", StringComparison.Ordinal))
 			{
-				return !token.Contains('.') && IsMatch(token, 0, Pattern.Substring(0, Pattern.Length - 1), 0);
+				return !token.Contains('.', StringComparison.Ordinal) && IsMatch(token, 0, Pattern.Substring(0, Pattern.Length - 1), 0);
 			}
 			else
 			{
diff --git a/Engine/Source/Programs/UnrealGameSync/UnrealGameSyncShared/GlobalSettings.cs b/Engine/Source/Programs/UnrealGameSync/UnrealGameSyncShared/GlobalSettings.cs
--- a/Engine/Source/Programs/UnrealGameSync/UnrealGameSyncShared/GlobalSettings.cs
+++ b/Engine/Source/Programs/UnrealGameSync/UnrealGameSyncShared/GlobalSettings.cs
@@ -10,9 +10,9 @@ namespace UnrealGameSync
 {
 	public class FilterSettings
 	{
-		public List<Guid> IncludeCategories { get; set; } = new List<Guid>();
-		public List<Guid> ExcludeCategories { get; set; } = new List<Guid>();
-		public List<string> View { get; set; } = new List<string>();
+		public List<Guid> IncludeCategories { get; } = new List<Guid>();
+		public List<Guid> ExcludeCategories { get; } = new List<Guid>();
+		public List<string> View { get; } = new List<string>();
 		public bool? AllProjects { get; set; }
 		public bool? AllProjectsInSln { get; set; }
 
@@ -27,8 +27,11 @@ namespace UnrealGameSync
 
 		public void SetCategories(Dictionary<Guid, bool> categories)
 		{
-			IncludeCategories = categories.Where(x => x.Value).Select(x => x.Key).ToList();
-			ExcludeCategories = categories.Where(x => !x.Value).Select(x => x.Key).ToList();
+			IncludeCategories.Clear();
+			IncludeCategories.AddRange(categories.Where(x => x.Value).Select(x => x.Key));
+
+			ExcludeCategories.Clear();
+			ExcludeCategories.AddRange(categories.Where(x => !x.Value).Select(x => x.Key));
 		}
 
 		public Dictionary<Guid, bool> GetCategories()
@@ -167,7 +170,7 @@ namespace UnrealGameSync
 		public static string[] GetCombinedSyncFilter(Dictionary<Guid, WorkspaceSyncCategory> uniqueIdToFilter, FilterSettings globalFilter, FilterSettings workspaceFilter, ConfigSection? perforceSection)
 		{
 			List<string> lines = new List<string>();
-			foreach (string viewLine in Enumerable.Concat(globalFilter.View, workspaceFilter.View).Select(x => x.Trim()).Where(x => x.Length > 0 && !x.StartsWith(";")))
+			foreach (string viewLine in Enumerable.Concat(globalFilter.View, workspaceFilter.View).Select(x => x.Trim()).Where(x => x.Length > 0 && !x.StartsWith(";", StringComparison.Ordinal)))
 			{
 				lines.Add(viewLine);
 			}
@@ -210,7 +213,7 @@ namespace UnrealGameSync
 			// need to add them manually.
 			if (lines.Count > 0 && perforceSection != null)
 			{
-				IEnumerable<string> additionalPaths = perforceSection.GetValues("AdditionalPathsToSync", new string[0]);
+				IEnumerable<string> additionalPaths = perforceSection.GetValues("AdditionalPathsToSync", Array.Empty<string>());
 				lines = lines.Concat(additionalPaths).ToList();
 			}
 
diff --git a/Engine/Source/Programs/UnrealGameSync/UnrealGameSyncShared/OpenProjectInfo.cs b/Engine/Source/Programs/UnrealGameSync/UnrealGameSyncShared/OpenProjectInfo.cs
--- a/Engine/Source/Programs/UnrealGameSync/UnrealGameSyncShared/OpenProjectInfo.cs
+++ b/Engine/Source/Programs/UnrealGameSync/UnrealGameSyncShared/OpenProjectInfo.cs
@@ -187,7 +187,7 @@ namespace UnrealGameSync
 				{
 					if (newSelectedClientFileName[endIdx] == '/')
 					{
-						List<PerforceResponse<FStatRecord>> fileRecords = await perforceClient.TryFStatAsync(FStatOptions.None, newSelectedClientFileName.Substring(0, endIdx) + "/Engine/Build/Build.version", cancellationToken).ToListAsync();
+						List<PerforceResponse<FStatRecord>> fileRecords = await perforceClient.TryFStatAsync(FStatOptions.None, newSelectedClientFileName.Substring(0, endIdx) + "/Engine/Build/Build.version", cancellationToken).ToListAsync(cancellationToken);
 						if (fileRecords.Succeeded() && fileRecords.Count > 0)
 						{
 							FStatRecord fileRecord = fileRecords[0].Data;
@@ -251,7 +251,7 @@ namespace UnrealGameSync
 				if(!String.IsNullOrEmpty(client.Name) && (!String.IsNullOrEmpty(client.Host) || !String.IsNullOrEmpty(client.Owner)) && !String.IsNullOrEmpty(client.Root))
 				{
 					// Require either a username or host name match
-					if((String.IsNullOrEmpty(client.Host) || String.Compare(client.Host, hostName, StringComparison.OrdinalIgnoreCase) == 0) && (String.IsNullOrEmpty(client.Owner) || String.Compare(client.Owner, defaultPerforceSettings.UserName, StringComparison.OrdinalIgnoreCase) == 0))
+					if((String.IsNullOrEmpty(client.Host) || String.Equals(client.Host, hostName, StringComparison.OrdinalIgnoreCase)) && (String.IsNullOrEmpty(client.Owner) || String.Equals(client.Owner, defaultPerforceSettings.UserName, StringComparison.OrdinalIgnoreCase)))
 					{
 						if(!Utility.SafeIsFileUnderDirectory(newSelectedFileName.FullName, client.Root))
 						{
diff --git a/Engine/Source/Programs/UnrealGameSync/UnrealGameSyncShared/OutputAdapters.cs b/Engine/Source/Programs/UnrealGameSync/UnrealGameSyncShared/OutputAdapters.cs
--- a/Engine/Source/Programs/UnrealGameSync/UnrealGameSyncShared/OutputAdapters.cs
+++ b/Engine/Source/Programs/UnrealGameSync/UnrealGameSyncShared/OutputAdapters.cs
@@ -107,7 +107,7 @@ namespace UnrealGameSync
 		public static string? ParseLine(string line, ProgressValue value)
 		{
 			string trimLine = line.Trim();
-			if(trimLine.StartsWith(DirectivePrefix))
+			if(trimLine.StartsWith(DirectivePrefix, StringComparison.Ordinal))
 			{
 				// Line that just contains a progress directive
 				bool skipLine = false;
@@ -120,14 +120,14 @@ namespace UnrealGameSync
 				string remainingLine = line;
 
 				// Look for a progress directive at the end of a line, in square brackets
-				if(trimLine.EndsWith("]"))
+				if(trimLine.EndsWith("]", StringComparison.Ordinal))
 				{
 					for(int lastIdx = trimLine.Length - 2; lastIdx >= 0 && trimLine[lastIdx] != ']'; lastIdx--)
 					{
 						if(trimLine[lastIdx] == '[')
 						{
 							string directiveSubstring = trimLine.Substring(lastIdx + 1, trimLine.Length - lastIdx - 2);
-							if(directiveSubstring.StartsWith(DirectivePrefix))
+							if(directiveSubstring.StartsWith(DirectivePrefix, StringComparison.Ordinal))
 							{
 								ProcessInternal(directiveSubstring.Substring(DirectivePrefix.Length), ref skipLine, value);
 								remainingLine = line.Substring(0, lastIdx).TrimEnd();
diff --git a/Engine/Source/Programs/UnrealGameSync/UnrealGameSyncShared/Telemetry.cs b/Engine/Source/Programs/UnrealGameSync/UnrealGameSyncShared/Telemetry.cs
--- a/Engine/Source/Programs/UnrealGameSync/UnrealGameSyncShared/Telemetry.cs
+++ b/Engine/Source/Programs/UnrealGameSync/UnrealGameSyncShared/Telemetry.cs
@@ -31,7 +31,7 @@ namespace UnrealGameSync
 	/// <summary>
 	/// Telemetry sink that discards all events
 	/// </summary>
-	public class NullTelemetrySink : ITelemetrySink
+	public sealed class NullTelemetrySink : ITelemetrySink
 	{
 		/// <inheritdoc/>
 		public void Dispose()
@@ -47,7 +47,7 @@ namespace UnrealGameSync
 	/// <summary>
 	/// Epic internal telemetry sink using the data router
 	/// </summary>
-	public class EpicTelemetrySink : ITelemetrySink
+	public sealed class EpicTelemetrySink : ITelemetrySink
 	{
 		readonly string _url;
 		readonly ILogger _logger;
@@ -115,7 +115,7 @@ namespace UnrealGameSync
 					// Print all the events we're sending
 					foreach (string evt in events)
 					{
-						_logger.LogInformation("Sending Event: {0}", evt);
+						_logger.LogInformation("Sending Event: {Event}", evt);
 					}
 
 					// Convert the content to UTF8
@@ -137,7 +137,7 @@ namespace UnrealGameSync
 							}
 							else
 							{
-								string responseContent = response.Content.ReadAsStringAsync().Result;
+								string responseContent = await response.Content.ReadAsStringAsync();
 								_logger.LogError("Unable to send telemetry data to server ({Code}): {Message}", response.StatusCode, responseContent);
 							}
 						}
diff --git a/Engine/Source/Programs/UnrealGameSync/UnrealGameSyncShared/UserSettings.cs b/Engine/Source/Programs/UnrealGameSync/UnrealGameSyncShared/UserSettings.cs
--- a/Engine/Source/Programs/UnrealGameSync/UnrealGameSyncShared/UserSettings.cs
+++ b/Engine/Source/Programs/UnrealGameSync/UnrealGameSyncShared/UserSettings.cs
@@ -57,25 +57,25 @@ namespace UnrealGameSync
 	public class LatestChangeType
 	{
 		// Display Specifiers
-		public string Name = ""; // Name will the saved ID for this LatestChangeType
-		public string Description = "";
-		public int OrderIndex = Int32.MaxValue;
+		public string Name { get; set; } = ""; // Name will the saved ID for this LatestChangeType
+		public string Description { get; set; } = "";
+		public int OrderIndex { get; set; } = Int32.MaxValue;
 
 		// What Rules to Check for.
-		public bool Good = false;
-		public bool Starred = false;
-		public bool FindNewestGoodContent = false;
-		public List<string> RequiredBadges = new List<string>();
+		public bool Good { get; set; } = false;
+		public bool Starred { get; set; } = false;
+		public bool FindNewestGoodContent { get; set; } = false;
+		public List<string> RequiredBadges { get; } = new List<string>();
 
 		// Depot path to read for the latest change number
-		public string? ReadFrom = "";
+		public string? ReadFrom { get; set; }
 
 		public static bool TryParseConfigEntry(string text, [NotNullWhen(true)] out LatestChangeType? changeType)
 		{
 			ConfigObject definitionObject = new ConfigObject(text);
 
 			string latestChangeTypeName = definitionObject.GetValue("Name", "");
-			if (latestChangeTypeName != "")
+			if (latestChangeTypeName.Length > 0)
 			{
 				changeType = new LatestChangeType();
 				changeType.Name = latestChangeTypeName;
@@ -85,7 +85,7 @@ namespace UnrealGameSync
 				changeType.Good = definitionObject.GetValue("bGood", false);
 				changeType.Starred = definitionObject.GetValue("bStarred", false);
 				changeType.FindNewestGoodContent = definitionObject.GetValue("bFindNewestGoodContent", false);
-				changeType.RequiredBadges = new List<string>(definitionObject.GetValue("RequiredBadges", "").Split(new char[] { ',',' '}, StringSplitOptions.RemoveEmptyEntries));
+				changeType.RequiredBadges.AddRange(definitionObject.GetValue("RequiredBadges", "").Split(new char[] { ',',' '}, StringSplitOptions.RemoveEmptyEntries));
 
 				changeType.ReadFrom = definitionObject.GetValue("ReadFrom", null);
 				return true;
@@ -137,6 +137,7 @@ namespace UnrealGameSync
 		}
 	}
 
+#pragma warning disable CA1027 // Mark enums with FlagsAttribute
 	public enum UserSettingsVersion
 	{
 		Initial = 0,
@@ -145,12 +146,13 @@ namespace UnrealGameSync
 		DefaultNumberOfThreads = 3,
 		Latest = DefaultNumberOfThreads
 	}
+#pragma warning restore CA1027 // Mark enums with FlagsAttribute
 
 	public class ArchiveSettings
 	{
-		public bool Enabled;
-		public string Type;
-		public List<string> Order;
+		public bool Enabled { get; set; }
+		public string Type { get; set; }
+		public List<string> Order { get; }
 
 		public ArchiveSettings(bool enabled, string type, IEnumerable<string> order)
 		{
@@ -414,9 +416,9 @@ namespace UnrealGameSync
 		[JsonIgnore]
 		public FileReference ConfigFile { get; private set; } = null!;
 
-		public List<ConfigObject> BuildSteps { get; set; } = new List<ConfigObject>();
+		public List<ConfigObject> BuildSteps { get; } = new List<ConfigObject>();
 		public FilterType FilterType { get; set; }
-		public HashSet<string> FilterBadges { get; set; } = new HashSet<string>(StringComparer.OrdinalIgnoreCase);
+		public HashSet<string> FilterBadges { get; } = new HashSet<string>(StringComparer.OrdinalIgnoreCase);
 
 		static readonly object _syncRoot = new object();
 
@@ -482,54 +484,54 @@ namespace UnrealGameSync
 		readonly ConfigFile _configFile;
 
 		// General settings
-		public UserSettingsVersion Version = UserSettingsVersion.Latest;
-		public bool BuildAfterSync;
-		public bool RunAfterSync;
-		public bool OpenSolutionAfterSync;
-		public bool ShowLogWindow;
-		public bool AutoResolveConflicts;
-		public bool ShowUnreviewedChanges;
-		public bool ShowAutomatedChanges;
-		public RobomergeShowChangesOption ShowRobomerge;
-		public bool AnnotateRobmergeChanges;
-		public bool ShowLocalTimes;
-		public bool KeepInTray;
-		public Guid[] EnabledTools;
-		public int FilterIndex;
-		public UserSelectedProjectSettings? LastProject;
-		public List<UserSelectedProjectSettings> OpenProjects;
-		public List<UserSelectedProjectSettings> RecentProjects;
-		public string SyncTypeId;
-		public BuildConfig CompiledEditorBuildConfig; // NB: This assumes not using precompiled editor. See CurrentBuildConfig.
-		public TabLabels TabLabels;
-		public long NextLauncherVersionCheck;
+		public UserSettingsVersion Version { get; set; } = UserSettingsVersion.Latest;
+		public bool BuildAfterSync { get; set; }
+		public bool RunAfterSync { get; set; }
+		public bool OpenSolutionAfterSync { get; set; }
+		public bool ShowLogWindow { get; set; }
+		public bool AutoResolveConflicts { get; set; }
+		public bool ShowUnreviewedChanges { get; set; }
+		public bool ShowAutomatedChanges { get; set; }
+		public RobomergeShowChangesOption ShowRobomerge { get; set; }
+		public bool AnnotateRobmergeChanges { get; set; }
+		public bool ShowLocalTimes { get; set; }
+		public bool KeepInTray { get; set; }
+		public HashSet<Guid> EnabledTools { get; } = new HashSet<Guid>();
+		public int FilterIndex { get; set; }
+		public UserSelectedProjectSettings? LastProject { get; set; }
+		public List<UserSelectedProjectSettings> OpenProjects { get; }
+		public List<UserSelectedProjectSettings> RecentProjects { get; }
+		public string SyncTypeId { get; set; }
+		public BuildConfig CompiledEditorBuildConfig { get; set; } // NB: This assumes not using precompiled editor. See CurrentBuildConfig.
+		public TabLabels TabLabels { get; set; }
+		public long NextLauncherVersionCheck { get; set; }
 
 		// Precompiled binaries
-		public List<ArchiveSettings> Archives = new List<ArchiveSettings>();
+		public List<ArchiveSettings> Archives { get; } = new List<ArchiveSettings>();
 
 		// OIDC Settings
-		public Dictionary<string, string> ProviderToRefreshTokens = new Dictionary<string, string>();
+		public Dictionary<string, string> ProviderToRefreshTokens { get; } = new Dictionary<string, string>();
 
 		// Window settings
-		public bool WindowVisible;
-		public string WindowState;
-		public Rectangle? WindowBounds;
+		public bool WindowVisible { get; set; }
+		public string WindowState { get; set; }
+		public Rectangle? WindowBounds { get; set; }
 		
 		// Schedule settings
-		public bool ScheduleEnabled;
-		public TimeSpan ScheduleTime;
-		public bool ScheduleAnyOpenProject;
-		public List<UserSelectedProjectSettings> ScheduleProjects;
+		public bool ScheduleEnabled { get; set; }
+		public TimeSpan ScheduleTime { get; set; }
+		public bool ScheduleAnyOpenProject { get; set; }
+		public List<UserSelectedProjectSettings> ScheduleProjects { get; } = new List<UserSelectedProjectSettings>();
 
 		// Run configuration
-		public List<Tuple<string, bool>> EditorArguments = new List<Tuple<string,bool>>();
-		public bool EditorArgumentsPrompt;
+		public List<Tuple<string, bool>> EditorArguments { get; } = new List<Tuple<string, bool>>();
+		public bool EditorArgumentsPrompt { get; set; }
 
 		// Notification settings
-		public List<string> NotifyProjects;
-		public int NotifyUnassignedMinutes;
-		public int NotifyUnacknowledgedMinutes;
-		public int NotifyUnresolvedMinutes;
+		public List<string> NotifyProjects { get; } = new List<string>();
+		public int NotifyUnassignedMinutes { get; set; }
+		public int NotifyUnacknowledgedMinutes { get; set; }
+		public int NotifyUnresolvedMinutes { get; set; }
 
 		// Project settings
 		readonly Dictionary<DirectoryReference, UserWorkspaceSettings> _workspaceDirToSettings = new Dictionary<DirectoryReference, UserWorkspaceSettings>();
@@ -605,7 +607,7 @@ namespace UnrealGameSync
 			if (coreSettingsData == null)
 			{
 				coreSettingsData = new GlobalSettings();
-				coreSettingsData.Filter.View = configFile.GetValues("General.SyncFilter", new string[0]).ToList();
+				coreSettingsData.Filter.View.AddRange(configFile.GetValues("General.SyncFilter", Array.Empty<string>()));
 				coreSettingsData.Filter.SetCategories(GetCategorySettings(configFile.FindSection("General"), "SyncIncludedCategories", "SyncExcludedCategories"));
 				coreSettingsData.Filter.AllProjects = configFile.GetValue("General.SyncAllProjects", false);
 				coreSettingsData.Filter.AllProjectsInSln = configFile.GetValue("General.IncludeAllProjectsInSolution", false);
@@ -631,32 +633,33 @@ namespace UnrealGameSync
 			ShowUnreviewedChanges = _configFile.GetValue("General.ShowUnreviewed", true);
 			ShowAutomatedChanges = _configFile.GetValue("General.ShowAutomated", false);
 			NextLauncherVersionCheck = _configFile.GetValue("General.NextLauncherVersionCheck", 0);
-
-			// safely parse the filter enum
-			ShowRobomerge = RobomergeShowChangesOption.All;
-			Enum.TryParse(_configFile.GetValue("General.RobomergeFilter", ""), out ShowRobomerge);
-
+			ShowRobomerge = _configFile.GetEnumValue("General.RobomergeFilter", RobomergeShowChangesOption.All);
 			AnnotateRobmergeChanges = _configFile.GetValue("General.AnnotateRobomerge", true);
 			ShowLocalTimes = _configFile.GetValue("General.ShowLocalTimes", false);
 			KeepInTray = _configFile.GetValue("General.KeepInTray", true);
 
-			List<Guid> enabledTools = _configFile.GetGuidValues("General.EnabledTools", new Guid[0]).ToList();
+			EnabledTools.Clear();
+			EnabledTools.UnionWith(_configFile.GetGuidValues("General.EnabledTools", Array.Empty<Guid>()));
 			if (_configFile.GetValue("General.EnableP4VExtensions", false))
 			{
-				enabledTools.Add(new Guid("963850A0-BF63-4E0E-B903-1C5954C7DCF8"));
+				EnabledTools.Add(new Guid("963850A0-BF63-4E0E-B903-1C5954C7DCF8"));
 			}
 			if (_configFile.GetValue("General.EnableUshell", false))
 			{
-				enabledTools.Add(new Guid("922EED87-E732-464C-92DC-5A8F7ED955E2"));
+				EnabledTools.Add(new Guid("922EED87-E732-464C-92DC-5A8F7ED955E2"));
 			}
-			EnabledTools = enabledTools.ToArray();
 
-			Int32.TryParse(_configFile.GetValue("General.FilterIndex", "0"), out FilterIndex);
+			FilterIndex = _configFile.GetValue("General.FilterIndex", 0);
 
 			string? lastProjectString = _configFile.GetValue("General.LastProject", null);
 			if(lastProjectString != null)
 			{
-				UserSelectedProjectSettings.TryParseConfigEntry(lastProjectString, out LastProject);
+				UserSelectedProjectSettings? lastProject;
+				if (!UserSelectedProjectSettings.TryParseConfigEntry(lastProjectString, out lastProject))
+				{
+					lastProject = null;
+				}
+				LastProject = lastProject;
 			}
 			else
 			{
@@ -688,18 +691,10 @@ namespace UnrealGameSync
 			}
 
 			// Build configuration
-			string compiledEditorBuildConfigName = _configFile.GetValue("General.BuildConfig", "");
-			if(!Enum.TryParse(compiledEditorBuildConfigName, true, out CompiledEditorBuildConfig))
-			{
-				CompiledEditorBuildConfig = BuildConfig.DebugGame;
-			}
+			CompiledEditorBuildConfig = _configFile.GetEnumValue("General.BuildConfig", BuildConfig.DebugGame);
 
 			// Tab names
-			string tabLabelsValue = _configFile.GetValue("General.TabLabels", "");
-			if(!Enum.TryParse(tabLabelsValue, true, out TabLabels))
-			{
-				TabLabels = TabLabels.Stream;
-			}
+			TabLabels = _configFile.GetEnumValue("General.TabLabels", TabLabels.Stream);
 
 			// Editor arguments
 			string[] arguments = _configFile.GetValues("General.EditorArguments", new string[]{ "0:-log", "0:-fastload" });
@@ -709,11 +704,11 @@ namespace UnrealGameSync
 			}
 			foreach(string argument in arguments)
 			{
-				if(argument.StartsWith("0:"))
+				if(argument.StartsWith("0:", StringComparison.Ordinal))
 				{
 					EditorArguments.Add(new Tuple<string,bool>(argument.Substring(2), false));
 				}
-				else if(argument.StartsWith("1:"))
+				else if(argument.StartsWith("1:", StringComparison.Ordinal))
 				{
 					EditorArguments.Add(new Tuple<string,bool>(argument.Substring(2), true));
 				}
@@ -725,7 +720,7 @@ namespace UnrealGameSync
 			EditorArgumentsPrompt = _configFile.GetValue("General.EditorArgumentsPrompt", false);
 
 			// Precompiled binaries
-			string[] archiveValues = _configFile.GetValues("PrecompiledBinaries.Archives", new string[0]);
+			string[] archiveValues = _configFile.GetValues("PrecompiledBinaries.Archives", Array.Empty<string>());
 			foreach (string archiveValue in archiveValues)
 			{
 				ArchiveSettings? settings;
@@ -737,11 +732,11 @@ namespace UnrealGameSync
 
 			if (syncPrecompiledEditor)
 			{
-				Archives.Add(new ArchiveSettings(true, "Editor", new string[0]));
+				Archives.Add(new ArchiveSettings(true, "Editor", Array.Empty<string>()));
 			}
 
 			// OIDC Settings
-			string[] tokens = _configFile.GetValues("OIDCProviders.Tokens", new string[0]);
+			string[] tokens = _configFile.GetValues("OIDCProviders.Tokens", Array.Empty<string>());
 			foreach (string tokenValue in tokens)
 			{
 				ConfigObject o = new ConfigObject(tokenValue);
@@ -760,16 +755,19 @@ namespace UnrealGameSync
 
 			// Schedule settings
 			ScheduleEnabled = _configFile.GetValue("Schedule.Enabled", false);
-			if(!TimeSpan.TryParse(_configFile.GetValue("Schedule.Time", ""), out ScheduleTime))
+
+			TimeSpan scheduleTime;
+			if(!TimeSpan.TryParse(_configFile.GetValue("Schedule.Time", ""), out scheduleTime))
 			{
-				ScheduleTime = new TimeSpan(6, 0, 0);
+				scheduleTime = new TimeSpan(6, 0, 0);
 			}
+			ScheduleTime = scheduleTime;
 
 			ScheduleAnyOpenProject = _configFile.GetValue("Schedule.AnyOpenProject", true);
 			ScheduleProjects = ReadProjectList("Schedule.Projects", "Schedule.ProjectFileNames");
 
 			// Notification settings
-			NotifyProjects = _configFile.GetValues("Notifications.NotifyProjects", new string[0]).ToList();
+			NotifyProjects = _configFile.GetValues("Notifications.NotifyProjects", Array.Empty<string>()).ToList();
 			NotifyUnassignedMinutes = _configFile.GetValue("Notifications.NotifyUnassignedMinutes", -1);
 			NotifyUnacknowledgedMinutes = _configFile.GetValue("Notifications.NotifyUnacknowledgedMinutes", -1);
 			NotifyUnresolvedMinutes = _configFile.GetValue("Notifications.NotifyUnresolvedMinutes", -1);
@@ -798,11 +796,11 @@ namespace UnrealGameSync
 			Dictionary<Guid, bool> result = new Dictionary<Guid, bool>();
 			if (section != null)
 			{
-				foreach (Guid uniqueId in section.GetValues(includedKey, new Guid[0]))
+				foreach (Guid uniqueId in section.GetValues(includedKey, Array.Empty<Guid>()))
 				{
 					result[uniqueId] = true;
 				}
-				foreach (Guid uniqueId in section.GetValues(excludedKey, new Guid[0]))
+				foreach (Guid uniqueId in section.GetValues(excludedKey, Array.Empty<Guid>()))
 				{
 					result[uniqueId] = false;
 				}
@@ -868,7 +866,7 @@ namespace UnrealGameSync
 			{
 				string legacyBranchAndClientKey = clientName + branchPath;
 
-				int slashIdx = legacyBranchAndClientKey.IndexOf('/');
+				int slashIdx = legacyBranchAndClientKey.IndexOf('/', StringComparison.Ordinal);
 				if(slashIdx != -1)
 				{
 					legacyBranchAndClientKey = legacyBranchAndClientKey.Substring(0, slashIdx) + "$" + legacyBranchAndClientKey.Substring(slashIdx + 1);
@@ -913,7 +911,7 @@ namespace UnrealGameSync
 				currentWorkspace.ProjectIdentifier = workspaceSection.GetValue("CurrentProjectPath", "");
 				currentWorkspace.CurrentChangeNumber = workspaceSection.GetValue("CurrentChangeNumber", -1);
 				currentWorkspace.CurrentSyncFilterHash = workspaceSection.GetValue("CurrentSyncFilterHash", null);
-				foreach(string additionalChangeNumberString in workspaceSection.GetValues("AdditionalChangeNumbers", new string[0]))
+				foreach(string additionalChangeNumberString in workspaceSection.GetValues("AdditionalChangeNumbers", Array.Empty<string>()))
 				{
 					int additionalChangeNumber;
 					if(Int32.TryParse(additionalChangeNumberString, out additionalChangeNumber))
@@ -922,10 +920,7 @@ namespace UnrealGameSync
 					}
 				}
 
-				WorkspaceUpdateResult lastSyncResult;
-				Enum.TryParse(workspaceSection.GetValue("LastSyncResult", ""), out lastSyncResult);
-				currentWorkspace.LastSyncResult = lastSyncResult;
-
+				currentWorkspace.LastSyncResult = workspaceSection.GetEnumValue("LastSyncResult", WorkspaceUpdateResult.Canceled);
 				currentWorkspace.LastSyncResultMessage = UnescapeText(workspaceSection.GetValue("LastSyncResultMessage"));
 				currentWorkspace.LastSyncChangeNumber = workspaceSection.GetValue("LastSyncChangeNumber", -1);
 
@@ -937,9 +932,11 @@ namespace UnrealGameSync
 
 				currentWorkspace.LastSyncDurationSeconds = workspaceSection.GetValue("LastSyncDuration", 0);
 				currentWorkspace.LastBuiltChangeNumber = workspaceSection.GetValue("LastBuiltChangeNumber", 0);
-				currentWorkspace.ExpandedArchiveTypes = workspaceSection.GetValues("ExpandedArchiveName", new string[0]);
 
-				string[] bisectEntries = workspaceSection.GetValues("Bisect", new string[0]);
+				currentWorkspace.ExpandedArchiveTypes.Clear();
+				currentWorkspace.ExpandedArchiveTypes.UnionWith(workspaceSection.GetValues("ExpandedArchiveName", Array.Empty<string>()));
+
+				string[] bisectEntries = workspaceSection.GetValues("Bisect", Array.Empty<string>());
 				foreach(string bisectEntry in bisectEntries)
 				{
 					ConfigObject bisectEntryObject = new ConfigObject(bisectEntry);
@@ -965,7 +962,9 @@ namespace UnrealGameSync
 			ConfigSection? workspaceSection = _configFile.FindSection(clientName + branchPath);
 			if (workspaceSection != null)
 			{
-				currentWorkspace.Filter.View = workspaceSection.GetValues("SyncFilter", new string[0]).ToList();
+				currentWorkspace.Filter.View.Clear();
+				currentWorkspace.Filter.View.AddRange(workspaceSection.GetValues("SyncFilter", Array.Empty<string>()));
+
 				currentWorkspace.Filter.SetCategories(GetCategorySettings(workspaceSection, "SyncIncludedCategories", "SyncExcludedCategories"));
 
 				int syncAllProjects = workspaceSection.GetValue("SyncAllProjects", -1);
@@ -982,16 +981,9 @@ namespace UnrealGameSync
 
 			// Read the project settings
 			ConfigSection projectSection = _configFile.FindOrAddSection(clientProjectFileName);
-			currentProject.BuildSteps.AddRange(projectSection.GetValues("BuildStep", new string[0]).Select(x => new ConfigObject(x)));
-
-			FilterType filterType;
-			if (!Enum.TryParse(projectSection.GetValue("FilterType", ""), true, out filterType))
-			{
-				filterType = FilterType.None;
-			}
-
-			currentProject.FilterType = filterType;
-			currentProject.FilterBadges.UnionWith(projectSection.GetValues("FilterBadges", new string[0]));
+			currentProject.BuildSteps.AddRange(projectSection.GetValues("BuildStep", Array.Empty<string>()).Select(x => new ConfigObject(x)));
+			currentProject.FilterType = projectSection.GetEnumValue("FilterType", FilterType.None);
+			currentProject.FilterBadges.UnionWith(projectSection.GetValues("FilterBadges", Array.Empty<string>()));
 		}
 
 		public override bool Save(ILogger logger)
@@ -1021,7 +1013,7 @@ namespace UnrealGameSync
 			}
 			generalSection.SetValues("OpenProjects", OpenProjects.Select(x => x.ToConfigEntry()).ToArray());
 			generalSection.SetValue("KeepInTray", KeepInTray);
-			generalSection.SetValues("EnabledTools", EnabledTools);
+			generalSection.SetValues("EnabledTools", EnabledTools.ToArray());
 			generalSection.SetValue("FilterIndex", FilterIndex);
 			generalSection.SetValues("RecentProjects", RecentProjects.Select(x => x.ToConfigEntry()).ToArray());
 			generalSection.SetValue("SyncTypeID", SyncTypeId);
diff --git a/Engine/Source/Programs/UnrealGameSync/UnrealGameSyncShared/Utility.cs b/Engine/Source/Programs/UnrealGameSync/UnrealGameSyncShared/Utility.cs
--- a/Engine/Source/Programs/UnrealGameSync/UnrealGameSyncShared/Utility.cs
+++ b/Engine/Source/Programs/UnrealGameSync/UnrealGameSyncShared/Utility.cs
@@ -257,10 +257,10 @@ namespace UnrealGameSync
 		public static string ExpandVariables(string inputString, Dictionary<string, string>? additionalVariables = null)
 		{
 			string result = inputString;
-			for (int idx = result.IndexOf("$("); idx != -1; idx = result.IndexOf("$(", idx))
+			for (int idx = result.IndexOf("$(", StringComparison.Ordinal); idx != -1; idx = result.IndexOf("$(", idx, StringComparison.Ordinal))
 			{
 				// Find the end of the variable name
-				int endIdx = result.IndexOf(')', idx + 2);
+				int endIdx = result.IndexOf(")", idx + 2, StringComparison.Ordinal);
 				if (endIdx == -1)
 				{
 					break;
@@ -271,7 +271,7 @@ namespace UnrealGameSync
 
 				// Strip the format from the name
 				string? format = null;
-				int formatIdx = name.IndexOf(':');
+				int formatIdx = name.IndexOf(':', StringComparison.Ordinal);
 				if(formatIdx != -1)
 				{ 
 					format = name.Substring(formatIdx + 1);
@@ -293,7 +293,7 @@ namespace UnrealGameSync
 				// Encode the variable if necessary
 				if(format != null)
 				{
-					if(String.Equals(format, "URI", StringComparison.InvariantCultureIgnoreCase))
+					if(String.Equals(format, "URI", StringComparison.OrdinalIgnoreCase))
 					{
 						value = Uri.EscapeDataString(value);
 					}
@@ -473,7 +473,7 @@ namespace UnrealGameSync
 				}
 				else
 				{
-					string[] lines = await FileReference.ReadAllLinesAsync(tempFile);
+					string[] lines = await FileReference.ReadAllLinesAsync(tempFile, cancellationToken);
 					try
 					{
 						FileReference.SetAttributes(tempFile, FileAttributes.Normal);
diff --git a/Engine/Source/Programs/UnrealGameSync/UnrealGameSyncShared/WorkspaceLock.cs b/Engine/Source/Programs/UnrealGameSync/UnrealGameSyncShared/WorkspaceLock.cs
--- a/Engine/Source/Programs/UnrealGameSync/UnrealGameSyncShared/WorkspaceLock.cs
+++ b/Engine/Source/Programs/UnrealGameSync/UnrealGameSyncShared/WorkspaceLock.cs
@@ -6,7 +6,6 @@ using System.Collections.Concurrent;
 using System.Security.Cryptography;
 using System.Text;
 using System.Threading;
-using System.Threading.Channels;
 using System.Threading.Tasks;
 
 namespace UnrealGameSync
@@ -14,7 +13,7 @@ namespace UnrealGameSync
 	/// <summary>
 	/// Encapsulates the state of cross-process workspace lock
 	/// </summary>
-	public class WorkspaceLock : IDisposable
+	public sealed class WorkspaceLock : IDisposable
 	{
 		const string Prefix = @"Global\ugs-workspace";
 
@@ -45,11 +44,13 @@ namespace UnrealGameSync
 		/// <param name="rootDir">Root directory for the workspace</param>
 		public WorkspaceLock(DirectoryReference rootDir)
 		{
+#pragma warning disable CA5351 // Do Not Use Broken Cryptographic Algorithms
 			using (MD5 md5 = MD5.Create())
 			{
 				byte[] idBytes = Encoding.UTF8.GetBytes(rootDir.FullName.ToUpperInvariant());
 				_objectName = StringUtils.FormatHexString(md5.ComputeHash(idBytes));
 			}
+#pragma warning restore CA5351 // Do Not Use Broken Cryptographic Algorithms
 
 			_mutex = new Mutex(false, $"{Prefix}.{_objectName}.mutex");
 
@@ -81,7 +82,10 @@ namespace UnrealGameSync
 				_monitorThread = null;
 			}
 
+			_lockedEvent?.Dispose();
+
 			_acquireCancellationSource.Dispose();
+			_acquireActions.Dispose();
 			_cancelMonitorEvent.Dispose();
 			_mutex.Dispose();
 		}
diff --git a/Engine/Source/Programs/UnrealGameSync/UnrealGameSyncShared/WorkspaceState.cs b/Engine/Source/Programs/UnrealGameSync/UnrealGameSyncShared/WorkspaceState.cs
--- a/Engine/Source/Programs/UnrealGameSync/UnrealGameSyncShared/WorkspaceState.cs
+++ b/Engine/Source/Programs/UnrealGameSync/UnrealGameSyncShared/WorkspaceState.cs
@@ -28,7 +28,7 @@ namespace UnrealGameSync
 		public int CurrentChangeNumber { get; set; } = -1;
 		public int CurrentCodeChangeNumber { get; set; } = -1;
 		public string? CurrentSyncFilterHash { get; set; }
-		public List<int> AdditionalChangeNumbers { get; set; } = new List<int>();
+		public List<int> AdditionalChangeNumbers { get; } = new List<int>();
 
 		// Settings for the last attempted sync. These values are set to persist error messages between runs.
 		public int LastSyncChangeNumber { get; set; }
@@ -41,10 +41,10 @@ namespace UnrealGameSync
 		public int LastBuiltChangeNumber { get; set; }
 
 		// Expanded archives in the workspace
-		public string[]? ExpandedArchiveTypes { get; set; }
+		public HashSet<string> ExpandedArchiveTypes { get; } = new HashSet<string>(StringComparer.Ordinal);
 
 		// The changes that we're regressing at the moment
-		public List<BisectEntry> BisectChanges { get; set; } = new List<BisectEntry>();
+		public List<BisectEntry> BisectChanges { get; } = new List<BisectEntry>();
 
 		public void UpdateCachedProjectInfo(ProjectInfo projectInfo, long settingsTimeUtc)
 		{
@@ -134,12 +134,16 @@ namespace UnrealGameSync
 		public int LastBuiltChangeNumber => _inner.LastBuiltChangeNumber;
 
 		// Expanded archives in the workspace
-		public IReadOnlyList<string> ExpandedArchiveTypes => _inner.ExpandedArchiveTypes ?? Array.Empty<string>();
+		public ReadOnlyHashSet<string> ExpandedArchiveTypes { get; }
 
 		// The changes that we're regressing at the moment
 		public IReadOnlyList<BisectEntry> BisectChanges => _inner.BisectChanges;
 
-		internal ReadOnlyWorkspaceState(WorkspaceState inner) => _inner = inner;
+		internal ReadOnlyWorkspaceState(WorkspaceState inner)
+		{
+			_inner = inner;
+			ExpandedArchiveTypes = new ReadOnlyHashSet<string>(_inner.ExpandedArchiveTypes);
+		}
 
 		public ReadOnlyWorkspaceState ResetForProject(ProjectInfo projectInfo)
 		{
@@ -166,7 +170,7 @@ namespace UnrealGameSync
 	/// <summary>
 	/// Monitors state of a workspace state file.
 	/// </summary>
-	public class WorkspaceStateWrapper : IDisposable
+	public sealed class WorkspaceStateWrapper : IDisposable
 	{
 		/// <summary>
 		/// Directory containing the config file
diff --git a/Engine/Source/Programs/UnrealGameSync/UnrealGameSyncShared/WorkspaceUpdate.cs b/Engine/Source/Programs/UnrealGameSync/UnrealGameSyncShared/WorkspaceUpdate.cs
--- a/Engine/Source/Programs/UnrealGameSync/UnrealGameSyncShared/WorkspaceUpdate.cs
+++ b/Engine/Source/Programs/UnrealGameSync/UnrealGameSyncShared/WorkspaceUpdate.cs
@@ -80,33 +80,39 @@ namespace UnrealGameSync
 
 	public class WorkspaceUpdateContext
 	{
-		public DateTime StartTime = DateTime.UtcNow;
-		public int ChangeNumber;
-		public int? CodeChangeNumber;
-		public WorkspaceUpdateOptions Options;
-		public BuildConfig EditorConfig;
-		public string[]? SyncFilter;
-		public Dictionary<string, Tuple<IArchiveInfo, string>?> ArchiveTypeToArchive = new Dictionary<string, Tuple<IArchiveInfo, string>?>();
-		public Dictionary<string, bool> DeleteFiles = new Dictionary<string,bool>();
-		public Dictionary<string, bool> ClobberFiles = new Dictionary<string,bool>();
-		public List<ConfigObject> UserBuildStepObjects;
-		public HashSet<Guid>? CustomBuildSteps;
-		public Dictionary<string, string> AdditionalVariables = new Dictionary<string, string>(StringComparer.OrdinalIgnoreCase);
-		public PerforceSyncOptions? PerforceSyncOptions;
-		public List<HaveRecord>? HaveFiles; // Cached when sync filter has changed
+		public DateTime StartTime { get; set; } = DateTime.UtcNow;
+		public int ChangeNumber { get; set; }
+		public int? CodeChangeNumber { get; set; }
+		public WorkspaceUpdateOptions Options { get; set; }
+		public BuildConfig EditorConfig { get; set; }
+		public List<string> SyncFilter { get; } = new List<string>();
+		public Dictionary<string, Tuple<IArchiveInfo, string>?> ArchiveTypeToArchive { get; } = new Dictionary<string, Tuple<IArchiveInfo, string>?>();
+		public Dictionary<string, bool> DeleteFiles { get; } = new Dictionary<string, bool>();
+		public Dictionary<string, bool> ClobberFiles { get; } = new Dictionary<string, bool>();
+		public List<ConfigObject> UserBuildStepObjects { get; } = new List<ConfigObject>();
+		public HashSet<Guid> CustomBuildSteps { get; } = new HashSet<Guid>();
+		public Dictionary<string, string> AdditionalVariables { get; } = new Dictionary<string, string>(StringComparer.OrdinalIgnoreCase);
+		public PerforceSyncOptions? PerforceSyncOptions { get; set; }
+		public List<HaveRecord> HaveFiles { get; } = new List<HaveRecord>(); // Cached when sync filter has changed
 
 		// May be updated during sync
-		public ConfigFile? ProjectConfigFile;
-		public IReadOnlyList<string>? ProjectStreamFilter;
+		public ConfigFile? ProjectConfigFile { get; set; }
+		public IReadOnlyList<string>? ProjectStreamFilter { get; set; }
 
-		public WorkspaceUpdateContext(int inChangeNumber, WorkspaceUpdateOptions inOptions, BuildConfig inEditorConfig, string[]? inSyncFilter, List<ConfigObject> inUserBuildSteps, HashSet<Guid>? inCustomBuildSteps)
+		public WorkspaceUpdateContext(int changeNumber, WorkspaceUpdateOptions options, BuildConfig editorConfig, string[]? syncFilter, List<ConfigObject> userBuildSteps, HashSet<Guid>? customBuildSteps)
 		{
-			ChangeNumber = inChangeNumber;
-			Options = inOptions;
-			EditorConfig = inEditorConfig;
-			SyncFilter = inSyncFilter;
-			UserBuildStepObjects = inUserBuildSteps;
-			CustomBuildSteps = inCustomBuildSteps;
+			ChangeNumber = changeNumber;
+			Options = options;
+			EditorConfig = editorConfig;
+			if (syncFilter != null)
+			{
+				SyncFilter.AddRange(syncFilter);
+			}
+			UserBuildStepObjects = userBuildSteps;
+			if (customBuildSteps != null)
+			{
+				CustomBuildSteps.UnionWith(customBuildSteps);
+			}
 		}
 
 		public static WorkspaceUpdateOptions GetOptionsFromConfig(GlobalSettings globalSettings, UserWorkspaceSettings workspaceSettings)
@@ -126,12 +132,12 @@ namespace UnrealGameSync
 
 	public class WorkspaceSyncCategory
 	{
-		public Guid UniqueId;
-		public bool Enable;
-		public string Name;
-		public string[] Paths;
-		public bool Hidden;
-		public Guid[] Requires;
+		public Guid UniqueId { get; set; }
+		public bool Enable { get; set; }
+		public string Name { get; set; }
+		public List<string> Paths { get; } = new List<string>();
+		public bool Hidden { get; set; }
+		public List<Guid> Requires { get; } = new List<Guid>();
 
 		public WorkspaceSyncCategory(Guid uniqueId) : this(uniqueId, "Unnamed")
 		{
@@ -142,8 +148,7 @@ namespace UnrealGameSync
 			UniqueId = uniqueId;
 			Enable = true;
 			Name = name;
-			Paths = paths;
-			Requires = new Guid[0];
+			Paths.AddRange(paths);
 		}
 
 		public static Dictionary<Guid, bool> GetDefault(IEnumerable<WorkspaceSyncCategory> categories)
@@ -198,7 +203,7 @@ namespace UnrealGameSync
 		public string ClientFileName => $"//{ClientName}{BranchPath}{ProjectPath}";
 		public string TelemetryProjectIdentifier => PerforceUtils.GetClientOrDepotDirectoryName(ProjectIdentifier);
 		public DirectoryReference EngineDir => DirectoryReference.Combine(LocalRootPath, "Engine");
-		public DirectoryReference? ProjectDir => ProjectPath.EndsWith(".uproject")? LocalFileName.Directory : null;
+		public DirectoryReference? ProjectDir => ProjectPath.EndsWith(".uproject", StringComparison.OrdinalIgnoreCase)? LocalFileName.Directory : null;
 		public DirectoryReference DataFolder => GetDataFolder(LocalRootPath);
 		public DirectoryReference CacheFolder => GetCacheFolder(LocalRootPath);
 
@@ -313,7 +318,7 @@ namespace UnrealGameSync
 
 		public static void ValidateBranchPath(string branchPath)
 		{
-			if (branchPath.Length > 0 && (!branchPath.StartsWith("/") || branchPath.EndsWith("/")))
+			if (branchPath.Length > 0 && (!branchPath.StartsWith("/", StringComparison.Ordinal) || branchPath.EndsWith("/", StringComparison.Ordinal)))
 			{
 				throw new ArgumentException("Branch path must start with a slash, and not end with a slash", nameof(branchPath));
 			}
@@ -321,7 +326,7 @@ namespace UnrealGameSync
 
 		public static void ValidateProjectPath(string projectPath)
 		{
-			if (!projectPath.StartsWith("/"))
+			if (!projectPath.StartsWith("/", StringComparison.Ordinal))
 			{
 				throw new ArgumentException("Project path must start with a slash", nameof(projectPath));
 			}
@@ -452,7 +457,7 @@ namespace UnrealGameSync
 
 			private void IncludeFile(string fullPath, string path, long size, ILogger logger)
 			{
-				int idx = path.IndexOf('/');
+				int idx = path.IndexOf('/', StringComparison.Ordinal);
 				if (idx == -1)
 				{
 					if (!IncludedFiles.ContainsKey(path))
@@ -471,7 +476,7 @@ namespace UnrealGameSync
 
 			public void ExcludeFile(string path)
 			{
-				int idx = path.IndexOf('/');
+				int idx = path.IndexOf('/', StringComparison.Ordinal);
 				if (idx != -1)
 				{
 					SyncTree subTree = FindOrAddSubTree(path.Substring(0, idx));
@@ -558,7 +563,7 @@ namespace UnrealGameSync
 
 		public static string ShellScriptExt { get; }= RuntimeInformation.IsOSPlatform(OSPlatform.Windows)? "bat" : "sh";
 
-		public Task<int> ExecuteShellCommandAsync(string commandLine, Action<string> processOutput, CancellationToken cancellationToken)
+		public static Task<int> ExecuteShellCommandAsync(string commandLine, Action<string> processOutput, CancellationToken cancellationToken)
 		{
 			if(RuntimeInformation.IsOSPlatform(OSPlatform.Windows))
 			{
@@ -600,13 +605,11 @@ namespace UnrealGameSync
 
 					// Get the user's sync filter
 					FileFilter userFilter = new FileFilter(FileFilterType.Include);
-					if (Context.SyncFilter != null)
-					{
-						userFilter.AddRules(Context.SyncFilter.Select(x => x.Trim()).Where(x => x.Length > 0 && !x.StartsWith(";") && !x.StartsWith("#")));
-					}
+					userFilter.AddRules(Context.SyncFilter.Select(x => x.Trim()).Where(x => x.Length > 0 && !x.StartsWith(";", StringComparison.Ordinal) && !x.StartsWith("#", StringComparison.Ordinal)));
 
 					// Check if the new sync filter matches the previous one. If not, we'll enumerate all files in the workspace and make sure there's nothing extra there.
 					string? nextSyncFilterHash = null;
+#pragma warning disable CA5350 // Do Not Use Weak Cryptographic Algorithms
 					using (SHA1 sha = SHA1.Create())
 					{
 						StringBuilder combinedFilter = new StringBuilder();
@@ -614,13 +617,14 @@ namespace UnrealGameSync
 						{
 							combinedFilter.AppendFormat("{0}\n", relativeSyncPath);
 						}
-						if (Context.SyncFilter != null)
+						if (Context.SyncFilter.Count > 0)
 						{
 							combinedFilter.Append("--FROM--\n");
 							combinedFilter.Append(String.Join("\n", Context.SyncFilter));
 						}
-						nextSyncFilterHash = BitConverter.ToString(sha.ComputeHash(Encoding.UTF8.GetBytes(combinedFilter.ToString()))).Replace("-", "");
+						nextSyncFilterHash = StringUtils.FormatHexString(sha.ComputeHash(Encoding.UTF8.GetBytes(combinedFilter.ToString())));
 					}
+#pragma warning restore CA5350 // Do Not Use Weak Cryptographic Algorithms
 
 					// If the hash differs, enumerate everything in the workspace to find what needs to be removed
 					if (nextSyncFilterHash != state.CurrentSyncFilterHash || (Context.Options & WorkspaceUpdateOptions.Refilter) != 0)
@@ -630,10 +634,9 @@ namespace UnrealGameSync
 							logger.LogInformation("Filter has changed ({PrevHash} -> {NextHash}); finding files in workspace that need to be removed.", (String.IsNullOrEmpty(state.CurrentSyncFilterHash)) ? "None" : state.CurrentSyncFilterHash, nextSyncFilterHash);
 
 							// Find all the files that are in this workspace
-							List<HaveRecord>? haveFiles = Context.HaveFiles;
-							if (haveFiles == null)
+							List<HaveRecord> haveFiles = Context.HaveFiles;
+							if (haveFiles.Count == 0)
 							{
-								haveFiles = new List<HaveRecord>();
 								using (RecordCounter haveCounter = new RecordCounter(Progress, "Sync filter changed; checking workspace..."))
 								{
 									await foreach (PerforceResponse<HaveRecord> record in perforce.TryHaveAsync(FileSpecList.Any, cancellationToken))
@@ -649,7 +652,6 @@ namespace UnrealGameSync
 										}
 									}
 								}
-								Context.HaveFiles = haveFiles;
 							}
 
 							// Build a filter for the current sync paths
@@ -683,22 +685,23 @@ namespace UnrealGameSync
 							{
 								bool deleteListMatches = true;
 
-								Dictionary<string, bool> newDeleteFiles = new Dictionary<string, bool>(StringComparer.OrdinalIgnoreCase);
+								Dictionary<string, bool> prevDeleteFiles = new Dictionary<string, bool>(Context.DeleteFiles, StringComparer.OrdinalIgnoreCase);
+
+								Context.DeleteFiles.Clear();
 								foreach (string removeDepotPath in removeDepotPaths)
 								{
 									bool delete;
-									if (!Context.DeleteFiles.TryGetValue(removeDepotPath, out delete))
+									if (!prevDeleteFiles.TryGetValue(removeDepotPath, out delete))
 									{
 										deleteListMatches = false;
 										delete = true;
 									}
-									newDeleteFiles[removeDepotPath] = delete;
+									Context.DeleteFiles[removeDepotPath] = delete;
 								}
-								Context.DeleteFiles = newDeleteFiles;
 
 								if (!deleteListMatches)
 								{
-									return (WorkspaceUpdateResult.FilesToDelete, $"Cancelled after finding {newDeleteFiles.Count} files excluded by filter");
+									return (WorkspaceUpdateResult.FilesToDelete, $"Cancelled after finding {Context.DeleteFiles.Count} files excluded by filter");
 								}
 
 								removeDepotPaths.RemoveAll(x => !Context.DeleteFiles[x]);
@@ -782,7 +785,7 @@ namespace UnrealGameSync
 								}
 								if (response.Info != null)
 								{
-									logger.LogInformation("Note: {0}", response.Info.Data);
+									logger.LogInformation("Note: {Note}", response.Info.Data);
 									continue;
 								}
 
@@ -918,7 +921,7 @@ namespace UnrealGameSync
 						else
 						{
 							// Otherwise use the depot path for GenerateProjectFiles.bat in the root of the workspace
-							List<WhereRecord> files = await perforce.WhereAsync(project.ClientRootPath + "/GenerateProjectFiles.bat").Where(x => !x.Unmap).ToListAsync(cancellationToken);
+							List<WhereRecord> files = await perforce.WhereAsync(project.ClientRootPath + "/GenerateProjectFiles.bat", cancellationToken).Where(x => !x.Unmap).ToListAsync(cancellationToken);
 							if (files.Count != 1)
 							{
 								return (WorkspaceUpdateResult.FailedToSync, $"Couldn't determine branch name for {project.ClientFileName}.");
@@ -1152,9 +1155,9 @@ namespace UnrealGameSync
 							IArchiveInfo archiveInfo = archiveTypeAndArchive.Value.Item1;
 							string archiveKey = archiveTypeAndArchive.Value.Item2;
 
-							logger.LogInformation("Syncing {ArchiveType} binaries...", archiveType.ToLowerInvariant());
-							Progress.Set(String.Format("Syncing {0} binaries...", archiveType.ToLowerInvariant()), 0.0f);
-							if (!archiveInfo.DownloadArchive(perforce, archiveKey, project.LocalRootPath, manifestFileName, logger, Progress, CancellationToken.None).Result)
+							logger.LogInformation("Syncing {ArchiveType} binaries...", archiveType);
+							Progress.Set(String.Format("Syncing {0} binaries...", archiveType), 0.0f);
+							if (!await archiveInfo.DownloadArchive(perforce, archiveKey, project.LocalRootPath, manifestFileName, logger, Progress, CancellationToken.None))
 							{
 								return (WorkspaceUpdateResult.FailedToSync, $"Couldn't read {archiveKey}");
 							}
@@ -1162,7 +1165,11 @@ namespace UnrealGameSync
 					}
 
 					// Update the state
-					state = stateMgr.Modify(x => x.ExpandedArchiveTypes = Context.ArchiveTypeToArchive.Where(x => x.Value != null).Select(x => x.Key).ToArray());
+					state = stateMgr.Modify(x =>
+					{
+						x.ExpandedArchiveTypes.Clear();
+						x.ExpandedArchiveTypes.UnionWith(Context.ArchiveTypeToArchive.Where(x => x.Value != null).Select(x => x.Key));
+					});
 
 					// Add the finish time
 					times.Add(new Tuple<string, TimeSpan>("Archive", stopwatch.Stop("Success")));
@@ -1226,7 +1233,7 @@ namespace UnrealGameSync
 				// Get the build steps
 				bool usingPrecompiledEditor = Context.ArchiveTypeToArchive.TryGetValue(IArchiveInfo.EditorArchiveType, out Tuple<IArchiveInfo, string>? archiveInfo) && archiveInfo != null;
 				Dictionary<Guid, ConfigObject> buildStepObjects = ConfigUtils.GetDefaultBuildStepObjects(project, editorTargetName, Context.EditorConfig, Context.ProjectConfigFile, usingPrecompiledEditor);
-				BuildStep.MergeBuildStepObjects(buildStepObjects, Context.ProjectConfigFile.GetValues("Build.Step", new string[0]).Select(x => new ConfigObject(x)));
+				BuildStep.MergeBuildStepObjects(buildStepObjects, Context.ProjectConfigFile.GetValues("Build.Step", Array.Empty<string>()).Select(x => new ConfigObject(x)));
 				BuildStep.MergeBuildStepObjects(buildStepObjects, Context.UserBuildStepObjects);
 
 				// Construct build steps from them
@@ -1248,7 +1255,7 @@ namespace UnrealGameSync
 				bool forceClean = false;
 				if (state.LastBuiltChangeNumber != 0)
 				{
-					foreach (string cleanBuildChange in Context.ProjectConfigFile.GetValues("ForceClean.Changelist", new string[0]))
+					foreach (string cleanBuildChange in Context.ProjectConfigFile.GetValues("ForceClean.Changelist", Array.Empty<string>()))
 					{
 						int changeNumber;
 						if (Int32.TryParse(cleanBuildChange, out changeNumber))
@@ -1367,7 +1374,7 @@ namespace UnrealGameSync
 										FileReference toolFileName = FileReference.Combine(project.LocalRootPath, Utility.ExpandVariables(step.FileName ?? "unknown", variables));
 										string toolWorkingDir = String.IsNullOrWhiteSpace(step.WorkingDir) ? toolFileName.Directory.FullName : Utility.ExpandVariables(step.WorkingDir, variables);
 										string toolArguments = Utility.ExpandVariables(step.Arguments ?? "", variables);
-										logger.LogInformation("tool> Running {0} {1}", toolFileName, toolArguments);
+										logger.LogInformation("tool> Running {Tool} {Arguments}", toolFileName, toolArguments);
 
 										if (!FileReference.Exists(toolFileName))
 										{
@@ -1459,12 +1466,12 @@ namespace UnrealGameSync
 			return syncPaths.Select(x => project.ClientRootPath + x).ToList();
 		}
 
-		public static List<string> GetRelativeSyncPaths(ProjectInfo project, bool syncAllProjects, string[]? syncFilter)
+		public static List<string> GetRelativeSyncPaths(ProjectInfo project, bool syncAllProjects, IReadOnlyList<string>? syncFilter)
 		{
 			List<string> syncPaths = new List<string>();
 
 			// Check the client path is formatted correctly
-			if (!project.ClientFileName.StartsWith(project.ClientRootPath + "/"))
+			if (!project.ClientFileName.StartsWith(project.ClientRootPath + "/", StringComparison.OrdinalIgnoreCase))
 			{
 				throw new Exception(String.Format("Expected '{0}' to start with '{1}'", project.ClientFileName, project.ClientRootPath));
 			}
@@ -1492,13 +1499,13 @@ namespace UnrealGameSync
 				foreach (string syncPath in syncFilter)
 				{
 					string trimSyncPath = syncPath.Trim();
-					if (trimSyncPath.StartsWith("/"))
+					if (trimSyncPath.StartsWith("/", StringComparison.Ordinal))
 					{
 						syncPaths.Add(trimSyncPath);
 					}
-					else if (trimSyncPath.StartsWith("-/") && trimSyncPath.EndsWith("..."))
+					else if (trimSyncPath.StartsWith("-/", StringComparison.Ordinal) && trimSyncPath.EndsWith("...", StringComparison.Ordinal))
 					{
-						syncPaths.RemoveAll(x => x.StartsWith(trimSyncPath.Substring(1, trimSyncPath.Length - 4)));
+						syncPaths.RemoveAll(x => x.StartsWith(trimSyncPath.Substring(1, trimSyncPath.Length - 4), StringComparison.Ordinal));
 					}
 				}
 			}
@@ -1508,12 +1515,12 @@ namespace UnrealGameSync
 			for (int idx = 0; idx < syncPaths.Count; idx++)
 			{
 				string syncPath = syncPaths[idx];
-				if (syncPath.EndsWith("..."))
+				if (syncPath.EndsWith("...", StringComparison.Ordinal))
 				{
 					string syncPathPrefix = syncPath.Substring(0, syncPath.Length - 3);
 					for (int otherIdx = syncPaths.Count - 1; otherIdx > idx; otherIdx--)
 					{
-						if (syncPaths[otherIdx].StartsWith(syncPathPrefix))
+						if (syncPaths[otherIdx].StartsWith(syncPathPrefix, StringComparison.Ordinal))
 						{
 							syncPaths.RemoveAt(otherIdx);
 						}
@@ -1593,7 +1600,7 @@ namespace UnrealGameSync
 					IPerforceConnection childConnection = await PerforceConnection.CreateAsync(perforce.Settings, perforce.Logger);
 					childConnections.Add(childConnection);
 
-					Task childTask = Task.Run(() => StaticSyncWorker(threadNumber, childConnection, context, state, syncOutput, logWrapper, cancellationToken));
+					Task childTask = Task.Run(() => StaticSyncWorker(threadNumber, childConnection, context, state, syncOutput, logWrapper, cancellationToken), cancellationToken);
 					childTasks.Add(childTask);
 				}
 
@@ -1613,7 +1620,7 @@ namespace UnrealGameSync
 			{
 				foreach (Task childTask in childTasks)
 				{
-					await childTask.ContinueWith(_ => { }); // Swallow exceptions until we've disposed the connections
+					await childTask.ContinueWith(_ => { }, TaskScheduler.Default); // Swallow exceptions until we've disposed the connections
 				}
 
 				foreach (IPerforceConnection childConnection in childConnections)
@@ -1678,7 +1685,7 @@ namespace UnrealGameSync
 
 					if (WorkspaceUpdateResult.FailedToSync == result && --retries >= 0)
 					{
-						threadLog.LogWarning("Sync Errors occurred.  Retrying: Remaining retries " + retries.ToString());
+						threadLog.LogWarning("Sync Errors occurred.  Retrying: Remaining retries {Count}", retries);
 					}
 				}
 
@@ -1756,7 +1763,7 @@ namespace UnrealGameSync
 					bool shouldClobber = (context.Options & WorkspaceUpdateOptions.Clobber) != 0 || context.ClobberFiles[tamperedFile];
 					if (shouldClobber)
 					{
-						List<PerforceResponse<SyncRecord>> response = await perforce.TrySyncAsync(SyncOptions.Force, -1, tamperedFile, cancellationToken).ToListAsync();
+						List<PerforceResponse<SyncRecord>> response = await perforce.TrySyncAsync(SyncOptions.Force, -1, tamperedFile, cancellationToken).ToListAsync(cancellationToken);
 						if (!response.Succeeded())
 						{
 							return (WorkspaceUpdateResult.FailedToSync, $"Couldn't sync {tamperedFile}.");
@@ -1910,7 +1917,7 @@ namespace UnrealGameSync
 
 		static string UpdateVersionStrings(string text, Dictionary<string, string> versionStrings)
 		{
-			StringWriter writer = new StringWriter();
+			using StringWriter writer = new StringWriter();
 			foreach (string line in text.Split('\n'))
 			{
 				string newLine = line;
@@ -1933,13 +1940,19 @@ namespace UnrealGameSync
 			int prevCompatibleChangelist = 0;
 			if (obj.TryGetValue("CompatibleChangelist", out object? prevCompatibleChangelistObj))
 			{
-				Int32.TryParse(prevCompatibleChangelistObj?.ToString(), out prevCompatibleChangelist);
+				if (!Int32.TryParse(prevCompatibleChangelistObj?.ToString(), out prevCompatibleChangelist))
+				{
+					prevCompatibleChangelist = 0;
+				}
 			}
 
 			int prevIsLicenseeVersion = 0;
 			if (obj.TryGetValue("IsLicenseeVersion", out object? prevIsLicenseeVersionObj))
 			{
-				Int32.TryParse(prevIsLicenseeVersionObj?.ToString(), out prevIsLicenseeVersion);
+				if (!Int32.TryParse(prevIsLicenseeVersionObj?.ToString(), out prevIsLicenseeVersion))
+				{
+					prevIsLicenseeVersion = 0;
+				}
 			}
 
 			obj["Changelist"] = changelist;
