UGS: Another wave of static analyzer fixes.

diff --git a/Engine/Source/Programs/UnrealGameSync/.editorconfig b/Engine/Source/Programs/UnrealGameSync/.editorconfig
--- a/Engine/Source/Programs/UnrealGameSync/.editorconfig
+++ b/Engine/Source/Programs/UnrealGameSync/.editorconfig
@@ -370,7 +370,7 @@ dotnet_diagnostic.CA1054.severity = none
 dotnet_diagnostic.CA1055.severity = warning
 
 # CA1056: URI-like properties should not be strings
-dotnet_diagnostic.CA1056.severity = warning
+dotnet_diagnostic.CA1056.severity = none
 
 # CA1058: Types should not extend certain base types
 dotnet_diagnostic.CA1058.severity = warning
@@ -424,7 +424,7 @@ dotnet_diagnostic.CA1305.severity = none
 dotnet_diagnostic.CA1307.severity = warning
 
 # CA1308: Normalize strings to uppercase
-dotnet_diagnostic.CA1308.severity = warning
+dotnet_diagnostic.CA1308.severity = none
 
 # CA1309: Use ordinal string comparison
 dotnet_diagnostic.CA1309.severity = warning
@@ -517,7 +517,7 @@ dotnet_diagnostic.CA1802.severity = warning
 dotnet_diagnostic.CA1805.severity = none
 
 # CA1806: Do not ignore method results
-dotnet_diagnostic.CA1806.severity = warning
+dotnet_diagnostic.CA1806.severity = none
 
 # CA1810: Initialize reference type static fields inline
 dotnet_diagnostic.CA1810.severity = warning
diff --git a/Engine/Source/Programs/UnrealGameSync/UnrealGameSync/AsyncDisposer.cs b/Engine/Source/Programs/UnrealGameSync/UnrealGameSync/AsyncDisposer.cs
--- a/Engine/Source/Programs/UnrealGameSync/UnrealGameSync/AsyncDisposer.cs
+++ b/Engine/Source/Programs/UnrealGameSync/UnrealGameSync/AsyncDisposer.cs
@@ -3,6 +3,7 @@
 using Microsoft.Extensions.Logging;
 using System;
 using System.Collections.Generic;
+using System.Threading;
 using System.Threading.Tasks;
 
 namespace UnrealGameSync
@@ -25,7 +26,7 @@ namespace UnrealGameSync
 
 		public void Add(Task task)
 		{
-			Task continuationTask = task.ContinueWith(Remove);
+			Task continuationTask = task.ContinueWith(Remove, TaskScheduler.Default);
 			lock (_lockObject)
 			{
 				_tasks.Add(continuationTask);
diff --git a/Engine/Source/Programs/UnrealGameSync/UnrealGameSync/Automation/P4Automation.cs b/Engine/Source/Programs/UnrealGameSync/UnrealGameSync/Automation/P4Automation.cs
--- a/Engine/Source/Programs/UnrealGameSync/UnrealGameSync/Automation/P4Automation.cs
+++ b/Engine/Source/Programs/UnrealGameSync/UnrealGameSync/Automation/P4Automation.cs
@@ -26,10 +26,10 @@ namespace UnrealGameSync
 
 		public static Task<string> PrintToTempFile(IPerforceConnection? connection, string depotPath, ILogger logger)
 		{
-			return PrintToTempFileAsync(connection, depotPath, CancellationToken.None, logger);
+			return PrintToTempFileAsync(connection, depotPath, logger, CancellationToken.None);
 		}
 
-		public static async Task<string> PrintToTempFileAsync(IPerforceConnection? connection, string depotPath, CancellationToken cancellationToken, ILogger logger)
+		public static async Task<string> PrintToTempFileAsync(IPerforceConnection? connection, string depotPath, ILogger logger, CancellationToken cancellationToken)
 		{
 			bool createNewConnection = (connection == null);
 			try
@@ -43,11 +43,11 @@ namespace UnrealGameSync
 				string depotFileName = Path.GetFileName(depotPath);
 
 				// Reorder CL and extension
-				int index = depotFileName.IndexOf('@');
+				int index = depotFileName.IndexOf('@', StringComparison.Ordinal);
 				if (index == -1)
 				{
 					depotFileName += "@Latest";
-					index = depotFileName.IndexOf('@');
+					index = depotFileName.IndexOf('@', StringComparison.Ordinal);
 				}
 
 				string cl = depotFileName.Substring(index + 1);
diff --git a/Engine/Source/Programs/UnrealGameSync/UnrealGameSync/Automation/VisualStudioAutomation.cs b/Engine/Source/Programs/UnrealGameSync/UnrealGameSync/Automation/VisualStudioAutomation.cs
--- a/Engine/Source/Programs/UnrealGameSync/UnrealGameSync/Automation/VisualStudioAutomation.cs
+++ b/Engine/Source/Programs/UnrealGameSync/UnrealGameSync/Automation/VisualStudioAutomation.cs
@@ -69,8 +69,7 @@ namespace UnrealGameSync
 			}
 			try
 			{
-
-				System.Diagnostics.Process vsProcess = new System.Diagnostics.Process { StartInfo = new ProcessStartInfo(install.DevEnvPath, arguments) };
+				using System.Diagnostics.Process vsProcess = new System.Diagnostics.Process { StartInfo = new ProcessStartInfo(install.DevEnvPath, arguments) };
 				vsProcess.Start();
 				vsProcess.WaitForInputIdle();
 			}
@@ -148,7 +147,7 @@ namespace UnrealGameSync
 			for (int idx = 0; idx < installs.Count; idx++)
 			{
 				string? moniker = installs[idx].RotMoniker;
-				if (moniker != null && inName.StartsWith(moniker))
+				if (moniker != null && inName.StartsWith(moniker, StringComparison.Ordinal))
 				{
 					return true;
 				}
diff --git a/Engine/Source/Programs/UnrealGameSync/UnrealGameSync/AutomationServer.cs b/Engine/Source/Programs/UnrealGameSync/UnrealGameSync/AutomationServer.cs
--- a/Engine/Source/Programs/UnrealGameSync/UnrealGameSync/AutomationServer.cs
+++ b/Engine/Source/Programs/UnrealGameSync/UnrealGameSync/AutomationServer.cs
@@ -37,7 +37,7 @@ namespace UnrealGameSync
 
 		public static AutomationRequestInput Read(Stream inputStream)
 		{
-			BinaryReader reader = new BinaryReader(inputStream);
+			using BinaryReader reader = new BinaryReader(inputStream);
 			
 			int type = reader.ReadInt32();
 			int inputSize = reader.ReadInt32();
@@ -48,7 +48,7 @@ namespace UnrealGameSync
 
 		public void Write(Stream outputStream)
 		{
-			BinaryWriter writer = new BinaryWriter(outputStream);
+			using BinaryWriter writer = new BinaryWriter(outputStream, Encoding.UTF8, true);
 
 			writer.Write((int)Type);
 			writer.Write(Data.Length);
@@ -74,7 +74,7 @@ namespace UnrealGameSync
 		public AutomationRequestOutput(AutomationRequestResult result)
 		{
 			Result = result;
-			Data = new byte[0];
+			Data = Array.Empty<byte>();
 		}
 
 		public AutomationRequestOutput(AutomationRequestResult result, byte[] data)
@@ -251,7 +251,7 @@ namespace UnrealGameSync
 				{
 					try
 					{
-						TcpClient client = await listener.AcceptTcpClientAsync().ConfigureAwait(false);
+						TcpClient client = await listener.AcceptTcpClientAsync(cancellationToken).ConfigureAwait(false);
 						try
 						{
 							_logger.LogInformation("Accepted connection from {Remote}", client.Client.RemoteEndPoint);
@@ -265,7 +265,7 @@ namespace UnrealGameSync
 							using (AutomationRequest request = new AutomationRequest(input))
 							{
 								_postRequest(request);
-								request.Complete.Wait();
+								request.Complete.Wait(cancellationToken);
 								output = request.Output!;
 							}
 
diff --git a/Engine/Source/Programs/UnrealGameSync/UnrealGameSync/Controls/CustomListViewControl.cs b/Engine/Source/Programs/UnrealGameSync/UnrealGameSync/Controls/CustomListViewControl.cs
--- a/Engine/Source/Programs/UnrealGameSync/UnrealGameSync/Controls/CustomListViewControl.cs
+++ b/Engine/Source/Programs/UnrealGameSync/UnrealGameSync/Controls/CustomListViewControl.cs
@@ -249,7 +249,7 @@ namespace UnrealGameSync
 			return FindWidget(HitTest(location));
 		}
 
-		protected CustomListViewWidget? FindWidget(ListViewHitTestInfo hitTest)
+		protected static CustomListViewWidget? FindWidget(ListViewHitTestInfo hitTest)
 		{
 			if(hitTest.Item != null)
 			{
@@ -384,6 +384,8 @@ namespace UnrealGameSync
 
 		public void DrawIcon(Graphics graphics, Rectangle bounds, Rectangle icon)
 		{
+			_ = this;
+
 			float dpiScaleX = graphics.DpiX / 96.0f;
 			float dpiScaleY = graphics.DpiY / 96.0f;
 
@@ -439,6 +441,7 @@ namespace UnrealGameSync
 
 		public void DrawDefaultBackground(Graphics graphics, Rectangle bounds)
 		{
+			_ = this;
 			graphics.FillRectangle(SystemBrushes.Window, bounds);
 		}
 
diff --git a/Engine/Source/Programs/UnrealGameSync/UnrealGameSync/Controls/LogControl.cs b/Engine/Source/Programs/UnrealGameSync/UnrealGameSync/Controls/LogControl.cs
--- a/Engine/Source/Programs/UnrealGameSync/UnrealGameSync/Controls/LogControl.cs
+++ b/Engine/Source/Programs/UnrealGameSync/UnrealGameSync/Controls/LogControl.cs
@@ -19,7 +19,7 @@ namespace UnrealGameSync
 	public class LogControl : UserControl
 	{
 		[Flags]
-		public enum ScrollInfoMask : uint
+		public enum ScrollInfoMask : int
 		{
 			SifRange = 0x1,
 			SifPage = 0x2,
@@ -193,7 +193,8 @@ namespace UnrealGameSync
 
 			_logFileStream.Seek(0, SeekOrigin.Begin);
 
-			string text = new StreamReader(_logFileStream).ReadToEnd().TrimEnd('\r', '\n');
+			using StreamReader reader = new StreamReader(_logFileStream, leaveOpen: true);
+			string text = reader.ReadToEnd().TrimEnd('\r', '\n');
 			if(text.Length > 0)
 			{
 				AddLinesInternal(text.Split('\n').Select(x => x + "\n").ToList());
@@ -381,12 +382,12 @@ namespace UnrealGameSync
 			}
 		}
 
-		protected void ContextMenu_CopySelection(object? sender, EventArgs e)
+		private void ContextMenu_CopySelection(object? sender, EventArgs e)
 		{
 			CopySelection();
 		}
 
-		protected void ContextMenu_SelectAll(object? sender, EventArgs e)
+		private void ContextMenu_SelectAll(object? sender, EventArgs e)
 		{
 			SelectAll();
 		}
@@ -790,7 +791,7 @@ namespace UnrealGameSync
 			_scrollLine = verticalScroll.nPos;
 		}
 
-		protected void SelectionScrollTimer_TimerElapsed(object? sender, EventArgs args)
+		private void SelectionScrollTimer_TimerElapsed(object? sender, EventArgs args)
 		{
 			if(_autoScrollRate != 0 && _selection != null)
 			{
diff --git a/Engine/Source/Programs/UnrealGameSync/UnrealGameSync/Controls/LogSplitContainer.cs b/Engine/Source/Programs/UnrealGameSync/UnrealGameSync/Controls/LogSplitContainer.cs
--- a/Engine/Source/Programs/UnrealGameSync/UnrealGameSync/Controls/LogSplitContainer.cs
+++ b/Engine/Source/Programs/UnrealGameSync/UnrealGameSync/Controls/LogSplitContainer.cs
@@ -40,7 +40,7 @@ namespace UnrealGameSync
 			SplitterMoved += OnSplitterMoved;
 		}
 
-		protected Font CaptionFont => SystemFonts.IconTitleFont!;
+		protected static Font CaptionFont => SystemFonts.IconTitleFont!;
 
 		protected override void OnCreateControl()
 		{
@@ -58,7 +58,7 @@ namespace UnrealGameSync
 			Invalidate();
 		}
 
-		public void OnSplitterMoved(object? obj, EventArgs args)
+		private void OnSplitterMoved(object? obj, EventArgs args)
 		{
 			_splitterMoved = true;
 
diff --git a/Engine/Source/Programs/UnrealGameSync/UnrealGameSync/Controls/StatusPanel.cs b/Engine/Source/Programs/UnrealGameSync/UnrealGameSync/Controls/StatusPanel.cs
--- a/Engine/Source/Programs/UnrealGameSync/UnrealGameSync/Controls/StatusPanel.cs
+++ b/Engine/Source/Programs/UnrealGameSync/UnrealGameSync/Controls/StatusPanel.cs
@@ -22,7 +22,7 @@ namespace UnrealGameSync
 		public static readonly Cursor Hand = new Cursor(LoadCursor(IntPtr.Zero, new IntPtr(IdcHand)));
 	}
 
-	class StatusElementResources
+	class StatusElementResources : IDisposable
 	{
 		readonly Dictionary<FontStyle, Font> _fontCache = new Dictionary<FontStyle, Font>();
 		public readonly Font BadgeFont;
diff --git a/Engine/Source/Programs/UnrealGameSync/UnrealGameSync/Controls/SyncFilterControl.cs b/Engine/Source/Programs/UnrealGameSync/UnrealGameSync/Controls/SyncFilterControl.cs
--- a/Engine/Source/Programs/UnrealGameSync/UnrealGameSync/Controls/SyncFilterControl.cs
+++ b/Engine/Source/Programs/UnrealGameSync/UnrealGameSync/Controls/SyncFilterControl.cs
@@ -27,7 +27,7 @@ namespace UnrealGameSync.Controls
 
 		private void SyntaxButton_LinkClicked(object sender, LinkLabelLinkClickedEventArgs e)
 		{
-			SyncFilterSyntax dialog = new SyncFilterSyntax();
+			using SyncFilterSyntax dialog = new SyncFilterSyntax();
 			dialog.ShowDialog(ParentForm);
 		}
 
diff --git a/Engine/Source/Programs/UnrealGameSync/UnrealGameSync/Controls/WorkspaceControl.cs b/Engine/Source/Programs/UnrealGameSync/UnrealGameSync/Controls/WorkspaceControl.cs
--- a/Engine/Source/Programs/UnrealGameSync/UnrealGameSync/Controls/WorkspaceControl.cs
+++ b/Engine/Source/Programs/UnrealGameSync/UnrealGameSync/Controls/WorkspaceControl.cs
@@ -28,13 +28,13 @@ namespace UnrealGameSync
 	{
 		public string Guid { get; set; } = String.Empty;
 		public string Name { get; set; } = String.Empty;
-		public List<string> Files { get; set; } = new List<string>();
+		public List<string> Files { get; } = new List<string>();
 	}
 
 	public class UncontrolledChangelistPersistency
 	{
 		public int Version { get; set; }
-		public List<UncontrolledChangelist> Changelists { get; set; } = new List<UncontrolledChangelist>();
+		public List<UncontrolledChangelist> Changelists { get; } = new List<UncontrolledChangelist>();
 	}
 
 	interface IWorkspaceControlOwner
@@ -142,7 +142,7 @@ namespace UnrealGameSync
 
 			public override int GetHashCode()
 			{
-				return Name.GetHashCode() + (Url ?? "").GetHashCode() * 3 + Result.GetHashCode() * 5;
+				return HashCode.Combine(Name, Url, Result);
 			}
 		}
 
@@ -441,7 +441,7 @@ namespace UnrealGameSync
 
 		public List<IssueData> GetOpenIssuesForUser()
 		{
-			return _issueMonitor.GetIssues().Where(x => x.FixChange <= 0 && String.Compare(x.Owner, _perforceSettings.UserName, StringComparison.OrdinalIgnoreCase) == 0).ToList();
+			return _issueMonitor.GetIssues().Where(x => x.FixChange <= 0 && String.Equals(x.Owner, _perforceSettings.UserName, StringComparison.OrdinalIgnoreCase)).ToList();
 		}
 
 		public void IssueMonitor_OnIssuesChanged()
@@ -1067,7 +1067,7 @@ namespace UnrealGameSync
 				DesiredTaskbarState = Tuple.Create(TaskbarState.Paused, 0.0f);
 				_owner.UpdateProgress();
 
-				DeleteWindow window = new DeleteWindow(context.DeleteFiles);
+				using DeleteWindow window = new DeleteWindow(context.DeleteFiles);
 				if (window.ShowDialog(this) == DialogResult.OK)
 				{
 					StartWorkspaceUpdate(context, _updateCallback);
@@ -1095,7 +1095,7 @@ namespace UnrealGameSync
 					ParseUncontrolledChangelistsPersistencyFile(Path.GetDirectoryName(SelectedProject.LocalPath)!, uncontrolledFiles);
 				}
 
-				ClobberWindow window = new ClobberWindow(context.ClobberFiles, uncontrolledFiles);
+				using ClobberWindow window = new ClobberWindow(context.ClobberFiles, uncontrolledFiles);
 
 				if (window.ShowDialog(this) == DialogResult.OK)
 				{
@@ -1172,7 +1172,7 @@ namespace UnrealGameSync
 
 				_promotedChangeNumbers = _perforceMonitor.GetPromotedChangeNumbers();
 
-				string[] excludeChanges = new string[0];
+				string[] excludeChanges = Array.Empty<string>();
 				if (_workspace != null)
 				{
 					ConfigFile projectConfigFile = _perforceMonitor.LatestProjectConfigFile;
@@ -1402,11 +1402,11 @@ namespace UnrealGameSync
 		/// </summary>
 		/// <param name="change"></param>
 		/// <returns></returns>
-		bool IsRobomergeChange(ChangesRecord change)
+		static bool IsRobomergeChange(ChangesRecord change)
 		{
 			// If hiding robomerge, filter out based on workspace name. Note - don't look at the description because we
 			// *do* want to see conflicts that were manually merged by someone
-			return change.Client!.IndexOf("ROBOMERGE", StringComparison.OrdinalIgnoreCase) >= 0;
+			return change.Client!.Contains("ROBOMERGE", StringComparison.OrdinalIgnoreCase);
 		}
 
 		bool ShouldShowChange(ChangesRecord change, string[] excludeChanges)
@@ -1472,7 +1472,7 @@ namespace UnrealGameSync
 					}
 				}
 
-				if (String.Compare(change.User, "buildmachine", true) == 0 && change.Description!.IndexOf("lightmaps", StringComparison.InvariantCultureIgnoreCase) == -1)
+				if (String.Equals(change.User, "buildmachine", StringComparison.OrdinalIgnoreCase) && !change.Description!.Contains("lightmaps", StringComparison.OrdinalIgnoreCase))
 				{
 					return false;
 				}
@@ -1685,7 +1685,7 @@ namespace UnrealGameSync
 			return ConfigUtils.GetSyncCategories(_perforceMonitor.LatestProjectConfigFile);
 		}
 
-		List<string> GetProjectRoots(string inUProjectDirsPath)
+		static List<string> GetProjectRoots(string inUProjectDirsPath)
 		{
 			List<string> projectRoots = new List<string>();
 
@@ -1702,7 +1702,7 @@ namespace UnrealGameSync
 				string trimmedLine = line.Trim();
 
 				// Remove every comments
-				int commentIndex = trimmedLine.IndexOf(';');
+				int commentIndex = trimmedLine.IndexOf(';', StringComparison.Ordinal);
 
 				if (commentIndex >= 0)
 				{
@@ -1807,7 +1807,7 @@ namespace UnrealGameSync
 
 			// Find all the badges which should notify users due to content changes
 			HashSet<string> contentBadges = new HashSet<string>();
-			contentBadges.UnionWith(_perforceMonitor.LatestProjectConfigFile.GetValues("Notifications.ContentBadges", new string[0]));
+			contentBadges.UnionWith(_perforceMonitor.LatestProjectConfigFile.GetValues("Notifications.ContentBadges", Array.Empty<string>()));
 
 			// Find the most recent build of each type, and the last time it succeeded
 			Dictionary<string, BadgeData> typeToLastBuild = new Dictionary<string, BadgeData>();
@@ -2055,7 +2055,7 @@ namespace UnrealGameSync
 			return archivePath;
 		}
 
-		private Color Blend(Color first, Color second, float t)
+		private static Color Blend(Color first, Color second, float t)
 		{
 			return Color.FromArgb((int)(first.R + (second.R - first.R) * t), (int)(first.G + (second.G - first.G) * t), (int)(first.B + (second.B - first.B) * t));
 		}
@@ -2078,7 +2078,7 @@ namespace UnrealGameSync
 		/// </summary>
 		/// <param name="inBadges"></param>
 		/// <returns></returns>
-		private bool DoRequiredBadgesExist(List<string> requiredBadgeList, List<BadgeData> inBadges)
+		private static bool DoRequiredBadgesExist(List<string> requiredBadgeList, List<BadgeData> inBadges)
 		{
 			Dictionary<string, BadgeData> inBadgeDictionary = new Dictionary<string, BadgeData>();
 			foreach (BadgeData badge in inBadges)
@@ -2188,7 +2188,7 @@ namespace UnrealGameSync
 				return;
 			}
 
-			ChangesRecord? change = (ChangesRecord?)e.Item?.Tag;
+			ChangesRecord? change = (ChangesRecord?)e.Item.Tag;
 			if (change == null)
 			{
 				return;
@@ -2199,7 +2199,7 @@ namespace UnrealGameSync
 
 			float iconY = e.Bounds.Top + (e.Bounds.Height - 16 * dpiScaleY) / 2;
 
-			StringFormat format = new StringFormat();
+			using StringFormat format = new StringFormat();
 			format.LineAlignment = StringAlignment.Center;
 			format.FormatFlags = StringFormatFlags.NoWrap;
 			format.Trimming = StringTrimming.EllipsisCharacter;
@@ -2416,9 +2416,9 @@ namespace UnrealGameSync
 						summaryText = (summaryTextBuilder.Length == 0) ? "No current users." : summaryTextBuilder.ToString();
 					}
 
-					if (summaryText != null && summaryText.Contains('\n'))
+					if (summaryText != null && summaryText.Contains('\n', StringComparison.Ordinal))
 					{
-						summaryText = summaryText.Substring(0, summaryText.IndexOf('\n')).TrimEnd() + "...";
+						summaryText = summaryText.Substring(0, summaryText.IndexOf('\n', StringComparison.Ordinal)).TrimEnd() + "...";
 					}
 
 					TextRenderer.DrawText(e.Graphics, summaryText, currentFont, new Rectangle(e.Bounds.Left, e.Bounds.Top, (int)maxX - e.Bounds.Left, e.Bounds.Height), textColor, TextFormatFlags.Left | TextFormatFlags.VerticalCenter | TextFormatFlags.EndEllipsis | TextFormatFlags.NoPrefix);
@@ -2465,7 +2465,7 @@ namespace UnrealGameSync
 			}
 		}
 
-		private Point GetBadgeListLocation(List<BadgeInfo> badges, Rectangle bounds, HorizontalAlign horzAlign, VerticalAlignment vertAlign)
+		private static Point GetBadgeListLocation(List<BadgeInfo> badges, Rectangle bounds, HorizontalAlign horzAlign, VerticalAlignment vertAlign)
 		{
 			if (badges.Count == 0)
 			{
@@ -2536,7 +2536,7 @@ namespace UnrealGameSync
 				ConfigFile projectConfigFile = _perforceMonitor.LatestProjectConfigFile;
 				if (projectConfigFile != null)
 				{
-					string[] badgeDefinitions = projectConfigFile.GetValues("Badges.DescriptionBadges", new string[0]);
+					string[] badgeDefinitions = projectConfigFile.GetValues("Badges.DescriptionBadges", Array.Empty<string>());
 					foreach (string badgeDefinition in badgeDefinitions.Distinct())
 					{
 						ConfigObject badgeDefinitionObject = new ConfigObject(badgeDefinition);
@@ -2703,14 +2703,14 @@ namespace UnrealGameSync
 		}
 
 		[return: NotNullIfNotNull("text")]
-		private string? ReplaceRegexMatches(string? text, Match matchResult)
+		private static string? ReplaceRegexMatches(string? text, Match matchResult)
 		{
 			if (text != null)
 			{
 				for (int idx = 1; idx < matchResult.Groups.Count; idx++)
 				{
 					string sourceText = String.Format("${0}", idx);
-					text = text.Replace(sourceText, matchResult.Groups[idx].Value);
+					text = text.Replace(sourceText, matchResult.Groups[idx].Value, StringComparison.Ordinal);
 				}
 			}
 			return text;
@@ -2795,7 +2795,7 @@ namespace UnrealGameSync
 			return badges;
 		}
 
-		private BadgeInfo CreateBadge(int changeNumber, string badgeName, string badgeGroup, BadgeData? badgeData)
+		private static BadgeInfo CreateBadge(int changeNumber, string badgeName, string badgeGroup, BadgeData? badgeData)
 		{
 			string badgeLabel = badgeName;
 			Color badgeColor = Color.FromArgb(0, Color.White);
@@ -3107,7 +3107,7 @@ namespace UnrealGameSync
 			FileReference[] processFileNames = GetProcessesRunningInWorkspace();
 			if (processFileNames.Length > 0)
 			{
-				ProgramsRunningWindow programsRunning = new ProgramsRunningWindow(GetProcessesRunningInWorkspace, processFileNames);
+				using ProgramsRunningWindow programsRunning = new ProgramsRunningWindow(GetProcessesRunningInWorkspace, processFileNames);
 				if (programsRunning.ShowDialog() != DialogResult.OK)
 				{
 					return false;
@@ -3563,14 +3563,14 @@ namespace UnrealGameSync
 
 						StatusLine successLine = new StatusLine();
 						successLine.AddIcon(Properties.Resources.StatusIcons, new Size(16, 16), 0);
-						successLine.AddText(String.Format("  Sync took {0}{1}s, completed at {2}.", (_workspace.State.LastSyncDurationSeconds >= 60) ? String.Format("{0}m ", _workspace.State.LastSyncDurationSeconds / 60) : "", _workspace.State.LastSyncDurationSeconds % 60, _workspace.State.LastSyncTime.Value.ToLocalTime().ToString("h\\:mmtt").ToLowerInvariant()));
+						successLine.AddText(String.Format("  Sync took {0}{1}s, completed at {2}.", (_workspace.State.LastSyncDurationSeconds >= 60) ? String.Format("{0}m ", _workspace.State.LastSyncDurationSeconds / 60) : "", _workspace.State.LastSyncDurationSeconds % 60, _workspace.State.LastSyncTime.Value.ToLocalTime().ToString("h\\:mmtt")));
 						lines.Add(successLine);
 					}
 					else if (GetLastUpdateMessage(_workspace.State.LastSyncResult, _workspace.State.LastSyncResultMessage, out summaryText))
 					{
 						lines.Add(new StatusLine() { LineHeight = 0.5f });
 
-						int summaryTextLength = summaryText.IndexOf('\n');
+						int summaryTextLength = summaryText.IndexOf('\n', StringComparison.Ordinal);
 						if (summaryTextLength == -1)
 						{
 							summaryTextLength = summaryText.Length;
@@ -3682,7 +3682,7 @@ namespace UnrealGameSync
 				{
 					issues.Add(issue);
 				}
-				else if (issue.FixChange == 0 && issue.Owner != null && String.Compare(issue.Owner, _perforceSettings.UserName, StringComparison.OrdinalIgnoreCase) == 0)
+				else if (issue.FixChange == 0 && issue.Owner != null && String.Equals(issue.Owner, _perforceSettings.UserName, StringComparison.OrdinalIgnoreCase))
 				{
 					hasOtherAssignedIssue = true;
 				}
@@ -3733,7 +3733,7 @@ namespace UnrealGameSync
 					{
 						item.ForeColor = SystemColors.GrayText;
 					}
-					else if (issue.Owner != null && String.Compare(issue.Owner, _perforceSettings.UserName, StringComparison.OrdinalIgnoreCase) == 0)
+					else if (issue.Owner != null && String.Equals(issue.Owner, _perforceSettings.UserName, StringComparison.OrdinalIgnoreCase))
 					{
 						item.Font = new Font(item.Font, FontStyle.Bold);
 					}
@@ -3777,7 +3777,7 @@ namespace UnrealGameSync
 			string? buildHealthProject;
 			TryGetProjectSetting(_perforceMonitor.LatestProjectConfigFile, "BuildHealthProject", out buildHealthProject);
 
-			IssueSettingsWindow issueSettings = new IssueSettingsWindow(_settings, buildHealthProject ?? "", _logger);
+			using IssueSettingsWindow issueSettings = new IssueSettingsWindow(_settings, buildHealthProject ?? "", _logger);
 			if (issueSettings.ShowDialog(this) == DialogResult.OK)
 			{
 				_owner.UpdateAlertWindows();
@@ -3790,12 +3790,12 @@ namespace UnrealGameSync
 			if (TryGetProjectSetting(_perforceMonitor.LatestProjectConfigFile, "SdkInfo", out sdkInfoEntries))
 			{
 				Dictionary<string, string> variables = _workspace.GetVariables(GetEditorBuildConfig());
-				SdkInfoWindow window = new SdkInfoWindow(sdkInfoEntries, variables, _badgeFont!);
+				using SdkInfoWindow window = new SdkInfoWindow(sdkInfoEntries, variables, _badgeFont!);
 				window.ShowDialog();
 			}
 		}
 
-		private StatusLine CreateStatusLineFromMarkdown(string text)
+		private static StatusLine CreateStatusLineFromMarkdown(string text)
 		{
 			StatusLine line = new StatusLine();
 
@@ -3842,7 +3842,7 @@ namespace UnrealGameSync
 						continue;
 					}
 
-					int endIdx = text.IndexOf("](", idx);
+					int endIdx = text.IndexOf("](", idx, StringComparison.Ordinal);
 					if (endIdx != -1)
 					{
 						int urlEndIdx = text.IndexOf(')', endIdx + 2);
@@ -3864,7 +3864,7 @@ namespace UnrealGameSync
 					int endIdx = text.IndexOf(':', idx + 1);
 					if (endIdx != -1)
 					{
-						if (String.Compare(":alert:", 0, text, idx, endIdx - idx) == 0)
+						if (String.Compare(":alert:", 0, text, idx, endIdx - idx, StringComparison.Ordinal) == 0)
 						{
 							FlushMarkdownText(line, elementText, style);
 							line.AddIcon(Properties.Resources.StatusIcons, new Size(16, 16), 4);
@@ -3882,7 +3882,7 @@ namespace UnrealGameSync
 			return line;
 		}
 
-		private void FlushMarkdownText(StatusLine line, StringBuilder elementText, FontStyle style)
+		private static void FlushMarkdownText(StatusLine line, StringBuilder elementText, FontStyle style)
 		{
 			if (elementText.Length > 0)
 			{
@@ -3907,10 +3907,10 @@ namespace UnrealGameSync
 
 					StreamContextMenu.Items.Add(new ToolStripSeparator());
 
-					foreach (string otherStreamName in otherStreamNames.OrderBy(x => x).Where(x => !x.EndsWith("/Dev-Binaries")))
+					foreach (string otherStreamName in otherStreamNames.OrderBy(x => x).Where(x => !x.EndsWith("/Dev-Binaries", StringComparison.OrdinalIgnoreCase)))
 					{
 						string thisStreamName = otherStreamName; // Local for lambda capture
-						if (String.Compare(StreamName, otherStreamName, StringComparison.InvariantCultureIgnoreCase) != 0)
+						if (String.Equals(StreamName, otherStreamName, StringComparison.OrdinalIgnoreCase))
 						{
 							ToolStripMenuItem item = new ToolStripMenuItem(thisStreamName, null, new EventHandler((s, e) => SelectStream(thisStreamName)));
 							StreamContextMenu.Items.Add(item);
@@ -3960,7 +3960,7 @@ namespace UnrealGameSync
 					{
 						async Task<bool> SwitchFunc(IPerforceConnection perforce, CancellationToken cancellationToken)
 						{
-							if (idx == 1 || !await perforce.OpenedAsync(OpenedOptions.None, FileSpecList.Any, cancellationToken).AnyAsync())
+							if (idx == 1 || !await perforce.OpenedAsync(OpenedOptions.None, FileSpecList.Any, cancellationToken).AnyAsync(cancellationToken))
 							{
 								await perforce.SwitchClientToStreamAsync(newStreamName, SwitchClientOptions.IgnoreOpenFiles, cancellationToken);
 								return true;
@@ -4268,7 +4268,7 @@ namespace UnrealGameSync
 				{
 					Dictionary<string, string> variables = _workspace.GetVariables(GetEditorBuildConfig(), _contextMenuChange.Number, -1);
 
-					string[] changeContextMenuEntries = projectConfigFile.GetValues("Options.ContextMenu", new string[0]);
+					string[] changeContextMenuEntries = projectConfigFile.GetValues("Options.ContextMenu", Array.Empty<string>());
 					foreach (string changeContextMenuEntry in changeContextMenuEntries)
 					{
 						ConfigObject obj = new ConfigObject(changeContextMenuEntry);
@@ -4381,7 +4381,7 @@ namespace UnrealGameSync
 			}
 		}
 
-		private BadgeInfo? HitTestBadge(Point location, List<BadgeInfo> badgeList, Point listLocation)
+		private static BadgeInfo? HitTestBadge(Point location, List<BadgeInfo> badgeList, Point listLocation)
 		{
 			foreach (BadgeInfo badge in badgeList)
 			{
@@ -4902,7 +4902,7 @@ namespace UnrealGameSync
 							}
 							if (summaryText.Length > 0)
 							{
-								summaryText.Append("\n");
+								summaryText.Append('\n');
 							}
 						}
 						AppendUserList(summaryText, "\n", "Compiled by {0}.", summary.Reviews.Where(x => x.Type == EventType.Compiles));
@@ -4976,7 +4976,7 @@ namespace UnrealGameSync
 
 		private bool ModifyEditorArguments()
 		{
-			ArgumentsWindow arguments = new ArgumentsWindow(_settings.EditorArguments, _settings.EditorArgumentsPrompt);
+			using ArgumentsWindow arguments = new ArgumentsWindow(_settings.EditorArguments, _settings.EditorArgumentsPrompt);
 			if (arguments.ShowDialog(this) == DialogResult.OK)
 			{
 				_settings.EditorArguments.Clear();
@@ -5086,7 +5086,7 @@ namespace UnrealGameSync
 					scheduledSyncTypeId = projectSetting.ScheduledSyncTypeId;
 				}
 
-				_logger.LogInformation("Scheduled sync at {0} for {1}.",
+				_logger.LogInformation("Scheduled sync at {Time} for {Type}.",
 					DateTime.Now,
 					scheduledSyncTypeId ?? "Default");
 
@@ -5223,7 +5223,7 @@ namespace UnrealGameSync
 		{
 			if (_contextMenuChange != null)
 			{
-				LeaveCommentWindow leaveComment = new LeaveCommentWindow();
+				using LeaveCommentWindow leaveComment = new LeaveCommentWindow();
 
 				string? commentText;
 				if (_eventMonitor.GetCommentByCurrentUser(_contextMenuChange.Number, out commentText))
@@ -5311,7 +5311,7 @@ namespace UnrealGameSync
 					{
 						if (!String.IsNullOrEmpty(step.StatusPanelLink) && CanRunStep(step))
 						{
-							int baseMenuIdx = step.StatusPanelLink.IndexOf('|');
+							int baseMenuIdx = step.StatusPanelLink.IndexOf('|', StringComparison.Ordinal);
 							if (baseMenuIdx == -1)
 							{
 								_customStatusPanelLinks.Add((step.StatusPanelLink, (p, r) => RunCustomTool(step, userSteps)));
@@ -5319,7 +5319,7 @@ namespace UnrealGameSync
 							else
 							{
 								string menuName = step.StatusPanelLink.Substring(0, baseMenuIdx);
-								string itemName = step.StatusPanelLink.Substring(baseMenuIdx + 1).Replace("&", "&&");
+								string itemName = step.StatusPanelLink.Substring(baseMenuIdx + 1).Replace("&", "&&", StringComparison.Ordinal);
 
 								ToolStripMenuItem newMenuItem = new ToolStripMenuItem(itemName);
 								newMenuItem.Click += new EventHandler((sender, e) => { RunCustomTool(step, userSteps); });
@@ -5436,7 +5436,7 @@ namespace UnrealGameSync
 					{
 						foreach (FileInfo targetFile in sourceDirectory.EnumerateFiles("*.target.cs", SearchOption.TopDirectoryOnly))
 						{
-							targetNames.Add(targetFile.Name.Substring(0, targetFile.Name.IndexOf('.')));
+							targetNames.Add(targetFile.Name.Substring(0, targetFile.Name.IndexOf('.', StringComparison.Ordinal)));
 						}
 					}
 				}
@@ -5447,7 +5447,7 @@ namespace UnrealGameSync
 
 				// Show the dialog
 				Dictionary<string, string> variables = _workspace.GetVariables(GetEditorBuildConfig());
-				ModifyBuildStepsWindow editStepsWindow = new ModifyBuildStepsWindow(targetNames, userSteps, new HashSet<Guid>(projectBuildStepObjects.Keys), BranchDirectoryName, variables);
+				using ModifyBuildStepsWindow editStepsWindow = new ModifyBuildStepsWindow(targetNames, userSteps, new HashSet<Guid>(projectBuildStepObjects.Keys), BranchDirectoryName, variables);
 				editStepsWindow.ShowDialog();
 
 				// Update the user settings
@@ -5477,7 +5477,7 @@ namespace UnrealGameSync
 			}
 		}
 
-		private void AddOrUpdateBuildStep(Dictionary<Guid, ConfigObject> steps, ConfigObject obj)
+		private static void AddOrUpdateBuildStep(Dictionary<Guid, ConfigObject> steps, ConfigObject obj)
 		{
 			Guid uniqueId;
 			if (Guid.TryParse(obj.GetValue(BuildStep.UniqueIdKey, ""), out uniqueId))
@@ -5508,7 +5508,7 @@ namespace UnrealGameSync
 			string editorTargetName = editorTargetFile.GetFileNameWithoutAnyExtensions();
 
 			Dictionary<Guid, ConfigObject> projectBuildSteps = ConfigUtils.GetDefaultBuildStepObjects(ProjectInfo, editorTargetName, _settings.CompiledEditorBuildConfig, projectConfigFile, ShouldSyncPrecompiledEditor);
-			foreach (string line in projectConfigFile.GetValues("Build.Step", new string[0]))
+			foreach (string line in projectConfigFile.GetValues("Build.Step", Array.Empty<string>()))
 			{
 				AddOrUpdateBuildStep(projectBuildSteps, new ConfigObject(line));
 			}
@@ -5561,13 +5561,13 @@ namespace UnrealGameSync
 			diagnosticsText.AppendFormat("Event monitor: {0}\n", (_eventMonitor == null) ? "(inactive)" : _eventMonitor.LastStatusMessage);
 			diagnosticsText.AppendFormat("Issue monitor: {0}\n", (_issueMonitor == null) ? "(inactive)" : _issueMonitor.LastStatusMessage);
 
-			DiagnosticsWindow diagnostics = new DiagnosticsWindow(_appDataFolder, _workspaceDataFolder, diagnosticsText.ToString(), _settings.GetCachedFilePaths());
+			using DiagnosticsWindow diagnostics = new DiagnosticsWindow(_appDataFolder, _workspaceDataFolder, diagnosticsText.ToString(), _settings.GetCachedFilePaths());
 			diagnostics.ShowDialog(this);
 		}
 
 		private void OptionsContextMenu_SyncFilter_Click(object sender, EventArgs e)
 		{
-			SyncFilter filter = new SyncFilter(GetSyncCategories(), _settings.Global.Filter, _workspaceSettings.Filter, _perforceMonitor.LatestPerforceConfigSection());
+			using SyncFilter filter = new SyncFilter(GetSyncCategories(), _settings.Global.Filter, _workspaceSettings.Filter, _perforceMonitor.LatestPerforceConfigSection());
 			if (filter.ShowDialog() == DialogResult.OK)
 			{
 				_settings.Global.Filter = filter.GlobalFilter;
@@ -5650,7 +5650,7 @@ namespace UnrealGameSync
 				return;
 			}
 
-			ChangelistWindow childWindow = new ChangelistWindow((_workspace == null) ? -1 : _workspace.CurrentChangeNumber);
+			using ChangelistWindow childWindow = new ChangelistWindow((_workspace == null) ? -1 : _workspace.CurrentChangeNumber);
 			if (childWindow.ShowDialog() == DialogResult.OK)
 			{
 				StartSync(childWindow.ChangeNumber);
diff --git a/Engine/Source/Programs/UnrealGameSync/UnrealGameSync/EventMonitor.cs b/Engine/Source/Programs/UnrealGameSync/UnrealGameSync/EventMonitor.cs
--- a/Engine/Source/Programs/UnrealGameSync/UnrealGameSync/EventMonitor.cs
+++ b/Engine/Source/Programs/UnrealGameSync/UnrealGameSync/EventMonitor.cs
@@ -93,7 +93,7 @@ namespace UnrealGameSync
 					return "Unknown";
 				}
 
-				int idx = BuildType.IndexOf(':');
+				int idx = BuildType.IndexOf(':', StringComparison.Ordinal);
 				if(idx == -1)
 				{
 					return BuildType;
@@ -114,7 +114,7 @@ namespace UnrealGameSync
 					return "Unknown";
 				}
 
-				int idx = BuildType.IndexOf(':');
+				int idx = BuildType.IndexOf(':', StringComparison.Ordinal);
 				if(idx == -1)
 				{
 					return BuildType;
@@ -288,7 +288,7 @@ namespace UnrealGameSync
 			if(_workerTask != null)
 			{
 				_cancellationSource.Cancel();
-				_asyncDisposer.Add(_workerTask.ContinueWith(_ => _cancellationSource.Dispose()));
+				_asyncDisposer.Add(_workerTask.ContinueWith(_ => _cancellationSource.Dispose(), CancellationToken.None, TaskContinuationOptions.None, TaskScheduler.Default));
 				_workerTask = null;
 			}
 		}
@@ -421,14 +421,14 @@ namespace UnrealGameSync
 			}
 			else if(evt.Type == EventType.Syncing)
 			{
-				summary.SyncEvents.RemoveAll(x => String.Compare(x.UserName, evt.UserName, true) == 0);
+				summary.SyncEvents.RemoveAll(x => String.Equals(x.UserName, evt.UserName, StringComparison.OrdinalIgnoreCase));
 				summary.SyncEvents.Add(evt);
 				ApplyFilteredUpdate(evt);
 			}
 			else if(IsReview(evt.Type))
 			{
 				// Try to find an existing review by this user. If we already have a newer review, ignore this one. Otherwise remove it.
-				EventData? existingReview = summary.Reviews.Find(x => String.Compare(x.UserName, evt.UserName, true) == 0);
+				EventData? existingReview = summary.Reviews.Find(x => String.Equals(x.UserName, evt.UserName, StringComparison.OrdinalIgnoreCase));
 				if(existingReview != null)
 				{
 					if(existingReview.Id <= evt.Id)
@@ -481,7 +481,7 @@ namespace UnrealGameSync
 		void ApplyCommentUpdate(CommentData comment)
 		{
 			EventSummary summary = FindOrAddSummary(comment.ChangeNumber);
-			if(String.Compare(comment.UserName, _currentUserName, true) == 0 && summary.Comments.Count > 0 && summary.Comments.Last().Id == Int64.MaxValue)
+			if(String.Equals(comment.UserName, _currentUserName, StringComparison.OrdinalIgnoreCase) && summary.Comments.Count > 0 && summary.Comments.Last().Id == Int64.MaxValue)
 			{
 				// This comment was added by PostComment(), to mask the latency of a round trip to the server. Remove it now we have the sorted comment.
 				summary.Comments.RemoveAt(summary.Comments.Count - 1);
@@ -495,7 +495,7 @@ namespace UnrealGameSync
 
 			for(; insertIdx > 0 && idSelector(items[insertIdx - 1]) >= idSelector(newItem); insertIdx--)
 			{
-				if(String.Compare(userSelector(items[insertIdx - 1]), userSelector(newItem), true) == 0)
+				if(String.Equals(userSelector(items[insertIdx - 1]), userSelector(newItem), StringComparison.OrdinalIgnoreCase))
 				{
 					return false;
 				}
@@ -505,7 +505,7 @@ namespace UnrealGameSync
 
 			for(; insertIdx > 0; insertIdx--)
 			{
-				if(String.Compare(userSelector(items[insertIdx - 1]), userSelector(newItem), true) == 0)
+				if(String.Equals(userSelector(items[insertIdx - 1]), userSelector(newItem), StringComparison.OrdinalIgnoreCase))
 				{
 					items.RemoveAt(insertIdx - 1);
 				}
@@ -681,7 +681,7 @@ namespace UnrealGameSync
 				{
 					if(evt.Id > lastSync.Id)
 					{
-						_changeNumberToSummary[lastSync.Change].CurrentUsers.RemoveAll(x => String.Compare(x, evt.UserName, true) == 0);
+						_changeNumberToSummary[lastSync.Change].CurrentUsers.RemoveAll(x => String.Equals(x, evt.UserName, StringComparison.OrdinalIgnoreCase));
 						FindOrAddSummary(evt.Change).CurrentUsers.Add(evt.UserName);
 						_userNameToLastSyncEvent[evt.UserName] = evt;
 					}
@@ -732,7 +732,7 @@ namespace UnrealGameSync
 					await ReadEventsFromBackendAsync(updateThrottledRequests, cancellationToken);
 
 					// Send a notification that we're ready to update
-					if(_incomingMetadata.Count > 0 || _incomingEvents.Count > 0 || _incomingBadges.Count > 0 || _incomingComments.Count > 0)
+					if(!_incomingMetadata.IsEmpty || !_incomingEvents.IsEmpty || !_incomingBadges.IsEmpty || !_incomingComments.IsEmpty)
 					{
 						_synchronizationContext.Post(_ => OnUpdatesReady?.Invoke(), null);
 					}
@@ -1012,7 +1012,7 @@ namespace UnrealGameSync
 				return false;
 			}
 
-			CommentData? comment = summary.Comments.Find(x => String.Compare(x.UserName, _currentUserName, true) == 0);
+			CommentData? comment = summary.Comments.Find(x => String.Equals(x.UserName, _currentUserName, StringComparison.OrdinalIgnoreCase));
 			if(comment == null || String.IsNullOrWhiteSpace(comment.Text))
 			{
 				commentText = null;
@@ -1031,7 +1031,7 @@ namespace UnrealGameSync
 				return null;
 			}
 
-			EventData? evt = summary.Reviews.FirstOrDefault(x => String.Compare(x.UserName, _currentUserName, true) == 0);
+			EventData? evt = summary.Reviews.FirstOrDefault(x => String.Equals(x.UserName, _currentUserName, StringComparison.OrdinalIgnoreCase));
 			if(evt == null || evt.Type == EventType.Unknown)
 			{
 				return null;
@@ -1070,7 +1070,7 @@ namespace UnrealGameSync
 		public bool WasSyncedByCurrentUser(int changeNumber)
 		{
 			EventSummary? summary = GetSummaryForChange(changeNumber);
-			return (summary != null && summary.SyncEvents.Any(x => x.Type == EventType.Syncing && String.Compare(x.UserName, _currentUserName, true) == 0));
+			return (summary != null && summary.SyncEvents.Any(x => x.Type == EventType.Syncing && String.Equals(x.UserName, _currentUserName, StringComparison.OrdinalIgnoreCase)));
 		}
 
 		public void StartInvestigating(int changeNumber)
@@ -1104,7 +1104,7 @@ namespace UnrealGameSync
 						}
 						else
 						{
-							_activeInvestigations.RemoveAll(x => String.Compare(x.UserName, investigationEvent.UserName, true) == 0 && x.Change <= investigationEvent.Change);
+							_activeInvestigations.RemoveAll(x => String.Equals(x.UserName, investigationEvent.UserName, StringComparison.OrdinalIgnoreCase) && x.Change <= investigationEvent.Change);
 						}
 					}
 				}
@@ -1114,7 +1114,7 @@ namespace UnrealGameSync
 				{
 					for(int otherIdx = 0; otherIdx < idx; otherIdx++)
 					{
-						if(String.Compare(_activeInvestigations[idx].UserName, _activeInvestigations[otherIdx].UserName, true) == 0)
+						if(String.Equals(_activeInvestigations[idx].UserName, _activeInvestigations[otherIdx].UserName, StringComparison.OrdinalIgnoreCase))
 						{
 							_activeInvestigations.RemoveAt(idx--);
 							break;
@@ -1133,7 +1133,7 @@ namespace UnrealGameSync
 		public bool IsUnderInvestigationByCurrentUser(int changeNumber)
 		{
 			UpdateActiveInvestigations();
-			return _activeInvestigations?.Any(x => x.Change <= changeNumber && String.Compare(x.UserName, _currentUserName, true) == 0) ?? false;
+			return _activeInvestigations?.Any(x => x.Change <= changeNumber && String.Equals(x.UserName, _currentUserName, StringComparison.OrdinalIgnoreCase)) ?? false;
 		}
 
 		public IEnumerable<string> GetInvestigatingUsers(int changeNumber)
@@ -1151,7 +1151,7 @@ namespace UnrealGameSync
 			{
 				foreach (EventData activeInvestigation in _activeInvestigations)
 				{
-					if (String.Compare(activeInvestigation.UserName, _currentUserName, true) == 0)
+					if (String.Equals(activeInvestigation.UserName, _currentUserName, StringComparison.OrdinalIgnoreCase))
 					{
 						if (activeInvestigation.Change <= lastChangeNumber && (startChangeNumber == -1 || activeInvestigation.Change < startChangeNumber))
 						{
diff --git a/Engine/Source/Programs/UnrealGameSync/UnrealGameSync/FindFoldersToCleanTask.cs b/Engine/Source/Programs/UnrealGameSync/UnrealGameSync/FindFoldersToCleanTask.cs
--- a/Engine/Source/Programs/UnrealGameSync/UnrealGameSync/FindFoldersToCleanTask.cs
+++ b/Engine/Source/Programs/UnrealGameSync/UnrealGameSync/FindFoldersToCleanTask.cs
@@ -169,7 +169,7 @@ namespace UnrealGameSync
 			// Remove any config files
 			if(perforceConfigFile != null)
 			{
-				localFolder._filesToDelete.RemoveAll(x => String.Compare(x.Name, perforceConfigFile, StringComparison.OrdinalIgnoreCase) == 0);
+				localFolder._filesToDelete.RemoveAll(x => String.Equals(x.Name, perforceConfigFile, StringComparison.OrdinalIgnoreCase));
 			}
 
 			// Figure out if this folder is just an empty directory that needs to be removed
@@ -183,7 +183,7 @@ namespace UnrealGameSync
 		{
 			if (type != null)
 			{
-				int idx = type.IndexOf('+');
+				int idx = type.IndexOf('+', StringComparison.Ordinal);
 				if (idx != -1)
 				{
 					return type.IndexOf('w', idx) != -1;
@@ -227,8 +227,8 @@ namespace UnrealGameSync
 			// Start enumerating all the files that exist locally
 			foreach(string syncPath in _syncPaths)
 			{
-				Debug.Assert(syncPath.StartsWith(_clientRootPath));
-				if(syncPath.StartsWith(_clientRootPath, StringComparison.InvariantCultureIgnoreCase))
+				Debug.Assert(syncPath.StartsWith(_clientRootPath, StringComparison.OrdinalIgnoreCase));
+				if(syncPath.StartsWith(_clientRootPath, StringComparison.OrdinalIgnoreCase))
 				{
 					string[] fragments = syncPath.Substring(_clientRootPath.Length).Split('/');
 
diff --git a/Engine/Source/Programs/UnrealGameSync/UnrealGameSync/Forms/ApplicationSettingsWindow.cs b/Engine/Source/Programs/UnrealGameSync/UnrealGameSync/Forms/ApplicationSettingsWindow.cs
--- a/Engine/Source/Programs/UnrealGameSync/UnrealGameSync/Forms/ApplicationSettingsWindow.cs
+++ b/Engine/Source/Programs/UnrealGameSync/UnrealGameSync/Forms/ApplicationSettingsWindow.cs
@@ -129,7 +129,7 @@ namespace UnrealGameSync
 
 		public static bool? ShowModal(IWin32Window owner, IPerforceSettings defaultPerforceSettings, bool preview, string originalExecutableFileName, UserSettings settings, ToolUpdateMonitor toolUpdateMonitor, ILogger<ApplicationSettingsWindow> logger)
 		{
-			ApplicationSettingsWindow applicationSettings = new ApplicationSettingsWindow(defaultPerforceSettings, preview, originalExecutableFileName, settings, toolUpdateMonitor, logger);
+			using ApplicationSettingsWindow applicationSettings = new ApplicationSettingsWindow(defaultPerforceSettings, preview, originalExecutableFileName, settings, toolUpdateMonitor, logger);
 			if(applicationSettings.ShowDialog(owner) == DialogResult.OK)
 			{
 				return applicationSettings._restartPreview;
@@ -140,7 +140,7 @@ namespace UnrealGameSync
 			}
 		}
 
-		private bool IsAutomaticallyRunAtStartup()
+		private static bool IsAutomaticallyRunAtStartup()
 		{
 			RegistryKey? key = Registry.CurrentUser.OpenSubKey("SOFTWARE\\Microsoft\\Windows\\CurrentVersion\\Run");
 			return (key?.GetValue("UnrealGameSync") != null);
@@ -269,7 +269,7 @@ namespace UnrealGameSync
 
 		private void AdvancedBtn_Click(object sender, EventArgs e)
 		{
-			PerforceSyncSettingsWindow window = new PerforceSyncSettingsWindow(_settings, _logger);
+			using PerforceSyncSettingsWindow window = new PerforceSyncSettingsWindow(_settings, _logger);
 			window.ShowDialog();
 		}
 	}
diff --git a/Engine/Source/Programs/UnrealGameSync/UnrealGameSync/Forms/AutomatedBuildWindow.cs b/Engine/Source/Programs/UnrealGameSync/UnrealGameSync/Forms/AutomatedBuildWindow.cs
--- a/Engine/Source/Programs/UnrealGameSync/UnrealGameSync/Forms/AutomatedBuildWindow.cs
+++ b/Engine/Source/Programs/UnrealGameSync/UnrealGameSync/Forms/AutomatedBuildWindow.cs
@@ -92,7 +92,7 @@ namespace UnrealGameSync.Forms
 				}
 			}
 
-			AutomatedBuildWindow window = new AutomatedBuildWindow(streamName, changelist, command, defaultPerforceSettings, defaultWorkspaceName, defaultProjectPath, loggerFactory);
+			using AutomatedBuildWindow window = new AutomatedBuildWindow(streamName, changelist, command, defaultPerforceSettings, defaultWorkspaceName, defaultProjectPath, loggerFactory);
 			if (window.ShowDialog() == DialogResult.OK)
 			{
 				buildInfo = window._result!;
diff --git a/Engine/Source/Programs/UnrealGameSync/UnrealGameSync/Forms/AutomatedSyncWindow.cs b/Engine/Source/Programs/UnrealGameSync/UnrealGameSync/Forms/AutomatedSyncWindow.cs
--- a/Engine/Source/Programs/UnrealGameSync/UnrealGameSync/Forms/AutomatedSyncWindow.cs
+++ b/Engine/Source/Programs/UnrealGameSync/UnrealGameSync/Forms/AutomatedSyncWindow.cs
@@ -138,7 +138,7 @@ namespace UnrealGameSync
 		{
 			string? workspaceName = FindDefaultWorkspace(owner, defaultPerforceSettings, streamName, serviceProvider);
 
-			AutomatedSyncWindow window = new AutomatedSyncWindow(streamName, projectPath, workspaceName, defaultPerforceSettings, serviceProvider);
+			using AutomatedSyncWindow window = new AutomatedSyncWindow(streamName, projectPath, workspaceName, defaultPerforceSettings, serviceProvider);
 			if(window.ShowDialog() == DialogResult.OK)
 			{
 				workspaceInfo = window._selectedWorkspaceInfo!;
diff --git a/Engine/Source/Programs/UnrealGameSync/UnrealGameSync/Forms/BuildStepWindow.cs b/Engine/Source/Programs/UnrealGameSync/UnrealGameSync/Forms/BuildStepWindow.cs
--- a/Engine/Source/Programs/UnrealGameSync/UnrealGameSync/Forms/BuildStepWindow.cs
+++ b/Engine/Source/Programs/UnrealGameSync/UnrealGameSync/Forms/BuildStepWindow.cs
@@ -33,7 +33,7 @@ namespace UnrealGameSync
 
 		private void BuildTaskWindow_Load(object sender, EventArgs e)
 		{
-			string defaultTargetName = _targetNames.FirstOrDefault(x => !x.EndsWith("Editor") && !x.EndsWith("Client") && !x.EndsWith("Server")) ?? ((_targetNames.Count > 0)? _targetNames[0] : "");
+			string defaultTargetName = _targetNames.FirstOrDefault(x => !x.EndsWith("Editor", StringComparison.OrdinalIgnoreCase) && !x.EndsWith("Client", StringComparison.OrdinalIgnoreCase) && !x.EndsWith("Server", StringComparison.OrdinalIgnoreCase)) ?? ((_targetNames.Count > 0)? _targetNames[0] : "");
 			CompileTargetComboBox.Items.AddRange(_targetNames.ToArray());
 			CompileTargetComboBox.Text = String.IsNullOrEmpty(_step.Target)? defaultTargetName : _step.Target;
 			StatusPanelLinkTextBox.Text = _step.StatusPanelLink;
@@ -169,7 +169,7 @@ namespace UnrealGameSync
 
 		private void CookFileNameButton_Click(object sender, EventArgs e)
 		{
-			OpenFileDialog dialog = new OpenFileDialog();
+			using OpenFileDialog dialog = new OpenFileDialog();
 			dialog.Filter = "Cook/Launch Profiles (*.ulp2)|*.ulp2";
 			dialog.FileName = AddBaseDirectory(CookFileNameTextBox.Text);
 			if(dialog.ShowDialog() == DialogResult.OK)
@@ -180,7 +180,7 @@ namespace UnrealGameSync
 
 		private void OtherFileNameButton_Click(object sender, EventArgs e)
 		{
-			OpenFileDialog dialog = new OpenFileDialog();
+			using OpenFileDialog dialog = new OpenFileDialog();
 			dialog.Filter = "Executable Files (*.exe,*.bat)|*.exe;*.bat|All files (*.*)|*.*";
 			dialog.FileName = AddBaseDirectory(OtherFileNameTextBox.Text);
 			if(dialog.ShowDialog() == DialogResult.OK)
@@ -191,7 +191,7 @@ namespace UnrealGameSync
 
 		private void OtherWorkingDirButton_Click(object sender, EventArgs e)
 		{
-			FolderBrowserDialog dialog = new FolderBrowserDialog();
+			using FolderBrowserDialog dialog = new FolderBrowserDialog();
 			dialog.SelectedPath = OtherWorkingDirTextBox.Text;
 			if (dialog.ShowDialog() == DialogResult.OK)
 			{
@@ -201,7 +201,7 @@ namespace UnrealGameSync
 
 		private string AddBaseDirectory(string fileName)
 		{
-			if(fileName.Contains("$("))
+			if(fileName.Contains("$(", StringComparison.Ordinal))
 			{
 				return "";
 			}
@@ -236,23 +236,20 @@ namespace UnrealGameSync
 		private void InsertVariable(string name)
 		{
 			IContainerControl container = this;
-			if(container != null)
+			for(;;)
 			{
-				for(;;)
+				IContainerControl? nextContainer = container.ActiveControl as IContainerControl;
+				if(nextContainer == null)
 				{
-					IContainerControl? nextContainer = container.ActiveControl as IContainerControl;
-					if(nextContainer == null)
-					{
-						break;
-					}
-					container = nextContainer;
+					break;
 				}
+				container = nextContainer;
+			}
 
-				TextBox? focusTextBox = container.ActiveControl as TextBox;
-				if(focusTextBox != null)
-				{
-					focusTextBox.SelectedText = name;
-				}
+			TextBox? focusTextBox = container.ActiveControl as TextBox;
+			if(focusTextBox != null)
+			{
+				focusTextBox.SelectedText = name;
 			}
 		}
 
diff --git a/Engine/Source/Programs/UnrealGameSync/UnrealGameSync/Forms/CleanWorkspaceWindow.cs b/Engine/Source/Programs/UnrealGameSync/UnrealGameSync/Forms/CleanWorkspaceWindow.cs
--- a/Engine/Source/Programs/UnrealGameSync/UnrealGameSync/Forms/CleanWorkspaceWindow.cs
+++ b/Engine/Source/Programs/UnrealGameSync/UnrealGameSync/Forms/CleanWorkspaceWindow.cs
@@ -136,7 +136,7 @@ namespace UnrealGameSync
 			}
 
 			// Populate the tree
-			CleanWorkspaceWindow cleanWorkspace = new CleanWorkspaceWindow(perforceSettings, rootFolderToClean, extraSafeToDeleteFolders, extraSafeToDeleteExtensions, logger);
+			using CleanWorkspaceWindow cleanWorkspace = new CleanWorkspaceWindow(perforceSettings, rootFolderToClean, extraSafeToDeleteFolders, extraSafeToDeleteExtensions, logger);
 			cleanWorkspace.ShowDialog();
 		}
 
@@ -299,12 +299,12 @@ namespace UnrealGameSync
 
 		private bool IsSafeToDeleteFolder(string folderPath)
 		{
-			return s_safeToDeleteFolders.Any(x => folderPath.EndsWith(x)) || _extraSafeToDeleteFolders.Any(x => folderPath.EndsWith(x));
+			return s_safeToDeleteFolders.Any(x => folderPath.EndsWith(x, StringComparison.OrdinalIgnoreCase)) || _extraSafeToDeleteFolders.Any(x => folderPath.EndsWith(x, StringComparison.OrdinalIgnoreCase));
 		}
 
 		private bool IsSafeToDeleteFile(string name)
 		{
-			return s_safeToDeleteExtensions.Any(x => name.EndsWith(x)) || _extraSafeToDeleteExtensions.Any(x => name.EndsWith(x));
+			return s_safeToDeleteExtensions.Any(x => name.EndsWith(x, StringComparison.OrdinalIgnoreCase)) || _extraSafeToDeleteExtensions.Any(x => name.EndsWith(x, StringComparison.OrdinalIgnoreCase));
 		}
 
 		private void TreeView_DrawNode(object sender, DrawTreeNodeEventArgs e)
@@ -385,7 +385,7 @@ namespace UnrealGameSync
 			}
 		}
 
-		private void UpdateImage(TreeNode node)
+		private static void UpdateImage(TreeNode node)
 		{
 			TreeNodeData nodeData = (TreeNodeData)node.Tag;
 			int imageIndex = (nodeData._folder != null)? 0 : 3;
@@ -410,7 +410,7 @@ namespace UnrealGameSync
 			ModalTask? result = ModalTask.Execute(this, "Clean Workspace", "Cleaning files, please wait...", x => DeleteFilesTask.RunAsync(_perforceSettings, filesToSync, filesToDelete, directoriesToDelete, _logger, x), ModalTaskFlags.Quiet);
 			if(result != null && result.Failed)
 			{
-				FailedToDeleteWindow failedToDelete = new FailedToDeleteWindow();
+				using FailedToDeleteWindow failedToDelete = new FailedToDeleteWindow();
 				failedToDelete.FileList.Text = result.Error;
 				failedToDelete.FileList.SelectionStart = 0;
 				failedToDelete.FileList.SelectionLength = 0;
diff --git a/Engine/Source/Programs/UnrealGameSync/UnrealGameSync/Forms/ClobberWindow.cs b/Engine/Source/Programs/UnrealGameSync/UnrealGameSync/Forms/ClobberWindow.cs
--- a/Engine/Source/Programs/UnrealGameSync/UnrealGameSync/Forms/ClobberWindow.cs
+++ b/Engine/Source/Programs/UnrealGameSync/UnrealGameSync/Forms/ClobberWindow.cs
@@ -28,7 +28,7 @@ namespace UnrealGameSync
 				item.SubItems.Add(Path.GetDirectoryName(fileToClobber.Key));
 				FileList.Items.Add(item);
 
-				if (inUncontrolledFiles.Contains(fileToClobber.Key.Replace("\\", "/")))
+				if (inUncontrolledFiles.Contains(fileToClobber.Key.Replace("\\", "/", StringComparison.Ordinal)))
 				{
 					uncontrolledChangeFound = true;
 					item.ForeColor = Color.Red;
diff --git a/Engine/Source/Programs/UnrealGameSync/UnrealGameSync/Forms/ConnectWindow.cs b/Engine/Source/Programs/UnrealGameSync/UnrealGameSync/Forms/ConnectWindow.cs
--- a/Engine/Source/Programs/UnrealGameSync/UnrealGameSync/Forms/ConnectWindow.cs
+++ b/Engine/Source/Programs/UnrealGameSync/UnrealGameSync/Forms/ConnectWindow.cs
@@ -40,7 +40,7 @@ namespace UnrealGameSync
 
 		public static bool ShowModal(IWin32Window owner, IPerforceSettings defaultSettings, ref string? serverAndPortOverride, ref string? userNameOverride, IServiceProvider serviceProvider)
 		{
-			ConnectWindow connect = new ConnectWindow(defaultSettings, serverAndPortOverride, userNameOverride, serviceProvider);
+			using ConnectWindow connect = new ConnectWindow(defaultSettings, serverAndPortOverride, userNameOverride, serviceProvider);
 			if(connect.ShowDialog(owner) == DialogResult.OK)
 			{
 				serverAndPortOverride = connect._serverAndPortOverride;
diff --git a/Engine/Source/Programs/UnrealGameSync/UnrealGameSync/Forms/DiagnosticsWindow.cs b/Engine/Source/Programs/UnrealGameSync/UnrealGameSync/Forms/DiagnosticsWindow.cs
--- a/Engine/Source/Programs/UnrealGameSync/UnrealGameSync/Forms/DiagnosticsWindow.cs
+++ b/Engine/Source/Programs/UnrealGameSync/UnrealGameSync/Forms/DiagnosticsWindow.cs
@@ -24,7 +24,7 @@ namespace UnrealGameSync
 			_appDataFolder = inAppDataFolder;
 			_workspaceDataFolder = inWorkspaceDataFolder;
 
-			DiagnosticsTextBox.Text = inDiagnosticsText.Replace("\n", "\r\n");
+			DiagnosticsTextBox.Text = inDiagnosticsText.Replace("\n", "\r\n", StringComparison.Ordinal);
 			_extraFiles = inExtraFiles.ToList();
 		}
 
@@ -40,7 +40,7 @@ namespace UnrealGameSync
 
 		private void SaveButton_Click(object sender, EventArgs e)
 		{
-			SaveFileDialog dialog = new SaveFileDialog();
+			using SaveFileDialog dialog = new SaveFileDialog();
 			dialog.Filter = "Zip Files (*.zip)|*.zip|AllFiles (*.*)|*.*";
 			dialog.InitialDirectory = Environment.GetFolderPath(Environment.SpecialFolder.Desktop);
 			dialog.FileName = Path.Combine(dialog.InitialDirectory, "UGS-Diagnostics.zip");
@@ -71,7 +71,7 @@ namespace UnrealGameSync
 							{
 								using (FileStream inputStream = FileReference.Open(extraFile, FileMode.Open, FileAccess.Read, FileShare.ReadWrite))
 								{
-									ZipArchiveEntry entry = zip.CreateEntry(extraFile.FullName.Replace(":", "").Replace('\\', '/'));
+									ZipArchiveEntry entry = zip.CreateEntry(extraFile.FullName.Replace(":", "", StringComparison.Ordinal).Replace("\\", "/", StringComparison.Ordinal));
 									using (Stream outputStream = entry.Open())
 									{
 										inputStream.CopyTo(outputStream);
diff --git a/Engine/Source/Programs/UnrealGameSync/UnrealGameSync/Forms/IssueAlertWindow.cs b/Engine/Source/Programs/UnrealGameSync/UnrealGameSync/Forms/IssueAlertWindow.cs
--- a/Engine/Source/Programs/UnrealGameSync/UnrealGameSync/Forms/IssueAlertWindow.cs
+++ b/Engine/Source/Programs/UnrealGameSync/UnrealGameSync/Forms/IssueAlertWindow.cs
@@ -97,7 +97,7 @@ namespace UnrealGameSync
 			}
 			else
 			{
-				if(String.Compare(newIssue.Owner, IssueMonitor.UserName, StringComparison.OrdinalIgnoreCase) == 0)
+				if(String.Equals(newIssue.Owner, IssueMonitor.UserName, StringComparison.OrdinalIgnoreCase))
 				{
 					if(newIssue.NominatedBy != null)
 					{
@@ -120,7 +120,7 @@ namespace UnrealGameSync
 					{
 						ownerTextBuilder.Append(" (not acknowledged)");
 					}
-					ownerTextBuilder.Append(".");
+					ownerTextBuilder.Append('.');
 				}
 			}
 			ownerTextBuilder.AppendFormat(" Open for {0}.", Utility.FormatDurationMinutes((int)(newIssue.RetrievedAt - newIssue.CreatedAt).TotalMinutes));
diff --git a/Engine/Source/Programs/UnrealGameSync/UnrealGameSync/Forms/IssueDetailsWindow.cs b/Engine/Source/Programs/UnrealGameSync/UnrealGameSync/Forms/IssueDetailsWindow.cs
--- a/Engine/Source/Programs/UnrealGameSync/UnrealGameSync/Forms/IssueDetailsWindow.cs
+++ b/Engine/Source/Programs/UnrealGameSync/UnrealGameSync/Forms/IssueDetailsWindow.cs
@@ -114,7 +114,7 @@ namespace UnrealGameSync
 
 					_cancellationSource.Cancel();
 
-					_workerTask.ContinueWith(x => _cancellationSource.Dispose());
+					_workerTask.ContinueWith(x => _cancellationSource.Dispose(), TaskScheduler.Default);
 					_workerTask = null;
 				}
 			}
@@ -163,7 +163,7 @@ namespace UnrealGameSync
 					List<ChangesRecord> describeChanges;
 					lock(_lockObject)
 					{
-						describeChanges = completedRequests.SelectMany(x => x.Changes).Where(x => !_changeNumberToDetails.ContainsKey(x.Number)).ToList();
+						describeChanges = completedRequests.SelectMany(x => x.Changes ?? new List<ChangesRecord>()).Where(x => !_changeNumberToDetails.ContainsKey(x.Number)).ToList();
 					}
 
 					// Fetch info on each individual change
@@ -177,7 +177,7 @@ namespace UnrealGameSync
 							}
 						}
 
-						PerforceResponse<DescribeRecord> response = await perforce.TryDescribeAsync(describeChange.Number);
+						PerforceResponse<DescribeRecord> response = await perforce.TryDescribeAsync(describeChange.Number, cancellationToken);
 						if (response.Succeeded)
 						{
 							DescribeRecord record = response.Data;
@@ -387,7 +387,7 @@ namespace UnrealGameSync
 			StopUpdateTimer();
 		}
 
-		void UpdateSummaryTextIfChanged(Label label, string newText)
+		static void UpdateSummaryTextIfChanged(Label label, string newText)
 		{
 			if (label.Text != newText)
 			{
@@ -424,11 +424,11 @@ namespace UnrealGameSync
 			richText.Append(@"{\field");
 			richText.Append(@"{\*\fldinst");
 			richText.AppendFormat("{{ HYPERLINK \"{0}\" }}", url);
-			richText.Append(@"}");
+			richText.Append('}');
 			richText.Append(@"{\fldrslt ");
 			AppendEscapedRtf(richText, label);
-			richText.Append(@"}");
-			richText.Append(@"}");
+			richText.Append('}');
+			richText.Append('}');
 		}
 
 		static string CreateRichTextErrors(List<IssueBuildData> issueBuilds, List<IssueDiagnosticData> diagnostics)
@@ -481,7 +481,7 @@ namespace UnrealGameSync
 					richText.Append(@"\li50");   // Other line indent
 					richText.Append(@"\sb100");  // Space before
 					richText.Append(@"\sa50");   // Space after
-					richText.Append(@" ");
+					richText.Append(' ');
 
 					richText.Append(@"\ul1");
 					AppendHyperlink(richText, String.Format("Error {0}/{1}", idx + 1, diagnosticsArray.Length), diagnostic.Url);
@@ -629,7 +629,7 @@ namespace UnrealGameSync
 			}
 		}
 
-		bool FilterMatch(Regex filterRegex, ChangesRecord summary)
+		static bool FilterMatch(Regex filterRegex, ChangesRecord summary)
 		{
 			if(filterRegex.IsMatch(summary.User))
 			{
@@ -642,7 +642,7 @@ namespace UnrealGameSync
 			return false;
 		}
 
-		bool FilterMatch(Regex filterRegex, DescribeRecord describeRecord)
+		static bool FilterMatch(Regex filterRegex, DescribeRecord describeRecord)
 		{
 			if(filterRegex.IsMatch(describeRecord.User))
 			{
@@ -679,9 +679,9 @@ namespace UnrealGameSync
 			foreach(string filterTerm in filterTerms)
 			{
 				string regexText = Regex.Escape(filterTerm);
-				regexText = regexText.Replace("\\?", ".");
-				regexText = regexText.Replace("\\*", "[^\\\\/]*");
-				regexText = regexText.Replace("\\.\\.\\.", ".*");
+				regexText = regexText.Replace("\\?", ".", StringComparison.Ordinal);
+				regexText = regexText.Replace("\\*", "[^\\\\/]*", StringComparison.Ordinal);
+				regexText = regexText.Replace("\\.\\.\\.", ".*", StringComparison.Ordinal);
 				filterRegexes.Add(new Regex(regexText, RegexOptions.IgnoreCase));
 			}
 
@@ -866,7 +866,7 @@ namespace UnrealGameSync
 			JobContextMenu.Show(BuildListView, point, ToolStripDropDownDirection.BelowRight);
 		}
 
-		void UpdateChangeTypeWidget(StatusLineListViewWidget typeWidget, PerforceChangeDetails? details)
+		static void UpdateChangeTypeWidget(StatusLineListViewWidget typeWidget, PerforceChangeDetails? details)
 		{
 			typeWidget.Line.Clear();
 			if(details == null)
@@ -1142,7 +1142,7 @@ namespace UnrealGameSync
 			else if(e.Item.Tag is ChangesRecord change)
 			{
 				Font changeFont = BuildListView.Font;
-				if(_issue.Owner != null && String.Compare(change.User, _issue.Owner, StringComparison.OrdinalIgnoreCase) == 0)
+				if(_issue.Owner != null && String.Equals(change.User, _issue.Owner, StringComparison.OrdinalIgnoreCase))
 				{
 					changeFont = _boldFont!;
 				}
@@ -1286,7 +1286,7 @@ namespace UnrealGameSync
 			update.Id = _issue.Id;
 			update.Owner = user;
 			update.FixChange = 0;
-			if(String.Compare(user, _perforceSettings.UserName, StringComparison.OrdinalIgnoreCase) == 0)
+			if(String.Equals(user, _perforceSettings.UserName, StringComparison.OrdinalIgnoreCase))
 			{
 				update.NominatedBy = "";
 				update.Acknowledged = true;
diff --git a/Engine/Source/Programs/UnrealGameSync/UnrealGameSync/Forms/IssueFixedWindow.cs b/Engine/Source/Programs/UnrealGameSync/UnrealGameSync/Forms/IssueFixedWindow.cs
--- a/Engine/Source/Programs/UnrealGameSync/UnrealGameSync/Forms/IssueFixedWindow.cs
+++ b/Engine/Source/Programs/UnrealGameSync/UnrealGameSync/Forms/IssueFixedWindow.cs
@@ -57,7 +57,7 @@ namespace UnrealGameSync
 					_onComplete = null;
 
 					_cancellationSource.Cancel();
-					stopTask = _backgroundTask.ContinueWith(_ => _cancellationSource.Dispose());
+					stopTask = _backgroundTask.ContinueWith(_ => _cancellationSource.Dispose(), TaskScheduler.Default);
 
 					_backgroundTask = null!;
 				}
@@ -201,7 +201,7 @@ namespace UnrealGameSync
 				{
 					foreach(DescribeRecord change in changes)
 					{
-						if(change.Description != null && change.Description.IndexOf("#ROBOMERGE-SOURCE", 0) == -1)
+						if(change.Description != null && !change.Description.Contains("#ROBOMERGE-SOURCE", StringComparison.Ordinal))
 						{
 							string stream = "";
 							if(change.Files.Count > 0)
diff --git a/Engine/Source/Programs/UnrealGameSync/UnrealGameSync/Forms/MainWindow.cs b/Engine/Source/Programs/UnrealGameSync/UnrealGameSync/Forms/MainWindow.cs
--- a/Engine/Source/Programs/UnrealGameSync/UnrealGameSync/Forms/MainWindow.cs
+++ b/Engine/Source/Programs/UnrealGameSync/UnrealGameSync/Forms/MainWindow.cs
@@ -268,7 +268,7 @@ namespace UnrealGameSync
 
 		AutomationRequestOutput StartExecCommand(AutomationRequest request)
 		{
-			BinaryReader reader = new BinaryReader(new MemoryStream(request.Input.Data));
+			using BinaryReader reader = new BinaryReader(new MemoryStream(request.Input.Data));
 			string streamName = reader.ReadString();
 			int changelist = reader.ReadInt32();
 			string command = reader.ReadString();
@@ -290,7 +290,7 @@ namespace UnrealGameSync
 			return new AutomationRequestOutput(AutomationRequestResult.Ok);
 		}
 
-		private void StartExecCommandAfterStartup(WorkspaceControl workspace, bool cancel, int changelist, string command)
+		private static void StartExecCommandAfterStartup(WorkspaceControl workspace, bool cancel, int changelist, string command)
 		{
 			if (!cancel)
 			{
@@ -305,7 +305,7 @@ namespace UnrealGameSync
 			}
 		}
 
-		private void StartExecCommandAfterSync(WorkspaceControl workspace, WorkspaceUpdateResult result, string command)
+		private static void StartExecCommandAfterSync(WorkspaceControl workspace, WorkspaceUpdateResult result, string command)
 		{
 			if (result == WorkspaceUpdateResult.Success && command != null)
 			{
@@ -318,7 +318,7 @@ namespace UnrealGameSync
 		{
 			ShowAndActivate();
 
-			BinaryReader reader = new BinaryReader(new MemoryStream(request.Input.Data));
+			using BinaryReader reader = new BinaryReader(new MemoryStream(request.Input.Data));
 			string streamName = reader.ReadString();
 			string projectPath = reader.ReadString();
 
@@ -351,7 +351,7 @@ namespace UnrealGameSync
 				for (int existingTabIdx = 0; existingTabIdx < TabControl.GetTabCount(); existingTabIdx++)
 				{
 					WorkspaceControl? existingWorkspace = TabControl.GetTabData(existingTabIdx) as WorkspaceControl;
-					if (existingWorkspace != null && existingWorkspace.ClientName.Equals(workspaceInfo.WorkspaceName))
+					if (existingWorkspace != null && existingWorkspace.ClientName.Equals(workspaceInfo.WorkspaceName, StringComparison.OrdinalIgnoreCase))
 					{
 						TabControl.RemoveTab(existingTabIdx);
 						break;
@@ -399,7 +399,7 @@ namespace UnrealGameSync
 			return true;
 		}
 
-		private void StartAutomatedSyncAfterStartup(WorkspaceControl workspace, bool cancel, AutomationRequest request)
+		private static void StartAutomatedSyncAfterStartup(WorkspaceControl workspace, bool cancel, AutomationRequest request)
 		{
 			if(cancel)
 			{
@@ -411,7 +411,7 @@ namespace UnrealGameSync
 			}
 		}
 
-		void CompleteAutomatedSync(WorkspaceUpdateResult result, FileReference selectedFileName, AutomationRequest request)
+		static void CompleteAutomatedSync(WorkspaceUpdateResult result, FileReference selectedFileName, AutomationRequest request)
 		{
 			if(result == WorkspaceUpdateResult.Success)
 			{
@@ -429,23 +429,23 @@ namespace UnrealGameSync
 
 		AutomationRequestOutput FindProject(AutomationRequest request)
 		{
-			BinaryReader reader = new BinaryReader(new MemoryStream(request.Input.Data));
+			using BinaryReader reader = new BinaryReader(new MemoryStream(request.Input.Data));
 			string streamName = reader.ReadString();
 			string projectPath = reader.ReadString();
 
 			for(int existingTabIdx = 0; existingTabIdx < TabControl.GetTabCount(); existingTabIdx++)
 			{
 				WorkspaceControl? existingWorkspace = TabControl.GetTabData(existingTabIdx) as WorkspaceControl;
-				if(existingWorkspace != null && String.Compare(existingWorkspace.StreamName, streamName, StringComparison.OrdinalIgnoreCase) == 0 && existingWorkspace.SelectedProject != null)
+				if(existingWorkspace != null && String.Equals(existingWorkspace.StreamName, streamName, StringComparison.OrdinalIgnoreCase) && existingWorkspace.SelectedProject != null)
 				{
 					string? clientPath = existingWorkspace.SelectedProject.ClientPath;
-					if(clientPath != null && clientPath.StartsWith("//"))
+					if(clientPath != null && clientPath.StartsWith("//", StringComparison.Ordinal))
 					{
 						int slashIdx = clientPath.IndexOf('/', 2);
 						if(slashIdx != -1)
 						{
 							string existingProjectPath = clientPath.Substring(slashIdx);
-							if(String.Compare(existingProjectPath, projectPath, StringComparison.OrdinalIgnoreCase) == 0)
+							if (String.Equals(existingProjectPath, projectPath, StringComparison.OrdinalIgnoreCase))
 							{
 								return new AutomationRequestOutput(AutomationRequestResult.Ok, Encoding.UTF8.GetBytes(existingWorkspace.SelectedFileName.FullName));
 							}
@@ -459,7 +459,7 @@ namespace UnrealGameSync
 
 		AutomationRequestOutput OpenIssue(AutomationRequest request)
 		{
-			BinaryReader reader = new BinaryReader(new MemoryStream(request.Input.Data));
+			using BinaryReader reader = new BinaryReader(new MemoryStream(request.Input.Data));
 			int issueId = reader.ReadInt32();
 
 			for (int existingTabIdx = 0; existingTabIdx < TabControl.GetTabCount(); existingTabIdx++)
@@ -468,7 +468,7 @@ namespace UnrealGameSync
 				if (existingWorkspace != null)
 				{
 					IssueMonitor issueMonitor = existingWorkspace.GetIssueMonitor();
-					if (issueMonitor != null && (DeploymentSettings.UrlHandleIssueApi == null || String.Compare(issueMonitor.ApiUrl, DeploymentSettings.UrlHandleIssueApi, StringComparison.OrdinalIgnoreCase) == 0))
+					if (issueMonitor != null && (DeploymentSettings.UrlHandleIssueApi == null || String.Equals(issueMonitor.ApiUrl, DeploymentSettings.UrlHandleIssueApi, StringComparison.OrdinalIgnoreCase)))
 					{
 						issueMonitor.AddRef();
 						try
@@ -488,7 +488,7 @@ namespace UnrealGameSync
 							}
 							else
 							{
-								_logger.LogError(issueTask?.Exception, "Unable to query issue {IssueId} from {ApiUrl}: {Error}", issueId, issueMonitor.ApiUrl, issueTask?.Error);
+								_logger.LogError(issueTask.Exception, "Unable to query issue {IssueId} from {ApiUrl}: {Error}", issueId, issueMonitor.ApiUrl, issueTask.Error);
 								return new AutomationRequestOutput(AutomationRequestResult.Error);
 							}
 						}
@@ -877,7 +877,7 @@ namespace UnrealGameSync
 				}
 			}
 
-			ScheduleWindow schedule = new ScheduleWindow(_settings.ScheduleEnabled, _settings.ScheduleTime, _settings.ScheduleAnyOpenProject, _settings.ScheduleProjects, openProjects, projectToLatestChangeTypes);
+			using ScheduleWindow schedule = new ScheduleWindow(_settings.ScheduleEnabled, _settings.ScheduleTime, _settings.ScheduleAnyOpenProject, _settings.ScheduleProjects, openProjects, projectToLatestChangeTypes);
 			if(schedule.ShowDialog() == System.Windows.Forms.DialogResult.OK)
 			{
 				schedule.CopySettings(_settings);
@@ -949,14 +949,16 @@ namespace UnrealGameSync
 				return;
 			}
 
-			LauncherUpdateWindow update = new LauncherUpdateWindow();
-			if (update.ShowDialog(this) == DialogResult.Ignore)
+			using (LauncherUpdateWindow update = new LauncherUpdateWindow())
 			{
-				_settings.NextLauncherVersionCheck = (now + TimeSpan.FromDays(1.0)).Ticks;
-				_settings.Save(_logger);
+				if (update.ShowDialog(this) == DialogResult.Ignore)
+				{
+					_settings.NextLauncherVersionCheck = (now + TimeSpan.FromDays(1.0)).Ticks;
+					_settings.Save(_logger);
 
-				CheckLauncherVersionTimer.Enabled = true;
-				return;
+					CheckLauncherVersionTimer.Enabled = true;
+					return;
+				}
 			}
 
 			using (Process childProcess = new Process())
@@ -1258,7 +1260,7 @@ namespace UnrealGameSync
 						}
 						else if(openProjectInfo.ProjectInfo.LocalFileName.IsUnderDirectory(workspace.BranchDirectoryName))
 						{
-							if((options & OpenProjectOptions.Quiet) == 0 && MessageBox.Show(String.Format("{0} is already open under {1}.\n\nWould you like to close it?", workspace.SelectedFileName.GetFileNameWithoutExtension(), workspace.BranchDirectoryName, openProjectInfo.ProjectInfo.LocalFileName.GetFileNameWithoutExtension()), "Branch already open", MessageBoxButtons.YesNo) == System.Windows.Forms.DialogResult.Yes)
+							if((options & OpenProjectOptions.Quiet) == 0 && MessageBox.Show($"{workspace.SelectedFileName.GetFileNameWithoutExtension()} is already open under {workspace.BranchDirectoryName}.\n\nWould you like to close it?", "Branch already open", MessageBoxButtons.YesNo) == System.Windows.Forms.DialogResult.Yes)
 							{
 								_logger.LogInformation("  Another project already open in this workspace, tab {TabIdx}. Replacing.", tabIdx);
 								TabControl.RemoveTab(tabIdx);
@@ -1597,7 +1599,7 @@ namespace UnrealGameSync
 						}
 						else if(!issue.AcknowledgedAt.HasValue)
 						{
-							if(String.Compare(issue.Owner, issueMonitor.UserName, StringComparison.OrdinalIgnoreCase) == 0)
+							if(String.Equals(issue.Owner, issueMonitor.UserName, StringComparison.OrdinalIgnoreCase))
 							{
 								reason |= IssueAlertReason.Owner;
 							}
@@ -1768,7 +1770,7 @@ namespace UnrealGameSync
 
 		public IssueMonitor CreateIssueMonitor(string? apiUrl, string userName)
 		{
-			WorkspaceIssueMonitor? workspaceIssueMonitor = _workspaceIssueMonitors.FirstOrDefault(x => String.Compare(x._issueMonitor.ApiUrl, apiUrl, StringComparison.OrdinalIgnoreCase) == 0 && String.Compare(x._issueMonitor.UserName, userName, StringComparison.OrdinalIgnoreCase) == 0);
+			WorkspaceIssueMonitor? workspaceIssueMonitor = _workspaceIssueMonitors.FirstOrDefault(x => String.Equals(x._issueMonitor.ApiUrl, apiUrl, StringComparison.OrdinalIgnoreCase) && String.Equals(x._issueMonitor.UserName, userName, StringComparison.OrdinalIgnoreCase));
 			if (workspaceIssueMonitor == null)
 			{
 				string serverId = apiUrl != null ? Regex.Replace(apiUrl, @"^.*://", "") : "noserver";
diff --git a/Engine/Source/Programs/UnrealGameSync/UnrealGameSync/Forms/ModalTaskWindow.cs b/Engine/Source/Programs/UnrealGameSync/UnrealGameSync/Forms/ModalTaskWindow.cs
--- a/Engine/Source/Programs/UnrealGameSync/UnrealGameSync/Forms/ModalTaskWindow.cs
+++ b/Engine/Source/Programs/UnrealGameSync/UnrealGameSync/Forms/ModalTaskWindow.cs
@@ -32,7 +32,7 @@ namespace UnrealGameSync
 		private void ModalTaskWindow_Load(object sender, EventArgs e)
 		{
 			SynchronizationContext syncContext = SynchronizationContext.Current!;
-			_task.ContinueWith(x => syncContext.Post(y => Close(), null));
+			_task.ContinueWith(x => syncContext.Post(y => Close(), null), TaskScheduler.Default);
 		}
 
 		private void ModalTaskWindow_FormClosing(object sender, FormClosingEventArgs e)
diff --git a/Engine/Source/Programs/UnrealGameSync/UnrealGameSync/Forms/ModifyBuildStepsWindow.cs b/Engine/Source/Programs/UnrealGameSync/UnrealGameSync/Forms/ModifyBuildStepsWindow.cs
--- a/Engine/Source/Programs/UnrealGameSync/UnrealGameSync/Forms/ModifyBuildStepsWindow.cs
+++ b/Engine/Source/Programs/UnrealGameSync/UnrealGameSync/Forms/ModifyBuildStepsWindow.cs
@@ -64,7 +64,7 @@ namespace UnrealGameSync
 			newStep.Description = "Untitled Task";
 			newStep.EstimatedDuration = 1;
 
-			BuildStepWindow newStepWindow = new BuildStepWindow(newStep, _targetNames, _baseDirectory, _variables);
+			using BuildStepWindow newStepWindow = new BuildStepWindow(newStep, _targetNames, _baseDirectory, _variables);
 			if(newStepWindow.ShowDialog() == DialogResult.OK)
 			{
 				AddTask(newStep);
@@ -77,7 +77,7 @@ namespace UnrealGameSync
 			{
 				if (item != null)
 				{
-					BuildStepWindow editStep = new BuildStepWindow((BuildStep)item.Tag, _targetNames, _baseDirectory, _variables);
+					using BuildStepWindow editStep = new BuildStepWindow((BuildStep)item.Tag, _targetNames, _baseDirectory, _variables);
 					editStep.ShowDialog();
 					item.Text = ((BuildStep)item.Tag).Description;
 					break;
diff --git a/Engine/Source/Programs/UnrealGameSync/UnrealGameSync/Forms/NewWorkspaceWindow.cs b/Engine/Source/Programs/UnrealGameSync/UnrealGameSync/Forms/NewWorkspaceWindow.cs
--- a/Engine/Source/Programs/UnrealGameSync/UnrealGameSync/Forms/NewWorkspaceWindow.cs
+++ b/Engine/Source/Programs/UnrealGameSync/UnrealGameSync/Forms/NewWorkspaceWindow.cs
@@ -102,7 +102,7 @@ namespace UnrealGameSync
 			Dictionary<DirectoryReference, int> rootPathToCount = new Dictionary<DirectoryReference, int>();
 			foreach(ClientsRecord client in clients)
 			{
-				if(client.Host == null || String.Compare(client.Host, info.ClientHost, StringComparison.OrdinalIgnoreCase) == 0)
+				if(client.Host == null || String.Equals(client.Host, info.ClientHost, StringComparison.OrdinalIgnoreCase))
 				{
 					if(!String.IsNullOrEmpty(client.Root) && client.Root != ".")
 					{
@@ -166,7 +166,7 @@ namespace UnrealGameSync
 				return false;
 			}
 
-			NewWorkspaceWindow window = new NewWorkspaceWindow(perforceSettings, forceStreamName, task.Result.CurrentStream, task.Result.Info, task.Result.Clients, serviceProvider);
+			using NewWorkspaceWindow window = new NewWorkspaceWindow(perforceSettings, forceStreamName, task.Result.CurrentStream, task.Result.Info, task.Result.Clients, serviceProvider);
 			if(window.ShowDialog(owner) == DialogResult.OK)
 			{
 				workspaceName = window._settings!.Name;
@@ -181,7 +181,7 @@ namespace UnrealGameSync
 
 		private void RootDirBrowseBtn_Click(object sender, EventArgs e)
 		{
-			FolderBrowserDialog dialog = new FolderBrowserDialog();
+			using FolderBrowserDialog dialog = new FolderBrowserDialog();
 			dialog.ShowNewFolderButton = true;
 			dialog.SelectedPath = RootDirTextBox.Text;
 			if (dialog.ShowDialog() == DialogResult.OK)
@@ -196,7 +196,7 @@ namespace UnrealGameSync
 			string baseName = Sanitize(String.Format("{0}_{1}_{2}", _info.UserName, _info.ClientHost, StreamTextBox.Text.Replace('/', '_').Trim('_'))).Trim('_');
 
 			string name = baseName;
-			for(int idx = 2; _clients.Any(x => x.Name != null && String.Compare(x.Name, name, StringComparison.InvariantCultureIgnoreCase) == 0); idx++)
+			for(int idx = 2; _clients.Any(x => x.Name != null && String.Equals(x.Name, name, StringComparison.OrdinalIgnoreCase)); idx++)
 			{
 				name = String.Format("{0}_{1}", baseName, idx);
 			}
@@ -217,7 +217,7 @@ namespace UnrealGameSync
 			return rootDir;
 		}
 
-		private string Sanitize(string text)
+		private static string Sanitize(string text)
 		{
 			StringBuilder result = new StringBuilder();
 			for(int idx = 0; idx < text.Length; idx++)
@@ -259,7 +259,7 @@ namespace UnrealGameSync
 			}
 
 			string newStream = StreamTextBox.Text.Trim();
-			if(!newStream.StartsWith("//") || newStream.IndexOf('/', 2) == -1)
+			if(!newStream.StartsWith("//", StringComparison.Ordinal) || newStream.IndexOf('/', 2) == -1)
 			{
 				settings = null;
 				return false;
diff --git a/Engine/Source/Programs/UnrealGameSync/UnrealGameSync/Forms/NotificationWindow.cs b/Engine/Source/Programs/UnrealGameSync/UnrealGameSync/Forms/NotificationWindow.cs
--- a/Engine/Source/Programs/UnrealGameSync/UnrealGameSync/Forms/NotificationWindow.cs
+++ b/Engine/Source/Programs/UnrealGameSync/UnrealGameSync/Forms/NotificationWindow.cs
@@ -37,8 +37,8 @@ namespace UnrealGameSync
 		bool _mouseDownOverClose;
 
 		// Notifications
-		public Action? OnMoreInformation;
-		public Action? OnDismiss;
+		public Action? OnMoreInformation { get; set; }
+		public Action? OnDismiss { get; set; }
 
 		public NotificationWindow(Image inLogo)
 		{
diff --git a/Engine/Source/Programs/UnrealGameSync/UnrealGameSync/Forms/OpenProjectWindow.cs b/Engine/Source/Programs/UnrealGameSync/UnrealGameSync/Forms/OpenProjectWindow.cs
--- a/Engine/Source/Programs/UnrealGameSync/UnrealGameSync/Forms/OpenProjectWindow.cs
+++ b/Engine/Source/Programs/UnrealGameSync/UnrealGameSync/Forms/OpenProjectWindow.cs
@@ -52,7 +52,7 @@ namespace UnrealGameSync
 					_userNameOverride = project.UserName;
 				}
 
-				if(project.ClientPath != null && project.ClientPath.StartsWith("//"))
+				if(project.ClientPath != null && project.ClientPath.StartsWith("//", StringComparison.Ordinal))
 				{
 					int slashIdx = project.ClientPath.IndexOf('/', 2);
 					if(slashIdx != -1)
@@ -87,7 +87,7 @@ namespace UnrealGameSync
 
 		public static OpenProjectInfo? ShowModal(IWin32Window owner, UserSelectedProjectSettings? project, UserSettings settings, IPerforceSettings defaultPerforceSettings, IServiceProvider serviceProvider, ILogger logger)
 		{
-			OpenProjectWindow window = new OpenProjectWindow(project, settings, defaultPerforceSettings, serviceProvider, logger);
+			using OpenProjectWindow window = new OpenProjectWindow(project, settings, defaultPerforceSettings, serviceProvider, logger);
 			if(window.ShowDialog(owner) == DialogResult.OK)
 			{
 				return window._openProjectInfo;
@@ -318,7 +318,7 @@ namespace UnrealGameSync
 		{
 			LocalFileRadioBtn.Checked = true;
 
-			OpenFileDialog dialog = new OpenFileDialog();
+			using OpenFileDialog dialog = new OpenFileDialog();
 			dialog.Filter = "Project files (*.uproject)|*.uproject|Project directory lists (*.uprojectdirs)|*.uprojectdirs|All supported files (*.uproject;*.uprojectdirs)|*.uproject;*.uprojectdirs|All files (*.*)|*.*" ;
 			dialog.FilterIndex = _settings.FilterIndex;
 			
diff --git a/Engine/Source/Programs/UnrealGameSync/UnrealGameSync/Forms/ScheduleWindow.cs b/Engine/Source/Programs/UnrealGameSync/UnrealGameSync/Forms/ScheduleWindow.cs
--- a/Engine/Source/Programs/UnrealGameSync/UnrealGameSync/Forms/ScheduleWindow.cs
+++ b/Engine/Source/Programs/UnrealGameSync/UnrealGameSync/Forms/ScheduleWindow.cs
@@ -102,7 +102,7 @@ namespace UnrealGameSync
 			SyncTypeDropDown.Close();
 		}
 
-		private void AddProjects(IEnumerable<UserSelectedProjectSettings> projects, Dictionary<string, UserSelectedProjectSettings> localFileToProject)
+		private static void AddProjects(IEnumerable<UserSelectedProjectSettings> projects, Dictionary<string, UserSelectedProjectSettings> localFileToProject)
 		{
 			foreach(UserSelectedProjectSettings project in projects)
 			{
diff --git a/Engine/Source/Programs/UnrealGameSync/UnrealGameSync/Forms/SdkInfoWindow.cs b/Engine/Source/Programs/UnrealGameSync/UnrealGameSync/Forms/SdkInfoWindow.cs
--- a/Engine/Source/Programs/UnrealGameSync/UnrealGameSync/Forms/SdkInfoWindow.cs
+++ b/Engine/Source/Programs/UnrealGameSync/UnrealGameSync/Forms/SdkInfoWindow.cs
@@ -79,13 +79,13 @@ namespace UnrealGameSync
 				SdkItem item = new SdkItem(category, description);
 
 				item.Install = Utility.ExpandVariables(obj.GetValue("Install", ""), variables);
-				if(item.Install.Contains("$("))
+				if(item.Install.Contains("$(", StringComparison.Ordinal))
 				{
 					item.Install = null;
 				}
 
 				item.Browse = Utility.ExpandVariables(obj.GetValue("Browse", ""), variables);
-				if(item.Browse.Contains("$("))
+				if(item.Browse.Contains("$(", StringComparison.Ordinal))
 				{
 					item.Browse = null;
 				}
@@ -244,7 +244,7 @@ namespace UnrealGameSync
 			return badges;
 		}
 
-		private void Browse(string directoryName)
+		private static void Browse(string directoryName)
 		{
 			try
 			{
@@ -256,7 +256,7 @@ namespace UnrealGameSync
 			}
 		}
 
-		private void Install(string fileName)
+		private static void Install(string fileName)
 		{
 			try
 			{
diff --git a/Engine/Source/Programs/UnrealGameSync/UnrealGameSync/Forms/SelectProjectFromWorkspaceWindow.cs b/Engine/Source/Programs/UnrealGameSync/UnrealGameSync/Forms/SelectProjectFromWorkspaceWindow.cs
--- a/Engine/Source/Programs/UnrealGameSync/UnrealGameSync/Forms/SelectProjectFromWorkspaceWindow.cs
+++ b/Engine/Source/Programs/UnrealGameSync/UnrealGameSync/Forms/SelectProjectFromWorkspaceWindow.cs
@@ -167,7 +167,7 @@ namespace UnrealGameSync
 				TreeNode node = FindOrAddChildNode(projectParentNodes[idx], projectNodes[idx].Name, 1);
 				node.Tag = projectNodes[idx].FullName;
 
-				if(String.Compare(projectNodes[idx].FullName, _selectedProjectFile, StringComparison.InvariantCultureIgnoreCase) == 0)
+				if(String.Equals(projectNodes[idx].FullName, _selectedProjectFile, StringComparison.OrdinalIgnoreCase))
 				{
 					ProjectTreeView.SelectedNode = node;
 					for(TreeNode parentNode = node.Parent; parentNode != rootNode; parentNode = parentNode.Parent)
@@ -182,7 +182,7 @@ namespace UnrealGameSync
 		{
 			foreach(TreeNode? childNode in parentNode.Nodes)
 			{
-				if(childNode != null && String.Compare(childNode.Text, text, StringComparison.InvariantCultureIgnoreCase) == 0)
+				if(childNode != null && String.Equals(childNode.Text, text, StringComparison.OrdinalIgnoreCase))
 				{
 					return childNode;
 				}
@@ -208,7 +208,7 @@ namespace UnrealGameSync
 				return false;
 			}
 
-			SelectProjectFromWorkspaceWindow selectProjectWindow = new SelectProjectFromWorkspaceWindow(workspaceName, pathsTask.Result, workspacePath);
+			using SelectProjectFromWorkspaceWindow selectProjectWindow = new SelectProjectFromWorkspaceWindow(workspaceName, pathsTask.Result, workspacePath);
 			if(selectProjectWindow.ShowDialog() == DialogResult.OK && !String.IsNullOrEmpty(selectProjectWindow._selectedProjectFile))
 			{
 				newWorkspacePath = selectProjectWindow._selectedProjectFile;
diff --git a/Engine/Source/Programs/UnrealGameSync/UnrealGameSync/Forms/SelectStreamWindow.cs b/Engine/Source/Programs/UnrealGameSync/UnrealGameSync/Forms/SelectStreamWindow.cs
--- a/Engine/Source/Programs/UnrealGameSync/UnrealGameSync/Forms/SelectStreamWindow.cs
+++ b/Engine/Source/Programs/UnrealGameSync/UnrealGameSync/Forms/SelectStreamWindow.cs
@@ -187,7 +187,7 @@ namespace UnrealGameSync
 			return base.ProcessCmdKey(ref msg, keyData);
 		}
 
-		private bool IncludeNodeInFilter(StreamNode node, string[] filter)
+		private static bool IncludeNodeInFilter(StreamNode node, string[] filter)
 		{
 			return filter.All(x => node.Record.Stream.IndexOf(x, StringComparison.InvariantCultureIgnoreCase) != -1 || node.Record.Name.IndexOf(x, StringComparison.InvariantCultureIgnoreCase) != -1);
 		}
@@ -330,7 +330,7 @@ namespace UnrealGameSync
 				return false;
 			}
 
-			SelectStreamWindow selectStream = new SelectStreamWindow(streamsTask.Result, streamName);
+			using SelectStreamWindow selectStream = new SelectStreamWindow(streamsTask.Result, streamName);
 			if(selectStream.ShowDialog(owner) == DialogResult.OK && selectStream._selectedStream != null)
 			{
 				newStreamName = selectStream._selectedStream;
diff --git a/Engine/Source/Programs/UnrealGameSync/UnrealGameSync/Forms/SelectUserWindow.cs b/Engine/Source/Programs/UnrealGameSync/UnrealGameSync/Forms/SelectUserWindow.cs
--- a/Engine/Source/Programs/UnrealGameSync/UnrealGameSync/Forms/SelectUserWindow.cs
+++ b/Engine/Source/Programs/UnrealGameSync/UnrealGameSync/Forms/SelectUserWindow.cs
@@ -86,13 +86,13 @@ namespace UnrealGameSync
 			return base.ProcessCmdKey(ref msg, keyData);
 		}
 
-		private bool IncludeInFilter(UsersRecord user, string[] filterWords)
+		private static bool IncludeInFilter(UsersRecord user, string[] filterWords)
 		{
 			foreach(string filterWord in filterWords)
 			{
-				if(user.UserName.IndexOf(filterWord, StringComparison.OrdinalIgnoreCase) == -1 
-					&& user.FullName.IndexOf(filterWord, StringComparison.OrdinalIgnoreCase) == -1
-					&& user.Email.IndexOf(filterWord, StringComparison.OrdinalIgnoreCase) == -1)
+				if(!user.UserName.Contains(filterWord, StringComparison.OrdinalIgnoreCase) 
+					&& !user.FullName.Contains(filterWord, StringComparison.OrdinalIgnoreCase)
+					&& !user.Email.Contains(filterWord, StringComparison.OrdinalIgnoreCase))
 				{
 					return false;
 				}
@@ -154,7 +154,7 @@ namespace UnrealGameSync
 				return false;
 			}
 
-			SelectUserWindow selectUser = new SelectUserWindow(usersTask.Result, 0);
+			using SelectUserWindow selectUser = new SelectUserWindow(usersTask.Result, 0);
 			if(selectUser.ShowDialog(owner) == DialogResult.OK)
 			{
 				selectedUserName = usersTask.Result[selectUser._selectedUserIndex].UserName;
diff --git a/Engine/Source/Programs/UnrealGameSync/UnrealGameSync/Forms/SelectWorkspaceWindow.cs b/Engine/Source/Programs/UnrealGameSync/UnrealGameSync/Forms/SelectWorkspaceWindow.cs
--- a/Engine/Source/Programs/UnrealGameSync/UnrealGameSync/Forms/SelectWorkspaceWindow.cs
+++ b/Engine/Source/Programs/UnrealGameSync/UnrealGameSync/Forms/SelectWorkspaceWindow.cs
@@ -67,7 +67,7 @@ namespace UnrealGameSync
 
 			foreach(ClientsRecord client in _clients.OrderBy(x => x.Name))
 			{
-				if(!OnlyForThisComputer.Checked || String.Compare(client.Host, _info.ClientHost, StringComparison.InvariantCultureIgnoreCase) == 0)
+				if(!OnlyForThisComputer.Checked || String.Equals(client.Host, _info.ClientHost, StringComparison.OrdinalIgnoreCase))
 				{
 					ListViewItem item = new ListViewItem(client.Name);
 					item.SubItems.Add(new ListViewItem.ListViewSubItem(item, client.Host));
@@ -93,7 +93,7 @@ namespace UnrealGameSync
 				return false;
 			}
 
-			SelectWorkspaceWindow selectWorkspace = new SelectWorkspaceWindow(task.Result.Info, task.Result.Clients, workspaceName);
+			using SelectWorkspaceWindow selectWorkspace = new SelectWorkspaceWindow(task.Result.Info, task.Result.Clients, workspaceName);
 			if(selectWorkspace.ShowDialog(owner) == DialogResult.OK)
 			{
 				newWorkspaceName = selectWorkspace._workspaceName;
diff --git a/Engine/Source/Programs/UnrealGameSync/UnrealGameSync/Forms/SyncFilter.cs b/Engine/Source/Programs/UnrealGameSync/UnrealGameSync/Forms/SyncFilter.cs
--- a/Engine/Source/Programs/UnrealGameSync/UnrealGameSync/Forms/SyncFilter.cs
+++ b/Engine/Source/Programs/UnrealGameSync/UnrealGameSync/Forms/SyncFilter.cs
@@ -120,7 +120,7 @@ namespace UnrealGameSync
 		{
 			GetSettings(out FilterSettings newGlobalFilter, out FilterSettings newWorkspaceFilter);
 
-			if(newGlobalFilter.View.Any(x => x.Contains("//")) || newWorkspaceFilter.View.Any(x => x.Contains("//")))
+			if(newGlobalFilter.View.Any(x => x.Contains("//", StringComparison.Ordinal)) || newWorkspaceFilter.View.Any(x => x.Contains("//", StringComparison.Ordinal)))
 			{
 				if(MessageBox.Show(this, "Custom views should be relative to the stream root (eg. -/Engine/...).\r\n\r\nFull depot paths (eg. //depot/...) will not match any files.\r\n\r\nAre you sure you want to continue?", "Invalid view", MessageBoxButtons.OKCancel) != System.Windows.Forms.DialogResult.OK)
 				{
diff --git a/Engine/Source/Programs/UnrealGameSync/UnrealGameSync/GlobalSettings.cs b/Engine/Source/Programs/UnrealGameSync/UnrealGameSync/GlobalSettings.cs
--- a/Engine/Source/Programs/UnrealGameSync/UnrealGameSync/GlobalSettings.cs
+++ b/Engine/Source/Programs/UnrealGameSync/UnrealGameSync/GlobalSettings.cs
@@ -49,7 +49,7 @@ namespace UnrealGameSync
 		public static void DeleteRegistryKey(RegistryKey key, string name)
 		{
 			string[] valueNames = key.GetValueNames();
-			if (valueNames.Any(x => String.Compare(x, name, StringComparison.OrdinalIgnoreCase) == 0))
+			if (valueNames.Any(x => String.Equals(x, name, StringComparison.OrdinalIgnoreCase)))
 			{
 				try
 				{
@@ -88,7 +88,7 @@ namespace UnrealGameSync
 						key.SetValue("UserName", userName);
 					}
 
-					if (String.IsNullOrEmpty(depotPath) || (DeploymentSettings.DefaultDepotPath != null && String.Equals(depotPath, DeploymentSettings.DefaultDepotPath, StringComparison.InvariantCultureIgnoreCase)))
+					if (String.IsNullOrEmpty(depotPath) || (DeploymentSettings.DefaultDepotPath != null && String.Equals(depotPath, DeploymentSettings.DefaultDepotPath, StringComparison.OrdinalIgnoreCase)))
 					{
 						DeleteRegistryKey(key, "DepotPath");
 					}
diff --git a/Engine/Source/Programs/UnrealGameSync/UnrealGameSync/IssueMonitor.cs b/Engine/Source/Programs/UnrealGameSync/UnrealGameSync/IssueMonitor.cs
--- a/Engine/Source/Programs/UnrealGameSync/UnrealGameSync/IssueMonitor.cs
+++ b/Engine/Source/Programs/UnrealGameSync/UnrealGameSync/IssueMonitor.cs
@@ -61,7 +61,7 @@ namespace UnrealGameSync
 		public bool Notify { get; set; }
 		public bool IsWarning { get; set; }
 		public string BuildUrl { get; set; } = String.Empty;
-		public List<string> Streams { get; set; } = new List<string>();
+		public List<string> Streams { get; } = new List<string>();
 
 		HashSet<string>? _cachedProjects;
 
@@ -311,7 +311,7 @@ namespace UnrealGameSync
 			if (_workerTask != null)
 			{
 				_cancellationSource.Cancel();
-				_asyncDisposer.Add(_workerTask.ContinueWith(_ => _cancellationSource.Dispose()));
+				_asyncDisposer.Add(_workerTask.ContinueWith(_ => _cancellationSource.Dispose(), TaskScheduler.Default));
 				_workerTask = null;
 			}
 		}
diff --git a/Engine/Source/Programs/UnrealGameSync/UnrealGameSync/Logging.cs b/Engine/Source/Programs/UnrealGameSync/UnrealGameSync/Logging.cs
--- a/Engine/Source/Programs/UnrealGameSync/UnrealGameSync/Logging.cs
+++ b/Engine/Source/Programs/UnrealGameSync/UnrealGameSync/Logging.cs
@@ -7,7 +7,7 @@ using Serilog.Templates;
 
 namespace UnrealGameSync
 {
-	public static class Logging
+	static class Logging
 	{
 		public static ILoggerProvider CreateLoggerProvider(FileReference file)
 		{
diff --git a/Engine/Source/Programs/UnrealGameSync/UnrealGameSync/ModalTask.cs b/Engine/Source/Programs/UnrealGameSync/UnrealGameSync/ModalTask.cs
--- a/Engine/Source/Programs/UnrealGameSync/UnrealGameSync/ModalTask.cs
+++ b/Engine/Source/Programs/UnrealGameSync/UnrealGameSync/ModalTask.cs
@@ -69,7 +69,7 @@ namespace UnrealGameSync
 			{
 				Task<T> backgroundTask = Task.Run(() => taskFunc(cancellationSource.Token));
 
-				ModalTaskWindow window = new ModalTaskWindow(title, message, (owner == null) ? FormStartPosition.CenterScreen : FormStartPosition.CenterParent, backgroundTask, cancellationSource);
+				using ModalTaskWindow window = new ModalTaskWindow(title, message, (owner == null) ? FormStartPosition.CenterScreen : FormStartPosition.CenterParent, backgroundTask, cancellationSource);
 				if(owner == null)
 				{
 					window.ShowInTaskbar = true;
diff --git a/Engine/Source/Programs/UnrealGameSync/UnrealGameSync/PerforceModalTask.cs b/Engine/Source/Programs/UnrealGameSync/UnrealGameSync/PerforceModalTask.cs
--- a/Engine/Source/Programs/UnrealGameSync/UnrealGameSync/PerforceModalTask.cs
+++ b/Engine/Source/Programs/UnrealGameSync/UnrealGameSync/PerforceModalTask.cs
@@ -87,7 +87,7 @@ namespace UnrealGameSync
 							passwordPrompt = $"Authentication failed. Enter the password for user '{perforceSettings.UserName}' on server '{perforceSettings.ServerAndPort}'.";
 						}
 
-						PasswordWindow passwordWindow = new PasswordWindow(passwordPrompt, password ?? String.Empty);
+						using PasswordWindow passwordWindow = new PasswordWindow(passwordPrompt, password ?? String.Empty);
 						if (owner == null)
 						{
 							passwordWindow.ShowInTaskbar = true;
diff --git a/Engine/Source/Programs/UnrealGameSync/UnrealGameSync/PerforceMonitor.cs b/Engine/Source/Programs/UnrealGameSync/UnrealGameSync/PerforceMonitor.cs
--- a/Engine/Source/Programs/UnrealGameSync/UnrealGameSync/PerforceMonitor.cs
+++ b/Engine/Source/Programs/UnrealGameSync/UnrealGameSync/PerforceMonitor.cs
@@ -61,6 +61,8 @@ namespace UnrealGameSync
 		public event Action? OnStreamChange;
 		public event Action? OnLoginExpired;
 
+		readonly object _lockObject = new object();
+
 		public TimeSpan ServerTimeZone
 		{
 			get;
@@ -105,7 +107,7 @@ namespace UnrealGameSync
 			if (_workerTask != null)
 			{
 				_cancellationSource.Cancel();
-				_asyncDisposeTasks.Add(_workerTask.ContinueWith(_ => _cancellationSource.Dispose()));
+				_asyncDisposeTasks.Add(_workerTask.ContinueWith(_ => _cancellationSource.Dispose(), TaskScheduler.Default));
 				_workerTask = null;
 			}
 		}
@@ -133,7 +135,7 @@ namespace UnrealGameSync
 			get => _pendingMaxChangesValue;
 			set 
 			{ 
-				lock(this)
+				lock(_lockObject)
 				{ 
 					if(value != _pendingMaxChangesValue)
 					{ 
@@ -163,7 +165,7 @@ namespace UnrealGameSync
 					{
 						Program.CaptureException(ex);
 					}
-					await Task.Delay(TimeSpan.FromSeconds(20.0), cancellationToken).ContinueWith(x => { });
+					await Task.Delay(TimeSpan.FromSeconds(20.0), cancellationToken).ContinueWith(x => { }, TaskScheduler.Default);
 				}
 			}
 		}
@@ -253,7 +255,7 @@ namespace UnrealGameSync
 			int newestChangeNumber = -1;
 			HashSet<int> currentChangelists;
 			SortedSet<int> prevPromotedChangelists;
-			lock(this)
+			lock(_lockObject)
 			{
 				maxChanges = PendingMaxChanges;
 				if(_changes.Count > 0)
@@ -285,7 +287,7 @@ namespace UnrealGameSync
 				ConfigSection? projectConfigSection = LatestProjectConfigFile.FindSection("Perforce");
 				if (projectConfigSection != null)
 				{
-					IEnumerable<string> additionalPaths = projectConfigSection.GetValues("AdditionalPathsToSync", new string[0]);
+					IEnumerable<string> additionalPaths = projectConfigSection.GetValues("AdditionalPathsToSync", Array.Empty<string>());
 
 					// turn into //ws/path
 					depotPaths.AddRange(additionalPaths.Select(p => String.Format("{0}/{1}", _branchClientPath, p.TrimStart('/'))));
@@ -342,7 +344,7 @@ namespace UnrealGameSync
 			const string roboMergePrefix = "#ROBOMERGE-AUTHOR:";
 			foreach (ChangesRecord change in newChanges)
 			{
-				if(change.Description != null && change.Description.StartsWith(roboMergePrefix))
+				if(change.Description != null && change.Description.StartsWith(roboMergePrefix, StringComparison.Ordinal))
 				{
 					int startIdx = roboMergePrefix.Length;
 					while(startIdx < change.Description.Length && change.Description[startIdx] == ' ')
@@ -368,7 +370,7 @@ namespace UnrealGameSync
 			if(newChanges.Count > 0 || maxChanges < CurrentMaxChanges)
 			{
 				// Insert them into the builds list
-				lock(this)
+				lock(_lockObject)
 				{
 					_changes.UnionWith(newChanges);
 					if(_changes.Count > maxChanges)
@@ -393,7 +395,7 @@ namespace UnrealGameSync
 				int newLastChangeByCurrentUser = -1;
 				foreach(ChangesRecord change in _changes)
 				{
-					if(String.Compare(change.User, perforce.Settings.UserName, StringComparison.InvariantCultureIgnoreCase) == 0)
+					if(String.Equals(change.User, perforce.Settings.UserName, StringComparison.OrdinalIgnoreCase))
 					{
 						newLastChangeByCurrentUser = Math.Max(newLastChangeByCurrentUser, change.Number);
 					}
@@ -431,7 +433,7 @@ namespace UnrealGameSync
 
 			// Find the changes we need to query
 			List<int> queryChangeNumbers = new List<int>();
-			lock(this)
+			lock(_lockObject)
 			{
 				foreach(ChangesRecord change in _changes)
 				{
@@ -484,7 +486,7 @@ namespace UnrealGameSync
 							details = new PerforceChangeDetails(describeRecord, isCodeFile);
 						}
 
-						lock (this)
+						lock (_lockObject)
 						{
 							if (!_changeDetails.ContainsKey(queryChangeNumber))
 							{
@@ -507,14 +509,14 @@ namespace UnrealGameSync
 					}
 				}
 				cancellationSource.Cancel();
-				await notifyTask.ContinueWith(_ => { }); // Ignore exceptions
+				await notifyTask.ContinueWith(_ => { }, TaskScheduler.Default); // Ignore exceptions
 			}
 
 			// Find the last submitted code change by the current user
 			int newLastCodeChangeByCurrentUser = -1;
 			foreach(ChangesRecord change in _changes)
 			{
-				if(String.Compare(change.User, perforce.Settings.UserName, StringComparison.InvariantCultureIgnoreCase) == 0)
+				if(String.Equals(change.User, perforce.Settings.UserName, StringComparison.OrdinalIgnoreCase))
 				{
 					PerforceChangeDetails? otherDetails;
 					if(_changeDetails.TryGetValue(change.Number, out otherDetails) && otherDetails.ContainsCode)
@@ -566,7 +568,7 @@ namespace UnrealGameSync
 
 		public List<ChangesRecord> GetChanges()
 		{
-			lock(this)
+			lock(_lockObject)
 			{
 				return new List<ChangesRecord>(_changes);
 			}
@@ -574,7 +576,7 @@ namespace UnrealGameSync
 
 		public bool TryGetChangeDetails(int number, [NotNullWhen(true)] out PerforceChangeDetails? details)
 		{
-			lock(this)
+			lock(_lockObject)
 			{
 				return _changeDetails.TryGetValue(number, out details);
 			}
@@ -582,7 +584,7 @@ namespace UnrealGameSync
 
 		public HashSet<int> GetPromotedChangeNumbers()
 		{
-			lock(this)
+			lock(_lockObject)
 			{
 				return new HashSet<int>(_promotedChangeNumbers);
 			}
diff --git a/Engine/Source/Programs/UnrealGameSync/UnrealGameSync/Program.cs b/Engine/Source/Programs/UnrealGameSync/UnrealGameSync/Program.cs
--- a/Engine/Source/Programs/UnrealGameSync/UnrealGameSync/Program.cs
+++ b/Engine/Source/Programs/UnrealGameSync/UnrealGameSync/Program.cs
@@ -269,7 +269,7 @@ namespace UnrealGameSync
 		{
 			SentrySdk.CaptureException(e.Exception);
 
-			ThreadExceptionDialog dialog = new ThreadExceptionDialog(e.Exception);
+			using ThreadExceptionDialog dialog = new ThreadExceptionDialog(e.Exception);
 			dialog.ShowDialog();
 		}
 
diff --git a/Engine/Source/Programs/UnrealGameSync/UnrealGameSync/ProgramApplicationContext.cs b/Engine/Source/Programs/UnrealGameSync/UnrealGameSync/ProgramApplicationContext.cs
--- a/Engine/Source/Programs/UnrealGameSync/UnrealGameSync/ProgramApplicationContext.cs
+++ b/Engine/Source/Programs/UnrealGameSync/UnrealGameSync/ProgramApplicationContext.cs
@@ -90,7 +90,7 @@ namespace UnrealGameSync
 			if(!String.IsNullOrEmpty(projectFileName))
 			{
 				string fullProjectFileName = Path.GetFullPath(projectFileName);
-				if(!_settings.OpenProjects.Any(x => x.LocalPath != null && String.Compare(x.LocalPath, fullProjectFileName, StringComparison.InvariantCultureIgnoreCase) == 0))
+				if(!_settings.OpenProjects.Any(x => x.LocalPath != null && String.Equals(x.LocalPath, fullProjectFileName, StringComparison.OrdinalIgnoreCase)))
 				{
 					_settings.OpenProjects.Add(new UserSelectedProjectSettings(null, null, UserSelectedProjectType.Local, null, fullProjectFileName));
 				}
@@ -237,9 +237,9 @@ namespace UnrealGameSync
 
 		private UserSelectedProjectSettings UpgradeSelectedProjectSettings(UserSelectedProjectSettings project)
 		{
-			if (project.ServerAndPort == null || String.Compare(project.ServerAndPort, _defaultPerforceSettings.ServerAndPort, StringComparison.OrdinalIgnoreCase) == 0)
+			if (project.ServerAndPort == null || String.Equals(project.ServerAndPort, _defaultPerforceSettings.ServerAndPort, StringComparison.OrdinalIgnoreCase))
 			{
-				if (project.UserName == null || String.Compare(project.UserName, _defaultPerforceSettings.UserName, StringComparison.OrdinalIgnoreCase) == 0)
+				if (project.UserName == null || String.Equals(project.UserName, _defaultPerforceSettings.UserName, StringComparison.OrdinalIgnoreCase))
 				{
 					project = new UserSelectedProjectSettings(null, null, project.Type, project.ClientPath, project.LocalPath);
 				}
diff --git a/Engine/Source/Programs/UnrealGameSync/UnrealGameSync/RESTApi.cs b/Engine/Source/Programs/UnrealGameSync/UnrealGameSync/RESTApi.cs
--- a/Engine/Source/Programs/UnrealGameSync/UnrealGameSync/RESTApi.cs
+++ b/Engine/Source/Programs/UnrealGameSync/UnrealGameSync/RESTApi.cs
@@ -10,9 +10,9 @@ using System.Threading.Tasks;
 
 namespace UnrealGameSync
 {
-	class RestException : Exception
+	public class RestException : Exception
 	{
-		public RestException(HttpMethod method, string uri, Exception innerException)
+		public RestException(HttpMethod method, string? uri, Exception innerException)
 			: base(String.Format("Error executing {0} {1}", method, uri), innerException)
 		{
 		}
@@ -45,12 +45,12 @@ namespace UnrealGameSync
 				{
 					using (HttpResponseMessage response = await s_httpClient.SendAsync(request, cancellationToken))
 					{
-						return await response.Content.ReadAsStringAsync();
+						return await response.Content.ReadAsStringAsync(cancellationToken);
 					}
 				}
 				catch (Exception ex)
 				{
-					throw new RestException(method, request.RequestUri.ToString(), ex);
+					throw new RestException(method, request.RequestUri?.ToString(), ex);
 				}
 			}
 		}
diff --git a/Engine/Source/Programs/UnrealGameSync/UnrealGameSync/Taskbar.cs b/Engine/Source/Programs/UnrealGameSync/UnrealGameSync/Taskbar.cs
--- a/Engine/Source/Programs/UnrealGameSync/UnrealGameSync/Taskbar.cs
+++ b/Engine/Source/Programs/UnrealGameSync/UnrealGameSync/Taskbar.cs
@@ -34,13 +34,17 @@ namespace UnrealGameSync
 		{
 		}
 
-		static readonly ITaskbarList3? _interface;
+		static readonly ITaskbarList3? _interface = CreateTaskbarListInterface();
 
-		static Taskbar()
+		static ITaskbarList3? CreateTaskbarListInterface()
 		{
-			if(Environment.OSVersion.Version >= new Version(6, 1))
+			if (Environment.OSVersion.Version >= new Version(6, 1))
 			{
-				_interface = new TaskbarList3() as ITaskbarList3;
+				return new TaskbarList3() as ITaskbarList3;
+			}
+			else
+			{
+				return null;
 			}
 		}
 
diff --git a/Engine/Source/Programs/UnrealGameSync/UnrealGameSync/ToolUpdateMonitor.cs b/Engine/Source/Programs/UnrealGameSync/UnrealGameSync/ToolUpdateMonitor.cs
--- a/Engine/Source/Programs/UnrealGameSync/UnrealGameSync/ToolUpdateMonitor.cs
+++ b/Engine/Source/Programs/UnrealGameSync/UnrealGameSync/ToolUpdateMonitor.cs
@@ -103,7 +103,7 @@ namespace UnrealGameSync
 			if (_workerTask != null)
 			{
 				_cancellationSource.Cancel();
-				_asyncDisposer.Add(_workerTask.ContinueWith(_ => _cancellationSource.Dispose()));
+				_asyncDisposer.Add(_workerTask.ContinueWith(_ => _cancellationSource.Dispose(), TaskScheduler.Default));
 				_workerTask = null;
 			}
 		}
@@ -178,7 +178,7 @@ namespace UnrealGameSync
 			List<ToolDefinition> newTools = new List<ToolDefinition>();
 			foreach (FStatRecord fileRecord in fileRecords)
 			{
-				if (fileRecord.DepotFile != null && fileRecord.DepotFile.EndsWith(".ini"))
+				if (fileRecord.DepotFile != null && fileRecord.DepotFile.EndsWith(".ini", StringComparison.OrdinalIgnoreCase))
 				{
 					ToolDefinition? tool = Tools.FirstOrDefault(x => x.ConfigPath.Equals(fileRecord.DepotFile, StringComparison.Ordinal));
 					if (tool == null || tool.ConfigChange != fileRecord.HeadChange)
@@ -292,7 +292,7 @@ namespace UnrealGameSync
 				};
 			}
 
-			string[] statusPanelLinks = configFile.GetValues("Settings.StatusPanelLinks", new string[0]);
+			string[] statusPanelLinks = configFile.GetValues("Settings.StatusPanelLinks", Array.Empty<string>());
 			foreach (string statusPanelLink in statusPanelLinks)
 			{
 				ConfigObject obj = new ConfigObject(statusPanelLink);
@@ -319,7 +319,7 @@ namespace UnrealGameSync
 			string commandExe = command;
 			string commandArgs = String.Empty;
 
-			int spaceIdx = command.IndexOf(' ');
+			int spaceIdx = command.IndexOf(' ', StringComparison.Ordinal);
 			if (spaceIdx != -1)
 			{
 				commandExe = command.Substring(0, spaceIdx);
@@ -332,7 +332,7 @@ namespace UnrealGameSync
 
 		async Task RemoveToolAsync(string toolName, Func<ILogger, CancellationToken, Task>? uninstallAction, ILogger logger, CancellationToken cancellationToken)
 		{
-			logger.LogInformation("Removing {0}", toolName);
+			logger.LogInformation("Removing {ToolName}", toolName);
 			DirectoryReference? toolPath = GetToolPath(toolName);
 
 			if (uninstallAction != null)
@@ -403,7 +403,7 @@ namespace UnrealGameSync
 				PerforceResponseList<PrintRecord> response = await perforce.TryPrintAsync(zipFile.FullName, $"{records[idx].DepotFile}#{records[idx].HeadRevision}", cancellationToken);
 				if(!response.Succeeded || !FileReference.Exists(zipFile))
 				{
-					logger.LogError("Unable to print {0}", records[idx].DepotFile);
+					logger.LogError("Unable to print {DepotFile}", records[idx].DepotFile);
 					return false;
 				}
 
diff --git a/Engine/Source/Programs/UnrealGameSync/UnrealGameSync/UpdateMonitor.cs b/Engine/Source/Programs/UnrealGameSync/UnrealGameSync/UpdateMonitor.cs
--- a/Engine/Source/Programs/UnrealGameSync/UnrealGameSync/UpdateMonitor.cs
+++ b/Engine/Source/Programs/UnrealGameSync/UnrealGameSync/UpdateMonitor.cs
@@ -49,7 +49,7 @@ namespace UnrealGameSync
 			if (_workerTask != null)
 			{
 				_cancellationSource.Cancel();
-				_asyncDisposer.Add(_workerTask.ContinueWith(_ => _cancellationSource.Dispose()));
+				_asyncDisposer.Add(_workerTask.ContinueWith(_ => _cancellationSource.Dispose(), TaskScheduler.Default));
 				_workerTask = null;
 			}
 		}
diff --git a/Engine/Source/Programs/UnrealGameSync/UnrealGameSync/UriHandlers/RdpHandler.cs b/Engine/Source/Programs/UnrealGameSync/UnrealGameSync/UriHandlers/RdpHandler.cs
--- a/Engine/Source/Programs/UnrealGameSync/UnrealGameSync/UriHandlers/RdpHandler.cs
+++ b/Engine/Source/Programs/UnrealGameSync/UnrealGameSync/UriHandlers/RdpHandler.cs
@@ -109,7 +109,7 @@ namespace UnrealGameSync
 			}
 		}
 
-		[UriHandler(Terminate = true)]
+		[UriHandler(terminate: true)]
 		public static UriResult Rdp(string host)
 		{
 			// Copy the credentials from a generic Windows credential to 
diff --git a/Engine/Source/Programs/UnrealGameSync/UnrealGameSync/UriHandlers/UriHandler.cs b/Engine/Source/Programs/UnrealGameSync/UnrealGameSync/UriHandlers/UriHandler.cs
--- a/Engine/Source/Programs/UnrealGameSync/UnrealGameSync/UriHandlers/UriHandler.cs
+++ b/Engine/Source/Programs/UnrealGameSync/UnrealGameSync/UriHandlers/UriHandler.cs
@@ -49,7 +49,7 @@ namespace UnrealGameSync
 
 				foreach (ParameterInfo param in info.GetParameters())
 				{					
-					string value = query.Get(param.Name);
+					string? value = query.Get(param.Name);
 
 					if (value == null)
 					{
@@ -154,7 +154,7 @@ namespace UnrealGameSync
 					}
 				}
 
-				if (uriIn == String.Empty)
+				if (String.IsNullOrEmpty(uriIn))
 				{
 					return false;
 				}
@@ -254,9 +254,9 @@ namespace UnrealGameSync
 	/// Method attribute for Uri handlers
 	/// </summary>
 	[AttributeUsage(AttributeTargets.Method)]
-	class UriHandlerAttribute : Attribute
+	sealed class UriHandlerAttribute : Attribute
 	{
-		public bool Terminate { get; set; }
+		public bool Terminate { get; }
 
 		public UriHandlerAttribute(bool terminate = false)
 		{
@@ -293,7 +293,7 @@ namespace UnrealGameSync
 			}
 		}
 
-		public static string CurrentProcessFilePath => Path.ChangeExtension(Path.GetFullPath(Process.GetCurrentProcess().MainModule!.FileName!), ".exe");
+		public static string CurrentProcessFilePath => Path.ChangeExtension(Path.GetFullPath(Environment.ProcessPath ?? String.Empty), ".exe");
 
 		static List<RegistrySetting> GetGlobalRegistrySettings()
 		{
diff --git a/Engine/Source/Programs/UnrealGameSync/UnrealGameSync/Workspace.cs b/Engine/Source/Programs/UnrealGameSync/UnrealGameSync/Workspace.cs
--- a/Engine/Source/Programs/UnrealGameSync/UnrealGameSync/Workspace.cs
+++ b/Engine/Source/Programs/UnrealGameSync/UnrealGameSync/Workspace.cs
@@ -128,7 +128,7 @@ namespace UnrealGameSync
 			{
 				CancellationTokenSource prevCancellationSourceCopy = _prevCancellationSource;
 				prevCancellationSourceCopy.Cancel();
-				_prevUpdateTask = _prevUpdateTask.ContinueWith(task => prevCancellationSourceCopy.Dispose());
+				_prevUpdateTask = _prevUpdateTask.ContinueWith(task => prevCancellationSourceCopy.Dispose(), TaskScheduler.Default);
 				_prevCancellationSource = null;
 			}
 			if(_currentUpdate != null)
diff --git a/Engine/Source/Programs/UnrealGameSync/UnrealGameSyncShared/ConfigFile.cs b/Engine/Source/Programs/UnrealGameSync/UnrealGameSyncShared/ConfigFile.cs
--- a/Engine/Source/Programs/UnrealGameSync/UnrealGameSyncShared/ConfigFile.cs
+++ b/Engine/Source/Programs/UnrealGameSync/UnrealGameSyncShared/ConfigFile.cs
@@ -468,7 +468,7 @@ namespace UnrealGameSync
 		public TEnum GetEnumValue<TEnum>(string key, TEnum defaultValue) where TEnum : struct
 		{
 			string? str = GetValue(key, null);
-			if (str != null && Enum.TryParse(str, out TEnum value))
+			if (str != null && Enum.TryParse(str, true, out TEnum value))
 			{
 				return value;
 			}
