UGS: Undo some static analysis fixes that rely on NET 6 features.

diff --git a/Engine/Source/Programs/UnrealGameSync/UnrealGameSync/AutomationServer.cs b/Engine/Source/Programs/UnrealGameSync/UnrealGameSync/AutomationServer.cs
--- a/Engine/Source/Programs/UnrealGameSync/UnrealGameSync/AutomationServer.cs
+++ b/Engine/Source/Programs/UnrealGameSync/UnrealGameSync/AutomationServer.cs
@@ -251,7 +251,7 @@ namespace UnrealGameSync
 				{
 					try
 					{
-						TcpClient client = await listener.AcceptTcpClientAsync(cancellationToken).ConfigureAwait(false);
+						TcpClient client = await listener.AcceptTcpClientAsync().ConfigureAwait(false);
 						try
 						{
 							_logger.LogInformation("Accepted connection from {Remote}", client.Client.RemoteEndPoint);
diff --git a/Engine/Source/Programs/UnrealGameSync/UnrealGameSync/RESTApi.cs b/Engine/Source/Programs/UnrealGameSync/UnrealGameSync/RESTApi.cs
--- a/Engine/Source/Programs/UnrealGameSync/UnrealGameSync/RESTApi.cs
+++ b/Engine/Source/Programs/UnrealGameSync/UnrealGameSync/RESTApi.cs
@@ -45,7 +45,7 @@ namespace UnrealGameSync
 				{
 					using (HttpResponseMessage response = await s_httpClient.SendAsync(request, cancellationToken))
 					{
-						return await response.Content.ReadAsStringAsync(cancellationToken);
+						return await response.Content.ReadAsStringAsync();
 					}
 				}
 				catch (Exception ex)
diff --git a/Engine/Source/Programs/UnrealGameSync/UnrealGameSync/UriHandlers/UriHandler.cs b/Engine/Source/Programs/UnrealGameSync/UnrealGameSync/UriHandlers/UriHandler.cs
--- a/Engine/Source/Programs/UnrealGameSync/UnrealGameSync/UriHandlers/UriHandler.cs
+++ b/Engine/Source/Programs/UnrealGameSync/UnrealGameSync/UriHandlers/UriHandler.cs
@@ -293,7 +293,7 @@ namespace UnrealGameSync
 			}
 		}
 
-		public static string CurrentProcessFilePath => Path.ChangeExtension(Path.GetFullPath(Environment.ProcessPath ?? String.Empty), ".exe");
+		public static string CurrentProcessFilePath => Path.ChangeExtension(Path.GetFullPath(Process.GetCurrentProcess().MainModule!.FileName!), ".exe");
 
 		static List<RegistrySetting> GetGlobalRegistrySettings()
 		{
