UGS: Fix CLI compile error.

diff --git a/Engine/Source/Programs/UnrealGameSync/UnrealGameSyncCmd/Program.cs b/Engine/Source/Programs/UnrealGameSync/UnrealGameSyncCmd/Program.cs
--- a/Engine/Source/Programs/UnrealGameSync/UnrealGameSyncCmd/Program.cs
+++ b/Engine/Source/Programs/UnrealGameSync/UnrealGameSyncCmd/Program.cs
@@ -1054,7 +1054,8 @@ namespace UnrealGameSyncCmd
 
 				if (commandOptions.View != null)
 				{
-					settings.View = commandOptions.View;
+					settings.View.Clear();
+					settings.View.AddRange(commandOptions.View);
 				}
 				if (commandOptions.RemoveView.Count > 0)
 				{
