UGS: Fix settings being lost on deserialization. NET 3.1 does not seem to support deserializing into existing list properties, unlike NET 6.

diff --git a/Engine/Source/Programs/UnrealGameSync/UnrealGameSyncShared/GlobalSettings.cs b/Engine/Source/Programs/UnrealGameSync/UnrealGameSyncShared/GlobalSettings.cs
--- a/Engine/Source/Programs/UnrealGameSync/UnrealGameSyncShared/GlobalSettings.cs
+++ b/Engine/Source/Programs/UnrealGameSync/UnrealGameSyncShared/GlobalSettings.cs
@@ -10,9 +10,9 @@ namespace UnrealGameSync
 {
 	public class FilterSettings
 	{
-		public List<Guid> IncludeCategories { get; } = new List<Guid>();
-		public List<Guid> ExcludeCategories { get; } = new List<Guid>();
-		public List<string> View { get; } = new List<string>();
+		public List<Guid> IncludeCategories { get; set; } = new List<Guid>();
+		public List<Guid> ExcludeCategories { get; set; } = new List<Guid>();
+		public List<string> View { get; set; } = new List<string>();
 		public bool? AllProjects { get; set; }
 		public bool? AllProjectsInSln { get; set; }
 
diff --git a/Engine/Source/Programs/UnrealGameSync/UnrealGameSyncShared/UserSettings.cs b/Engine/Source/Programs/UnrealGameSync/UnrealGameSyncShared/UserSettings.cs
--- a/Engine/Source/Programs/UnrealGameSync/UnrealGameSyncShared/UserSettings.cs
+++ b/Engine/Source/Programs/UnrealGameSync/UnrealGameSyncShared/UserSettings.cs
@@ -65,7 +65,7 @@ namespace UnrealGameSync
 		public bool Good { get; set; } = false;
 		public bool Starred { get; set; } = false;
 		public bool FindNewestGoodContent { get; set; } = false;
-		public List<string> RequiredBadges { get; } = new List<string>();
+		public List<string> RequiredBadges { get; set; } = new List<string>();
 
 		// Depot path to read for the latest change number
 		public string? ReadFrom { get; set; }
@@ -152,7 +152,7 @@ namespace UnrealGameSync
 	{
 		public bool Enabled { get; set; }
 		public string Type { get; set; }
-		public List<string> Order { get; }
+		public List<string> Order { get; set; }
 
 		public ArchiveSettings(bool enabled, string type, IEnumerable<string> order)
 		{
@@ -416,9 +416,9 @@ namespace UnrealGameSync
 		[JsonIgnore]
 		public FileReference ConfigFile { get; private set; } = null!;
 
-		public List<ConfigObject> BuildSteps { get; } = new List<ConfigObject>();
+		public List<ConfigObject> BuildSteps { get; set; } = new List<ConfigObject>();
 		public FilterType FilterType { get; set; }
-		public HashSet<string> FilterBadges { get; } = new HashSet<string>(StringComparer.OrdinalIgnoreCase);
+		public HashSet<string> FilterBadges { get; set; } = new HashSet<string>(StringComparer.OrdinalIgnoreCase);
 
 		static readonly object _syncRoot = new object();
 
@@ -496,21 +496,21 @@ namespace UnrealGameSync
 		public bool AnnotateRobmergeChanges { get; set; }
 		public bool ShowLocalTimes { get; set; }
 		public bool KeepInTray { get; set; }
-		public HashSet<Guid> EnabledTools { get; } = new HashSet<Guid>();
+		public HashSet<Guid> EnabledTools { get; set; } = new HashSet<Guid>();
 		public int FilterIndex { get; set; }
 		public UserSelectedProjectSettings? LastProject { get; set; }
-		public List<UserSelectedProjectSettings> OpenProjects { get; }
-		public List<UserSelectedProjectSettings> RecentProjects { get; }
+		public List<UserSelectedProjectSettings> OpenProjects { get; set; }
+		public List<UserSelectedProjectSettings> RecentProjects { get; set; }
 		public string SyncTypeId { get; set; }
 		public BuildConfig CompiledEditorBuildConfig { get; set; } // NB: This assumes not using precompiled editor. See CurrentBuildConfig.
 		public TabLabels TabLabels { get; set; }
 		public long NextLauncherVersionCheck { get; set; }
 
 		// Precompiled binaries
-		public List<ArchiveSettings> Archives { get; } = new List<ArchiveSettings>();
+		public List<ArchiveSettings> Archives { get; set; } = new List<ArchiveSettings>();
 
 		// OIDC Settings
-		public Dictionary<string, string> ProviderToRefreshTokens { get; } = new Dictionary<string, string>();
+		public Dictionary<string, string> ProviderToRefreshTokens { get; set; } = new Dictionary<string, string>();
 
 		// Window settings
 		public bool WindowVisible { get; set; }
@@ -521,14 +521,14 @@ namespace UnrealGameSync
 		public bool ScheduleEnabled { get; set; }
 		public TimeSpan ScheduleTime { get; set; }
 		public bool ScheduleAnyOpenProject { get; set; }
-		public List<UserSelectedProjectSettings> ScheduleProjects { get; } = new List<UserSelectedProjectSettings>();
+		public List<UserSelectedProjectSettings> ScheduleProjects { get; set; } = new List<UserSelectedProjectSettings>();
 
 		// Run configuration
-		public List<Tuple<string, bool>> EditorArguments { get; } = new List<Tuple<string, bool>>();
+		public List<Tuple<string, bool>> EditorArguments { get; set; } = new List<Tuple<string, bool>>();
 		public bool EditorArgumentsPrompt { get; set; }
 
 		// Notification settings
-		public List<string> NotifyProjects { get; } = new List<string>();
+		public List<string> NotifyProjects { get; set; } = new List<string>();
 		public int NotifyUnassignedMinutes { get; set; }
 		public int NotifyUnacknowledgedMinutes { get; set; }
 		public int NotifyUnresolvedMinutes { get; set; }
