UGS: Fix exception when switching projects.

diff --git a/Engine/Source/Programs/UnrealGameSync/UnrealGameSyncShared/WorkspaceState.cs b/Engine/Source/Programs/UnrealGameSync/UnrealGameSyncShared/WorkspaceState.cs
--- a/Engine/Source/Programs/UnrealGameSync/UnrealGameSyncShared/WorkspaceState.cs
+++ b/Engine/Source/Programs/UnrealGameSync/UnrealGameSyncShared/WorkspaceState.cs
@@ -15,6 +15,12 @@ namespace UnrealGameSync
 	/// </summary>
 	public class WorkspaceState
 	{
+		// **********************************************************************
+		//
+		// NOTE: UPDATE CopyFrom() BELOW WHEN ADDING ANY PROPERTIES TO THIS CLASS
+		//
+		// **********************************************************************
+
 		// Cached state about the project, configured using UserWorkspaceSettings and taken from computed values in ProjectInfo. Assumed valid unless manually updated.
 		public long SettingsTimeUtc { get; set; }
 		public string ClientName { get; set; } = String.Empty;
@@ -28,7 +34,7 @@ namespace UnrealGameSync
 		public int CurrentChangeNumber { get; set; } = -1;
 		public int CurrentCodeChangeNumber { get; set; } = -1;
 		public string? CurrentSyncFilterHash { get; set; }
-		public List<int> AdditionalChangeNumbers { get; } = new List<int>();
+		public List<int> AdditionalChangeNumbers { get; set; } = new List<int>();
 
 		// Settings for the last attempted sync. These values are set to persist error messages between runs.
 		public int LastSyncChangeNumber { get; set; }
@@ -41,10 +47,10 @@ namespace UnrealGameSync
 		public int LastBuiltChangeNumber { get; set; }
 
 		// Expanded archives in the workspace
-		public HashSet<string> ExpandedArchiveTypes { get; } = new HashSet<string>(StringComparer.Ordinal);
+		public HashSet<string> ExpandedArchiveTypes { get; set; } = new HashSet<string>(StringComparer.Ordinal);
 
 		// The changes that we're regressing at the moment
-		public List<BisectEntry> BisectChanges { get; } = new List<BisectEntry>();
+		public List<BisectEntry> BisectChanges { get; set; } = new List<BisectEntry>();
 
 		public void UpdateCachedProjectInfo(ProjectInfo projectInfo, long settingsTimeUtc)
 		{
@@ -74,10 +80,33 @@ namespace UnrealGameSync
 
 		public void CopyFrom(WorkspaceState other)
 		{
-			foreach (PropertyInfo propertyInfo in GetType().GetProperties(BindingFlags.Instance | BindingFlags.Public | BindingFlags.SetProperty | BindingFlags.GetProperty))
-			{
-				propertyInfo.SetValue(this, propertyInfo.GetValue(other));
-			}
+			SettingsTimeUtc = other.SettingsTimeUtc;
+			ClientName = other.ClientName;
+			BranchPath = other.BranchPath;
+			ProjectPath = other.ProjectPath;
+			StreamName = other.StreamName;
+			ProjectIdentifier = other.ProjectIdentifier;
+			IsEnterpriseProject = other.IsEnterpriseProject;
+
+			CurrentChangeNumber = other.CurrentChangeNumber;
+			CurrentCodeChangeNumber = other.CurrentCodeChangeNumber;
+			CurrentSyncFilterHash = other.CurrentSyncFilterHash;
+			AdditionalChangeNumbers.Clear();
+			AdditionalChangeNumbers.AddRange(other.AdditionalChangeNumbers);
+
+			LastSyncChangeNumber = other.LastSyncChangeNumber;
+			LastSyncResult = other.LastSyncResult;
+			LastSyncResultMessage = other.LastSyncResultMessage;
+			LastSyncTime = other.LastSyncTime;
+			LastSyncDurationSeconds = other.LastSyncDurationSeconds;
+
+			LastBuiltChangeNumber = other.LastBuiltChangeNumber;
+
+			ExpandedArchiveTypes.Clear();
+			ExpandedArchiveTypes.UnionWith(other.ExpandedArchiveTypes);
+
+			BisectChanges.Clear();
+			BisectChanges.AddRange(other.BisectChanges);
 		}
 
 		public void SetBisectState(int change, BisectState state)
