UGS: Only create workspace lock on Windows for CLI client. Named synchronization primitives are not supported on other platforms.

diff --git a/Engine/Source/Programs/UnrealGameSync/UnrealGameSyncCmd/Program.cs b/Engine/Source/Programs/UnrealGameSync/UnrealGameSyncCmd/Program.cs
--- a/Engine/Source/Programs/UnrealGameSync/UnrealGameSyncCmd/Program.cs
+++ b/Engine/Source/Programs/UnrealGameSync/UnrealGameSyncCmd/Program.cs
@@ -657,8 +657,8 @@ namespace UnrealGameSyncCmd
 
 				string[] syncFilter = ReadSyncFilter(settings, context.UserSettings, projectConfig);
 
-				using WorkspaceLock workspaceLock = new WorkspaceLock(settings.RootDir);
-				if (!await workspaceLock.TryAcquireAsync())
+				using WorkspaceLock? workspaceLock = CreateWorkspaceLock(settings.RootDir);
+				if(workspaceLock != null && !await workspaceLock.TryAcquireAsync())
 				{
 					logger.LogError("Another process is already syncing this workspace.");
 					return;
@@ -728,6 +728,18 @@ namespace UnrealGameSyncCmd
 
 				state.Modify(x => x.SetLastSyncState(result, updateContext, message));
 			}
+
+			static WorkspaceLock? CreateWorkspaceLock(DirectoryReference rootDir)
+			{
+				if (RuntimeInformation.IsOSPlatform(OSPlatform.Windows))
+				{
+					return new WorkspaceLock(rootDir);
+				}
+				else
+				{
+					return null;
+				}
+			}
 		}
 
 		class ClientsCommand : Command
