UGS: Upgrade to NET 6.

diff --git a/Engine/Source/Programs/UnrealGameSync/Shared/EpicGames.Core/EpicGames.Core.csproj b/Engine/Source/Programs/UnrealGameSync/Shared/EpicGames.Core/EpicGames.Core.csproj
deleted file mode 100644
diff --git a/Engine/Source/Programs/UnrealGameSync/Shared/EpicGames.Jupiter/EpicGames.Jupiter.csproj b/Engine/Source/Programs/UnrealGameSync/Shared/EpicGames.Jupiter/EpicGames.Jupiter.csproj
deleted file mode 100644
diff --git a/Engine/Source/Programs/UnrealGameSync/Shared/EpicGames.OIDC/EpicGames.OIDC.csproj b/Engine/Source/Programs/UnrealGameSync/Shared/EpicGames.OIDC/EpicGames.OIDC.csproj
deleted file mode 100644
diff --git a/Engine/Source/Programs/UnrealGameSync/Shared/EpicGames.Perforce/EpicGames.Perforce.csproj b/Engine/Source/Programs/UnrealGameSync/Shared/EpicGames.Perforce/EpicGames.Perforce.csproj
deleted file mode 100644
diff --git a/Engine/Source/Programs/UnrealGameSync/UnrealGameSync/Program.cs b/Engine/Source/Programs/UnrealGameSync/UnrealGameSync/Program.cs
--- a/Engine/Source/Programs/UnrealGameSync/UnrealGameSync/Program.cs
+++ b/Engine/Source/Programs/UnrealGameSync/UnrealGameSync/Program.cs
@@ -67,6 +67,7 @@ namespace UnrealGameSync
 				{
 					Application.EnableVisualStyles();
 					Application.SetCompatibleTextRenderingDefault(false);
+					Application.SetHighDpiMode(HighDpiMode.PerMonitorV2);
 				}
 
 				using (EventWaitHandle activateEvent = new EventWaitHandle(false, EventResetMode.AutoReset, "ActivateUnrealGameSync"))
diff --git a/Engine/Source/Programs/UnrealGameSync/UnrealGameSync/Properties/app.manifest b/Engine/Source/Programs/UnrealGameSync/UnrealGameSync/Properties/app.manifest
--- a/Engine/Source/Programs/UnrealGameSync/UnrealGameSync/Properties/app.manifest
+++ b/Engine/Source/Programs/UnrealGameSync/UnrealGameSync/Properties/app.manifest
@@ -26,23 +26,18 @@
   </trustInfo>
   <compatibility xmlns="urn:schemas-microsoft-com:compatibility.v1">
     <application>
-      <!-- A list of all Windows versions that this application is designed to work with. 
-      Windows will automatically select the most compatible environment.-->
-      <!-- If your application is designed to work with Windows Vista, uncomment the following supportedOS node-->
-      <!--<supportedOS Id="{e2011457-1546-43c5-a5fe-008deee3d3f0}"></supportedOS>-->
-      <!-- If your application is designed to work with Windows 7, uncomment the following supportedOS node-->
-      <supportedOS Id="{35138b9a-5d96-4fbd-8e2d-a2440225f93a}"/>
-      <!-- If your application is designed to work with Windows 8, uncomment the following supportedOS node-->
-      <!--<supportedOS Id="{4a2f28e3-53b9-4441-ba9c-d69d4a4a6e38}"></supportedOS>-->
-      <!-- If your application is designed to work with Windows 8.1, uncomment the following supportedOS node-->
-      <!--<supportedOS Id="{1f676c76-80e1-4239-95bb-83d0f6d0da78}"/>-->
-    </application>
+		<!-- Windows 10 and Windows 11 -->
+		<supportedOS Id="{8e0f7a12-bfb3-4fe8-b9a5-48fd50a15a9a}"/>
+		<!-- Windows 8.1 -->
+		<supportedOS Id="{1f676c76-80e1-4239-95bb-83d0f6d0da78}"/>
+		<!-- Windows 8 -->
+		<supportedOS Id="{4a2f28e3-53b9-4441-ba9c-d69d4a4a6e38}"/>
+		<!-- Windows 7 -->
+		<supportedOS Id="{35138b9a-5d96-4fbd-8e2d-a2440225f93a}"/>
+		<!-- Windows Vista -->
+		<supportedOS Id="{e2011457-1546-43c5-a5fe-008deee3d3f0}"/>
+	</application>
   </compatibility>
-  <asmv3:application>
-    <asmv3:windowsSettings xmlns="http://schemas.microsoft.com/SMI/2005/WindowsSettings">
-      <dpiAware>true</dpiAware>
-    </asmv3:windowsSettings>
-  </asmv3:application>
 
 	<!-- Enable themes for Windows common controls and dialogs (Windows XP and later) -->
   <!-- <dependency>
diff --git a/Engine/Source/Programs/UnrealGameSync/UnrealGameSync/UnrealGameSync.csproj b/Engine/Source/Programs/UnrealGameSync/UnrealGameSync/UnrealGameSync.csproj
--- a/Engine/Source/Programs/UnrealGameSync/UnrealGameSync/UnrealGameSync.csproj
+++ b/Engine/Source/Programs/UnrealGameSync/UnrealGameSync/UnrealGameSync.csproj
@@ -1,7 +1,7 @@
-<Project Sdk="Microsoft.NET.Sdk.WindowsDesktop">
+<Project Sdk="Microsoft.NET.Sdk">
   <PropertyGroup>
     <OutputType>WinExe</OutputType>
-    <TargetFramework>netcoreapp3.1</TargetFramework>
+    <TargetFramework>net6.0-windows</TargetFramework>
     <UseWindowsForms>true</UseWindowsForms>
     <Deterministic>true</Deterministic>
     <Product>UnrealGameSync</Product>
@@ -29,9 +29,10 @@
     <PackageReference Include="Serilog.Sinks.File" Version="5.0.0" />
   </ItemGroup>
   <ItemGroup>
-    <ProjectReference Include="..\Shared\EpicGames.Core\EpicGames.Core.csproj" />
-    <ProjectReference Include="..\Shared\EpicGames.OIDC\EpicGames.OIDC.csproj" />
-    <ProjectReference Include="..\Shared\EpicGames.Perforce\EpicGames.Perforce.csproj" />
+    <ProjectReference Include="..\..\Shared\EpicGames.Core\EpicGames.Core.csproj" />
+    <ProjectReference Include="..\..\Shared\EpicGames.Horde\EpicGames.Horde.csproj" />
+    <ProjectReference Include="..\..\Shared\EpicGames.OIDC\EpicGames.OIDC.csproj" />
+    <ProjectReference Include="..\..\Shared\EpicGames.Perforce\EpicGames.Perforce.csproj" />
     <ProjectReference Include="..\UnrealGameSyncShared\UnrealGameSyncShared.csproj" />
   </ItemGroup>
   <ItemGroup>
diff --git a/Engine/Source/Programs/UnrealGameSync/UnrealGameSyncCmd/UnrealGameSyncCmd.csproj b/Engine/Source/Programs/UnrealGameSync/UnrealGameSyncCmd/UnrealGameSyncCmd.csproj
--- a/Engine/Source/Programs/UnrealGameSync/UnrealGameSyncCmd/UnrealGameSyncCmd.csproj
+++ b/Engine/Source/Programs/UnrealGameSync/UnrealGameSyncCmd/UnrealGameSyncCmd.csproj
@@ -2,7 +2,7 @@
 
   <PropertyGroup>
     <OutputType>Exe</OutputType>
-    <TargetFramework>netcoreapp3.1</TargetFramework>
+    <TargetFramework>net6.0</TargetFramework>
     <Nullable>enable</Nullable>
     <AssemblyName>ugs</AssemblyName>
   </PropertyGroup>
@@ -15,7 +15,6 @@
   </ItemGroup>
 
   <ItemGroup>
-    <ProjectReference Include="..\Shared\EpicGames.Core\EpicGames.Core.csproj" />
     <ProjectReference Include="..\UnrealGameSyncShared\UnrealGameSyncShared.csproj" />
   </ItemGroup>
 
diff --git a/Engine/Source/Programs/UnrealGameSync/UnrealGameSyncShared/UnrealGameSyncShared.csproj b/Engine/Source/Programs/UnrealGameSync/UnrealGameSyncShared/UnrealGameSyncShared.csproj
--- a/Engine/Source/Programs/UnrealGameSync/UnrealGameSyncShared/UnrealGameSyncShared.csproj
+++ b/Engine/Source/Programs/UnrealGameSync/UnrealGameSyncShared/UnrealGameSyncShared.csproj
@@ -1,15 +1,19 @@
 <Project Sdk="Microsoft.NET.Sdk">
 
   <PropertyGroup>
-    <TargetFramework>netcoreapp3.1</TargetFramework>
+    <TargetFramework>net6.0</TargetFramework>
     <Nullable>enable</Nullable>
     <RootNamespace>UnrealGameSync</RootNamespace>
   </PropertyGroup>
 
   <ItemGroup>
-    <ProjectReference Include="..\Shared\EpicGames.Core\EpicGames.Core.csproj" />
-    <ProjectReference Include="..\Shared\EpicGames.OIDC\EpicGames.OIDC.csproj" />
-    <ProjectReference Include="..\Shared\EpicGames.Perforce\EpicGames.Perforce.csproj" />
+    <PackageReference Include="System.IO.FileSystem.DriveInfo" Version="4.3.1" />
+  </ItemGroup>
+
+  <ItemGroup>
+    <ProjectReference Include="..\..\Shared\EpicGames.Core\EpicGames.Core.csproj" />
+    <ProjectReference Include="..\..\Shared\EpicGames.OIDC\EpicGames.OIDC.csproj" />
+    <ProjectReference Include="..\..\Shared\EpicGames.Perforce\EpicGames.Perforce.csproj" />
   </ItemGroup>
 
 </Project>
diff --git a/Engine/Source/Programs/UnrealGameSync/UnrealGameSyncShared/Utility.cs b/Engine/Source/Programs/UnrealGameSync/UnrealGameSyncShared/Utility.cs
--- a/Engine/Source/Programs/UnrealGameSync/UnrealGameSyncShared/Utility.cs
+++ b/Engine/Source/Programs/UnrealGameSync/UnrealGameSyncShared/Utility.cs
@@ -123,7 +123,7 @@ namespace UnrealGameSync
 
 		public static byte[] SerializeJson<T>(T obj)
 		{
-			JsonSerializerOptions options = new JsonSerializerOptions { IgnoreNullValues = true, WriteIndented = true };
+			JsonSerializerOptions options = new JsonSerializerOptions { DefaultIgnoreCondition = JsonIgnoreCondition.WhenWritingNull, WriteIndented = true };
 			options.Converters.Add(new JsonStringEnumConverter());
 
 			byte[] buffer;
