UGS: Fix NET 6 DPI issues.

Winforms 6.0 seems to reset DPI scaling values whenever font is assigned through generated code, overwriting existing values. Assign font after ResumeLayout() call instead.

diff --git a/Engine/Source/Programs/UnrealGameSync/UnrealGameSync/Controls/WorkspaceControl.cs b/Engine/Source/Programs/UnrealGameSync/UnrealGameSync/Controls/WorkspaceControl.cs
--- a/Engine/Source/Programs/UnrealGameSync/UnrealGameSync/Controls/WorkspaceControl.cs
+++ b/Engine/Source/Programs/UnrealGameSync/UnrealGameSync/Controls/WorkspaceControl.cs
@@ -382,6 +382,7 @@ namespace UnrealGameSync
 //			_jupiterMonitor = JupiterMonitor.CreateFromConfigFile(inOidcTokenManager, jupiterLogger, openProjectInfo.LatestProjectConfigFile, SelectedProjectIdentifier);
 
 			UpdateColumnSettings(true);
+			Font = new System.Drawing.Font("Segoe UI", 8.25F, System.Drawing.FontStyle.Regular, System.Drawing.GraphicsUnit.Point);
 
 			SyncLog.OpenFile(projectLogBaseName);
 
@@ -565,61 +566,58 @@ namespace UnrealGameSync
 					_columnWidths[idx] = BuildList.Columns[idx].Width;
 				}
 
-				using (Graphics graphics = Graphics.FromHwnd(IntPtr.Zero))
-				{
-					float dpiScaleX = graphics.DpiX / 96.0f;
-
-					_minColumnWidths = Enumerable.Repeat(32, BuildList.Columns.Count).ToArray();
-					_minColumnWidths[IconColumn.Index] = (int)(50 * dpiScaleX);
-					_minColumnWidths[TypeColumn.Index] = (int)(100 * dpiScaleX);
-					_minColumnWidths[TimeColumn.Index] = (int)(75 * dpiScaleX);
-					_minColumnWidths[ChangeColumn.Index] = (int)(75 * dpiScaleX);
-					_minColumnWidths[CISColumn.Index] = (int)(200 * dpiScaleX);
-
-					_desiredColumnWidths = Enumerable.Repeat(65536, BuildList.Columns.Count).ToArray();
-					_desiredColumnWidths[IconColumn.Index] = _minColumnWidths[IconColumn.Index];
-					_desiredColumnWidths[TypeColumn.Index] = _minColumnWidths[TypeColumn.Index];
-					_desiredColumnWidths[TimeColumn.Index] = _minColumnWidths[TimeColumn.Index];
-					_desiredColumnWidths[ChangeColumn.Index] = _minColumnWidths[ChangeColumn.Index];
-					_desiredColumnWidths[AuthorColumn.Index] = (int)(120 * dpiScaleX);
-					_desiredColumnWidths[CISColumn.Index] = (int)(200 * dpiScaleX);
-					_desiredColumnWidths[StatusColumn.Index] = (int)(300 * dpiScaleX);
-
-					_columnWeights = Enumerable.Repeat(1.0f, BuildList.Columns.Count).ToArray();
-					_columnWeights[IconColumn.Index] = 3.0f;
-					_columnWeights[TypeColumn.Index] = 3.0f;
-					_columnWeights[TimeColumn.Index] = 3.0f;
-					_columnWeights[ChangeColumn.Index] = 3.0f;
-					_columnWeights[DescriptionColumn.Index] = 1.25f;
-					_columnWeights[CISColumn.Index] = 1.5f;
+				float dpiScaleX = DeviceDpi / 96.0f;
+
+				_minColumnWidths = Enumerable.Repeat((int)(32 * dpiScaleX), BuildList.Columns.Count).ToArray();
+				_minColumnWidths[IconColumn.Index] = (int)(50 * dpiScaleX);
+				_minColumnWidths[TypeColumn.Index] = (int)(100 * dpiScaleX);
+				_minColumnWidths[TimeColumn.Index] = (int)(75 * dpiScaleX);
+				_minColumnWidths[ChangeColumn.Index] = (int)(75 * dpiScaleX);
+				_minColumnWidths[CISColumn.Index] = (int)(200 * dpiScaleX);
+
+				_desiredColumnWidths = Enumerable.Repeat(65536, BuildList.Columns.Count).ToArray();
+				_desiredColumnWidths[IconColumn.Index] = _minColumnWidths[IconColumn.Index];
+				_desiredColumnWidths[TypeColumn.Index] = _minColumnWidths[TypeColumn.Index];
+				_desiredColumnWidths[TimeColumn.Index] = _minColumnWidths[TimeColumn.Index];
+				_desiredColumnWidths[ChangeColumn.Index] = _minColumnWidths[ChangeColumn.Index];
+				_desiredColumnWidths[AuthorColumn.Index] = (int)(120 * dpiScaleX);
+				_desiredColumnWidths[CISColumn.Index] = (int)(200 * dpiScaleX);
+				_desiredColumnWidths[StatusColumn.Index] = (int)(300 * dpiScaleX);
+
+				_columnWeights = Enumerable.Repeat(1.0f, BuildList.Columns.Count).ToArray();
+				_columnWeights[IconColumn.Index] = 3.0f;
+				_columnWeights[TypeColumn.Index] = 3.0f;
+				_columnWeights[TimeColumn.Index] = 3.0f;
+				_columnWeights[ChangeColumn.Index] = 3.0f;
+				_columnWeights[DescriptionColumn.Index] = 1.25f;
+				_columnWeights[CISColumn.Index] = 1.5f;
 
-					foreach (ColumnHeader? column in BuildList.Columns)
+				foreach (ColumnHeader? column in BuildList.Columns)
+				{
+					if (column != null)
 					{
-						if (column != null)
+						ConfigObject? columnConfig = (ConfigObject?)column.Tag;
+						if (columnConfig != null)
 						{
-							ConfigObject? columnConfig = (ConfigObject?)column.Tag;
-							if (columnConfig != null)
-							{
-								_minColumnWidths[column.Index] = (int)(columnConfig.GetValue("MinWidth", _minColumnWidths[column.Index]) * dpiScaleX);
-								_desiredColumnWidths[column.Index] = (int)(columnConfig.GetValue("DesiredWidth", _desiredColumnWidths[column.Index]) * dpiScaleX);
-								_columnWeights[column.Index] = columnConfig.GetValue("Weight", _minColumnWidths[column.Index]);
-							}
+							_minColumnWidths[column.Index] = (int)(columnConfig.GetValue("MinWidth", _minColumnWidths[column.Index]) * dpiScaleX);
+							_desiredColumnWidths[column.Index] = (int)(columnConfig.GetValue("DesiredWidth", _desiredColumnWidths[column.Index]) * dpiScaleX);
+							_columnWeights[column.Index] = columnConfig.GetValue("Weight", _minColumnWidths[column.Index]);
 						}
 					}
+				}
 
-					ConfigFile projectConfigFile = _perforceMonitor.LatestProjectConfigFile;
-					for (int idx = 0; idx < BuildList.Columns.Count; idx++)
+				ConfigFile projectConfigFile = _perforceMonitor.LatestProjectConfigFile;
+				for (int idx = 0; idx < BuildList.Columns.Count; idx++)
+				{
+					if (!String.IsNullOrEmpty(BuildList.Columns[idx].Text))
 					{
-						if (!String.IsNullOrEmpty(BuildList.Columns[idx].Text))
+						string? stringValue;
+						if (TryGetProjectSetting(projectConfigFile, String.Format("ColumnWidth_{0}", BuildList.Columns[idx].Text), out stringValue))
 						{
-							string? stringValue;
-							if (TryGetProjectSetting(projectConfigFile, String.Format("ColumnWidth_{0}", BuildList.Columns[idx].Text), out stringValue))
+							int intValue;
+							if (Int32.TryParse(stringValue, out intValue))
 							{
-								int intValue;
-								if (Int32.TryParse(stringValue, out intValue))
-								{
-									_desiredColumnWidths[idx] = (int)(intValue * dpiScaleX);
-								}
+								_desiredColumnWidths[idx] = (int)(intValue * dpiScaleX);
 							}
 						}
 					}
diff --git a/Engine/Source/Programs/UnrealGameSync/UnrealGameSync/Forms/ApplicationSettingsWindow.cs b/Engine/Source/Programs/UnrealGameSync/UnrealGameSync/Forms/ApplicationSettingsWindow.cs
--- a/Engine/Source/Programs/UnrealGameSync/UnrealGameSync/Forms/ApplicationSettingsWindow.cs
+++ b/Engine/Source/Programs/UnrealGameSync/UnrealGameSync/Forms/ApplicationSettingsWindow.cs
@@ -62,6 +62,7 @@ namespace UnrealGameSync
 		private ApplicationSettingsWindow(IPerforceSettings defaultPerforceSettings, bool preview, string originalExecutableFileName, UserSettings settings, ToolUpdateMonitor toolUpdateMonitor, ILogger<ApplicationSettingsWindow> logger)
 		{
 			InitializeComponent();
+			Font = new System.Drawing.Font("Segoe UI", 8.25F, System.Drawing.FontStyle.Regular, System.Drawing.GraphicsUnit.Point, ((byte)(0)));
 
 			_originalExecutableFileName = originalExecutableFileName;
 			_defaultPerforceSettings = defaultPerforceSettings;
diff --git a/Engine/Source/Programs/UnrealGameSync/UnrealGameSync/Forms/ArgumentsWindow.cs b/Engine/Source/Programs/UnrealGameSync/UnrealGameSync/Forms/ArgumentsWindow.cs
--- a/Engine/Source/Programs/UnrealGameSync/UnrealGameSync/Forms/ArgumentsWindow.cs
+++ b/Engine/Source/Programs/UnrealGameSync/UnrealGameSync/Forms/ArgumentsWindow.cs
@@ -25,6 +25,8 @@ namespace UnrealGameSync
 		public ArgumentsWindow(List<Tuple<string, bool>> arguments, bool promptBeforeLaunch)
 		{
 			InitializeComponent();
+			Font = new System.Drawing.Font("Segoe UI", 8.25F, System.Drawing.FontStyle.Regular, System.Drawing.GraphicsUnit.Point, ((byte)(0)));
+
 			ActiveControl = ArgumentsList;
 
 			ArgumentsList.Font = SystemFonts.IconTitleFont;
diff --git a/Engine/Source/Programs/UnrealGameSync/UnrealGameSync/Forms/AutomatedBuildWindow.cs b/Engine/Source/Programs/UnrealGameSync/UnrealGameSync/Forms/AutomatedBuildWindow.cs
--- a/Engine/Source/Programs/UnrealGameSync/UnrealGameSync/Forms/AutomatedBuildWindow.cs
+++ b/Engine/Source/Programs/UnrealGameSync/UnrealGameSync/Forms/AutomatedBuildWindow.cs
@@ -42,6 +42,7 @@ namespace UnrealGameSync.Forms
 			_serviceProvider = serviceProvider;
 
 			InitializeComponent();
+			Font = new System.Drawing.Font("Segoe UI", 8.25F, System.Drawing.FontStyle.Regular, System.Drawing.GraphicsUnit.Point, ((byte)(0)));
 
 			ActiveControl = WorkspaceNameTextBox;
 
diff --git a/Engine/Source/Programs/UnrealGameSync/UnrealGameSync/Forms/AutomatedSyncWindow.cs b/Engine/Source/Programs/UnrealGameSync/UnrealGameSync/Forms/AutomatedSyncWindow.cs
--- a/Engine/Source/Programs/UnrealGameSync/UnrealGameSync/Forms/AutomatedSyncWindow.cs
+++ b/Engine/Source/Programs/UnrealGameSync/UnrealGameSync/Forms/AutomatedSyncWindow.cs
@@ -107,6 +107,7 @@ namespace UnrealGameSync
 			_serviceProvider = serviceProvider;
 
 			InitializeComponent();
+			Font = new System.Drawing.Font("Segoe UI", 8.25F, System.Drawing.FontStyle.Regular, System.Drawing.GraphicsUnit.Point, ((byte)(0)));
 
 			ActiveControl = WorkspaceNameTextBox;
 
diff --git a/Engine/Source/Programs/UnrealGameSync/UnrealGameSync/Forms/BuildStepWindow.cs b/Engine/Source/Programs/UnrealGameSync/UnrealGameSync/Forms/BuildStepWindow.cs
--- a/Engine/Source/Programs/UnrealGameSync/UnrealGameSync/Forms/BuildStepWindow.cs
+++ b/Engine/Source/Programs/UnrealGameSync/UnrealGameSync/Forms/BuildStepWindow.cs
@@ -26,6 +26,7 @@ namespace UnrealGameSync
 			_variables = inVariables;
 
 			InitializeComponent();
+			Font = new System.Drawing.Font("Segoe UI", 8.25F, System.Drawing.FontStyle.Regular, System.Drawing.GraphicsUnit.Point, ((byte)(0)));
 
 			MaximumSize = new System.Drawing.Size(32768, Height);
 			MinimumSize = new System.Drawing.Size(Width, Height);
diff --git a/Engine/Source/Programs/UnrealGameSync/UnrealGameSync/Forms/ChangelistWindow.cs b/Engine/Source/Programs/UnrealGameSync/UnrealGameSync/Forms/ChangelistWindow.cs
--- a/Engine/Source/Programs/UnrealGameSync/UnrealGameSync/Forms/ChangelistWindow.cs
+++ b/Engine/Source/Programs/UnrealGameSync/UnrealGameSync/Forms/ChangelistWindow.cs
@@ -12,6 +12,7 @@ namespace UnrealGameSync
 		public ChangelistWindow(int changeNumber)
 		{
 			InitializeComponent();
+			Font = new System.Drawing.Font("Segoe UI", 8.25F, System.Drawing.FontStyle.Regular, System.Drawing.GraphicsUnit.Point, ((byte)(0)));
 
 			MaximumSize = new System.Drawing.Size(32768, Height);
 			MinimumSize = new System.Drawing.Size(Width, Height);
diff --git a/Engine/Source/Programs/UnrealGameSync/UnrealGameSync/Forms/CleanWorkspaceWindow.cs b/Engine/Source/Programs/UnrealGameSync/UnrealGameSync/Forms/CleanWorkspaceWindow.cs
--- a/Engine/Source/Programs/UnrealGameSync/UnrealGameSync/Forms/CleanWorkspaceWindow.cs
+++ b/Engine/Source/Programs/UnrealGameSync/UnrealGameSync/Forms/CleanWorkspaceWindow.cs
@@ -113,6 +113,7 @@ namespace UnrealGameSync
 			_logger = logger;
 
 			InitializeComponent();
+			Font = new System.Drawing.Font("Segoe UI", 8.25F, System.Drawing.FontStyle.Regular, System.Drawing.GraphicsUnit.Point, ((byte)(0)));
 		}
 
 		public static void DoClean(IWin32Window owner, IPerforceSettings perforceSettings, DirectoryReference localRootPath, string clientRootPath, IReadOnlyList<string> syncPaths, string[] extraSafeToDeleteFolders, string[] extraSafeToDeleteExtensions, ILogger<CleanWorkspaceWindow> logger)
diff --git a/Engine/Source/Programs/UnrealGameSync/UnrealGameSync/Forms/ClobberWindow.cs b/Engine/Source/Programs/UnrealGameSync/UnrealGameSync/Forms/ClobberWindow.cs
--- a/Engine/Source/Programs/UnrealGameSync/UnrealGameSync/Forms/ClobberWindow.cs
+++ b/Engine/Source/Programs/UnrealGameSync/UnrealGameSync/Forms/ClobberWindow.cs
@@ -17,6 +17,7 @@ namespace UnrealGameSync
 			bool uncontrolledChangeFound = false;
 
 			InitializeComponent();
+			Font = new System.Drawing.Font("Segoe UI", 8.25F, System.Drawing.FontStyle.Regular, System.Drawing.GraphicsUnit.Point, ((byte)(0)));
 
 			_filesToClobber = inFilesToClobber;
 
diff --git a/Engine/Source/Programs/UnrealGameSync/UnrealGameSync/Forms/ConnectWindow.cs b/Engine/Source/Programs/UnrealGameSync/UnrealGameSync/Forms/ConnectWindow.cs
--- a/Engine/Source/Programs/UnrealGameSync/UnrealGameSync/Forms/ConnectWindow.cs
+++ b/Engine/Source/Programs/UnrealGameSync/UnrealGameSync/Forms/ConnectWindow.cs
@@ -16,6 +16,7 @@ namespace UnrealGameSync
 		private ConnectWindow(IPerforceSettings defaultPerforceSettings, string? serverAndPortOverride, string? userNameOverride, IServiceProvider serviceProvider)
 		{
 			InitializeComponent();
+			Font = new System.Drawing.Font("Segoe UI", 8.25F, System.Drawing.FontStyle.Regular, System.Drawing.GraphicsUnit.Point, ((byte)(0)));
 
 			_defaultPerforceSettings = defaultPerforceSettings;
 			_serviceProvider = serviceProvider;
diff --git a/Engine/Source/Programs/UnrealGameSync/UnrealGameSync/Forms/DeleteWindow.cs b/Engine/Source/Programs/UnrealGameSync/UnrealGameSync/Forms/DeleteWindow.cs
--- a/Engine/Source/Programs/UnrealGameSync/UnrealGameSync/Forms/DeleteWindow.cs
+++ b/Engine/Source/Programs/UnrealGameSync/UnrealGameSync/Forms/DeleteWindow.cs
@@ -13,6 +13,7 @@ namespace UnrealGameSync
 		public DeleteWindow(Dictionary<string, bool> inFilesToDelete)
 		{
 			InitializeComponent();
+			Font = new System.Drawing.Font("Segoe UI", 8.25F, System.Drawing.FontStyle.Regular, System.Drawing.GraphicsUnit.Point, ((byte)(0)));
 
 			_filesToDelete = inFilesToDelete;
 
diff --git a/Engine/Source/Programs/UnrealGameSync/UnrealGameSync/Forms/DiagnosticsWindow.cs b/Engine/Source/Programs/UnrealGameSync/UnrealGameSync/Forms/DiagnosticsWindow.cs
--- a/Engine/Source/Programs/UnrealGameSync/UnrealGameSync/Forms/DiagnosticsWindow.cs
+++ b/Engine/Source/Programs/UnrealGameSync/UnrealGameSync/Forms/DiagnosticsWindow.cs
@@ -20,6 +20,7 @@ namespace UnrealGameSync
 		public DiagnosticsWindow(DirectoryReference inAppDataFolder, DirectoryReference inWorkspaceDataFolder, string inDiagnosticsText, IEnumerable<FileReference> inExtraFiles)
 		{
 			InitializeComponent();
+			Font = new System.Drawing.Font("Segoe UI", 8.25F, System.Drawing.FontStyle.Regular, System.Drawing.GraphicsUnit.Point, ((byte)(0)));
 
 			_appDataFolder = inAppDataFolder;
 			_workspaceDataFolder = inWorkspaceDataFolder;
diff --git a/Engine/Source/Programs/UnrealGameSync/UnrealGameSync/Forms/FailedToDeleteWindow.cs b/Engine/Source/Programs/UnrealGameSync/UnrealGameSync/Forms/FailedToDeleteWindow.cs
--- a/Engine/Source/Programs/UnrealGameSync/UnrealGameSync/Forms/FailedToDeleteWindow.cs
+++ b/Engine/Source/Programs/UnrealGameSync/UnrealGameSync/Forms/FailedToDeleteWindow.cs
@@ -1,4 +1,4 @@
-﻿// Copyright Epic Games, Inc. All Rights Reserved.
+// Copyright Epic Games, Inc. All Rights Reserved.
 
 using System;
 using System.Collections.Generic;
@@ -17,6 +17,7 @@ namespace UnrealGameSync
 		public FailedToDeleteWindow()
 		{
 			InitializeComponent();
+			Font = new System.Drawing.Font("Segoe UI", 8.25F, System.Drawing.FontStyle.Regular, System.Drawing.GraphicsUnit.Point, ((byte)(0)));
 		}
 	}
 }
diff --git a/Engine/Source/Programs/UnrealGameSync/UnrealGameSync/Forms/IssueAlertWindow.cs b/Engine/Source/Programs/UnrealGameSync/UnrealGameSync/Forms/IssueAlertWindow.cs
--- a/Engine/Source/Programs/UnrealGameSync/UnrealGameSync/Forms/IssueAlertWindow.cs
+++ b/Engine/Source/Programs/UnrealGameSync/UnrealGameSync/Forms/IssueAlertWindow.cs
@@ -32,6 +32,7 @@ namespace UnrealGameSync
 			Issue = issue;
 
 			InitializeComponent();
+			Font = new System.Drawing.Font("Segoe UI", 8.25F, System.Drawing.FontStyle.Regular, System.Drawing.GraphicsUnit.Point, ((byte)(0)));
 
 			SetIssue(issue, reason);
 		}
diff --git a/Engine/Source/Programs/UnrealGameSync/UnrealGameSync/Forms/IssueBrowserWindow.cs b/Engine/Source/Programs/UnrealGameSync/UnrealGameSync/Forms/IssueBrowserWindow.cs
--- a/Engine/Source/Programs/UnrealGameSync/UnrealGameSync/Forms/IssueBrowserWindow.cs
+++ b/Engine/Source/Programs/UnrealGameSync/UnrealGameSync/Forms/IssueBrowserWindow.cs
@@ -42,16 +42,14 @@ namespace UnrealGameSync
 			issueMonitor.AddRef();
 
 			InitializeComponent();
+			Font = new System.Drawing.Font("Segoe UI", 8.25F, System.Drawing.FontStyle.Regular, System.Drawing.GraphicsUnit.Point, ((byte)(0)));
 
-			using (Graphics graphics = Graphics.FromHwnd(IntPtr.Zero))
+			float dpiScaleX = DeviceDpi / 96.0f;
+			foreach (ColumnHeader? column in IssueListView.Columns)
 			{
-				float dpiScaleX = graphics.DpiX / 96.0f;
-				foreach (ColumnHeader? column in IssueListView.Columns)
+				if (column != null)
 				{
-					if (column != null)
-					{
-						column.Width = (int)(column.Width * dpiScaleX);
-					}
+					column.Width = (int)(column.Width * dpiScaleX);
 				}
 			}
 
diff --git a/Engine/Source/Programs/UnrealGameSync/UnrealGameSync/Forms/IssueDetailsWindow.cs b/Engine/Source/Programs/UnrealGameSync/UnrealGameSync/Forms/IssueDetailsWindow.cs
--- a/Engine/Source/Programs/UnrealGameSync/UnrealGameSync/Forms/IssueDetailsWindow.cs
+++ b/Engine/Source/Programs/UnrealGameSync/UnrealGameSync/Forms/IssueDetailsWindow.cs
@@ -293,6 +293,7 @@ namespace UnrealGameSync
 			_mainThreadSynchronizationContext = SynchronizationContext.Current!;
 
 			InitializeComponent();
+			Font = new System.Drawing.Font("Segoe UI", 8.25F, System.Drawing.FontStyle.Regular, System.Drawing.GraphicsUnit.Point, ((byte)(0)));
 
 			Text = String.Format("Issue {0}", issue.Id);
 			_statusElementResources = new StatusElementResources(BuildListView.Font);
diff --git a/Engine/Source/Programs/UnrealGameSync/UnrealGameSync/Forms/IssueFixedWindow.cs b/Engine/Source/Programs/UnrealGameSync/UnrealGameSync/Forms/IssueFixedWindow.cs
--- a/Engine/Source/Programs/UnrealGameSync/UnrealGameSync/Forms/IssueFixedWindow.cs
+++ b/Engine/Source/Programs/UnrealGameSync/UnrealGameSync/Forms/IssueFixedWindow.cs
@@ -114,6 +114,7 @@ namespace UnrealGameSync
 		public IssueFixedWindow(IPerforceSettings perforceSettings, int initialChangeNumber, IServiceProvider serviceProvider)
 		{
 			InitializeComponent();
+			Font = new System.Drawing.Font("Segoe UI", 8.25F, System.Drawing.FontStyle.Regular, System.Drawing.GraphicsUnit.Point, ((byte)(0)));
 
 			_perforceSettings = perforceSettings;
 			_worker = new FindChangesWorker(perforceSettings, PopulateChanges, serviceProvider.GetRequiredService<ILogger<FindChangesWorker>>());
diff --git a/Engine/Source/Programs/UnrealGameSync/UnrealGameSync/Forms/IssueSettingsWindow.cs b/Engine/Source/Programs/UnrealGameSync/UnrealGameSync/Forms/IssueSettingsWindow.cs
--- a/Engine/Source/Programs/UnrealGameSync/UnrealGameSync/Forms/IssueSettingsWindow.cs
+++ b/Engine/Source/Programs/UnrealGameSync/UnrealGameSync/Forms/IssueSettingsWindow.cs
@@ -18,6 +18,7 @@ namespace UnrealGameSync
 			_logger = logger;
 
 			InitializeComponent();
+			Font = new System.Drawing.Font("Segoe UI", 8.25F, System.Drawing.FontStyle.Regular, System.Drawing.GraphicsUnit.Point, ((byte)(0)));
 
 			if (settings.NotifyProjects.Count == 0)
 			{
diff --git a/Engine/Source/Programs/UnrealGameSync/UnrealGameSync/Forms/LauncherUpdateWindow.cs b/Engine/Source/Programs/UnrealGameSync/UnrealGameSync/Forms/LauncherUpdateWindow.cs
--- a/Engine/Source/Programs/UnrealGameSync/UnrealGameSync/Forms/LauncherUpdateWindow.cs
+++ b/Engine/Source/Programs/UnrealGameSync/UnrealGameSync/Forms/LauncherUpdateWindow.cs
@@ -9,6 +9,7 @@ namespace UnrealGameSync.Forms
 		public LauncherUpdateWindow()
 		{
 			InitializeComponent();
+			Font = new System.Drawing.Font("Segoe UI", 8.25F, System.Drawing.FontStyle.Regular, System.Drawing.GraphicsUnit.Point, ((byte)(0)));
 		}
 	}
 }
diff --git a/Engine/Source/Programs/UnrealGameSync/UnrealGameSync/Forms/LeaveCommentWindow.cs b/Engine/Source/Programs/UnrealGameSync/UnrealGameSync/Forms/LeaveCommentWindow.cs
--- a/Engine/Source/Programs/UnrealGameSync/UnrealGameSync/Forms/LeaveCommentWindow.cs
+++ b/Engine/Source/Programs/UnrealGameSync/UnrealGameSync/Forms/LeaveCommentWindow.cs
@@ -9,6 +9,7 @@ namespace UnrealGameSync
 		public LeaveCommentWindow()
 		{
 			InitializeComponent();
+			Font = new System.Drawing.Font("Segoe UI", 8.25F, System.Drawing.FontStyle.Regular, System.Drawing.GraphicsUnit.Point, ((byte)(0)));
 		}
 	}
 }
diff --git a/Engine/Source/Programs/UnrealGameSync/UnrealGameSync/Forms/MainWindow.cs b/Engine/Source/Programs/UnrealGameSync/UnrealGameSync/Forms/MainWindow.cs
--- a/Engine/Source/Programs/UnrealGameSync/UnrealGameSync/Forms/MainWindow.cs
+++ b/Engine/Source/Programs/UnrealGameSync/UnrealGameSync/Forms/MainWindow.cs
@@ -130,6 +130,7 @@ namespace UnrealGameSync
 			_logger.LogInformation("Opening Main Window for {NumProject} projects. Last Project {LastProject}", startupTasks.Count, settings.LastProject);
 
 			InitializeComponent();
+			Font = new Font("Segoe UI", 8.25F, FontStyle.Regular, GraphicsUnit.Point, ((byte)(0)));
 
 			_updateMonitor = inUpdateMonitor;
 			_mainThreadSynchronizationContext = SynchronizationContext.Current!;
diff --git a/Engine/Source/Programs/UnrealGameSync/UnrealGameSync/Forms/ModalTaskWindow.cs b/Engine/Source/Programs/UnrealGameSync/UnrealGameSync/Forms/ModalTaskWindow.cs
--- a/Engine/Source/Programs/UnrealGameSync/UnrealGameSync/Forms/ModalTaskWindow.cs
+++ b/Engine/Source/Programs/UnrealGameSync/UnrealGameSync/Forms/ModalTaskWindow.cs
@@ -15,6 +15,7 @@ namespace UnrealGameSync
 		public ModalTaskWindow(string inTitle, string inMessage, FormStartPosition inStartPosition, Task inTask, CancellationTokenSource inCancellationSource)//, Func<CancellationToken, Task> InTaskFunc)
 		{
 			InitializeComponent();
+			Font = new System.Drawing.Font("Segoe UI", 8.25F, System.Drawing.FontStyle.Regular, System.Drawing.GraphicsUnit.Point, ((byte)(0)));
 
 			Text = inTitle;
 			MessageLabel.Text = inMessage;
diff --git a/Engine/Source/Programs/UnrealGameSync/UnrealGameSync/Forms/ModifyBuildStepsWindow.cs b/Engine/Source/Programs/UnrealGameSync/UnrealGameSync/Forms/ModifyBuildStepsWindow.cs
--- a/Engine/Source/Programs/UnrealGameSync/UnrealGameSync/Forms/ModifyBuildStepsWindow.cs
+++ b/Engine/Source/Programs/UnrealGameSync/UnrealGameSync/Forms/ModifyBuildStepsWindow.cs
@@ -27,15 +27,12 @@ namespace UnrealGameSync
 			_variables = inVariables;
 
 			InitializeComponent();
+			Font = new System.Drawing.Font("Segoe UI", 8.25F, System.Drawing.FontStyle.Regular, System.Drawing.GraphicsUnit.Point, ((byte)(0)));
 
-			using(Graphics graphics = Graphics.FromHwnd(IntPtr.Zero))
-			{
-				float dpiScaleX = graphics.DpiX / 96.0f;
-
-				NormalSyncColumn.Width = (int)(104 * dpiScaleX);
-				ScheduledSyncColumn.Width = (int)(104 * dpiScaleX);
-				DescriptionColumn.Width = BuildStepList.ClientSize.Width - NormalSyncColumn.Width - ScheduledSyncColumn.Width - 10;
-			}
+			float dpiScaleX = DeviceDpi / 96.0f;
+			NormalSyncColumn.Width = (int)(104 * dpiScaleX);
+			ScheduledSyncColumn.Width = (int)(104 * dpiScaleX);
+			DescriptionColumn.Width = BuildStepList.ClientSize.Width - NormalSyncColumn.Width - ScheduledSyncColumn.Width - 10;
 
 			BuildStepList.Font = SystemFonts.IconTitleFont;
 		}
diff --git a/Engine/Source/Programs/UnrealGameSync/UnrealGameSync/Forms/NewWorkspaceWindow.cs b/Engine/Source/Programs/UnrealGameSync/UnrealGameSync/Forms/NewWorkspaceWindow.cs
--- a/Engine/Source/Programs/UnrealGameSync/UnrealGameSync/Forms/NewWorkspaceWindow.cs
+++ b/Engine/Source/Programs/UnrealGameSync/UnrealGameSync/Forms/NewWorkspaceWindow.cs
@@ -93,6 +93,7 @@ namespace UnrealGameSync
 		private NewWorkspaceWindow(IPerforceSettings perforceSettings, string? forceStream, string? defaultStream, InfoRecord info, List<ClientsRecord> clients, IServiceProvider serviceProvider)
 		{
 			InitializeComponent();
+			Font = new System.Drawing.Font("Segoe UI", 8.25F, System.Drawing.FontStyle.Regular, System.Drawing.GraphicsUnit.Point, ((byte)(0)));
 
 			_perforceSettings = perforceSettings;
 			_info = info;
diff --git a/Engine/Source/Programs/UnrealGameSync/UnrealGameSync/Forms/NotificationWindow.cs b/Engine/Source/Programs/UnrealGameSync/UnrealGameSync/Forms/NotificationWindow.cs
--- a/Engine/Source/Programs/UnrealGameSync/UnrealGameSync/Forms/NotificationWindow.cs
+++ b/Engine/Source/Programs/UnrealGameSync/UnrealGameSync/Forms/NotificationWindow.cs
@@ -44,6 +44,7 @@ namespace UnrealGameSync
 		{
 			_logoBitmap = new Bitmap(inLogo);
 			InitializeComponent();
+			Font = new System.Drawing.Font("Segoe UI", 8.25F, System.Drawing.FontStyle.Regular, System.Drawing.GraphicsUnit.Point, ((byte)(0)));
 		}
 
 		public void Show(NotificationType inType, string inCaption, string inMessage)
diff --git a/Engine/Source/Programs/UnrealGameSync/UnrealGameSync/Forms/OpenProjectWindow.cs b/Engine/Source/Programs/UnrealGameSync/UnrealGameSync/Forms/OpenProjectWindow.cs
--- a/Engine/Source/Programs/UnrealGameSync/UnrealGameSync/Forms/OpenProjectWindow.cs
+++ b/Engine/Source/Programs/UnrealGameSync/UnrealGameSync/Forms/OpenProjectWindow.cs
@@ -30,6 +30,7 @@ namespace UnrealGameSync
 		private OpenProjectWindow(UserSelectedProjectSettings? project, UserSettings settings, IPerforceSettings defaultSettings, IServiceProvider serviceProvider, ILogger logger)
 		{
 			InitializeComponent();
+			Font = new System.Drawing.Font("Segoe UI", 8.25F, System.Drawing.FontStyle.Regular, System.Drawing.GraphicsUnit.Point, ((byte)(0)));
 
 			_settings = settings;
 			_openProjectInfo = null;
diff --git a/Engine/Source/Programs/UnrealGameSync/UnrealGameSync/Forms/PasswordWindow.cs b/Engine/Source/Programs/UnrealGameSync/UnrealGameSync/Forms/PasswordWindow.cs
--- a/Engine/Source/Programs/UnrealGameSync/UnrealGameSync/Forms/PasswordWindow.cs
+++ b/Engine/Source/Programs/UnrealGameSync/UnrealGameSync/Forms/PasswordWindow.cs
@@ -12,6 +12,7 @@ namespace UnrealGameSync
 		public PasswordWindow(string prompt, string password)
 		{
 			InitializeComponent();
+			Font = new System.Drawing.Font("Segoe UI", 8.25F, System.Drawing.FontStyle.Regular, System.Drawing.GraphicsUnit.Point, ((byte)(0)));
 
 			Password = password;
 
diff --git a/Engine/Source/Programs/UnrealGameSync/UnrealGameSync/Forms/PerforceSyncSettingsWindow.cs b/Engine/Source/Programs/UnrealGameSync/UnrealGameSync/Forms/PerforceSyncSettingsWindow.cs
--- a/Engine/Source/Programs/UnrealGameSync/UnrealGameSync/Forms/PerforceSyncSettingsWindow.cs
+++ b/Engine/Source/Programs/UnrealGameSync/UnrealGameSync/Forms/PerforceSyncSettingsWindow.cs
@@ -17,6 +17,7 @@ namespace UnrealGameSync
 			_logger = logger;
 
 			InitializeComponent();
+			Font = new System.Drawing.Font("Segoe UI", 8.25F, System.Drawing.FontStyle.Regular, System.Drawing.GraphicsUnit.Point, ((byte)(0)));
 		}
 
 		private void PerforceSettingsWindow_Load(object sender, EventArgs e)
diff --git a/Engine/Source/Programs/UnrealGameSync/UnrealGameSync/Forms/ProgramsRunningWindow.cs b/Engine/Source/Programs/UnrealGameSync/UnrealGameSync/Forms/ProgramsRunningWindow.cs
--- a/Engine/Source/Programs/UnrealGameSync/UnrealGameSync/Forms/ProgramsRunningWindow.cs
+++ b/Engine/Source/Programs/UnrealGameSync/UnrealGameSync/Forms/ProgramsRunningWindow.cs
@@ -20,6 +20,7 @@ namespace UnrealGameSync
 		public ProgramsRunningWindow(Func<FileReference[]> enumeratePrograms, FileReference[] programs)
 		{
 			InitializeComponent();
+			Font = new System.Drawing.Font("Segoe UI", 8.25F, System.Drawing.FontStyle.Regular, System.Drawing.GraphicsUnit.Point, ((byte)(0)));
 
 			_programs = programs.OrderBy(x => x).ToArray();
 			_enumeratePrograms = enumeratePrograms;
diff --git a/Engine/Source/Programs/UnrealGameSync/UnrealGameSync/Forms/ScheduleWindow.cs b/Engine/Source/Programs/UnrealGameSync/UnrealGameSync/Forms/ScheduleWindow.cs
--- a/Engine/Source/Programs/UnrealGameSync/UnrealGameSync/Forms/ScheduleWindow.cs
+++ b/Engine/Source/Programs/UnrealGameSync/UnrealGameSync/Forms/ScheduleWindow.cs
@@ -22,6 +22,7 @@ namespace UnrealGameSync
 			Dictionary<UserSelectedProjectSettings, List<LatestChangeType>> inProjectToLatestChangeTypes)
 		{
 			InitializeComponent();
+			Font = new System.Drawing.Font("Segoe UI", 8.25F, System.Drawing.FontStyle.Regular, System.Drawing.GraphicsUnit.Point, ((byte)(0)));
 
 			EnableCheckBox.Checked = inEnabled;
 
diff --git a/Engine/Source/Programs/UnrealGameSync/UnrealGameSync/Forms/SdkInfoWindow.cs b/Engine/Source/Programs/UnrealGameSync/UnrealGameSync/Forms/SdkInfoWindow.cs
--- a/Engine/Source/Programs/UnrealGameSync/UnrealGameSync/Forms/SdkInfoWindow.cs
+++ b/Engine/Source/Programs/UnrealGameSync/UnrealGameSync/Forms/SdkInfoWindow.cs
@@ -50,6 +50,7 @@ namespace UnrealGameSync
 		public SdkInfoWindow(string[] sdkInfoEntries, Dictionary<string, string> variables, Font badgeFont)
 		{
 			InitializeComponent();
+			Font = new System.Drawing.Font("Segoe UI", 8.25F, System.Drawing.FontStyle.Regular, System.Drawing.GraphicsUnit.Point, ((byte)(0)));
 
 			_badgeFont = badgeFont;
 
diff --git a/Engine/Source/Programs/UnrealGameSync/UnrealGameSync/Forms/SelectProjectFromWorkspaceWindow.cs b/Engine/Source/Programs/UnrealGameSync/UnrealGameSync/Forms/SelectProjectFromWorkspaceWindow.cs
--- a/Engine/Source/Programs/UnrealGameSync/UnrealGameSync/Forms/SelectProjectFromWorkspaceWindow.cs
+++ b/Engine/Source/Programs/UnrealGameSync/UnrealGameSync/Forms/SelectProjectFromWorkspaceWindow.cs
@@ -94,7 +94,8 @@ namespace UnrealGameSync
 		public SelectProjectFromWorkspaceWindow(string workspaceName, List<string> projectFiles, string selectedProjectFile)
 		{
 			InitializeComponent();
-			
+			Font = new System.Drawing.Font("Segoe UI", 8.25F, System.Drawing.FontStyle.Regular, System.Drawing.GraphicsUnit.Point, ((byte)(0)));
+
 			_projectFiles = projectFiles;
 			_selectedProjectFile = selectedProjectFile;
 
diff --git a/Engine/Source/Programs/UnrealGameSync/UnrealGameSync/Forms/SelectStreamWindow.cs b/Engine/Source/Programs/UnrealGameSync/UnrealGameSync/Forms/SelectStreamWindow.cs
--- a/Engine/Source/Programs/UnrealGameSync/UnrealGameSync/Forms/SelectStreamWindow.cs
+++ b/Engine/Source/Programs/UnrealGameSync/UnrealGameSync/Forms/SelectStreamWindow.cs
@@ -73,6 +73,7 @@ namespace UnrealGameSync
 		private SelectStreamWindow(List<StreamsRecord> streams, string? streamName)
 		{
 			InitializeComponent();
+			Font = new System.Drawing.Font("Segoe UI", 8.25F, System.Drawing.FontStyle.Regular, System.Drawing.GraphicsUnit.Point, ((byte)(0)));
 
 			_selectedStream = streamName;
 
diff --git a/Engine/Source/Programs/UnrealGameSync/UnrealGameSync/Forms/SelectUserWindow.cs b/Engine/Source/Programs/UnrealGameSync/UnrealGameSync/Forms/SelectUserWindow.cs
--- a/Engine/Source/Programs/UnrealGameSync/UnrealGameSync/Forms/SelectUserWindow.cs
+++ b/Engine/Source/Programs/UnrealGameSync/UnrealGameSync/Forms/SelectUserWindow.cs
@@ -30,6 +30,7 @@ namespace UnrealGameSync
 		private SelectUserWindow(List<UsersRecord> users, int selectedUserIndex)
 		{
 			InitializeComponent();
+			Font = new System.Drawing.Font("Segoe UI", 8.25F, System.Drawing.FontStyle.Regular, System.Drawing.GraphicsUnit.Point, ((byte)(0)));
 
 			_selectedUserIndex = selectedUserIndex;
 			_users = users;
diff --git a/Engine/Source/Programs/UnrealGameSync/UnrealGameSync/Forms/SelectWorkspaceWindow.cs b/Engine/Source/Programs/UnrealGameSync/UnrealGameSync/Forms/SelectWorkspaceWindow.cs
--- a/Engine/Source/Programs/UnrealGameSync/UnrealGameSync/Forms/SelectWorkspaceWindow.cs
+++ b/Engine/Source/Programs/UnrealGameSync/UnrealGameSync/Forms/SelectWorkspaceWindow.cs
@@ -43,6 +43,7 @@ namespace UnrealGameSync
 		private SelectWorkspaceWindow(InfoRecord info, List<ClientsRecord> clients, string? workspaceName)
 		{
 			InitializeComponent();
+			Font = new System.Drawing.Font("Segoe UI", 8.25F, System.Drawing.FontStyle.Regular, System.Drawing.GraphicsUnit.Point, ((byte)(0)));
 
 			_info = info;
 			_clients = clients;
diff --git a/Engine/Source/Programs/UnrealGameSync/UnrealGameSync/Forms/SyncFilter.cs b/Engine/Source/Programs/UnrealGameSync/UnrealGameSync/Forms/SyncFilter.cs
--- a/Engine/Source/Programs/UnrealGameSync/UnrealGameSync/Forms/SyncFilter.cs
+++ b/Engine/Source/Programs/UnrealGameSync/UnrealGameSync/Forms/SyncFilter.cs
@@ -17,6 +17,7 @@ namespace UnrealGameSync
 		public SyncFilter(Dictionary<Guid, WorkspaceSyncCategory> uniqueIdToCategory, FilterSettings globalFilter, FilterSettings workspaceFilter, ConfigSection? perforceSection)
 		{
 			InitializeComponent();
+			Font = new System.Drawing.Font("Segoe UI", 8.25F, System.Drawing.FontStyle.Regular, System.Drawing.GraphicsUnit.Point, ((byte)(0)));
 
 			_uniqueIdToCategory = uniqueIdToCategory;
 			GlobalFilter = globalFilter;
diff --git a/Engine/Source/Programs/UnrealGameSync/UnrealGameSync/Forms/SyncFilterSyntax.cs b/Engine/Source/Programs/UnrealGameSync/UnrealGameSync/Forms/SyncFilterSyntax.cs
--- a/Engine/Source/Programs/UnrealGameSync/UnrealGameSync/Forms/SyncFilterSyntax.cs
+++ b/Engine/Source/Programs/UnrealGameSync/UnrealGameSync/Forms/SyncFilterSyntax.cs
@@ -1,4 +1,4 @@
-﻿// Copyright Epic Games, Inc. All Rights Reserved.
+// Copyright Epic Games, Inc. All Rights Reserved.
 
 using System;
 using System.Collections.Generic;
@@ -17,6 +17,7 @@ namespace UnrealGameSync
 		public SyncFilterSyntax()
 		{
 			InitializeComponent();
+			Font = new System.Drawing.Font("Segoe UI", 8.25F, System.Drawing.FontStyle.Regular, System.Drawing.GraphicsUnit.Point, ((byte)(0)));
 		}
 	}
 }
diff --git a/Engine/Source/Programs/UnrealGameSync/UnrealGameSync/Forms/VariablesWindow.cs b/Engine/Source/Programs/UnrealGameSync/UnrealGameSync/Forms/VariablesWindow.cs
--- a/Engine/Source/Programs/UnrealGameSync/UnrealGameSync/Forms/VariablesWindow.cs
+++ b/Engine/Source/Programs/UnrealGameSync/UnrealGameSync/Forms/VariablesWindow.cs
@@ -28,6 +28,7 @@ namespace UnrealGameSync
 		public VariablesWindow(IReadOnlyDictionary<string, string> variables)
 		{
 			InitializeComponent();
+			Font = new System.Drawing.Font("Segoe UI", 8.25F, System.Drawing.FontStyle.Regular, System.Drawing.GraphicsUnit.Point, ((byte)(0)));
 
 			ListViewGroup currentProjectGroup = new ListViewGroup("Current Project");
 			MacrosList.Groups.Add(currentProjectGroup);
