UGS: Fix NET 6 static analysis warnings.

diff --git a/Engine/Source/Programs/UnrealGameSync/UnrealGameSync/AutomationServer.cs b/Engine/Source/Programs/UnrealGameSync/UnrealGameSync/AutomationServer.cs
--- a/Engine/Source/Programs/UnrealGameSync/UnrealGameSync/AutomationServer.cs
+++ b/Engine/Source/Programs/UnrealGameSync/UnrealGameSync/AutomationServer.cs
@@ -251,7 +251,7 @@ namespace UnrealGameSync
 				{
 					try
 					{
-						TcpClient client = await listener.AcceptTcpClientAsync().ConfigureAwait(false);
+						TcpClient client = await listener.AcceptTcpClientAsync(cancellationToken).ConfigureAwait(false);
 						try
 						{
 							_logger.LogInformation("Accepted connection from {Remote}", client.Client.RemoteEndPoint);
diff --git a/Engine/Source/Programs/UnrealGameSync/UnrealGameSync/Controls/LogControl.cs b/Engine/Source/Programs/UnrealGameSync/UnrealGameSync/Controls/LogControl.cs
--- a/Engine/Source/Programs/UnrealGameSync/UnrealGameSync/Controls/LogControl.cs
+++ b/Engine/Source/Programs/UnrealGameSync/UnrealGameSync/Controls/LogControl.cs
@@ -175,6 +175,12 @@ namespace UnrealGameSync
 				_updateTimer = null;
 			}
 
+			if (_selectionScrollTimer != null)
+			{
+				_selectionScrollTimer.Dispose();
+				_selectionScrollTimer = null;
+			}
+
 			CloseFile();
 		}
 
diff --git a/Engine/Source/Programs/UnrealGameSync/UnrealGameSync/Controls/WorkspaceControl.cs b/Engine/Source/Programs/UnrealGameSync/UnrealGameSync/Controls/WorkspaceControl.cs
--- a/Engine/Source/Programs/UnrealGameSync/UnrealGameSync/Controls/WorkspaceControl.cs
+++ b/Engine/Source/Programs/UnrealGameSync/UnrealGameSync/Controls/WorkspaceControl.cs
@@ -227,9 +227,13 @@ namespace UnrealGameSync
 //		JupiterMonitor _jupiterMonitor;
 		PerforceMonitor _perforceMonitor;
 		Workspace _workspace;
+#pragma warning disable CA2213 //warning CA2213: 'WorkspaceControl' contains field '_issueMonitor' that is of IDisposable type 'IssueMonitor', but it is never disposed. Change the Dispose method on 'WorkspaceControl' to call Close or Dispose on this field.
 		IssueMonitor _issueMonitor;
+#pragma warning restore CA2213
 		EventMonitor _eventMonitor;
+#pragma warning disable CA2213 // warning CA2213: 'WorkspaceControl' contains field '_updateTimer' that is of IDisposable type 'Timer', but it is never disposed. Change the Dispose method on 'WorkspaceControl' to call Close or Dispose on this field.
 		readonly System.Windows.Forms.Timer _updateTimer;
+#pragma warning restore CA2213
 		HashSet<int> _promotedChangeNumbers = new HashSet<int>();
 		List<int> _listIndexToChangeIndex = new List<int>();
 		List<int> _sortedChangeNumbers = new List<int>();
@@ -322,6 +326,7 @@ namespace UnrealGameSync
 			_updateTimer = new System.Windows.Forms.Timer();
 			_updateTimer.Interval = 30;
 			_updateTimer.Tick += TimerCallback;
+			components!.Add(_updateTimer);
 
 			UpdateCheckedBuildConfig();
 
@@ -699,8 +704,6 @@ namespace UnrealGameSync
 				components.Dispose();
 			}
 
-			_updateTimer?.Stop();
-
 			if (_startupCallbacks != null)
 			{
 				foreach (WorkspaceStartupCallback startupCallback in _startupCallbacks)
@@ -3726,6 +3729,7 @@ namespace UnrealGameSync
 
 					description.AppendFormat(" ({0})", Utility.FormatDurationMinutes((int)((issue.RetrievedAt - issue.CreatedAt).TotalMinutes + 1)));
 
+#pragma warning disable CA2000 // warning CA2000: Call System.IDisposable.Dispose on object created by 'new ToolStripMenuItem(description.ToString())' before all references to it are out of scope
 					ToolStripMenuItem item = new ToolStripMenuItem(description.ToString());
 					if (issue.FixChange > 0)
 					{
@@ -3737,6 +3741,7 @@ namespace UnrealGameSync
 					}
 					item.Click += (s, e) => { BuildHealthContextMenu_Issue_Click(issue); };
 					BuildHealthContextMenu.Items.Insert(minSeparatorIdx + idx + 1, item);
+#pragma warning restore CA2000
 				}
 			}
 
@@ -4281,9 +4286,11 @@ namespace UnrealGameSync
 							execute = Utility.ExpandVariables(execute, variables);
 							arguments = Utility.ExpandVariables(arguments ?? "", variables);
 
+#pragma warning disable CA2000 // warning CA2000: Call System.IDisposable.Dispose on object created by 'new ToolStripMenuItem(label, null, new EventHandler((o, a) => SafeProcessStart(execute, arguments)))' before all references to it are out of scope
 							ToolStripMenuItem item = new ToolStripMenuItem(label, null, new EventHandler((o, a) => SafeProcessStart(execute, arguments)));
-
 							BuildListContextMenu.Items.Insert(customToolEnd, item);
+#pragma warning restore CA2000
+
 							customToolEnd++;
 						}
 					}
@@ -5319,6 +5326,7 @@ namespace UnrealGameSync
 								string menuName = step.StatusPanelLink.Substring(0, baseMenuIdx);
 								string itemName = step.StatusPanelLink.Substring(baseMenuIdx + 1).Replace("&", "&&", StringComparison.Ordinal);
 
+#pragma warning disable CA2000 // warning CA2000: Use recommended dispose pattern to ensure that object created by 'new ToolStripMenuItem(itemName)' is disposed on all paths. If possible, wrap the creation within a 'using' statement or a 'using' declaration. Otherwise, use a try-finally pattern, with a dedicated local variable declared before the try region and an unconditional Dispose invocation on non-null value in the 'finally' region, say 'x?.Dispose()'. If the object is explicitly disposed within the try region or the dispose ownership is transfered to another object or method, assign 'null' to the local variable just after such an operation to prevent double dispose in 'finally'.
 								ToolStripMenuItem newMenuItem = new ToolStripMenuItem(itemName);
 								newMenuItem.Click += new EventHandler((sender, e) => { RunCustomTool(step, userSteps); });
 
@@ -5339,6 +5347,7 @@ namespace UnrealGameSync
 									}
 									menu.Items.Add(newMenuItem);
 								}
+#pragma warning restore CA2000
 							}
 						}
 
@@ -5701,11 +5710,13 @@ namespace UnrealGameSync
 				RecentMenu.Items.RemoveAt(2);
 			}
 
+#pragma warning disable CA2000 //  warning CA2000: Call System.IDisposable.Dispose on object created by 'new ToolStripMenuItem(recentProject.ToString(), null, new EventHandler((o, e) => _owner.RequestProjectChange(this, recentProject, true)))' before all references to it are out of scope
 			foreach (UserSelectedProjectSettings recentProject in _settings.RecentProjects)
 			{
 				ToolStripMenuItem item = new ToolStripMenuItem(recentProject.ToString(), null, new EventHandler((o, e) => _owner.RequestProjectChange(this, recentProject, true)));
 				RecentMenu.Items.Insert(RecentMenu.Items.Count - 2, item);
 			}
+#pragma warning restore CA2000
 
 			RecentMenu_Separator.Visible = (_settings.RecentProjects.Count > 0);
 			RecentMenu.Show(StatusPanel, new Point(bounds.Left, bounds.Bottom), ToolStripDropDownDirection.BelowRight);
diff --git a/Engine/Source/Programs/UnrealGameSync/UnrealGameSync/DeploymentSettings.cs b/Engine/Source/Programs/UnrealGameSync/UnrealGameSync/DeploymentSettings.cs
--- a/Engine/Source/Programs/UnrealGameSync/UnrealGameSync/DeploymentSettings.cs
+++ b/Engine/Source/Programs/UnrealGameSync/UnrealGameSync/DeploymentSettings.cs
@@ -1,6 +1,7 @@
 // Copyright Epic Games, Inc. All Rights Reserved.
 
 #pragma warning disable IDE0005
+#pragma warning disable CA1802 // warning CA1802: Field 'EnableAlerts' is declared as 'readonly' but is initialized with a constant value. Mark this field as 'const' instead.
 
 using Microsoft.Extensions.Logging;
 using Microsoft.Extensions.Logging.Abstractions;
diff --git a/Engine/Source/Programs/UnrealGameSync/UnrealGameSync/EventMonitor.cs b/Engine/Source/Programs/UnrealGameSync/UnrealGameSync/EventMonitor.cs
--- a/Engine/Source/Programs/UnrealGameSync/UnrealGameSync/EventMonitor.cs
+++ b/Engine/Source/Programs/UnrealGameSync/UnrealGameSync/EventMonitor.cs
@@ -210,7 +210,9 @@ namespace UnrealGameSync
 		readonly string _project;
 		readonly string _currentUserName;
 		readonly SynchronizationContext _synchronizationContext;
+#pragma warning disable CA2213 // warning CA2213: 'EventMonitor' contains field '_cancellationSource' that is of IDisposable type 'CancellationTokenSource', but it is never disposed. Change the Dispose method on 'EventMonitor' to call Close or Dispose on this field.
 		readonly CancellationTokenSource _cancellationSource;
+#pragma warning restore CA2213
 		Task? _workerTask;
 		readonly AsyncEvent _refreshEvent = new AsyncEvent();
 		readonly ConcurrentQueue<EventData> _outgoingEvents = new ConcurrentQueue<EventData>();
diff --git a/Engine/Source/Programs/UnrealGameSync/UnrealGameSync/Forms/BuildStepWindow.cs b/Engine/Source/Programs/UnrealGameSync/UnrealGameSync/Forms/BuildStepWindow.cs
--- a/Engine/Source/Programs/UnrealGameSync/UnrealGameSync/Forms/BuildStepWindow.cs
+++ b/Engine/Source/Programs/UnrealGameSync/UnrealGameSync/Forms/BuildStepWindow.cs
@@ -16,7 +16,9 @@ namespace UnrealGameSync
 		readonly List<string> _targetNames;
 		readonly DirectoryReference _baseDirectory;
 		readonly IReadOnlyDictionary<string, string> _variables;
+#pragma warning disable CA2213 // warning CA2213: 'BuildStepWindow' contains field '_variablesWindow' that is of IDisposable type 'VariablesWindow?', but it is never disposed. Change the Dispose method on 'BuildStepWindow' to call Close or Dispose on this field.
 		VariablesWindow? _variablesWindow;
+#pragma warning restore CA2213
 
 		public BuildStepWindow(BuildStep inTask, List<string> inTargetNames, DirectoryReference inBaseDirectory, IReadOnlyDictionary<string, string> inVariables)
 		{
@@ -264,6 +266,8 @@ namespace UnrealGameSync
 				_variablesWindow.Size = new Size(_variablesWindow.Size.Width, Size.Height);
 				_variablesWindow.FormClosed += OnClosedVariablesWindow;
 				_variablesWindow.Show(this);
+
+				components.Add(_variablesWindow);
 			}
 			else
 			{
diff --git a/Engine/Source/Programs/UnrealGameSync/UnrealGameSync/Forms/IssueDetailsWindow.cs b/Engine/Source/Programs/UnrealGameSync/UnrealGameSync/Forms/IssueDetailsWindow.cs
--- a/Engine/Source/Programs/UnrealGameSync/UnrealGameSync/Forms/IssueDetailsWindow.cs
+++ b/Engine/Source/Programs/UnrealGameSync/UnrealGameSync/Forms/IssueDetailsWindow.cs
@@ -77,7 +77,9 @@ namespace UnrealGameSync
 			readonly object _lockObject = new object();
 			readonly SynchronizationContext _synchronizationContext;
 			Task? _workerTask;
+#pragma warning disable CA2213 // warning CA2213: 'PerforceWorkerTask' contains field '_cancellationSource' that is of IDisposable type 'CancellationTokenSource', but it is never disposed. Change the Dispose method on 'PerforceWorkerTask' to call Close or Dispose on this field.
 			readonly CancellationTokenSource _cancellationSource;
+#pragma warning restore CA2213
 			readonly AsyncEvent _refreshEvent;
 			readonly List<PerforceChangeRange> _requests = new List<PerforceChangeRange>();
 			readonly Dictionary<int, PerforceChangeDetailsWithDescribeRecord> _changeNumberToDetails = new Dictionary<int, PerforceChangeDetailsWithDescribeRecord>();
@@ -275,7 +277,9 @@ namespace UnrealGameSync
 		ChangesRecord? _contextMenuChange;
 		string? _lastOwner;
 		string? _lastDetailsText;
+#pragma warning disable CA2213 // warning CA2213: 'IssueDetailsWindow' contains field '_updateTimer' that is of IDisposable type 'Timer?', but it is never disposed. Change the Dispose method on 'IssueDetailsWindow' to call Close or Dispose on this field.
 		System.Windows.Forms.Timer? _updateTimer;
+#pragma warning restore CA2213
 		StatusElementResources _statusElementResources;
 
 		IssueDetailsWindow(IssueMonitor issueMonitor, IssueData issue, List<IssueBuildData> issueBuilds, List<IssueDiagnosticData> diagnostics, IPerforceSettings perforceSettings, TimeSpan? serverTimeOffset, IServiceProvider serviceProvider, string? currentStream)
@@ -355,6 +359,12 @@ namespace UnrealGameSync
 				_issueMonitor = null!;
 			}
 
+			_statusElementResources.Dispose();
+			_boldFont?.Dispose();
+			_boldFont = null;
+			_perforceWorker?.Dispose();
+			_perforceWorker = null;
+
 			base.Dispose(disposing);
 		}
 
@@ -859,9 +869,11 @@ namespace UnrealGameSync
 
 			foreach (IssueBuildData build in buildGroup.Builds.OrderBy(x => x.JobStepName))
 			{
+#pragma warning disable CA2000 // warning CA2000: Call System.IDisposable.Dispose on object created by 'new ToolStripMenuItem(String.Format("View Step: {0}", build.JobStepName))' before all references to it are out of scope
 				ToolStripMenuItem menuItem = new ToolStripMenuItem(String.Format("View Step: {0}", build.JobStepName));
 				menuItem.Click += (s, e) => Utility.OpenUrl(build.JobStepUrl);
 				JobContextMenu.Items.Insert(minIndex++, menuItem);
+#pragma warning restore CA2000
 			}
 
 			JobContextMenu.Show(BuildListView, point, ToolStripDropDownDirection.BelowRight);
diff --git a/Engine/Source/Programs/UnrealGameSync/UnrealGameSync/Forms/IssueFixedWindow.cs b/Engine/Source/Programs/UnrealGameSync/UnrealGameSync/Forms/IssueFixedWindow.cs
--- a/Engine/Source/Programs/UnrealGameSync/UnrealGameSync/Forms/IssueFixedWindow.cs
+++ b/Engine/Source/Programs/UnrealGameSync/UnrealGameSync/Forms/IssueFixedWindow.cs
@@ -46,7 +46,7 @@ namespace UnrealGameSync
 
 			public void Stop()
 			{
-				_ = StopAsync();
+				StopAsync().Wait();
 			}
 
 			Task StopAsync()
@@ -57,7 +57,7 @@ namespace UnrealGameSync
 					_onComplete = null;
 
 					_cancellationSource.Cancel();
-					stopTask = _backgroundTask.ContinueWith(_ => _cancellationSource.Dispose(), TaskScheduler.Default);
+					stopTask = _backgroundTask;
 
 					_backgroundTask = null!;
 				}
@@ -75,6 +75,8 @@ namespace UnrealGameSync
 				base.Dispose(disposing);
 
 				Stop();
+
+				_cancellationSource.Dispose();
 			}
 
 			public async Task DoWork(CancellationToken cancellationToken)
@@ -108,7 +110,9 @@ namespace UnrealGameSync
 
 		readonly IPerforceSettings _perforceSettings;
 		int _changeNumber;
+#pragma warning disable CA2213 // warning CA2213: 'IssueFixedWindow' contains field '_worker' that is of IDisposable type 'FindChangesWorker', but it is never disposed. Change the Dispose method on 'IssueFixedWindow' to call Close or Dispose on this field.
 		readonly FindChangesWorker _worker;
+#pragma warning restore CA2213
 		readonly IServiceProvider _serviceProvider;
 	
 		public IssueFixedWindow(IPerforceSettings perforceSettings, int initialChangeNumber, IServiceProvider serviceProvider)
diff --git a/Engine/Source/Programs/UnrealGameSync/UnrealGameSync/Forms/MainWindow.cs b/Engine/Source/Programs/UnrealGameSync/UnrealGameSync/Forms/MainWindow.cs
--- a/Engine/Source/Programs/UnrealGameSync/UnrealGameSync/Forms/MainWindow.cs
+++ b/Engine/Source/Programs/UnrealGameSync/UnrealGameSync/Forms/MainWindow.cs
@@ -544,8 +544,10 @@ namespace UnrealGameSync
 
 				foreach(UserSelectedProjectSettings recentProject in _settings.RecentProjects)
 				{
+#pragma warning disable CA2000 // warning CA2000: Call System.IDisposable.Dispose on object created by 'new ToolStripMenuItem(recentProject.ToString(), null, new EventHandler((o, e) => TabMenu_OpenRecentProject_Click(recentProject, _tabMenuTabIdx)))' before all references to it are out of scope
 					ToolStripMenuItem item = new ToolStripMenuItem(recentProject.ToString(), null, new EventHandler((o, e) => TabMenu_OpenRecentProject_Click(recentProject, _tabMenuTabIdx)));
 					TabMenu_RecentProjects.DropDownItems.Insert(insertIdx, item);
+#pragma warning restore CA2000
 					insertIdx++;
 				}
 
diff --git a/Engine/Source/Programs/UnrealGameSync/UnrealGameSync/Forms/NotificationWindow.cs b/Engine/Source/Programs/UnrealGameSync/UnrealGameSync/Forms/NotificationWindow.cs
--- a/Engine/Source/Programs/UnrealGameSync/UnrealGameSync/Forms/NotificationWindow.cs
+++ b/Engine/Source/Programs/UnrealGameSync/UnrealGameSync/Forms/NotificationWindow.cs
@@ -66,15 +66,24 @@ namespace UnrealGameSync
 
 		protected override void Dispose(bool disposing)
 		{
-			if(disposing && _logoBitmap != null)
+			if (disposing)
 			{
-				_logoBitmap.Dispose();
-				_logoBitmap = null;
-			}
-			if(disposing && components != null)
-			{
-				components.Dispose();
+				if (_captionFont != null)
+				{
+					_captionFont.Dispose();
+					_captionFont = null;
+				}
+				if (_logoBitmap != null)
+				{
+					_logoBitmap.Dispose();
+					_logoBitmap = null;
+				}
+				if (components != null)
+				{
+					components.Dispose();
+				}
 			}
+
 			base.Dispose(disposing);
 		}
 
diff --git a/Engine/Source/Programs/UnrealGameSync/UnrealGameSync/Forms/ProgramsRunningWindow.cs b/Engine/Source/Programs/UnrealGameSync/UnrealGameSync/Forms/ProgramsRunningWindow.cs
--- a/Engine/Source/Programs/UnrealGameSync/UnrealGameSync/Forms/ProgramsRunningWindow.cs
+++ b/Engine/Source/Programs/UnrealGameSync/UnrealGameSync/Forms/ProgramsRunningWindow.cs
@@ -27,6 +27,16 @@ namespace UnrealGameSync
 			ProgramListBox.Items.AddRange(programs);
 		}
 
+		protected override void Dispose(bool disposing)
+		{
+			if (disposing)
+			{
+				_terminateEvent?.Dispose();
+				components?.Dispose();
+			}
+			base.Dispose(disposing);
+		}
+
 		private void ProgramsRunningWindow_Load(object sender, EventArgs e)
 		{
 			_terminateEvent = new ManualResetEvent(false);
diff --git a/Engine/Source/Programs/UnrealGameSync/UnrealGameSync/IssueMonitor.cs b/Engine/Source/Programs/UnrealGameSync/UnrealGameSync/IssueMonitor.cs
--- a/Engine/Source/Programs/UnrealGameSync/UnrealGameSync/IssueMonitor.cs
+++ b/Engine/Source/Programs/UnrealGameSync/UnrealGameSync/IssueMonitor.cs
@@ -135,7 +135,9 @@ namespace UnrealGameSync
 		int _refCount = 1;
 		Task? _workerTask;
 		readonly ILogger _logger;
+#pragma warning disable CA2213 // warning CA2213: 'IssueMonitor' contains field '_cancellationSource' that is of IDisposable type 'CancellationTokenSource', but it is never disposed. Change the Dispose method on 'IssueMonitor' to call Close or Dispose on this field.
 		readonly CancellationTokenSource _cancellationSource;
+#pragma warning restore CA2213
 		readonly AsyncEvent _refreshEvent;
 		int _updateIntervalMs;
 		readonly List<long> _trackingIssueIds = new List<long>();
diff --git a/Engine/Source/Programs/UnrealGameSync/UnrealGameSync/PerforceModalTask.cs b/Engine/Source/Programs/UnrealGameSync/UnrealGameSync/PerforceModalTask.cs
--- a/Engine/Source/Programs/UnrealGameSync/UnrealGameSync/PerforceModalTask.cs
+++ b/Engine/Source/Programs/UnrealGameSync/UnrealGameSync/PerforceModalTask.cs
@@ -46,7 +46,9 @@ namespace UnrealGameSync
 			}
 			finally
 			{
+#pragma warning disable CA1508 // warning CA1508: 'connection' is always 'null'. Remove or refactor the condition(s) to avoid dead code.
 				connection?.Dispose();
+#pragma warning restore CA1508
 			}
 		}
 
diff --git a/Engine/Source/Programs/UnrealGameSync/UnrealGameSync/PerforceMonitor.cs b/Engine/Source/Programs/UnrealGameSync/UnrealGameSync/PerforceMonitor.cs
--- a/Engine/Source/Programs/UnrealGameSync/UnrealGameSync/PerforceMonitor.cs
+++ b/Engine/Source/Programs/UnrealGameSync/UnrealGameSync/PerforceMonitor.cs
@@ -39,7 +39,9 @@ namespace UnrealGameSync
 		readonly string _selectedClientFileName;
 		readonly string _selectedProjectIdentifier;
 		Task? _workerTask;
+#pragma warning disable CA2213 //  warning CA2213: 'PerforceMonitor' contains field '_cancellationSource' that is of IDisposable type 'CancellationTokenSource', but it is never disposed. Change the Dispose method on 'PerforceMonitor' to call Close or Dispose on this field.
 		readonly CancellationTokenSource _cancellationSource;
+#pragma warning restore CA2213
 		int _pendingMaxChangesValue;
 		SortedSet<ChangesRecord> _changes = new SortedSet<ChangesRecord>(new PerforceChangeSorter());
 		readonly SortedDictionary<int, PerforceChangeDetails> _changeDetails = new SortedDictionary<int,PerforceChangeDetails>();
@@ -504,7 +506,7 @@ namespace UnrealGameSync
 						// Notify the caller after a fixed period of time, in case further updates are slow to arrive
 						if (notifyTask.IsCompleted)
 						{
-							notifyTask = Task.Delay(TimeSpan.FromSeconds(5.0), cancellationSource.Token).ContinueWith(_ => _synchronizationContext.Post(_ => OnUpdateMetadata?.Invoke(), null));
+							notifyTask = Task.Delay(TimeSpan.FromSeconds(5.0), cancellationSource.Token).ContinueWith(_ => _synchronizationContext.Post(_ => OnUpdateMetadata?.Invoke(), null), cancellationSource.Token, new TaskContinuationOptions(), TaskScheduler.Default);
 						}
 					}
 				}
diff --git a/Engine/Source/Programs/UnrealGameSync/UnrealGameSync/Program.cs b/Engine/Source/Programs/UnrealGameSync/UnrealGameSync/Program.cs
--- a/Engine/Source/Programs/UnrealGameSync/UnrealGameSync/Program.cs
+++ b/Engine/Source/Programs/UnrealGameSync/UnrealGameSync/Program.cs
@@ -161,9 +161,6 @@ namespace UnrealGameSync
 			DirectoryReference dataFolder = DirectoryReference.Combine(DirectoryReference.GetSpecialFolder(Environment.SpecialFolder.LocalApplicationData)!, "UnrealGameSync");
 			DirectoryReference.CreateDirectory(dataFolder);
 
-			// Enable TLS 1.1 and 1.2. TLS 1.0 is now deprecated and not allowed by default in NET Core servers.
-			ServicePointManager.SecurityProtocol |= SecurityProtocolType.Tls11 | SecurityProtocolType.Tls12;
-
 			// Create a new logger
 			using (ILoggerProvider loggerProvider = Logging.CreateLoggerProvider(FileReference.Combine(dataFolder, "UnrealGameSync.log")))
 			{
@@ -322,7 +319,7 @@ namespace UnrealGameSync
 		{
 			for(int idx = 0; idx < remainingArgs.Count; idx++)
 			{
-				if(remainingArgs[idx].Equals(option, StringComparison.InvariantCultureIgnoreCase))
+				if(remainingArgs[idx].Equals(option, StringComparison.OrdinalIgnoreCase))
 				{
 					value = true;
 					remainingArgs.RemoveAt(idx);
@@ -338,7 +335,7 @@ namespace UnrealGameSync
 		{
 			for(int idx = 0; idx < remainingArgs.Count; idx++)
 			{
-				if(remainingArgs[idx].StartsWith(prefix, StringComparison.InvariantCultureIgnoreCase))
+				if(remainingArgs[idx].StartsWith(prefix, StringComparison.OrdinalIgnoreCase))
 				{
 					value = remainingArgs[idx].Substring(prefix.Length);
 					remainingArgs.RemoveAt(idx);
diff --git a/Engine/Source/Programs/UnrealGameSync/UnrealGameSync/ProgramApplicationContext.cs b/Engine/Source/Programs/UnrealGameSync/UnrealGameSync/ProgramApplicationContext.cs
--- a/Engine/Source/Programs/UnrealGameSync/UnrealGameSync/ProgramApplicationContext.cs
+++ b/Engine/Source/Programs/UnrealGameSync/UnrealGameSync/ProgramApplicationContext.cs
@@ -44,6 +44,8 @@ namespace UnrealGameSync
 
 		Container _components = new Container();
 		NotifyIcon _notifyIcon;
+#pragma warning disable CA2213 // warning CA2213: 'ProgramApplicationContext' contains field '_notifyMenuOpenUnrealGameSync' that is of IDisposable type 'ToolStripMenuItem', but it is never disposed. Change the Dispose method on 'ProgramApplicationContext' to call Close or Dispose on this field.
+		// Disposed via _components
 		readonly ContextMenuStrip _notifyMenu;
 		readonly ToolStripMenuItem _notifyMenuOpenUnrealGameSync;
 		readonly ToolStripSeparator _notifyMenuOpenUnrealGameSyncSeparator;
@@ -51,11 +53,16 @@ namespace UnrealGameSync
 		readonly ToolStripMenuItem _notifyMenuLaunchEditor;
 		readonly ToolStripSeparator _notifyMenuExitSeparator;
 		readonly ToolStripMenuItem _notifyMenuExit;
+#pragma warning restore CA2213
 
 		CancellationTokenSource _startupCancellationSource = new CancellationTokenSource();
 		readonly Task _startupTask;
 		ModalTaskWindow? _startupWindow;
+#pragma warning disable CA2213 // warning CA2213: 'ProgramApplicationContext' contains field '_mainWindowInstance' that is of IDisposable type 'MainWindow?', but it is never disposed. Change the Dispose method on 'ProgramApplicationContext' to call Close or Dispose on this field.
 		MainWindow? _mainWindowInstance;
+#pragma warning restore CA2213
+
+		WindowsFormsSynchronizationContext? _synchronizationContext;
 
 		public ProgramApplicationContext(IPerforceSettings defaultPerforceSettings, UpdateMonitor updateMonitor, string? apiUrl, DirectoryReference dataFolder, EventWaitHandle activateEvent, bool restoreState, string? updateSpawn, string? projectFileName, bool preview, IServiceProvider serviceProvider, string? uri)
 		{
@@ -79,7 +86,8 @@ namespace UnrealGameSync
 			// back to the main thread at any time.
 			if(SynchronizationContext.Current == null)
 			{
-				SynchronizationContext.SetSynchronizationContext(new WindowsFormsSynchronizationContext());
+				_synchronizationContext = new WindowsFormsSynchronizationContext();
+				SynchronizationContext.SetSynchronizationContext(_synchronizationContext);
 			}
 
 			// Capture the main thread's synchronization context for callbacks
@@ -269,7 +277,8 @@ namespace UnrealGameSync
 
 			// Create the main window 
 			_mainWindowInstance = new MainWindow(_updateMonitor, _apiUrl, _dataFolder, _updateSpawn ?? originalExe, _preview, startupTasks, _defaultPerforceSettings, _serviceProvider, _settings, _uri);
-			if(visible)
+			_components.Add(_mainWindowInstance);
+			if (visible)
 			{
 				_mainWindowInstance.Show();
 				if(!_restoreState)
@@ -363,6 +372,12 @@ namespace UnrealGameSync
 				_startupCancellationSource.Dispose();
 				_startupCancellationSource = null!;
 			}
+
+			if (_synchronizationContext != null)
+			{
+				_synchronizationContext.Dispose();
+				_synchronizationContext = null;
+			}
 		}
 
 		private void NotifyIcon_MouseDown(object? sender, MouseEventArgs e)
diff --git a/Engine/Source/Programs/UnrealGameSync/UnrealGameSync/RESTApi.cs b/Engine/Source/Programs/UnrealGameSync/UnrealGameSync/RESTApi.cs
--- a/Engine/Source/Programs/UnrealGameSync/UnrealGameSync/RESTApi.cs
+++ b/Engine/Source/Programs/UnrealGameSync/UnrealGameSync/RESTApi.cs
@@ -45,7 +45,7 @@ namespace UnrealGameSync
 				{
 					using (HttpResponseMessage response = await s_httpClient.SendAsync(request, cancellationToken))
 					{
-						return await response.Content.ReadAsStringAsync();
+						return await response.Content.ReadAsStringAsync(cancellationToken);
 					}
 				}
 				catch (Exception ex)
diff --git a/Engine/Source/Programs/UnrealGameSync/UnrealGameSync/ToolUpdateMonitor.cs b/Engine/Source/Programs/UnrealGameSync/UnrealGameSync/ToolUpdateMonitor.cs
--- a/Engine/Source/Programs/UnrealGameSync/UnrealGameSync/ToolUpdateMonitor.cs
+++ b/Engine/Source/Programs/UnrealGameSync/UnrealGameSync/ToolUpdateMonitor.cs
@@ -58,7 +58,9 @@ namespace UnrealGameSync
 
 	class ToolUpdateMonitor : IDisposable
 	{
+#pragma warning disable CA2213 // warning CA2213: 'ToolUpdateMonitor' contains field '_cancellationSource' that is of IDisposable type 'CancellationTokenSource', but it is never disposed. Change the Dispose method on 'ToolUpdateMonitor' to call Close or Dispose on this field.
 		readonly CancellationTokenSource _cancellationSource;
+#pragma warning restore CA2213
 		readonly SynchronizationContext _synchronizationContext;
 		Task? _workerTask;
 		readonly AsyncEvent _wakeEvent;
diff --git a/Engine/Source/Programs/UnrealGameSync/UnrealGameSync/UnrealGameSync.csproj b/Engine/Source/Programs/UnrealGameSync/UnrealGameSync/UnrealGameSync.csproj
--- a/Engine/Source/Programs/UnrealGameSync/UnrealGameSync/UnrealGameSync.csproj
+++ b/Engine/Source/Programs/UnrealGameSync/UnrealGameSync/UnrealGameSync.csproj
@@ -13,6 +13,7 @@
     <ApplicationIcon>Resources\Icon.ico</ApplicationIcon>
     <ApplicationManifest>Properties\app.manifest</ApplicationManifest>
     <Nullable>enable</Nullable>
+    <NeutralLanguage>en-US</NeutralLanguage>
   </PropertyGroup>
   <PropertyGroup Condition="'$(Configuration)|$(Platform)'=='Release|AnyCPU'">
     <DefineConstants>TRACE;DEBUG;WITH_TELEMETRY</DefineConstants>
diff --git a/Engine/Source/Programs/UnrealGameSync/UnrealGameSync/UpdateMonitor.cs b/Engine/Source/Programs/UnrealGameSync/UnrealGameSync/UpdateMonitor.cs
--- a/Engine/Source/Programs/UnrealGameSync/UnrealGameSync/UpdateMonitor.cs
+++ b/Engine/Source/Programs/UnrealGameSync/UnrealGameSync/UpdateMonitor.cs
@@ -18,7 +18,9 @@ namespace UnrealGameSync
 	class UpdateMonitor : IDisposable
 	{
 		Task? _workerTask;
+#pragma warning disable CA2213 // warning CA2213: 'UpdateMonitor' contains field '_cancellationSource' that is of IDisposable type 'CancellationTokenSource', but it is never disposed. Change the Dispose method on 'UpdateMonitor' to call Close or Dispose on this field.
 		readonly CancellationTokenSource _cancellationSource = new CancellationTokenSource();
+#pragma warning restore CA2213
 		readonly ILogger _logger;
 		readonly IAsyncDisposer _asyncDisposer;
 
diff --git a/Engine/Source/Programs/UnrealGameSync/UnrealGameSync/UriHandlers/UriHandler.cs b/Engine/Source/Programs/UnrealGameSync/UnrealGameSync/UriHandlers/UriHandler.cs
--- a/Engine/Source/Programs/UnrealGameSync/UnrealGameSync/UriHandlers/UriHandler.cs
+++ b/Engine/Source/Programs/UnrealGameSync/UnrealGameSync/UriHandlers/UriHandler.cs
@@ -293,7 +293,7 @@ namespace UnrealGameSync
 			}
 		}
 
-		public static string CurrentProcessFilePath => Path.ChangeExtension(Path.GetFullPath(Process.GetCurrentProcess().MainModule!.FileName!), ".exe");
+		public static string CurrentProcessFilePath => Path.ChangeExtension(Path.GetFullPath(Environment.ProcessPath ?? "UnrealGameSync.exe"), ".exe");
 
 		static List<RegistrySetting> GetGlobalRegistrySettings()
 		{
diff --git a/Engine/Source/Programs/UnrealGameSync/UnrealGameSync/Workspace.cs b/Engine/Source/Programs/UnrealGameSync/UnrealGameSync/Workspace.cs
--- a/Engine/Source/Programs/UnrealGameSync/UnrealGameSync/Workspace.cs
+++ b/Engine/Source/Programs/UnrealGameSync/UnrealGameSync/Workspace.cs
@@ -103,7 +103,9 @@ namespace UnrealGameSync
 			get; private set;
 		}
 
+#pragma warning disable CA2213 // warning CA2213: 'Workspace' contains field '_prevCancellationSource' that is of IDisposable type 'CancellationTokenSource?', but it is never disposed. Change the Dispose method on 'Workspace' to call Close or Dispose on this field.
 		CancellationTokenSource? _prevCancellationSource;
+#pragma warning restore CA2213
 		Task _prevUpdateTask = Task.CompletedTask;
 
 		public void Update(WorkspaceUpdateContext context)
diff --git a/Engine/Source/Programs/UnrealGameSync/UnrealGameSyncLauncher/Forms/SettingsWindow.cs b/Engine/Source/Programs/UnrealGameSync/UnrealGameSyncLauncher/Forms/SettingsWindow.cs
--- a/Engine/Source/Programs/UnrealGameSync/UnrealGameSyncLauncher/Forms/SettingsWindow.cs
+++ b/Engine/Source/Programs/UnrealGameSync/UnrealGameSyncLauncher/Forms/SettingsWindow.cs
@@ -53,7 +53,7 @@ namespace UnrealGameSyncLauncher
 
 		private void ViewLogBtn_Click(object sender, EventArgs e)
 		{
-			LogWindow log = new LogWindow(_logText ?? String.Empty);
+			using LogWindow log = new LogWindow(_logText ?? String.Empty);
 			log.ShowDialog(this);
 		}
 
diff --git a/Engine/Source/Programs/UnrealGameSync/UnrealGameSyncLauncher/Program.cs b/Engine/Source/Programs/UnrealGameSync/UnrealGameSyncLauncher/Program.cs
--- a/Engine/Source/Programs/UnrealGameSync/UnrealGameSyncLauncher/Program.cs
+++ b/Engine/Source/Programs/UnrealGameSync/UnrealGameSyncLauncher/Program.cs
@@ -26,6 +26,7 @@ namespace UnrealGameSyncLauncher
 		{
 			Application.EnableVisualStyles();
 			Application.SetCompatibleTextRenderingDefault(false);
+			Application.SetHighDpiMode(HighDpiMode.PerMonitorV2);
 
 			bool firstInstance;
 			using(Mutex instanceMutex = new Mutex(true, "UnrealGameSyncRunning", out firstInstance))
@@ -53,7 +54,7 @@ namespace UnrealGameSyncLauncher
 				if ((Control.ModifierKeys & Keys.Shift) != 0)
 				{
 					// Show the settings window immediately
-					SettingsWindow updateError = new SettingsWindow(null, null, serverAndPort, userName, depotPath, preview, SyncAndRunWrapper);
+					using SettingsWindow updateError = new SettingsWindow(null, null, serverAndPort, userName, depotPath, preview, SyncAndRunWrapper);
 					if(updateError.ShowDialog() == DialogResult.OK)
 					{
 						return 0;
@@ -76,7 +77,7 @@ namespace UnrealGameSyncLauncher
 						return 0;
 					}
 
-					SettingsWindow updateError = new SettingsWindow("Unable to update UnrealGameSync from Perforce. Verify that your connection settings are correct.", logger.Render(Environment.NewLine), serverAndPort, userName, depotPath, preview, SyncAndRunWrapper);
+					using SettingsWindow updateError = new SettingsWindow("Unable to update UnrealGameSync from Perforce. Verify that your connection settings are correct.", logger.Render(Environment.NewLine), serverAndPort, userName, depotPath, preview, SyncAndRunWrapper);
 					if(updateError.ShowDialog() == DialogResult.OK)
 					{
 						return 0;
@@ -104,10 +105,12 @@ namespace UnrealGameSyncLauncher
 				{
 					syncPath = baseDepotPathPrefix + (preview ? "/UnstableRelease/..." : "/Release/...");
 					changes = await perforce.GetChangesAsync(ChangesOptions.None, 1, ChangeStatus.Submitted, syncPath, cancellationToken);
+#pragma warning disable CA1508 // warning CA1508: 'changes.Count == 0' is always 'true'. Remove or refactor the condition(s) to avoid dead code.
 					if (changes.Count == 0)
 					{
 						throw new UserErrorException($"Unable to find any UGS binaries under {syncPath}");
 					}
+#pragma warning restore CA1508
 				}
 
 				int requiredChangeNumber = changes[0].Number;
@@ -239,7 +242,7 @@ namespace UnrealGameSyncLauncher
 
 		static string QuoteArgument(string arg)
 		{
-			if(arg.IndexOf(' ') != -1 && !arg.StartsWith("\""))
+			if(arg.IndexOf(' ', StringComparison.Ordinal) != -1 && !arg.StartsWith("\"", StringComparison.Ordinal))
 			{
 				return String.Format("\"{0}\"", arg);
 			}
diff --git a/Engine/Source/Programs/UnrealGameSync/UnrealGameSyncLauncher/Properties/AssemblyInfo.cs b/Engine/Source/Programs/UnrealGameSync/UnrealGameSyncLauncher/Properties/AssemblyInfo.cs
--- a/Engine/Source/Programs/UnrealGameSync/UnrealGameSyncLauncher/Properties/AssemblyInfo.cs
+++ b/Engine/Source/Programs/UnrealGameSync/UnrealGameSyncLauncher/Properties/AssemblyInfo.cs
@@ -1,7 +1,12 @@
 // Copyright Epic Games, Inc. All Rights Reserved.
 
 using System.Reflection;
+using System.Resources;
 using System.Runtime.InteropServices;
+using System.Runtime.Versioning;
+
+[assembly: SupportedOSPlatform("windows10.0")]
+[assembly: NeutralResourcesLanguage("en-US")]
 
 // General Information about an assembly is controlled through the following 
 // set of attributes. Change these attribute values to modify the information
diff --git a/Engine/Source/Programs/UnrealGameSync/UnrealGameSyncLauncher/UnrealGameSyncLauncher.csproj b/Engine/Source/Programs/UnrealGameSync/UnrealGameSyncLauncher/UnrealGameSyncLauncher.csproj
--- a/Engine/Source/Programs/UnrealGameSync/UnrealGameSyncLauncher/UnrealGameSyncLauncher.csproj
+++ b/Engine/Source/Programs/UnrealGameSync/UnrealGameSyncLauncher/UnrealGameSyncLauncher.csproj
@@ -1,10 +1,9 @@
 <?xml version="1.0" encoding="utf-8"?>
-<Project Sdk="Microsoft.NET.Sdk.WindowsDesktop">
-
+<Project Sdk="Microsoft.NET.Sdk">
   <PropertyGroup>
     <OutputType>WinExe</OutputType>
     <DefineConstants>UGS_LAUNCHER</DefineConstants>
-    <TargetFramework>netcoreapp3.1</TargetFramework>
+    <TargetFramework>net6.0-windows</TargetFramework>
     <UseWindowsForms>true</UseWindowsForms>
     <GenerateAssemblyInfo>false</GenerateAssemblyInfo>
     <RootNamespace>UnrealGameSyncLauncher</RootNamespace>
@@ -14,6 +13,7 @@
     <ApplicationManifest>app.manifest</ApplicationManifest>
     <RuntimeIdentifiers>win10-x64</RuntimeIdentifiers>
     <Nullable>enable</Nullable>
+    <NeutralLanguage>en-US</NeutralLanguage>
   </PropertyGroup>
   <ItemGroup>
     <Compile Include="..\UnrealGameSync\Forms\ModalTaskWindow.cs">
@@ -49,8 +49,8 @@
     </Compile>
   </ItemGroup>
   <ItemGroup>
-    <ProjectReference Include="..\Shared\EpicGames.Core\EpicGames.Core.csproj" />
-    <ProjectReference Include="..\Shared\EpicGames.Perforce\EpicGames.Perforce.csproj" />
+    <ProjectReference Include="..\..\Shared\EpicGames.Core\EpicGames.Core.csproj" />
+    <ProjectReference Include="..\..\Shared\EpicGames.Perforce\EpicGames.Perforce.csproj" />
   </ItemGroup>
   <ItemGroup>
     <None Update="Properties\Settings.settings">
diff --git a/Engine/Source/Programs/UnrealGameSync/UnrealGameSyncLauncher/app.manifest b/Engine/Source/Programs/UnrealGameSync/UnrealGameSyncLauncher/app.manifest
--- a/Engine/Source/Programs/UnrealGameSync/UnrealGameSyncLauncher/app.manifest
+++ b/Engine/Source/Programs/UnrealGameSync/UnrealGameSyncLauncher/app.manifest
@@ -27,34 +27,20 @@
            and is designed to work with. Uncomment the appropriate elements
            and Windows will automatically select the most compatible environment. -->
 
-      <!-- Windows Vista -->
-      <!--<supportedOS Id="{e2011457-1546-43c5-a5fe-008deee3d3f0}" />-->
-
-      <!-- Windows 7 -->
-      <!--<supportedOS Id="{35138b9a-5d96-4fbd-8e2d-a2440225f93a}" />-->
-
-      <!-- Windows 8 -->
-      <!--<supportedOS Id="{4a2f28e3-53b9-4441-ba9c-d69d4a4a6e38}" />-->
-
-      <!-- Windows 8.1 -->
-      <!--<supportedOS Id="{1f676c76-80e1-4239-95bb-83d0f6d0da78}" />-->
-
-      <!-- Windows 10 -->
-      <!--<supportedOS Id="{8e0f7a12-bfb3-4fe8-b9a5-48fd50a15a9a}" />-->
-
-    </application>
+		<!-- Windows 10 and Windows 11 -->
+		<supportedOS Id="{8e0f7a12-bfb3-4fe8-b9a5-48fd50a15a9a}"/>
+		<!-- Windows 8.1 -->
+		<supportedOS Id="{1f676c76-80e1-4239-95bb-83d0f6d0da78}"/>
+		<!-- Windows 8 -->
+		<supportedOS Id="{4a2f28e3-53b9-4441-ba9c-d69d4a4a6e38}"/>
+		<!-- Windows 7 -->
+		<supportedOS Id="{35138b9a-5d96-4fbd-8e2d-a2440225f93a}"/>
+		<!-- Windows Vista -->
+		<supportedOS Id="{e2011457-1546-43c5-a5fe-008deee3d3f0}"/>
+
+	</application>
   </compatibility>
 
-  <!-- Indicates that the application is DPI-aware and will not be automatically scaled by Windows at higher
-       DPIs. Windows Presentation Foundation (WPF) applications are automatically DPI-aware and do not need 
-       to opt in. Windows Forms applications targeting .NET Framework 4.6 that opt into this setting, should 
-       also set the 'EnableWindowsFormsHighDpiAutoResizing' setting to 'true' in their app.config. -->
-  <application xmlns="urn:schemas-microsoft-com:asm.v3">
-    <windowsSettings>
-      <dpiAware xmlns="http://schemas.microsoft.com/SMI/2005/WindowsSettings">true</dpiAware>
-    </windowsSettings>
-  </application>
-
   <!-- Enable themes for Windows common controls and dialogs (Windows XP and later) -->
   <!--
   <dependency>
diff --git a/Engine/Source/Programs/UnrealGameSync/UnrealGameSyncShared/GlobalSettings.cs b/Engine/Source/Programs/UnrealGameSync/UnrealGameSyncShared/GlobalSettings.cs
--- a/Engine/Source/Programs/UnrealGameSync/UnrealGameSyncShared/GlobalSettings.cs
+++ b/Engine/Source/Programs/UnrealGameSync/UnrealGameSyncShared/GlobalSettings.cs
@@ -10,9 +10,9 @@ namespace UnrealGameSync
 {
 	public class FilterSettings
 	{
-		public List<Guid> IncludeCategories { get; set; } = new List<Guid>();
-		public List<Guid> ExcludeCategories { get; set; } = new List<Guid>();
-		public List<string> View { get; set; } = new List<string>();
+		public List<Guid> IncludeCategories { get; } = new List<Guid>();
+		public List<Guid> ExcludeCategories { get; } = new List<Guid>();
+		public List<string> View { get; } = new List<string>();
 		public bool? AllProjects { get; set; }
 		public bool? AllProjectsInSln { get; set; }
 
diff --git a/Engine/Source/Programs/UnrealGameSync/UnrealGameSyncShared/OpenProjectInfo.cs b/Engine/Source/Programs/UnrealGameSync/UnrealGameSyncShared/OpenProjectInfo.cs
--- a/Engine/Source/Programs/UnrealGameSync/UnrealGameSyncShared/OpenProjectInfo.cs
+++ b/Engine/Source/Programs/UnrealGameSync/UnrealGameSyncShared/OpenProjectInfo.cs
@@ -129,11 +129,13 @@ namespace UnrealGameSync
 						// Filter this list of clients
 						candidateClients = await FilterClients(clients, newSelectedFileName, defaultConnection.Settings, perforceInfo.ClientHost, logger, cancellationToken);
 
+#pragma warning disable CA1508 // False positive? "warning CA1508: 'newCandidateClients.Count == 0' is always 'true'. Remove or refactor the condition(s) to avoid dead code."
 						// If we still couldn't find any, fail.
 						if (candidateClients.Count == 0)
 						{
 							throw new UserErrorException($"Couldn't find any Perforce workspace containing {newSelectedFileName}. Check your connection settings.");
 						}
+#pragma warning restore CA1508
 					}
 
 					// Check there's only one client
diff --git a/Engine/Source/Programs/UnrealGameSync/UnrealGameSyncShared/UserSettings.cs b/Engine/Source/Programs/UnrealGameSync/UnrealGameSyncShared/UserSettings.cs
--- a/Engine/Source/Programs/UnrealGameSync/UnrealGameSyncShared/UserSettings.cs
+++ b/Engine/Source/Programs/UnrealGameSync/UnrealGameSyncShared/UserSettings.cs
@@ -65,7 +65,7 @@ namespace UnrealGameSync
 		public bool Good { get; set; } = false;
 		public bool Starred { get; set; } = false;
 		public bool FindNewestGoodContent { get; set; } = false;
-		public List<string> RequiredBadges { get; set; } = new List<string>();
+		public List<string> RequiredBadges { get; } = new List<string>();
 
 		// Depot path to read for the latest change number
 		public string? ReadFrom { get; set; }
@@ -152,7 +152,7 @@ namespace UnrealGameSync
 	{
 		public bool Enabled { get; set; }
 		public string Type { get; set; }
-		public List<string> Order { get; set; }
+		public List<string> Order { get; }
 
 		public ArchiveSettings(bool enabled, string type, IEnumerable<string> order)
 		{
@@ -416,9 +416,9 @@ namespace UnrealGameSync
 		[JsonIgnore]
 		public FileReference ConfigFile { get; private set; } = null!;
 
-		public List<ConfigObject> BuildSteps { get; set; } = new List<ConfigObject>();
+		public List<ConfigObject> BuildSteps { get; } = new List<ConfigObject>();
 		public FilterType FilterType { get; set; }
-		public HashSet<string> FilterBadges { get; set; } = new HashSet<string>(StringComparer.OrdinalIgnoreCase);
+		public HashSet<string> FilterBadges { get; } = new HashSet<string>(StringComparer.OrdinalIgnoreCase);
 
 		static readonly object _syncRoot = new object();
 
@@ -496,21 +496,21 @@ namespace UnrealGameSync
 		public bool AnnotateRobmergeChanges { get; set; }
 		public bool ShowLocalTimes { get; set; }
 		public bool KeepInTray { get; set; }
-		public HashSet<Guid> EnabledTools { get; set; } = new HashSet<Guid>();
+		public HashSet<Guid> EnabledTools { get; } = new HashSet<Guid>();
 		public int FilterIndex { get; set; }
 		public UserSelectedProjectSettings? LastProject { get; set; }
-		public List<UserSelectedProjectSettings> OpenProjects { get; set; }
-		public List<UserSelectedProjectSettings> RecentProjects { get; set; }
+		public List<UserSelectedProjectSettings> OpenProjects { get; }
+		public List<UserSelectedProjectSettings> RecentProjects { get; }
 		public string SyncTypeId { get; set; }
 		public BuildConfig CompiledEditorBuildConfig { get; set; } // NB: This assumes not using precompiled editor. See CurrentBuildConfig.
 		public TabLabels TabLabels { get; set; }
 		public long NextLauncherVersionCheck { get; set; }
 
 		// Precompiled binaries
-		public List<ArchiveSettings> Archives { get; set; } = new List<ArchiveSettings>();
+		public List<ArchiveSettings> Archives { get; } = new List<ArchiveSettings>();
 
 		// OIDC Settings
-		public Dictionary<string, string> ProviderToRefreshTokens { get; set; } = new Dictionary<string, string>();
+		public Dictionary<string, string> ProviderToRefreshTokens { get; } = new Dictionary<string, string>();
 
 		// Window settings
 		public bool WindowVisible { get; set; }
@@ -521,14 +521,14 @@ namespace UnrealGameSync
 		public bool ScheduleEnabled { get; set; }
 		public TimeSpan ScheduleTime { get; set; }
 		public bool ScheduleAnyOpenProject { get; set; }
-		public List<UserSelectedProjectSettings> ScheduleProjects { get; set; } = new List<UserSelectedProjectSettings>();
+		public List<UserSelectedProjectSettings> ScheduleProjects { get; } = new List<UserSelectedProjectSettings>();
 
 		// Run configuration
-		public List<Tuple<string, bool>> EditorArguments { get; set; } = new List<Tuple<string, bool>>();
+		public List<Tuple<string, bool>> EditorArguments { get; } = new List<Tuple<string, bool>>();
 		public bool EditorArgumentsPrompt { get; set; }
 
 		// Notification settings
-		public List<string> NotifyProjects { get; set; } = new List<string>();
+		public List<string> NotifyProjects { get; } = new List<string>();
 		public int NotifyUnassignedMinutes { get; set; }
 		public int NotifyUnacknowledgedMinutes { get; set; }
 		public int NotifyUnresolvedMinutes { get; set; }
diff --git a/Engine/Source/Programs/UnrealGameSync/UnrealGameSyncShared/WorkspaceState.cs b/Engine/Source/Programs/UnrealGameSync/UnrealGameSyncShared/WorkspaceState.cs
--- a/Engine/Source/Programs/UnrealGameSync/UnrealGameSyncShared/WorkspaceState.cs
+++ b/Engine/Source/Programs/UnrealGameSync/UnrealGameSyncShared/WorkspaceState.cs
@@ -34,7 +34,7 @@ namespace UnrealGameSync
 		public int CurrentChangeNumber { get; set; } = -1;
 		public int CurrentCodeChangeNumber { get; set; } = -1;
 		public string? CurrentSyncFilterHash { get; set; }
-		public List<int> AdditionalChangeNumbers { get; set; } = new List<int>();
+		public List<int> AdditionalChangeNumbers { get; } = new List<int>();
 
 		// Settings for the last attempted sync. These values are set to persist error messages between runs.
 		public int LastSyncChangeNumber { get; set; }
@@ -47,10 +47,10 @@ namespace UnrealGameSync
 		public int LastBuiltChangeNumber { get; set; }
 
 		// Expanded archives in the workspace
-		public HashSet<string> ExpandedArchiveTypes { get; set; } = new HashSet<string>(StringComparer.Ordinal);
+		public HashSet<string> ExpandedArchiveTypes { get; } = new HashSet<string>(StringComparer.Ordinal);
 
 		// The changes that we're regressing at the moment
-		public List<BisectEntry> BisectChanges { get; set; } = new List<BisectEntry>();
+		public List<BisectEntry> BisectChanges { get; } = new List<BisectEntry>();
 
 		public void UpdateCachedProjectInfo(ProjectInfo projectInfo, long settingsTimeUtc)
 		{
