UGS: Include the NET runtime in UGS deployments.

diff --git a/Engine/Source/Programs/UnrealGameSync/BuildUGS.xml b/Engine/Source/Programs/UnrealGameSync/BuildUGS.xml
--- a/Engine/Source/Programs/UnrealGameSync/BuildUGS.xml
+++ b/Engine/Source/Programs/UnrealGameSync/BuildUGS.xml
@@ -32,7 +32,7 @@
 			<Property Name="UgsLauncherPublishDir" Value="$(UgsDir)/UnrealGameSyncLauncher/bin/Release/netcoreapp3.1/win-x64/publish"/>
 
 			<Delete Files="$(UgsLauncherPublishDir)/..."/>
-			<DotNet Arguments="publish &quot;$(UgsDir)/UnrealGameSyncLauncher/UnrealGameSyncLauncher.csproj&quot; --output=&quot;$(UgsLauncherPublishDir)&quot; --no-self-contained -c Release -r win-x64 $(VersionArguments)"/>
+			<DotNet Arguments="publish &quot;$(UgsDir)/UnrealGameSyncLauncher/UnrealGameSyncLauncher.csproj&quot; --output=&quot;$(UgsLauncherPublishDir)&quot; --self-contained -c Release -r win-x64 $(VersionArguments)"/>
 
 			<Spawn Exe="$(WixDir)/heat.exe" Arguments="dir &quot;$(UgsLauncherPublishDir)&quot; -cg UGSLauncher_Project -dr INSTALLFOLDER -scom -sreg -srd -var var.BasePath -gg -sfrag -out obj/Release/UGSLauncher.wxs" WorkingDir="$(InstallerDir)"/>
 			<Spawn Exe="$(WixDir)/candle.exe" Arguments="-dBasePath=&quot;$(UgsLauncherPublishDir)&quot; -out obj\Release\ -dConfiguration=Release -dPlatform=x64 -arch x86 -ext ..\..\..\ThirdParty\WiX\3.8\WixUtilExtension.dll Product.wxs obj/Release/UGSLauncher.wxs" WorkingDir="$(InstallerDir)"/>
@@ -44,10 +44,10 @@
 		
 		<Node Name="Build UnrealGameSync" Requires="$(InstallerDependency)">
 			<Delete Files="$(UgsOutputDir)/..."/>
-			<DotNet Arguments="publish &quot;$(UgsDir)/UnrealGameSyncCmd/UnrealGameSyncCmd.csproj&quot; --output=&quot;$(UgsPublishDir)&quot; --no-self-contained -c Release -r win-x64 $(VersionArguments)"/>
+			<DotNet Arguments="publish &quot;$(UgsDir)/UnrealGameSyncCmd/UnrealGameSyncCmd.csproj&quot; --output=&quot;$(UgsPublishDir)&quot; --self-contained -c Release -r win-x64 $(VersionArguments)"/>
 			<DotNet Arguments="publish &quot;$(UgsDir)/UnrealGameSyncCmd/UnrealGameSyncCmd.csproj&quot; --output=&quot;$(UgsPublishDir)&quot; --no-self-contained -c Release -r osx-x64 $(VersionArguments)"/>
 			<DotNet Arguments="publish &quot;$(UgsDir)/UnrealGameSyncCmd/UnrealGameSyncCmd.csproj&quot; --output=&quot;$(UgsPublishDir)&quot; --no-self-contained -c Release -r linux-x64 $(VersionArguments)"/>
-			<DotNet Arguments="publish &quot;$(UgsDir)/UnrealGameSync/UnrealGameSync.csproj&quot; --output=&quot;$(UgsPublishDir)&quot; --no-self-contained -c Release -r win-x64 $(VersionArguments)"/>
+			<DotNet Arguments="publish &quot;$(UgsDir)/UnrealGameSync/UnrealGameSync.csproj&quot; --output=&quot;$(UgsPublishDir)&quot; --self-contained -c Release -r win-x64 $(VersionArguments)"/>
 			<Copy From="$(InstallerDir)/bin/Release/UnrealGameSync.msi" To="$(UgsPublishDir)/UnrealGameSync.msi" If="$(IncludeInstaller)"/>
 		</Node>
 		<Node Name="Publish Unstable" Requires="Build UnrealGameSync">
