UGS: Fix NET 6 static analyzer warnings in UnrealGameSyncCmd.

diff --git a/Engine/Source/Programs/UnrealGameSync/UnrealGameSyncCmd/Program.cs b/Engine/Source/Programs/UnrealGameSync/UnrealGameSyncCmd/Program.cs
--- a/Engine/Source/Programs/UnrealGameSync/UnrealGameSyncCmd/Program.cs
+++ b/Engine/Source/Programs/UnrealGameSync/UnrealGameSyncCmd/Program.cs
@@ -23,7 +23,7 @@ namespace UnrealGameSyncCmd
 {
 	using ILogger = Microsoft.Extensions.Logging.ILogger;
 
-	sealed class UserErrorException : Exception
+	public sealed class UserErrorException : Exception
 	{
 		public LogEvent Event { get; }
 		public int Code { get; }
@@ -41,7 +41,7 @@ namespace UnrealGameSyncCmd
 		}
 	}
 
-	public class Program
+	public static class Program
 	{
 		static BuildConfig EditorConfig => BuildConfig.Development;
 
@@ -228,7 +228,7 @@ namespace UnrealGameSyncCmd
 				CommandInfo? command = _commands.FirstOrDefault(x => x.Name.Equals(commandName, StringComparison.OrdinalIgnoreCase));
 				if (command == null)
 				{
-					logger.LogError($"unknown command '{commandName}'");
+					logger.LogError("Unknown command '{Command}'", commandName);
 					Console.WriteLine();
 					PrintHelp();
 					return 1;
@@ -351,7 +351,7 @@ namespace UnrealGameSyncCmd
 			{
 				searchPath = $"//{clientName}{branchPath}/*.uprojectdirs";
 			}
-			else if (projectName.Contains('.'))
+			else if (projectName.Contains('.', StringComparison.Ordinal))
 			{
 				searchPath = $"//{clientName}{branchPath}/{projectName.TrimStart('/')}";
 			}
@@ -413,7 +413,7 @@ namespace UnrealGameSyncCmd
 				}
 			}
 
-			async Task InitNewClientAsync(IPerforceConnection perforce, string streamName, string hostName, ProjectInitOptions options, ILogger logger)
+			static async Task InitNewClientAsync(IPerforceConnection perforce, string streamName, string hostName, ProjectInitOptions options, ILogger logger)
 			{
 				logger.LogInformation("Checking stream...");
 
@@ -479,7 +479,7 @@ namespace UnrealGameSyncCmd
 				}
 			}
 
-			async Task InitExistingClientAsync(IPerforceConnection perforce, string hostName, ProjectInitOptions options, ILogger logger)
+			static async Task InitExistingClientAsync(IPerforceConnection perforce, string hostName, ProjectInitOptions options, ILogger logger)
 			{
 				DirectoryReference currentDir = DirectoryReference.GetCurrentDirectory();
 
@@ -507,7 +507,7 @@ namespace UnrealGameSyncCmd
 
 				// If a project path was specified in local syntax, try to convert it to client-relative syntax
 				string? projectName = options.ProjectName;
-				if (options.ProjectName != null && options.ProjectName.Contains('.'))
+				if (options.ProjectName != null && options.ProjectName.Contains('.', StringComparison.Ordinal))
 				{
 					options.ProjectName = FileReference.Combine(currentDir, options.ProjectName).MakeRelativeTo(clientDir).Replace('\\', '/');
 				}
@@ -533,7 +533,7 @@ namespace UnrealGameSyncCmd
 				List<ClientsRecord> clients = await perforce.GetClientsAsync(ClientsOptions.None, perforce.Settings.UserName);
 				foreach (ClientsRecord client in clients)
 				{
-					if (!String.IsNullOrEmpty(client.Root) && !String.IsNullOrEmpty(client.Host) && String.Compare(hostName, client.Host, StringComparison.OrdinalIgnoreCase) == 0)
+					if (!String.IsNullOrEmpty(client.Root) && !String.IsNullOrEmpty(client.Host) && String.Equals(hostName, client.Host, StringComparison.OrdinalIgnoreCase))
 					{
 						DirectoryReference? rootDir;
 						try
@@ -580,13 +580,13 @@ namespace UnrealGameSyncCmd
 				public bool Refilter { get; set; }
 			}
 
-			async Task<bool> IsCodeChangeAsync(IPerforceConnection perforce, int change)
+			static async Task<bool> IsCodeChangeAsync(IPerforceConnection perforce, int change)
 			{
 				DescribeRecord describeRecord = await perforce.DescribeAsync(change);
 				return IsCodeChange(describeRecord);
 			}
 
-			bool IsCodeChange(DescribeRecord describeRecord)
+			static bool IsCodeChange(DescribeRecord describeRecord)
 			{
 				foreach (DescribeFileRecord file in describeRecord.Files)
 				{
@@ -816,7 +816,7 @@ namespace UnrealGameSyncCmd
 
 				if (!Utility.SpawnProcess(receipt.Launch, commandLine))
 				{
-					logger.LogError("Unable to spawn {0} {1}", receipt.Launch, launchArguments.ToString());
+					logger.LogError("Unable to spawn {App} {Args}", receipt.Launch, launchArguments.ToString());
 				}
 			}
 		}
@@ -1250,7 +1250,7 @@ namespace UnrealGameSyncCmd
 				}
 			}
 
-			public async Task SwitchStreamAsync(IPerforceConnection perforceClient, string streamName, bool force, ILogger logger)
+			public static async Task SwitchStreamAsync(IPerforceConnection perforceClient, string streamName, bool force, ILogger logger)
 			{
 				if (!force && await perforceClient.OpenedAsync(OpenedOptions.None, FileSpecList.Any).AnyAsync())
 				{
@@ -1262,7 +1262,7 @@ namespace UnrealGameSyncCmd
 				logger.LogInformation("Switched to stream {StreamName}", streamName);
 			}
 
-			public async Task SwitchProjectAsync(IPerforceConnection perforceClient, UserWorkspaceSettings settings, string projectName, ILogger logger)
+			public static async Task SwitchProjectAsync(IPerforceConnection perforceClient, UserWorkspaceSettings settings, string projectName, ILogger logger)
 			{
 				settings.ProjectPath = await FindProjectPathAsync(perforceClient, settings.ClientName, settings.BranchPath, projectName);
 				settings.Save(logger);
