UGS: Disable false positive warning for not disposing of memory stream.

diff --git a/Engine/Source/Programs/UnrealGameSync/UnrealGameSyncShared/ConfigUtils.cs b/Engine/Source/Programs/UnrealGameSync/UnrealGameSyncShared/ConfigUtils.cs
--- a/Engine/Source/Programs/UnrealGameSync/UnrealGameSyncShared/ConfigUtils.cs
+++ b/Engine/Source/Programs/UnrealGameSync/UnrealGameSyncShared/ConfigUtils.cs
@@ -217,8 +217,10 @@ namespace UnrealGameSync
 				ConfigurationBuilder configBuilder = new ConfigurationBuilder();
 				foreach (string[] configFile in latestOidcConfigFiles)
 				{
+#pragma warning disable CA2000
 					MemoryStream configStream = new MemoryStream(Encoding.UTF8.GetBytes(String.Join("\n", configFile)));
 					configBuilder.AddJsonStream(configStream);
+#pragma warning restore CA2000
 				}
 
 				IConfigurationRoot config = configBuilder.Build();
