UGS: Fix data being lost due to collections not being serialized.

diff --git a/Engine/Source/Programs/UnrealGameSync/UnrealGameSyncShared/GlobalSettings.cs b/Engine/Source/Programs/UnrealGameSync/UnrealGameSyncShared/GlobalSettings.cs
--- a/Engine/Source/Programs/UnrealGameSync/UnrealGameSyncShared/GlobalSettings.cs
+++ b/Engine/Source/Programs/UnrealGameSync/UnrealGameSyncShared/GlobalSettings.cs
@@ -10,9 +10,9 @@ namespace UnrealGameSync
 {
 	public class FilterSettings
 	{
-		public List<Guid> IncludeCategories { get; } = new List<Guid>();
-		public List<Guid> ExcludeCategories { get; } = new List<Guid>();
-		public List<string> View { get; } = new List<string>();
+		public List<Guid> IncludeCategories { get; init; } = new List<Guid>();
+		public List<Guid> ExcludeCategories { get; init; } = new List<Guid>();
+		public List<string> View { get; init; } = new List<string>();
 		public bool? AllProjects { get; set; }
 		public bool? AllProjectsInSln { get; set; }
 
diff --git a/Engine/Source/Programs/UnrealGameSync/UnrealGameSyncShared/UserSettings.cs b/Engine/Source/Programs/UnrealGameSync/UnrealGameSyncShared/UserSettings.cs
--- a/Engine/Source/Programs/UnrealGameSync/UnrealGameSyncShared/UserSettings.cs
+++ b/Engine/Source/Programs/UnrealGameSync/UnrealGameSyncShared/UserSettings.cs
@@ -65,7 +65,7 @@ namespace UnrealGameSync
 		public bool Good { get; set; } = false;
 		public bool Starred { get; set; } = false;
 		public bool FindNewestGoodContent { get; set; } = false;
-		public List<string> RequiredBadges { get; } = new List<string>();
+		public List<string> RequiredBadges { get; init; } = new List<string>();
 
 		// Depot path to read for the latest change number
 		public string? ReadFrom { get; set; }
@@ -152,7 +152,7 @@ namespace UnrealGameSync
 	{
 		public bool Enabled { get; set; }
 		public string Type { get; set; }
-		public List<string> Order { get; }
+		public List<string> Order { get; init; }
 
 		public ArchiveSettings(bool enabled, string type, IEnumerable<string> order)
 		{
@@ -416,9 +416,9 @@ namespace UnrealGameSync
 		[JsonIgnore]
 		public FileReference ConfigFile { get; private set; } = null!;
 
-		public List<ConfigObject> BuildSteps { get; } = new List<ConfigObject>();
+		public List<ConfigObject> BuildSteps { get; init; } = new List<ConfigObject>();
 		public FilterType FilterType { get; set; }
-		public HashSet<string> FilterBadges { get; } = new HashSet<string>(StringComparer.OrdinalIgnoreCase);
+		public HashSet<string> FilterBadges { get; init; } = new HashSet<string>(StringComparer.OrdinalIgnoreCase);
 
 		static readonly object _syncRoot = new object();
 
@@ -496,21 +496,21 @@ namespace UnrealGameSync
 		public bool AnnotateRobmergeChanges { get; set; }
 		public bool ShowLocalTimes { get; set; }
 		public bool KeepInTray { get; set; }
-		public HashSet<Guid> EnabledTools { get; } = new HashSet<Guid>();
+		public HashSet<Guid> EnabledTools { get; init; } = new HashSet<Guid>();
 		public int FilterIndex { get; set; }
 		public UserSelectedProjectSettings? LastProject { get; set; }
-		public List<UserSelectedProjectSettings> OpenProjects { get; }
-		public List<UserSelectedProjectSettings> RecentProjects { get; }
+		public List<UserSelectedProjectSettings> OpenProjects { get; init; }
+		public List<UserSelectedProjectSettings> RecentProjects { get; init; }
 		public string SyncTypeId { get; set; }
 		public BuildConfig CompiledEditorBuildConfig { get; set; } // NB: This assumes not using precompiled editor. See CurrentBuildConfig.
 		public TabLabels TabLabels { get; set; }
 		public long NextLauncherVersionCheck { get; set; }
 
 		// Precompiled binaries
-		public List<ArchiveSettings> Archives { get; } = new List<ArchiveSettings>();
+		public List<ArchiveSettings> Archives { get; init; } = new List<ArchiveSettings>();
 
 		// OIDC Settings
-		public Dictionary<string, string> ProviderToRefreshTokens { get; } = new Dictionary<string, string>();
+		public Dictionary<string, string> ProviderToRefreshTokens { get; init; } = new Dictionary<string, string>();
 
 		// Window settings
 		public bool WindowVisible { get; set; }
@@ -521,14 +521,14 @@ namespace UnrealGameSync
 		public bool ScheduleEnabled { get; set; }
 		public TimeSpan ScheduleTime { get; set; }
 		public bool ScheduleAnyOpenProject { get; set; }
-		public List<UserSelectedProjectSettings> ScheduleProjects { get; } = new List<UserSelectedProjectSettings>();
+		public List<UserSelectedProjectSettings> ScheduleProjects { get; init; } = new List<UserSelectedProjectSettings>();
 
 		// Run configuration
-		public List<Tuple<string, bool>> EditorArguments { get; } = new List<Tuple<string, bool>>();
+		public List<Tuple<string, bool>> EditorArguments { get; init; } = new List<Tuple<string, bool>>();
 		public bool EditorArgumentsPrompt { get; set; }
 
 		// Notification settings
-		public List<string> NotifyProjects { get; } = new List<string>();
+		public List<string> NotifyProjects { get; init; } = new List<string>();
 		public int NotifyUnassignedMinutes { get; set; }
 		public int NotifyUnacknowledgedMinutes { get; set; }
 		public int NotifyUnresolvedMinutes { get; set; }
diff --git a/Engine/Source/Programs/UnrealGameSync/UnrealGameSyncShared/WorkspaceState.cs b/Engine/Source/Programs/UnrealGameSync/UnrealGameSyncShared/WorkspaceState.cs
--- a/Engine/Source/Programs/UnrealGameSync/UnrealGameSyncShared/WorkspaceState.cs
+++ b/Engine/Source/Programs/UnrealGameSync/UnrealGameSyncShared/WorkspaceState.cs
@@ -34,7 +34,7 @@ namespace UnrealGameSync
 		public int CurrentChangeNumber { get; set; } = -1;
 		public int CurrentCodeChangeNumber { get; set; } = -1;
 		public string? CurrentSyncFilterHash { get; set; }
-		public List<int> AdditionalChangeNumbers { get; } = new List<int>();
+		public List<int> AdditionalChangeNumbers { get; init; } = new List<int>();
 
 		// Settings for the last attempted sync. These values are set to persist error messages between runs.
 		public int LastSyncChangeNumber { get; set; }
@@ -47,10 +47,10 @@ namespace UnrealGameSync
 		public int LastBuiltChangeNumber { get; set; }
 
 		// Expanded archives in the workspace
-		public HashSet<string> ExpandedArchiveTypes { get; } = new HashSet<string>(StringComparer.Ordinal);
+		public HashSet<string> ExpandedArchiveTypes { get; init; } = new HashSet<string>(StringComparer.Ordinal);
 
 		// The changes that we're regressing at the moment
-		public List<BisectEntry> BisectChanges { get; } = new List<BisectEntry>();
+		public List<BisectEntry> BisectChanges { get; init; } = new List<BisectEntry>();
 
 		public void UpdateCachedProjectInfo(ProjectInfo projectInfo, long settingsTimeUtc)
 		{
