UGS: Fix 'Clean Workspace' functionality when project is synced to the root of a drive.

diff --git a/Engine/Source/Programs/UnrealGameSync/UnrealGameSync/FindFoldersToCleanTask.cs b/Engine/Source/Programs/UnrealGameSync/UnrealGameSync/FindFoldersToCleanTask.cs
--- a/Engine/Source/Programs/UnrealGameSync/UnrealGameSync/FindFoldersToCleanTask.cs
+++ b/Engine/Source/Programs/UnrealGameSync/UnrealGameSync/FindFoldersToCleanTask.cs
@@ -263,7 +263,11 @@ namespace UnrealGameSync
 			}
 
 			// Get the prefix for any local file
-			string localRootPrefix = _rootFolderToClean._directory.FullName.Replace(Path.AltDirectorySeparatorChar, Path.DirectorySeparatorChar) + Path.DirectorySeparatorChar;
+			string localRootPrefix = _rootFolderToClean._directory.FullName.Replace(Path.AltDirectorySeparatorChar, Path.DirectorySeparatorChar);
+			if (!localRootPrefix.EndsWith(Path.DirectorySeparatorChar))
+			{
+				localRootPrefix.Append(Path.DirectorySeparatorChar);
+			}
 
 			// Query the have table and build a separate tree from it
 			PerforceHaveFolder rootHaveFolder = new PerforceHaveFolder();
