UGS: Add an option to enumerate all changes covered by PCBs (Perforce.FindAllChangesForPCBs). This code was disabled previously due to the large load it can place on the Perforce server, but is technically correct and may be desirable in certain situations.

diff --git a/Engine/Source/Programs/UnrealGameSync/UnrealGameSync/PerforceMonitor.cs b/Engine/Source/Programs/UnrealGameSync/UnrealGameSync/PerforceMonitor.cs
--- a/Engine/Source/Programs/UnrealGameSync/UnrealGameSync/PerforceMonitor.cs
+++ b/Engine/Source/Programs/UnrealGameSync/UnrealGameSync/PerforceMonitor.cs
@@ -257,10 +257,10 @@ namespace UnrealGameSync
 			int newestChangeNumber = -1;
 			HashSet<int> currentChangelists;
 			SortedSet<int> prevPromotedChangelists;
-			lock(_lockObject)
+			lock (_lockObject)
 			{
 				maxChanges = PendingMaxChanges;
-				if(_changes.Count > 0)
+				if (_changes.Count > 0)
 				{
 					newestChangeNumber = _changes.First().Number;
 					oldestChangeNumber = _changes.Last().Number;
@@ -269,6 +269,9 @@ namespace UnrealGameSync
 				prevPromotedChangelists = new SortedSet<int>(_promotedChangeNumbers);
 			}
 
+			// Get the Perforce section from the config file
+			ConfigSection? perforceConfigSection = LatestProjectConfigFile.FindSection("Perforce");
+
 			// Build a full list of all the paths to sync
 			List<string> depotPaths = new List<string>();
 			if (ShowChangesForAllProjects || _selectedClientFileName.EndsWith(".uprojectdirs", StringComparison.InvariantCultureIgnoreCase))
@@ -286,10 +289,9 @@ namespace UnrealGameSync
 				}
 
 				// Add in additional paths property
-				ConfigSection? projectConfigSection = LatestProjectConfigFile.FindSection("Perforce");
-				if (projectConfigSection != null)
+				if (perforceConfigSection != null)
 				{
-					IEnumerable<string> additionalPaths = projectConfigSection.GetValues("AdditionalPathsToSync", Array.Empty<string>());
+					IEnumerable<string> additionalPaths = perforceConfigSection.GetValues("AdditionalPathsToSync", Array.Empty<string>());
 
 					// turn into //ws/path
 					depotPaths.AddRange(additionalPaths.Select(p => String.Format("{0}/{1}", _branchClientPath, p.TrimStart('/'))));
@@ -298,14 +300,14 @@ namespace UnrealGameSync
 
 			// Read any new changes
 			List<ChangesRecord> newChanges;
-			if(maxChanges > CurrentMaxChanges || newestChangeNumber == -1)
+			if (maxChanges > CurrentMaxChanges || newestChangeNumber == -1)
 			{
 				newChanges = await perforce.GetChangesAsync(ChangesOptions.IncludeTimes | ChangesOptions.LongOutput, maxChanges, ChangeStatus.Submitted, depotPaths, cancellationToken);
 			}
 			else
 			{
 				List<string> depotPathsWithRange = depotPaths.ConvertAll(x => $"{x}@{newestChangeNumber + 1},#head");
-//				newChanges = await perforce.GetChangesAsync(ChangesOptions.IncludeTimes | ChangesOptions.LongOutput, maxChanges, ChangeStatus.Submitted, depotPaths.Select(x => $"{x}@>{newestChangeNumber}").ToArray(), cancellationToken);
+				//				newChanges = await perforce.GetChangesAsync(ChangesOptions.IncludeTimes | ChangesOptions.LongOutput, maxChanges, ChangeStatus.Submitted, depotPaths.Select(x => $"{x}@>{newestChangeNumber}").ToArray(), cancellationToken);
 				newChanges = await perforce.GetChangesAsync(ChangesOptions.IncludeTimes | ChangesOptions.LongOutput, clientName: null, minChangeNumber: newestChangeNumber + 1, maxChanges: maxChanges, status: ChangeStatus.Submitted, userName: null, fileSpecs: depotPathsWithRange, cancellationToken);
 			}
 
@@ -313,34 +315,35 @@ namespace UnrealGameSync
 			newChanges.RemoveAll(x => currentChangelists.Contains(x.Number));
 
 			// Update the change ranges
-			if(newChanges.Count > 0)
+			if (newChanges.Count > 0)
 			{
 				oldestChangeNumber = Math.Max(oldestChangeNumber, newChanges.Last().Number);
 				newestChangeNumber = Math.Min(newestChangeNumber, newChanges.First().Number);
 			}
 
 			// The code below is correct, but can cause a lot of load on the Perforce server when we query a large number of changes because PCBs are far behind.
-#if false
 			// If we are using zipped binaries, make sure we have every change since the last zip containing them. This is necessary for ensuring that content changes show as
 			// syncable in the workspace view if there have been a large number of content changes since the last code change.
-			int minZippedChangeNumber = -1;
-			foreach (PerforceArchiveInfo archive in _archives)
+			if (perforceConfigSection != null && perforceConfigSection.GetValue("FindAllChangesForPCBs", false))
 			{
-				foreach (int changeNumber in archive.ChangeNumberToFileRevision.Keys)
+				int minZippedChangeNumber = -1;
+				foreach (PerforceArchiveInfo archive in _archives)
 				{
-					if (changeNumber > minZippedChangeNumber && changeNumber <= oldestChangeNumber)
+					foreach (int changeNumber in archive.ChangeNumberToFileRevision.Keys)
 					{
-						minZippedChangeNumber = changeNumber;
+						if (changeNumber > minZippedChangeNumber && changeNumber <= oldestChangeNumber)
+						{
+							minZippedChangeNumber = changeNumber;
+						}
 					}
 				}
+				if (minZippedChangeNumber != -1 && minZippedChangeNumber < oldestChangeNumber)
+				{
+					string[] filteredPaths = depotPaths.Select(x => $"{x}@{minZippedChangeNumber},{oldestChangeNumber - 1}").ToArray();
+					List<ChangesRecord> zipChanges = await perforce.GetChangesAsync(ChangesOptions.None, -1, ChangeStatus.Submitted, filteredPaths, cancellationToken);
+					newChanges.AddRange(zipChanges);
+				}
 			}
-			if(minZippedChangeNumber != -1 && minZippedChangeNumber < oldestChangeNumber)
-			{
-				string[] filteredPaths = depotPaths.Select(x => $"{x}@{minZippedChangeNumber},{oldestChangeNumber - 1}").ToArray();
-				List<ChangesRecord> zipChanges = await perforce.GetChangesAsync(ChangesOptions.None, -1, ChangeStatus.Submitted, filteredPaths, cancellationToken);
-				newChanges.AddRange(zipChanges);
-			}
-#endif
 
 			// Fixup any ROBOMERGE authors
 			const string roboMergePrefix = "#ROBOMERGE-AUTHOR:";
