PR #10187: UGS: Expose Perforce settings to build steps (Contributed by [REDACTED_PERSON])

diff --git a/Engine/Source/Programs/UnrealGameSync/UnrealGameSync/Workspace.cs b/Engine/Source/Programs/UnrealGameSync/UnrealGameSync/Workspace.cs
--- a/Engine/Source/Programs/UnrealGameSync/UnrealGameSync/Workspace.cs
+++ b/Engine/Source/Programs/UnrealGameSync/UnrealGameSync/Workspace.cs
@@ -225,7 +225,7 @@ namespace UnrealGameSync
 				editorReceipt = ConfigUtils.CreateDefaultEditorReceipt(Project, ProjectConfigFile, editorConfig);
 			}
 
-			Dictionary<string, string> variables = ConfigUtils.GetWorkspaceVariables(Project, overrideChange ?? State.CurrentChangeNumber, overrideCodeChange ?? State.CurrentCodeChangeNumber, editorReceipt, ProjectConfigFile);
+			Dictionary<string, string> variables = ConfigUtils.GetWorkspaceVariables(Project, overrideChange ?? State.CurrentChangeNumber, overrideCodeChange ?? State.CurrentCodeChangeNumber, editorReceipt, ProjectConfigFile, PerforceSettings);
 			return variables;
 		}
 
diff --git a/Engine/Source/Programs/UnrealGameSync/UnrealGameSyncShared/ConfigUtils.cs b/Engine/Source/Programs/UnrealGameSync/UnrealGameSyncShared/ConfigUtils.cs
--- a/Engine/Source/Programs/UnrealGameSync/UnrealGameSyncShared/ConfigUtils.cs
+++ b/Engine/Source/Programs/UnrealGameSync/UnrealGameSyncShared/ConfigUtils.cs
@@ -425,7 +425,7 @@ namespace UnrealGameSync
 			return defaultBuildSteps.ToDictionary(x => x.UniqueId, x => x.ToConfigObject());
 		}
 
-		public static Dictionary<string, string> GetWorkspaceVariables(ProjectInfo projectInfo, int changeNumber, int codeChangeNumber, TargetReceipt? editorTarget, ConfigFile? projectConfigFile)
+		public static Dictionary<string, string> GetWorkspaceVariables(ProjectInfo projectInfo, int changeNumber, int codeChangeNumber, TargetReceipt? editorTarget, ConfigFile? projectConfigFile, IPerforceSettings perforceSettings)
 		{
 			Dictionary<string, string> variables = new Dictionary<string, string>(StringComparer.OrdinalIgnoreCase);
 
@@ -466,12 +466,19 @@ namespace UnrealGameSync
 				}
 			}
 
+			variables.Add("PerforceServerAndPort", perforceSettings.ServerAndPort);
+			variables.Add("PerforceUserName", perforceSettings.UserName);
+			if (perforceSettings.ClientName != null)
+			{
+				variables.Add("PerforceClientName", perforceSettings.ClientName);
+			}
+
 			return variables;
 		}
 
-		public static Dictionary<string, string> GetWorkspaceVariables(ProjectInfo projectInfo, int changeNumber, int codeChangeNumber, TargetReceipt? editorTarget, ConfigFile? projectConfigFile, IEnumerable<KeyValuePair<string, string>> additionalVariables)
+		public static Dictionary<string, string> GetWorkspaceVariables(ProjectInfo projectInfo, int changeNumber, int codeChangeNumber, TargetReceipt? editorTarget, ConfigFile? projectConfigFile, IPerforceSettings perforceSettings, IEnumerable<KeyValuePair<string, string>> additionalVariables)
 		{
-			Dictionary<string, string> variables = GetWorkspaceVariables(projectInfo, changeNumber, codeChangeNumber, editorTarget, projectConfigFile);
+			Dictionary<string, string> variables = GetWorkspaceVariables(projectInfo, changeNumber, codeChangeNumber, editorTarget, projectConfigFile, perforceSettings);
 			foreach ((string key, string value) in additionalVariables)
 			{
 				variables[key] = value;
diff --git a/Engine/Source/Programs/UnrealGameSync/UnrealGameSyncShared/WorkspaceUpdate.cs b/Engine/Source/Programs/UnrealGameSync/UnrealGameSyncShared/WorkspaceUpdate.cs
--- a/Engine/Source/Programs/UnrealGameSync/UnrealGameSyncShared/WorkspaceUpdate.cs
+++ b/Engine/Source/Programs/UnrealGameSync/UnrealGameSyncShared/WorkspaceUpdate.cs
@@ -1066,7 +1066,7 @@ namespace UnrealGameSync
 							logger.LogInformation("");
 							logger.LogInformation("Executing post-sync steps...");
 
-							Dictionary<string, string> postSyncVariables = ConfigUtils.GetWorkspaceVariables(project, Context.ChangeNumber, versionChangeNumber, null, Context.ProjectConfigFile);
+							Dictionary<string, string> postSyncVariables = ConfigUtils.GetWorkspaceVariables(project, Context.ChangeNumber, versionChangeNumber, null, Context.ProjectConfigFile, perforceSettings);
 							foreach (string postSyncStep in postSyncSteps.Select(x => x.Trim()))
 							{
 								ConfigObject postSyncStepObject = new ConfigObject(postSyncStep);
@@ -1305,7 +1305,7 @@ namespace UnrealGameSync
 							{
 								editorReceipt = defaultEditorReceipt;
 							}
-							Dictionary<string, string> variables = ConfigUtils.GetWorkspaceVariables(project, state.CurrentChangeNumber, state.CurrentCodeChangeNumber, editorReceipt, Context.ProjectConfigFile, Context.AdditionalVariables);
+							Dictionary<string, string> variables = ConfigUtils.GetWorkspaceVariables(project, state.CurrentChangeNumber, state.CurrentCodeChangeNumber, editorReceipt, Context.ProjectConfigFile, perforceSettings, Context.AdditionalVariables);
 
 							// Handle all the different types of steps
 							switch (step.Type)
