PR #10152: UGS: Speed up self-update in some cases (Contributed by [REDACTED_PERSON])

diff --git a/Engine/Source/Programs/UnrealGameSync/UnrealGameSyncLauncher/Program.cs b/Engine/Source/Programs/UnrealGameSync/UnrealGameSyncLauncher/Program.cs
--- a/Engine/Source/Programs/UnrealGameSync/UnrealGameSyncLauncher/Program.cs
+++ b/Engine/Source/Programs/UnrealGameSync/UnrealGameSyncLauncher/Program.cs
@@ -87,6 +87,9 @@ namespace UnrealGameSyncLauncher
 			return 1;
 		}
 
+		// Values of the Perforce "action" field that means that the file should no longer be synced
+		public static readonly string[] s_deleteActions = { "delete", "move/delete", "purge", "archive" };
+
 		public static async Task SyncAndRun(IPerforceConnection perforce, string? baseDepotPath, bool preview, string[] args, Mutex instanceMutex, ILogger logger, CancellationToken cancellationToken)
 		{
 			try
@@ -155,6 +158,12 @@ namespace UnrealGameSyncLauncher
 					string depotPathPrefix = syncPath.Substring(0, syncPath.LastIndexOf('/') + 1);
 					foreach (FStatRecord fileRecord in fileRecords)
 					{
+						// Skip deleted files
+						if (Array.IndexOf(s_deleteActions, fileRecord.Action) != -1)
+						{
+							continue;
+						}
+
 						if (fileRecord.DepotFile == null)
 						{
 							throw new UserErrorException("Missing depot path for returned file");
