UGS: Use a data-driven method for printing CLI arguments, to prevent them getting out of data.

Also add a -Clean option to the sync command to allow cleaning before building.

diff --git a/Engine/Source/Programs/UnrealGameSync/UnrealGameSyncCmd/Program.cs b/Engine/Source/Programs/UnrealGameSync/UnrealGameSyncCmd/Program.cs
--- a/Engine/Source/Programs/UnrealGameSync/UnrealGameSyncCmd/Program.cs
+++ b/Engine/Source/Programs/UnrealGameSync/UnrealGameSyncCmd/Program.cs
@@ -18,6 +18,8 @@ using System.Threading.Tasks;
 using EpicGames.OIDC;
 using Microsoft.Extensions.Configuration;
 using UnrealGameSync;
+using System.Text;
+using System.Globalization;
 
 namespace UnrealGameSyncCmd
 {
@@ -49,13 +51,15 @@ namespace UnrealGameSyncCmd
 		{
 			public string Name { get; }
 			public Type Type { get; }
+			public Type? OptionsType { get; }
 			public string Usage { get; }
 			public string Brief { get; }
 
-			public CommandInfo(string name, Type type, string usage, string brief)
+			public CommandInfo(string name, Type type, Type? optionsType, string usage, string brief)
 			{
 				Name = name;
 				Type = type;
+				OptionsType = optionsType;
 				Usage = usage;
 				Brief = brief;
 			}
@@ -63,51 +67,51 @@ namespace UnrealGameSyncCmd
 
 		static readonly CommandInfo[] _commands =
 		{
-			new CommandInfo("init", typeof(InitCommand),
-				"ugs init [stream-path] [-client=..] [-server=..] [-user=..] [-branch=..] [-project=..]",
+			new CommandInfo("init", typeof(InitCommand), typeof(InitCommandOptions),
+				"ugs init [stream-path]",
 				"Create a client for the given stream, or initializes an existing client for use by UGS."
 			),
-			new CommandInfo("switch", typeof(SwitchCommand),
+			new CommandInfo("switch", typeof(SwitchCommand), typeof(SwitchCommandOptions),
 				"ugs switch [project name|project path|stream]",
 				"Changes the active project to the one in the workspace with the given name, or switches to a new stream."
 			),
-			new CommandInfo("changes", typeof(ChangesCommand),
+			new CommandInfo("changes", typeof(ChangesCommand), typeof(ChangesCommandOptions),
 				"ugs changes",
 				"List recently submitted changes to the current branch."
 			),
-			new CommandInfo("config", typeof(ConfigCommand),
+			new CommandInfo("config", typeof(ConfigCommand), typeof(ConfigCommandOptions),
 				"ugs config",
 				"Updates the configuration for the current workspace."
 			),
-			new CommandInfo("filter", typeof(FilterCommand),
-				"ugs filter [-reset] [-include=..] [-exclude=..] [-view=..] [-addview=..] [-removeview=..] [-global]",
+			new CommandInfo("filter", typeof(FilterCommand), typeof(FilterCommandOptions),
+				"ugs filter",
 				"Displays or updates the workspace or global sync filter"
 			),
-			new CommandInfo("sync", typeof(SyncCommand),
-				"ugs sync [change|'latest'] [-build] [-binaries] [-remove] [-only]",
+			new CommandInfo("sync", typeof(SyncCommand), typeof(SyncCommandOptions),
+				"ugs sync [change|'latest']",
 				"Syncs the current workspace to the given changelist, optionally removing all local state."
 			),
-			new CommandInfo("clients", typeof(ClientsCommand),
+			new CommandInfo("clients", typeof(ClientsCommand), typeof(ClientsCommandOptions),
 				"ugs clients",
 				"Lists all clients suitable for use on the current machine."
 			),
-			new CommandInfo("run", typeof(RunCommand),
+			new CommandInfo("run", typeof(RunCommand), null,
 				"ugs run",
 				"Runs the editor for the current branch."
 			),
-			new CommandInfo("build", typeof(BuildCommand),
-				"ugs build [id] [-list]",
+			new CommandInfo("build", typeof(BuildCommand), typeof(BuildCommandOptions),
+				"ugs build [id]",
 				"Runs the default build steps for the current project, or a particular step referenced by id."
 			),
-			new CommandInfo("status", typeof(StatusCommand),
+			new CommandInfo("status", typeof(StatusCommand), null,
 				"ugs status [-update]",
 				"Shows the status of the currently synced branch."
 			),
-			new CommandInfo("login", typeof(LoginCommand),
+			new CommandInfo("login", typeof(LoginCommand), null,
 				"ugs login",
 				"Starts a interactive login flow against the configured Identity Provider"
 			),
-			new CommandInfo("version", typeof(VersionCommand),
+			new CommandInfo("version", typeof(VersionCommand), null,
 				"ugs version",
 				"Prints the current application version"
 			),
@@ -138,7 +142,7 @@ namespace UnrealGameSyncCmd
 			public string? UserName { get; set; }
 		}
 
-		class ProjectConfigOptions : ServerOptions
+		class ConfigCommandOptions : ServerOptions
 		{
 			public void ApplyTo(UserWorkspaceSettings settings)
 			{
@@ -153,7 +157,7 @@ namespace UnrealGameSyncCmd
 			}
 		}
 
-		class ProjectInitOptions : ProjectConfigOptions
+		class InitCommandOptions : ConfigCommandOptions
 		{
 			[CommandLine("-Client=")]
 			public string? ClientName { get; set; }
@@ -261,11 +265,40 @@ namespace UnrealGameSyncCmd
 			foreach (CommandInfo command in _commands)
 			{
 				Console.WriteLine();
-				ConsoleUtils.WriteLineWithWordWrap(command.Usage, 2, 8);
+				ConsoleUtils.WriteLineWithWordWrap(GetUsage(command), 2, 8);
 				ConsoleUtils.WriteLineWithWordWrap(command.Brief, 4, 4);
 			}
 		}
 
+		static string GetUsage(CommandInfo commandInfo)
+		{
+			StringBuilder result = new StringBuilder(commandInfo.Usage);
+			if (commandInfo.OptionsType != null)
+			{
+				foreach (PropertyInfo propertyInfo in commandInfo.OptionsType.GetProperties(BindingFlags.Instance | BindingFlags.Public))
+				{
+					List<string> names = new List<string>();
+					foreach (CommandLineAttribute attribute in propertyInfo.GetCustomAttributes<CommandLineAttribute>())
+					{
+						string name = (attribute.Prefix ?? propertyInfo.Name).ToLower(CultureInfo.InvariantCulture);
+						if (propertyInfo.PropertyType == typeof(bool) || propertyInfo.PropertyType == typeof(bool?))
+						{
+							names.Add(name);
+						}
+						else
+						{
+							names.Add($"{name}..");
+						}
+					}
+					if (names.Count > 0)
+					{
+						result.Append($" [{String.Join('|', names)}]");
+					}
+				}
+			}
+			return result.ToString();
+		}
+
 		public static UserWorkspaceSettings? ReadOptionalUserWorkspaceSettings()
 		{
 			DirectoryReference? dir = DirectoryReference.GetCurrentDirectory();
@@ -393,7 +426,7 @@ namespace UnrealGameSyncCmd
 				context.Arguments.TryGetPositionalArgument(out initName);
 
 				// Get the config settings from the command line
-				ProjectInitOptions options = new ProjectInitOptions();
+				InitCommandOptions options = new InitCommandOptions();
 				context.Arguments.ApplyTo(options);
 				context.Arguments.CheckAllArgumentsUsed();
 
@@ -413,7 +446,7 @@ namespace UnrealGameSyncCmd
 				}
 			}
 
-			static async Task InitNewClientAsync(IPerforceConnection perforce, string streamName, string hostName, ProjectInitOptions options, ILogger logger)
+			static async Task InitNewClientAsync(IPerforceConnection perforce, string streamName, string hostName, InitCommandOptions options, ILogger logger)
 			{
 				logger.LogInformation("Checking stream...");
 
@@ -479,7 +512,7 @@ namespace UnrealGameSyncCmd
 				}
 			}
 
-			static async Task InitExistingClientAsync(IPerforceConnection perforce, string hostName, ProjectInitOptions options, ILogger logger)
+			static async Task InitExistingClientAsync(IPerforceConnection perforce, string hostName, InitCommandOptions options, ILogger logger)
 			{
 				DirectoryReference currentDir = DirectoryReference.GetCurrentDirectory();
 
@@ -556,30 +589,33 @@ namespace UnrealGameSyncCmd
 			}
 		}
 
-		class SyncCommand : Command
+		class SyncCommandOptions
 		{
-			class SyncOptions
-			{
-				[CommandLine("-Only")]
-				public bool SingleChange { get; set; }
+			[CommandLine("-Clean")]
+			public bool Clean { get; set; }
 
-				[CommandLine("-Build")]
-				public bool Build { get; set; }
+			[CommandLine("-Build")]
+			public bool Build { get; set; }
 
-				[CommandLine("-Binaries")]
-				public bool Binaries { get; set; }
+			[CommandLine("-Binaries")]
+			public bool Binaries { get; set; }
 
-				[CommandLine("-NoGPF", Value = "false")]
-				[CommandLine("-NoProjectFiles", Value = "false")]
-				public bool ProjectFiles { get; set; } = true;
+			[CommandLine("-NoGPF", Value = "false")]
+			[CommandLine("-NoProjectFiles", Value = "false")]
+			public bool ProjectFiles { get; set; } = true;
 
-				[CommandLine("-Clobber")]
-				public bool Clobber { get; set; }
+			[CommandLine("-Clobber")]
+			public bool Clobber { get; set; }
 
-				[CommandLine("-Refilter")]
-				public bool Refilter { get; set; }
-			}
+			[CommandLine("-Refilter")]
+			public bool Refilter { get; set; }
+
+			[CommandLine("-Only")]
+			public bool SingleChange { get; set; }
+		}
 
+		class SyncCommand : Command
+		{
 			static async Task<bool> IsCodeChangeAsync(IPerforceConnection perforce, int change)
 			{
 				DescribeRecord describeRecord = await perforce.DescribeAsync(change);
@@ -603,7 +639,7 @@ namespace UnrealGameSyncCmd
 				ILogger logger = context.Logger;
 				context.Arguments.TryGetPositionalArgument(out string? changeString);
 
-				SyncOptions syncOptions = new SyncOptions();
+				SyncCommandOptions syncOptions = new SyncCommandOptions();
 				context.Arguments.ApplyTo(syncOptions);
 
 				context.Arguments.CheckAllArgumentsUsed();
@@ -636,6 +672,10 @@ namespace UnrealGameSyncCmd
 				}
 
 				WorkspaceUpdateOptions options = syncOptions.SingleChange? WorkspaceUpdateOptions.SyncSingleChange : WorkspaceUpdateOptions.Sync;
+				if (syncOptions.Clean)
+				{
+					options |= WorkspaceUpdateOptions.Clean;
+				}
 				if (syncOptions.Build)
 				{
 					options |= WorkspaceUpdateOptions.Build;
@@ -742,17 +782,17 @@ namespace UnrealGameSyncCmd
 			}
 		}
 
-		class ClientsCommand : Command
+		class ClientsCommandOptions : ServerOptions
 		{
-			public class ClientsOptions : ServerOptions
-			{
-			}
+		}
 
+		class ClientsCommand : Command
+		{
 			public override async Task ExecuteAsync(CommandContext context)
 			{
 				ILogger logger = context.Logger;
 
-				ClientsOptions options = context.Arguments.ApplyTo<ClientsOptions>(logger);
+				ClientsCommandOptions options = context.Arguments.ApplyTo<ClientsCommandOptions>(logger);
 				context.Arguments.CheckAllArgumentsUsed();
 
 				using IPerforceConnection perforceClient = await ConnectAsync(options.ServerAndPort, options.UserName, null, context.LoggerFactory);
@@ -821,20 +861,29 @@ namespace UnrealGameSyncCmd
 			}
 		}
 
+		class ChangesCommandOptions
+		{
+			[CommandLine("-Count=")]
+			public int Count { get; set; } = 10;
+
+			[CommandLine("-Lines=")]
+			public int LineCount { get; set; } = 3;
+		}
+
 		class ChangesCommand : Command
 		{
 			public override async Task ExecuteAsync(CommandContext context)
 			{
 				ILogger logger = context.Logger;
 
-				int count = context.Arguments.GetIntegerOrDefault("-Count=", 10);
-				int lineCount = context.Arguments.GetIntegerOrDefault("-Lines=", 3);
+				ChangesCommandOptions options = new ChangesCommandOptions();
+				context.Arguments.ApplyTo(options);
 				context.Arguments.CheckAllArgumentsUsed(context.Logger);
 
 				UserWorkspaceSettings settings = ReadRequiredUserWorkspaceSettings();
 				using IPerforceConnection perforceClient = await ConnectAsync(settings, context.LoggerFactory);
 
-				List<ChangesRecord> changes = await perforceClient.GetChangesAsync(ChangesOptions.None, count, ChangeStatus.Submitted, $"//{settings.ClientName}/...");
+				List<ChangesRecord> changes = await perforceClient.GetChangesAsync(EpicGames.Perforce.ChangesOptions.None, options.Count, ChangeStatus.Submitted, $"//{settings.ClientName}/...");
 				foreach(IEnumerable<ChangesRecord> changesBatch in changes.Batch(10))
 				{
 					List<DescribeRecord> describeRecords = await perforceClient.DescribeAsync(changesBatch.Select(x => x.Number).ToArray());
@@ -876,7 +925,7 @@ namespace UnrealGameSyncCmd
 							lines.Add(String.Empty);
 						}
 
-						lineCount = Math.Min(lineCount, lines.Count);
+						int lineCount = Math.Min(options.LineCount, lines.Count);
 
 						logger.LogInformation("  {Change,-9} {Type,-8} {Author,-15} {Description}", describeRecord.Number, type, author, lines[0]);
 						for (int lineIndex = 1; lineIndex < lineCount; lineIndex++)
@@ -910,7 +959,7 @@ namespace UnrealGameSyncCmd
 				}
 				else
 				{
-					ProjectConfigOptions options = new ProjectConfigOptions();
+					ConfigCommandOptions options = new ConfigCommandOptions();
 					context.Arguments.ApplyTo(options);
 					context.Arguments.CheckAllArgumentsUsed(context.Logger);
 
@@ -924,40 +973,40 @@ namespace UnrealGameSyncCmd
 			}
 		}
 
-		class FilterCommand : Command
+		class FilterCommandOptions
 		{
-			class FilterCommandOptions
-			{
-				[CommandLine("-Reset")]
-				public bool Reset = false;
+			[CommandLine("-Reset")]
+			public bool Reset { get; set; } = false;
 
-				[CommandLine("-Include=")]
-				public List<string> Include { get; set; } = new List<string>();
+			[CommandLine("-Include=")]
+			public List<string> Include { get; set; } = new List<string>();
 
-				[CommandLine("-Exclude=")]
-				public List<string> Exclude { get; set; } = new List<string>();
+			[CommandLine("-Exclude=")]
+			public List<string> Exclude { get; set; } = new List<string>();
 
-				[CommandLine("-View=", ListSeparator = ';')]
-				public List<string>? View { get; set; } 
+			[CommandLine("-View=", ListSeparator = ';')]
+			public List<string>? View { get; set; }
 
-				[CommandLine("-AddView=", ListSeparator = ';')]
-				public List<string> AddView { get; set; } = new List<string>();
+			[CommandLine("-AddView=", ListSeparator = ';')]
+			public List<string> AddView { get; set; } = new List<string>();
 
-				[CommandLine("-RemoveView=", ListSeparator = ';')]
-				public List<string> RemoveView { get; set; } = new List<string>();
+			[CommandLine("-RemoveView=", ListSeparator = ';')]
+			public List<string> RemoveView { get; set; } = new List<string>();
 
-				[CommandLine("-AllProjects", Value = "true")]
-				[CommandLine("-OnlyCurrent", Value = "false")]
-				public bool? AllProjects = null;
+			[CommandLine("-AllProjects", Value = "true")]
+			[CommandLine("-OnlyCurrent", Value = "false")]
+			public bool? AllProjects { get; set; } = null;
 
-				[CommandLine("-GpfAllProjects", Value ="true")]
-				[CommandLine("-GpfOnlyCurrent", Value = "false")]
-				public bool? AllProjectsInSln = null;
+			[CommandLine("-GpfAllProjects", Value = "true")]
+			[CommandLine("-GpfOnlyCurrent", Value = "false")]
+			public bool? AllProjectsInSln { get; set; } = null;
 
-				[CommandLine("-Global")]
-				public bool Global { get; set; }
-			}
+			[CommandLine("-Global")]
+			public bool Global { get; set; }
+		}
 
+		class FilterCommand : Command
+		{
 			public override async Task ExecuteAsync(CommandContext context)
 			{
 				ILogger logger = context.Logger;
@@ -1102,13 +1151,21 @@ namespace UnrealGameSyncCmd
 			}
 		}
 
+		class BuildCommandOptions
+		{
+			[CommandLine("-List")]
+			public bool ListOnly { get; set; }
+		}
+
 		class BuildCommand : Command
 		{
 			public override async Task ExecuteAsync(CommandContext context)
 			{
 				ILogger logger = context.Logger;
 				context.Arguments.TryGetPositionalArgument(out string? target);
-				bool listOnly = context.Arguments.HasOption("-List");
+
+				BuildCommandOptions options = new BuildCommandOptions();
+				context.Arguments.ApplyTo(options);
 				context.Arguments.CheckAllArgumentsUsed();
 
 				UserWorkspaceSettings settings = ReadRequiredUserWorkspaceSettings();
@@ -1117,7 +1174,7 @@ namespace UnrealGameSyncCmd
 
 				ProjectInfo projectInfo = new ProjectInfo(settings.RootDir, state.Current);
 
-				if (listOnly)
+				if (options.ListOnly)
 				{
 					ConfigFile projectConfig = await ConfigUtils.ReadProjectConfigFileAsync(perforceClient, projectInfo, logger, CancellationToken.None);
 
@@ -1219,6 +1276,12 @@ namespace UnrealGameSyncCmd
 			}
 		}
 
+		class SwitchCommandOptions
+		{
+			[CommandLine("-Force")]
+			public bool Force { get; set; }
+		}
+
 		class SwitchCommand : Command
 		{
 			public override async Task ExecuteAsync(CommandContext context)
@@ -1230,11 +1293,16 @@ namespace UnrealGameSyncCmd
 					throw new UserErrorException("Missing stream or project name to switch to.");
 				}
 
-				bool force = targetName.StartsWith("//", StringComparison.Ordinal) && context.Arguments.HasOption("-Force");
-
 				// Finish argument parsing
+				SwitchCommandOptions options = new SwitchCommandOptions();
+				context.Arguments.ApplyTo(options);
 				context.Arguments.CheckAllArgumentsUsed();
 
+				if (targetName.StartsWith("//", StringComparison.Ordinal))
+				{
+					options.Force = true;
+				}
+
 				// Get a connection to the client for this workspace
 				UserWorkspaceSettings settings = ReadRequiredUserWorkspaceSettings();
 				using IPerforceConnection perforceClient = await ConnectAsync(settings, context.LoggerFactory);
@@ -1242,7 +1310,7 @@ namespace UnrealGameSyncCmd
 				// Check whether we're switching stream or project
 				if (targetName.StartsWith("//", StringComparison.Ordinal))
 				{
-					await SwitchStreamAsync(perforceClient, targetName, force, context.Logger);
+					await SwitchStreamAsync(perforceClient, targetName, options.Force, context.Logger);
 				}
 				else
 				{
