UGS: Fix exception opening variables window for custom build steps.

diff --git a/Engine/Source/Programs/UnrealGameSync/UnrealGameSync/Forms/BuildStepWindow.cs b/Engine/Source/Programs/UnrealGameSync/UnrealGameSync/Forms/BuildStepWindow.cs
--- a/Engine/Source/Programs/UnrealGameSync/UnrealGameSync/Forms/BuildStepWindow.cs
+++ b/Engine/Source/Programs/UnrealGameSync/UnrealGameSync/Forms/BuildStepWindow.cs
@@ -267,6 +267,7 @@ namespace UnrealGameSync
 				_variablesWindow.FormClosed += OnClosedVariablesWindow;
 				_variablesWindow.Show(this);
 
+				components ??= new System.ComponentModel.Container();
 				components.Add(_variablesWindow);
 			}
 			else
