Use self-contained distributions for CLI on Mac & Linux.

diff --git a/Engine/Source/Programs/UnrealGameSync/BuildUGS.xml b/Engine/Source/Programs/UnrealGameSync/BuildUGS.xml
--- a/Engine/Source/Programs/UnrealGameSync/BuildUGS.xml
+++ b/Engine/Source/Programs/UnrealGameSync/BuildUGS.xml
@@ -71,7 +71,7 @@
 		
 		<Node Name="Build UnrealGameSyncCmd Mac">
 			<Delete Files="$(UgsCmdMacOutputDir)/..."/>
-			<DotNet Arguments="publish &quot;$(UgsDir)/UnrealGameSyncCmd/UnrealGameSyncCmd.csproj&quot; --output=&quot;$(UgsCmdMacOutputDir)&quot; --no-self-contained -c Release -r osx-x64 $(VersionArguments)"/>
+			<DotNet Arguments="publish &quot;$(UgsDir)/UnrealGameSyncCmd/UnrealGameSyncCmd.csproj&quot; --output=&quot;$(UgsCmdMacOutputDir)&quot; --self-contained -c Release -r osx-x64 $(VersionArguments)"/>
 		</Node>
 		<Node Name="Publish UnrealGameSyncCmd Mac" Requires="Build UnrealGameSyncCmd Mac">
 			<Error Message="Missing -HordeSettings=... argument" If="'$(HordeSettings)' == ''"/>
@@ -80,7 +80,7 @@
 		
 		<Node Name="Build UnrealGameSyncCmd Linux">
 			<Delete Files="$(UgsCmdLinuxOutputDir)/..."/>
-			<DotNet Arguments="publish &quot;$(UgsDir)/UnrealGameSyncCmd/UnrealGameSyncCmd.csproj&quot; --output=&quot;$(UgsCmdLinuxOutputDir)&quot; --no-self-contained -c Release -r linux-x64 $(VersionArguments)"/>
+			<DotNet Arguments="publish &quot;$(UgsDir)/UnrealGameSyncCmd/UnrealGameSyncCmd.csproj&quot; --output=&quot;$(UgsCmdLinuxOutputDir)&quot; --self-contained -c Release -r linux-x64 $(VersionArguments)"/>
 		</Node>
 		<Node Name="Publish UnrealGameSyncCmd Linux" Requires="Build UnrealGameSyncCmd Linux">
 			<Error Message="Missing -HordeSettings=... argument" If="'$(HordeSettings)' == ''"/>
