UGS: Read deployment settings from a Deployment.json file in the application folder.

diff --git a/Engine/Source/Programs/UnrealGameSync/UnrealGameSync/Controls/WorkspaceControl.cs b/Engine/Source/Programs/UnrealGameSync/UnrealGameSync/Controls/WorkspaceControl.cs
--- a/Engine/Source/Programs/UnrealGameSync/UnrealGameSync/Controls/WorkspaceControl.cs
+++ b/Engine/Source/Programs/UnrealGameSync/UnrealGameSync/Controls/WorkspaceControl.cs
@@ -1792,7 +1792,7 @@ namespace UnrealGameSync
 
 		void UpdateBuildFailureNotification()
 		{
-			if (!DeploymentSettings.EnableAlerts)
+			if (!DeploymentSettings.Instance.EnableAlerts)
 			{
 				return;
 			}
diff --git a/Engine/Source/Programs/UnrealGameSync/UnrealGameSync/Deployment.json b/Engine/Source/Programs/UnrealGameSync/UnrealGameSync/Deployment.json
new file mode 100644
--- /dev/null
+++ b/Engine/Source/Programs/UnrealGameSync/UnrealGameSync/Deployment.json
@@ -0,0 +1,2 @@
+{
+}
diff --git a/Engine/Source/Programs/UnrealGameSync/UnrealGameSync/DeploymentSettings.cs b/Engine/Source/Programs/UnrealGameSync/UnrealGameSync/DeploymentSettings.cs
--- a/Engine/Source/Programs/UnrealGameSync/UnrealGameSync/DeploymentSettings.cs
+++ b/Engine/Source/Programs/UnrealGameSync/UnrealGameSync/DeploymentSettings.cs
@@ -3,18 +3,47 @@
 #pragma warning disable IDE0005
 #pragma warning disable CA1802 // warning CA1802: Field 'EnableAlerts' is declared as 'readonly' but is initialized with a constant value. Mark this field as 'const' instead.
 
+using EpicGames.Core;
 using Microsoft.Extensions.Logging;
 using Microsoft.Extensions.Logging.Abstractions;
 using System.Collections.Generic;
 using System.Diagnostics.CodeAnalysis;
+using System.IO;
+using System.Reflection;
+using System.Text.Json;
+using System.Xml;
 
 namespace UnrealGameSync
 {
 	/// <summary>
-	/// This class contains settings for a site-specific deployment of UGS. Epic's internal implementation uses a static constructor in a NotForLicensees folder to initialize these values.
+	/// This class contains settings for a site-specific deployment of UGS, and is read from the Deployment.json file in the application directory.
 	/// </summary>
-	static partial class DeploymentSettings
+	partial class DeploymentSettings
 	{
+		static DeploymentSettings? s_instance;
+
+		public static DeploymentSettings Instance
+		{
+			get
+			{
+				if (s_instance == null)
+				{
+					FileReference assemblyFile = new FileReference(Assembly.GetExecutingAssembly().Location);
+					FileReference settingsFile = FileReference.Combine(assemblyFile.Directory, "Deployment.json");
+
+					if (FileReference.Exists(settingsFile))
+					{
+						byte[] data = FileReference.ReadAllBytes(settingsFile);
+						JsonSerializerOptions options = new JsonSerializerOptions { AllowTrailingCommas = true, PropertyNameCaseInsensitive = true, ReadCommentHandling = JsonCommentHandling.Skip };
+						s_instance = JsonSerializer.Deserialize<DeploymentSettings>(data, options);
+					}
+
+					s_instance ??= new DeploymentSettings();
+				}
+				return s_instance;
+			}
+		}
+
 #if WITH_TELEMETRY
 		/// <summary>
 		/// Delegate used to create a telemetry sink
@@ -27,54 +56,53 @@ namespace UnrealGameSync
 #endif
 
 		/// <summary>
-		/// SQL connection string used to connect to the database for telemetry and review data. The 'Program' class is a partial class, to allow an
-		/// opportunistically included C# source file in NotForLicensees/ProgramSettings.cs to override this value in a static constructor.
+		/// SQL connection string used to connect to the database for telemetry and review data.
 		/// </summary>
-		public static readonly string? ApiUrl = null;
+		public string? ApiUrl { get; set; }
 
 		/// <summary>
 		/// Servers to connect to for issue details by default
 		/// </summary>
-		public static readonly List<string> DefaultIssueApiUrls = new List<string>();
+		public List<string> DefaultIssueApiUrls { get; set; } = new List<string>();
 
 		/// <summary>
 		/// The issue api to use for URL handler events
 		/// </summary>
-		public static readonly string? UrlHandleIssueApi = null;
+		public string? UrlHandleIssueApi { get; set; } = null;
 
 		/// <summary>
 		/// Specifies the depot path to sync down the stable version of UGS from, without a trailing slash (eg. //depot/UnrealGameSync/bin). This is a site-specific setting. 
 		/// The UnrealGameSync executable should be located at Release/UnrealGameSync.exe under this path, with any dependent DLLs.
 		/// </summary>
-		public static readonly string? DefaultDepotPath = null;
+		public string? DefaultDepotPath { get; set; } = null;
 
 		/// <summary>
 		/// Depot path to sync additional tools from
 		/// </summary>
-		public static readonly string? ToolsDepotPath = null;
+		public string? ToolsDepotPath { get; set; } = null;
 
 		/// <summary>
 		/// DSN for sending crash reports to Sentry.
 		/// </summary>
-		public static readonly string? SentryDsn = null;
+		public string? SentryDsn { get; set; } = null;
 
 		/// <summary>
 		/// Whether to allow notifications about build failures
 		/// </summary>
-		public static readonly bool EnableAlerts = true;
+		public bool EnableAlerts { get; set; } = true;
 
 #if WITH_TELEMETRY
 		/// <summary>
 		/// Delegate used to create a new telemetry sink
 		/// </summary>
-		public static readonly CreateTelemetrySinkDelegate CreateTelemetrySink = (userName, sessionId, log) => new NullTelemetrySink();
+		public readonly CreateTelemetrySinkDelegate CreateTelemetrySink = (userName, sessionId, log) => new NullTelemetrySink();
 #endif
 
 #if !UGS_LAUNCHER
 		/// <summary>
 		/// Delegate to allow validating a project being opened
 		/// </summary>
-		/// <param name="Workspace">The detected settings for the project</param>
+		/// <param name="openProjectInfo">The detected settings for the project</param>
 		/// <param name="logger">The logger</param>
 		/// <param name="error">Receives an error on failure</param>
 		/// <returns></returns>
@@ -83,7 +111,7 @@ namespace UnrealGameSync
 		/// <summary>
 		/// Called to validate the project settings
 		/// </summary>
-		public static DetectProjectSettingsEvent? OnDetectProjectSettings = null;
+		public DetectProjectSettingsEvent? OnDetectProjectSettings = null;
 #endif
 	}
 }
diff --git a/Engine/Source/Programs/UnrealGameSync/UnrealGameSync/Forms/ApplicationSettingsWindow.cs b/Engine/Source/Programs/UnrealGameSync/UnrealGameSync/Forms/ApplicationSettingsWindow.cs
--- a/Engine/Source/Programs/UnrealGameSync/UnrealGameSync/Forms/ApplicationSettingsWindow.cs
+++ b/Engine/Source/Programs/UnrealGameSync/UnrealGameSync/Forms/ApplicationSettingsWindow.cs
@@ -91,7 +91,7 @@ namespace UnrealGameSync
 
 			DepotPathTextBox.Text = _initialDepotPath;
 			DepotPathTextBox.Select(DepotPathTextBox.TextLength, 0);
-			DepotPathTextBox.CueBanner = DeploymentSettings.DefaultDepotPath ?? String.Empty;
+			DepotPathTextBox.CueBanner = DeploymentSettings.Instance.DefaultDepotPath ?? String.Empty;
 
 			UsePreviewBuildCheckBox.Checked = preview;
 
@@ -163,7 +163,7 @@ namespace UnrealGameSync
 			}
 
 			string? depotPath = DepotPathTextBox.Text.Trim();
-			if(depotPath.Length == 0 || depotPath == DeploymentSettings.DefaultDepotPath)
+			if(depotPath.Length == 0 || depotPath == DeploymentSettings.Instance.DefaultDepotPath)
 			{
 				depotPath = null;
 			}
@@ -183,7 +183,7 @@ namespace UnrealGameSync
 				{
 					PerforceSettings settings = Utility.OverridePerforceSettings(_defaultPerforceSettings, serverAndPort, userName);
 
-					string? testDepotPath = depotPath ?? DeploymentSettings.DefaultDepotPath;
+					string? testDepotPath = depotPath ?? DeploymentSettings.Instance.DefaultDepotPath;
 					if (testDepotPath != null)
 					{
 						ModalTask? task = PerforceModalTask.Execute(this, "Checking connection", "Checking connection, please wait...", settings, (p, c) => PerforceTestConnectionTask.RunAsync(p, testDepotPath, c), _logger);
diff --git a/Engine/Source/Programs/UnrealGameSync/UnrealGameSync/Forms/MainWindow.cs b/Engine/Source/Programs/UnrealGameSync/UnrealGameSync/Forms/MainWindow.cs
--- a/Engine/Source/Programs/UnrealGameSync/UnrealGameSync/Forms/MainWindow.cs
+++ b/Engine/Source/Programs/UnrealGameSync/UnrealGameSync/Forms/MainWindow.cs
@@ -201,7 +201,7 @@ namespace UnrealGameSync
 			TabPanel.ResumeLayout(false);
 			ResumeLayout(false);
 
-			foreach (string defaultIssueApiUrl in DeploymentSettings.DefaultIssueApiUrls)
+			foreach (string defaultIssueApiUrl in DeploymentSettings.Instance.DefaultIssueApiUrls)
 			{
 				_defaultIssueMonitors.Add(CreateIssueMonitor(defaultIssueApiUrl, defaultPerforceSettings.UserName));
 			}
@@ -469,7 +469,7 @@ namespace UnrealGameSync
 				if (existingWorkspace != null)
 				{
 					IssueMonitor issueMonitor = existingWorkspace.GetIssueMonitor();
-					if (issueMonitor != null && (DeploymentSettings.UrlHandleIssueApi == null || String.Equals(issueMonitor.ApiUrl, DeploymentSettings.UrlHandleIssueApi, StringComparison.OrdinalIgnoreCase)))
+					if (issueMonitor != null && (DeploymentSettings.Instance.UrlHandleIssueApi == null || String.Equals(issueMonitor.ApiUrl, DeploymentSettings.Instance.UrlHandleIssueApi, StringComparison.OrdinalIgnoreCase)))
 					{
 						issueMonitor.AddRef();
 						try
@@ -1575,7 +1575,7 @@ namespace UnrealGameSync
 
 		public void UpdateAlertWindows()
 		{
-			if (!DeploymentSettings.EnableAlerts)
+			if (!DeploymentSettings.Instance.EnableAlerts)
 			{
 				return;
 			}
diff --git a/Engine/Source/Programs/UnrealGameSync/UnrealGameSync/Forms/OpenProjectWindow.cs b/Engine/Source/Programs/UnrealGameSync/UnrealGameSync/Forms/OpenProjectWindow.cs
--- a/Engine/Source/Programs/UnrealGameSync/UnrealGameSync/Forms/OpenProjectWindow.cs
+++ b/Engine/Source/Programs/UnrealGameSync/UnrealGameSync/Forms/OpenProjectWindow.cs
@@ -296,10 +296,10 @@ namespace UnrealGameSync
 		public static async Task<OpenProjectInfo> DetectSettingsAsync(IPerforceConnection perforce, UserSelectedProjectSettings selectedProject, UserSettings userSettings, OidcTokenManager oidcTokenManager, ILogger<OpenProjectInfo> logger, CancellationToken cancellationToken)
 		{
 			OpenProjectInfo settings = await OpenProjectInfo.CreateAsync(perforce, selectedProject, userSettings, oidcTokenManager, logger, cancellationToken);
-			if (DeploymentSettings.OnDetectProjectSettings != null)
+			if (DeploymentSettings.Instance.OnDetectProjectSettings != null)
 			{
 				string? message;
-				if (!DeploymentSettings.OnDetectProjectSettings(settings, logger, out message))
+				if (!DeploymentSettings.Instance.OnDetectProjectSettings(settings, logger, out message))
 				{
 					throw new UserErrorException(message);
 				}
diff --git a/Engine/Source/Programs/UnrealGameSync/UnrealGameSync/GlobalSettings.cs b/Engine/Source/Programs/UnrealGameSync/UnrealGameSync/GlobalSettings.cs
--- a/Engine/Source/Programs/UnrealGameSync/UnrealGameSync/GlobalSettings.cs
+++ b/Engine/Source/Programs/UnrealGameSync/UnrealGameSync/GlobalSettings.cs
@@ -88,7 +88,7 @@ namespace UnrealGameSync
 						key.SetValue("UserName", userName);
 					}
 
-					if (String.IsNullOrEmpty(depotPath) || (DeploymentSettings.DefaultDepotPath != null && String.Equals(depotPath, DeploymentSettings.DefaultDepotPath, StringComparison.OrdinalIgnoreCase)))
+					if (String.IsNullOrEmpty(depotPath) || (DeploymentSettings.Instance.DefaultDepotPath != null && String.Equals(depotPath, DeploymentSettings.Instance.DefaultDepotPath, StringComparison.OrdinalIgnoreCase)))
 					{
 						DeleteRegistryKey(key, "DepotPath");
 					}
diff --git a/Engine/Source/Programs/UnrealGameSync/UnrealGameSync/Program.cs b/Engine/Source/Programs/UnrealGameSync/UnrealGameSync/Program.cs
--- a/Engine/Source/Programs/UnrealGameSync/UnrealGameSync/Program.cs
+++ b/Engine/Source/Programs/UnrealGameSync/UnrealGameSync/Program.cs
@@ -31,7 +31,7 @@ namespace UnrealGameSync
 
 		public static void CaptureException(Exception exception)
 		{
-			if (DeploymentSettings.SentryDsn != null)
+			if (DeploymentSettings.Instance.SentryDsn != null)
 			{
 				SentrySdk.CaptureException(exception);
 			}
@@ -40,10 +40,10 @@ namespace UnrealGameSync
 		[STAThread]
 		static void Main(string[] args)
 		{
-			if (DeploymentSettings.SentryDsn != null)
+			if (DeploymentSettings.Instance.SentryDsn != null)
 			{
 				SentryOptions sentryOptions = new SentryOptions();
-				sentryOptions.Dsn = DeploymentSettings.SentryDsn;
+				sentryOptions.Dsn = DeploymentSettings.Instance.SentryDsn;
 				sentryOptions.StackTraceMode = StackTraceMode.Enhanced;
 				sentryOptions.AttachStacktrace = true;
 				sentryOptions.TracesSampleRate = 1.0;
@@ -194,7 +194,7 @@ namespace UnrealGameSync
 					ILogger telemetryLogger = loggerProvider.CreateLogger("Telemetry");
 					telemetryLogger.LogInformation("Creating telemetry sink for session {SessionId}", sessionId);
 
-					using (ITelemetrySink telemetrySink = DeploymentSettings.CreateTelemetrySink(userName, sessionId, telemetryLogger))
+					using (ITelemetrySink telemetrySink = DeploymentSettings.Instance.CreateTelemetrySink(userName, sessionId, telemetryLogger))
 					{
 						ITelemetrySink? prevTelemetrySink = Telemetry.ActiveSink;
 						try
@@ -211,7 +211,7 @@ namespace UnrealGameSync
 
 							using (UpdateMonitor updateMonitor = new UpdateMonitor(defaultSettings, updatePath, serviceProvider))
 							{
-								using ProgramApplicationContext context = new ProgramApplicationContext(defaultSettings, updateMonitor, DeploymentSettings.ApiUrl, dataFolder, activateEvent, restoreState, updateSpawn, projectFileName, preview, serviceProvider, uri);
+								using ProgramApplicationContext context = new ProgramApplicationContext(defaultSettings, updateMonitor, DeploymentSettings.Instance.ApiUrl, dataFolder, activateEvent, restoreState, updateSpawn, projectFileName, preview, serviceProvider, uri);
 								Application.Run(context);
 
 								if (updateMonitor.IsUpdateAvailable && updateSpawn != null)
@@ -238,7 +238,7 @@ namespace UnrealGameSync
 
 		private static void TraceException(Exception ex, ILogger logger)
 		{
-			if (DeploymentSettings.SentryDsn != null)
+			if (DeploymentSettings.Instance.SentryDsn != null)
 			{
 				SentrySdk.CaptureException(ex);
 			}
diff --git a/Engine/Source/Programs/UnrealGameSync/UnrealGameSync/ToolUpdateMonitor.cs b/Engine/Source/Programs/UnrealGameSync/UnrealGameSync/ToolUpdateMonitor.cs
--- a/Engine/Source/Programs/UnrealGameSync/UnrealGameSync/ToolUpdateMonitor.cs
+++ b/Engine/Source/Programs/UnrealGameSync/UnrealGameSync/ToolUpdateMonitor.cs
@@ -92,7 +92,7 @@ namespace UnrealGameSync
 
 		public void Start()
 		{
-			if (DeploymentSettings.ToolsDepotPath != null)
+			if (DeploymentSettings.Instance.ToolsDepotPath != null)
 			{
 				_workerTask = Task.Run(() => PollForUpdatesAsync(_cancellationSource.Token));
 			}
@@ -168,13 +168,13 @@ namespace UnrealGameSync
 		{
 			using IPerforceConnection perforce = await PerforceConnection.CreateAsync(PerforceSettings, logger);
 
-			List<ChangesRecord> changes = await perforce.GetChangesAsync(ChangesOptions.None, 1, ChangeStatus.Submitted, $"{DeploymentSettings.ToolsDepotPath}/...", cancellationToken);
+			List<ChangesRecord> changes = await perforce.GetChangesAsync(ChangesOptions.None, 1, ChangeStatus.Submitted, $"{DeploymentSettings.Instance.ToolsDepotPath}/...", cancellationToken);
 			if (changes.Count == 0 || changes[0].Number == _lastChange)
 			{
 				return;
 			}
 
-			List<FStatRecord> fileRecords = await perforce.FStatAsync($"{DeploymentSettings.ToolsDepotPath}/...", cancellationToken).ToListAsync(cancellationToken);
+			List<FStatRecord> fileRecords = await perforce.FStatAsync($"{DeploymentSettings.Instance.ToolsDepotPath}/...", cancellationToken).ToListAsync(cancellationToken);
 
 			// Update the tools list
 			List<ToolDefinition> newTools = new List<ToolDefinition>();
diff --git a/Engine/Source/Programs/UnrealGameSync/UnrealGameSync/UnrealGameSync.csproj b/Engine/Source/Programs/UnrealGameSync/UnrealGameSync/UnrealGameSync.csproj
--- a/Engine/Source/Programs/UnrealGameSync/UnrealGameSync/UnrealGameSync.csproj
+++ b/Engine/Source/Programs/UnrealGameSync/UnrealGameSync/UnrealGameSync.csproj
@@ -43,6 +43,9 @@
     </Reference>
   </ItemGroup>
   <ItemGroup>
+    <None Update="Deployment.json">
+      <CopyToOutputDirectory>Always</CopyToOutputDirectory>
+    </None>
     <Compile Update="Properties\Settings.Designer.cs">
       <DesignTimeSharedInput>True</DesignTimeSharedInput>
       <AutoGen>True</AutoGen>
