UGS: Move custom hook for opening a new project into OpenProjectWindow.

diff --git a/Engine/Source/Programs/UnrealGameSync/UnrealGameSync/DeploymentSettings.cs b/Engine/Source/Programs/UnrealGameSync/UnrealGameSync/DeploymentSettings.cs
--- a/Engine/Source/Programs/UnrealGameSync/UnrealGameSync/DeploymentSettings.cs
+++ b/Engine/Source/Programs/UnrealGameSync/UnrealGameSync/DeploymentSettings.cs
@@ -97,21 +97,5 @@ namespace UnrealGameSync
 		/// </summary>
 		public readonly CreateTelemetrySinkDelegate CreateTelemetrySink = (userName, sessionId, log) => new NullTelemetrySink();
 #endif
-
-#if !UGS_LAUNCHER
-		/// <summary>
-		/// Delegate to allow validating a project being opened
-		/// </summary>
-		/// <param name="openProjectInfo">The detected settings for the project</param>
-		/// <param name="logger">The logger</param>
-		/// <param name="error">Receives an error on failure</param>
-		/// <returns></returns>
-		public delegate bool DetectProjectSettingsEvent(OpenProjectInfo openProjectInfo, ILogger logger, [NotNullWhen(false)] out string? error);
-
-		/// <summary>
-		/// Called to validate the project settings
-		/// </summary>
-		public DetectProjectSettingsEvent? OnDetectProjectSettings = null;
-#endif
 	}
 }
diff --git a/Engine/Source/Programs/UnrealGameSync/UnrealGameSync/Forms/OpenProjectWindow.cs b/Engine/Source/Programs/UnrealGameSync/UnrealGameSync/Forms/OpenProjectWindow.cs
--- a/Engine/Source/Programs/UnrealGameSync/UnrealGameSync/Forms/OpenProjectWindow.cs
+++ b/Engine/Source/Programs/UnrealGameSync/UnrealGameSync/Forms/OpenProjectWindow.cs
@@ -27,6 +27,13 @@ namespace UnrealGameSync
 		readonly UserSettings _settings;
 		readonly ILogger _logger;
 
+		delegate bool DetectProjectSettingsEvent(OpenProjectInfo openProjectInfo, ILogger logger, [NotNullWhen(false)] out string? error);
+
+#pragma warning disable CS0649 // Avoid unused private fields
+		// Hook called once the project settings have been established. Can be overridden in a specific deployment to log telemetry or perform additional validation.
+		static readonly DetectProjectSettingsEvent? s_onDetectProjectSettings;
+#pragma warning restore CS0649 // Avoid unused private fields
+
 		private OpenProjectWindow(UserSelectedProjectSettings? project, UserSettings settings, IPerforceSettings defaultSettings, IServiceProvider serviceProvider, ILogger logger)
 		{
 			InitializeComponent();
@@ -296,10 +303,10 @@ namespace UnrealGameSync
 		public static async Task<OpenProjectInfo> DetectSettingsAsync(IPerforceConnection perforce, UserSelectedProjectSettings selectedProject, UserSettings userSettings, OidcTokenManager oidcTokenManager, ILogger<OpenProjectInfo> logger, CancellationToken cancellationToken)
 		{
 			OpenProjectInfo settings = await OpenProjectInfo.CreateAsync(perforce, selectedProject, userSettings, oidcTokenManager, logger, cancellationToken);
-			if (DeploymentSettings.Instance.OnDetectProjectSettings != null)
+			if (s_onDetectProjectSettings != null)
 			{
 				string? message;
-				if (!DeploymentSettings.Instance.OnDetectProjectSettings(settings, logger, out message))
+				if (!s_onDetectProjectSettings(settings, logger, out message))
 				{
 					throw new UserErrorException(message);
 				}
diff --git a/Engine/Source/Programs/UnrealGameSync/UnrealGameSync/UnrealGameSync.csproj b/Engine/Source/Programs/UnrealGameSync/UnrealGameSync/UnrealGameSync.csproj
--- a/Engine/Source/Programs/UnrealGameSync/UnrealGameSync/UnrealGameSync.csproj
+++ b/Engine/Source/Programs/UnrealGameSync/UnrealGameSync/UnrealGameSync.csproj
@@ -58,6 +58,15 @@
       <LastGenOutput>Settings.Designer.cs</LastGenOutput>
     </None>
   </ItemGroup>
+  <Choose>
+    <When Condition="Exists('..\..\..\..\Restricted\NotForLicensees\Source\Programs\UnrealGameSync\UnrealGameSync\OpenProjectWindowInternal.cs')">
+      <ItemGroup>
+        <Compile Include="..\..\..\..\Restricted\NotForLicensees\Source\Programs\UnrealGameSync\UnrealGameSync\OpenProjectWindowInternal.cs">
+          <Link>NotForLicensees\OpenProjectWindowInternal.cs</Link>
+        </Compile>
+      </ItemGroup>
+    </When>
+  </Choose>
   <Choose>
     <When Condition="Exists('..\..\..\..\Restricted\NotForLicensees\Source\Programs\UnrealGameSync\UnrealGameSync\DeploymentSettingsInternal.cs')">
       <ItemGroup>
