UGS: Move telemetry sink into the Program class, leaving DeploymentSettings a pure configuration driven object.

diff --git a/Engine/Source/Programs/UnrealGameSync/UnrealGameSync/DeploymentSettings.cs b/Engine/Source/Programs/UnrealGameSync/UnrealGameSync/DeploymentSettings.cs
--- a/Engine/Source/Programs/UnrealGameSync/UnrealGameSync/DeploymentSettings.cs
+++ b/Engine/Source/Programs/UnrealGameSync/UnrealGameSync/DeploymentSettings.cs
@@ -44,17 +44,6 @@ namespace UnrealGameSync
 			}
 		}
 
-#if WITH_TELEMETRY
-		/// <summary>
-		/// Delegate used to create a telemetry sink
-		/// </summary>
-		/// <param name="userName">The default Perforce user name</param>
-		/// <param name="sessionId">Unique identifier for this session</param>
-		/// <param name="logger">Log writer</param>
-		/// <returns>New telemetry sink instance</returns>
-		public delegate ITelemetrySink CreateTelemetrySinkDelegate(string userName, string sessionId, ILogger logger);
-#endif
-
 		/// <summary>
 		/// SQL connection string used to connect to the database for telemetry and review data.
 		/// </summary>
@@ -90,12 +79,5 @@ namespace UnrealGameSync
 		/// Whether to allow notifications about build failures
 		/// </summary>
 		public bool EnableAlerts { get; set; } = true;
-
-#if WITH_TELEMETRY
-		/// <summary>
-		/// Delegate used to create a new telemetry sink
-		/// </summary>
-		public readonly CreateTelemetrySinkDelegate CreateTelemetrySink = (userName, sessionId, log) => new NullTelemetrySink();
-#endif
 	}
 }
diff --git a/Engine/Source/Programs/UnrealGameSync/UnrealGameSync/Program.cs b/Engine/Source/Programs/UnrealGameSync/UnrealGameSync/Program.cs
--- a/Engine/Source/Programs/UnrealGameSync/UnrealGameSync/Program.cs
+++ b/Engine/Source/Programs/UnrealGameSync/UnrealGameSync/Program.cs
@@ -19,8 +19,22 @@ using System.Windows.Forms;
 
 namespace UnrealGameSync
 {
-	static class Program
+	/// <summary>
+	/// Delegate used to create a telemetry sink
+	/// </summary>
+	/// <param name="userName">The default Perforce user name</param>
+	/// <param name="sessionId">Unique identifier for this session</param>
+	/// <param name="logger">Log writer</param>
+	/// <returns>New telemetry sink instance</returns>
+	public delegate ITelemetrySink CreateTelemetrySinkDelegate(string userName, string sessionId, ILogger logger);
+
+	static partial class Program
 	{
+		/// <summary>
+		/// Delegate used to create a new telemetry sink
+		/// </summary>
+		static CreateTelemetrySinkDelegate CreateTelemetrySink { get; } = (userName, sessionId, log) => new NullTelemetrySink();
+
 		public static string GetVersionString()
 		{
 			AssemblyInformationalVersionAttribute? version = Assembly.GetExecutingAssembly().GetCustomAttribute<AssemblyInformationalVersionAttribute>();
@@ -194,7 +208,7 @@ namespace UnrealGameSync
 					ILogger telemetryLogger = loggerProvider.CreateLogger("Telemetry");
 					telemetryLogger.LogInformation("Creating telemetry sink for session {SessionId}", sessionId);
 
-					using (ITelemetrySink telemetrySink = DeploymentSettings.Instance.CreateTelemetrySink(userName, sessionId, telemetryLogger))
+					using (ITelemetrySink telemetrySink = CreateTelemetrySink(userName, sessionId, telemetryLogger))
 					{
 						ITelemetrySink? prevTelemetrySink = Telemetry.ActiveSink;
 						try
diff --git a/Engine/Source/Programs/UnrealGameSync/UnrealGameSync/UnrealGameSync.csproj b/Engine/Source/Programs/UnrealGameSync/UnrealGameSync/UnrealGameSync.csproj
--- a/Engine/Source/Programs/UnrealGameSync/UnrealGameSync/UnrealGameSync.csproj
+++ b/Engine/Source/Programs/UnrealGameSync/UnrealGameSync/UnrealGameSync.csproj
@@ -58,23 +58,10 @@
       <LastGenOutput>Settings.Designer.cs</LastGenOutput>
     </None>
   </ItemGroup>
-  <Choose>
-    <When Condition="Exists('..\..\..\..\Restricted\NotForLicensees\Source\Programs\UnrealGameSync\UnrealGameSync\OpenProjectWindowInternal.cs')">
-      <ItemGroup>
-        <Compile Include="..\..\..\..\Restricted\NotForLicensees\Source\Programs\UnrealGameSync\UnrealGameSync\OpenProjectWindowInternal.cs">
-          <Link>NotForLicensees\OpenProjectWindowInternal.cs</Link>
-        </Compile>
-      </ItemGroup>
-    </When>
-  </Choose>
-  <Choose>
-    <When Condition="Exists('..\..\..\..\Restricted\NotForLicensees\Source\Programs\UnrealGameSync\UnrealGameSync\DeploymentSettingsInternal.cs')">
-      <ItemGroup>
-        <Compile Include="..\..\..\..\Restricted\NotForLicensees\Source\Programs\UnrealGameSync\UnrealGameSync\DeploymentSettingsInternal.cs">
-          <Link>NotForLicensees\DeploymentSettingsInternal.cs</Link>
-        </Compile>
-      </ItemGroup>
-    </When>
-  </Choose>
+  <ItemGroup>
+    <Compile Include="..\..\..\..\Restricted\NotForLicensees\Source\Programs\UnrealGameSync\UnrealGameSync\*.cs">
+      <Link>NotForLicensees\%(FileName).cs</Link>
+    </Compile>
+  </ItemGroup>
   <Import Project="..\..\..\..\Restricted\NotForLicensees\Source\Programs\UnrealGameSync\UnrealGameSync\UnrealGameSync.csproj" Condition="Exists('..\..\..\..\Restricted\NotForLicensees\Source\Programs\UnrealGameSync\UnrealGameSync\UnrealGameSync.csproj')" />
 </Project>
