UGS: Move deployment settings into UnrealGameSyncShared, so they can be read by the command line client.

diff --git a/Engine/Source/Programs/UnrealGameSync/UnrealGameSync/Deployment.json b/Engine/Source/Programs/UnrealGameSync/UnrealGameSync/Deployment.json
deleted file mode 100644
diff --git a/Engine/Source/Programs/UnrealGameSync/UnrealGameSync/DeploymentSettings.cs b/Engine/Source/Programs/UnrealGameSync/UnrealGameSync/DeploymentSettings.cs
deleted file mode 100644
diff --git a/Engine/Source/Programs/UnrealGameSync/UnrealGameSync/UnrealGameSync.csproj b/Engine/Source/Programs/UnrealGameSync/UnrealGameSync/UnrealGameSync.csproj
--- a/Engine/Source/Programs/UnrealGameSync/UnrealGameSync/UnrealGameSync.csproj
+++ b/Engine/Source/Programs/UnrealGameSync/UnrealGameSync/UnrealGameSync.csproj
@@ -43,9 +43,6 @@
     </Reference>
   </ItemGroup>
   <ItemGroup>
-    <None Update="Deployment.json">
-      <CopyToOutputDirectory>Always</CopyToOutputDirectory>
-    </None>
     <Compile Update="Properties\Settings.Designer.cs">
       <DesignTimeSharedInput>True</DesignTimeSharedInput>
       <AutoGen>True</AutoGen>
diff --git a/Engine/Source/Programs/UnrealGameSync/UnrealGameSyncLauncher/Program.cs b/Engine/Source/Programs/UnrealGameSync/UnrealGameSyncLauncher/Program.cs
--- a/Engine/Source/Programs/UnrealGameSync/UnrealGameSyncLauncher/Program.cs
+++ b/Engine/Source/Programs/UnrealGameSync/UnrealGameSyncLauncher/Program.cs
@@ -46,7 +46,7 @@ namespace UnrealGameSyncLauncher
 				// Read the settings
 				string? serverAndPort = null;
 				string? userName = null;
-				string? depotPath = DeploymentSettings.DefaultDepotPath;
+				string? depotPath = DeploymentSettings.Instance.DefaultDepotPath;
 				GlobalPerforceSettings.ReadGlobalPerforceSettings(ref serverAndPort, ref userName, ref depotPath, ref preview);
 
 				// If the shift key is held down, immediately show the settings window
diff --git a/Engine/Source/Programs/UnrealGameSync/UnrealGameSyncLauncher/UnrealGameSyncLauncher.csproj b/Engine/Source/Programs/UnrealGameSync/UnrealGameSyncLauncher/UnrealGameSyncLauncher.csproj
--- a/Engine/Source/Programs/UnrealGameSync/UnrealGameSyncLauncher/UnrealGameSyncLauncher.csproj
+++ b/Engine/Source/Programs/UnrealGameSync/UnrealGameSyncLauncher/UnrealGameSyncLauncher.csproj
@@ -32,9 +32,6 @@
       <Link>UnrealGameSync\PasswordWindow.Designer.cs</Link>
       <DependentUpon>PasswordWindow.cs</DependentUpon>
     </Compile>
-    <Compile Include="..\UnrealGameSync\DeploymentSettings.cs">
-      <Link>UnrealGameSync\DeploymentSettings.cs</Link>
-    </Compile>
     <Compile Include="..\UnrealGameSync\ModalTask.cs">
       <Link>UnrealGameSync\ModalTask.cs</Link>
     </Compile>
@@ -44,26 +41,15 @@
     <Compile Include="..\UnrealGameSync\GlobalSettings.cs">
       <Link>UnrealGameSync\GlobalSettings.cs</Link>
     </Compile>
-    <Compile Include="..\UnrealGameSyncShared\Utility.cs">
-      <Link>UnrealGameSyncShared\Utility.cs</Link>
-    </Compile>
   </ItemGroup>
   <ItemGroup>
     <ProjectReference Include="..\..\Shared\EpicGames.Core\EpicGames.Core.csproj" />
     <ProjectReference Include="..\..\Shared\EpicGames.Perforce\EpicGames.Perforce.csproj" />
+    <ProjectReference Include="..\UnrealGameSyncShared\UnrealGameSyncShared.csproj" />
   </ItemGroup>
   <ItemGroup>
     <None Update="Properties\Settings.settings">
       <Generator>SettingsSingleFileGenerator</Generator>
     </None>
   </ItemGroup>
-  <Choose>
-    <When Condition="Exists('..\..\..\..\Restricted\NotForLicensees\Source\Programs\UnrealGameSync\UnrealGameSync\DeploymentSettingsInternal.cs')">
-      <ItemGroup>
-        <Compile Include="..\..\..\..\Restricted\NotForLicensees\Source\Programs\UnrealGameSync\UnrealGameSync\DeploymentSettingsInternal.cs">
-          <Link>Shared\NotForLicensees\DeploymentSettingsInternal.cs</Link>
-        </Compile>
-      </ItemGroup>
-    </When>
-  </Choose>
 </Project>
\ No newline at end of file
diff --git a/Engine/Source/Programs/UnrealGameSync/UnrealGameSyncShared/Deployment.json b/Engine/Source/Programs/UnrealGameSync/UnrealGameSyncShared/Deployment.json
new file mode 100644
--- /dev/null
+++ b/Engine/Source/Programs/UnrealGameSync/UnrealGameSyncShared/Deployment.json
@@ -0,0 +1,2 @@
+{
+}
diff --git a/Engine/Source/Programs/UnrealGameSync/UnrealGameSyncShared/DeploymentSettings.cs b/Engine/Source/Programs/UnrealGameSync/UnrealGameSyncShared/DeploymentSettings.cs
new file mode 100644
--- /dev/null
+++ b/Engine/Source/Programs/UnrealGameSync/UnrealGameSyncShared/DeploymentSettings.cs
@@ -0,0 +1,85 @@
+// Copyright Epic Games, Inc. All Rights Reserved.
+
+#pragma warning disable IDE0005
+#pragma warning disable CA1802 // warning CA1802: Field 'EnableAlerts' is declared as 'readonly' but is initialized with a constant value. Mark this field as 'const' instead.
+
+using EpicGames.Core;
+using Microsoft.Extensions.Logging;
+using Microsoft.Extensions.Logging.Abstractions;
+using System.Collections.Generic;
+using System.Diagnostics.CodeAnalysis;
+using System.IO;
+using System.Reflection;
+using System.Text.Json;
+using System.Xml;
+
+namespace UnrealGameSync
+{
+	/// <summary>
+	/// This class contains settings for a site-specific deployment of UGS, and is read from the Deployment.json file in the application directory.
+	/// </summary>
+	public partial class DeploymentSettings
+	{
+		/// <summary>
+		/// SQL connection string used to connect to the database for telemetry and review data.
+		/// </summary>
+		public string? ApiUrl { get; set; }
+
+		/// <summary>
+		/// Servers to connect to for issue details by default
+		/// </summary>
+#pragma warning disable CA2227 // Collection properties should be read only
+		public List<string> DefaultIssueApiUrls { get; set; } = new List<string>();
+#pragma warning restore CA2227 // Collection properties should be read only
+
+		/// <summary>
+		/// The issue api to use for URL handler events
+		/// </summary>
+		public string? UrlHandleIssueApi { get; set; } = null;
+
+		/// <summary>
+		/// Specifies the depot path to sync down the stable version of UGS from, without a trailing slash (eg. //depot/UnrealGameSync/bin). This is a site-specific setting. 
+		/// The UnrealGameSync executable should be located at Release/UnrealGameSync.exe under this path, with any dependent DLLs.
+		/// </summary>
+		public string? DefaultDepotPath { get; set; } = null;
+
+		/// <summary>
+		/// Depot path to sync additional tools from
+		/// </summary>
+		public string? ToolsDepotPath { get; set; } = null;
+
+		/// <summary>
+		/// DSN for sending crash reports to Sentry.
+		/// </summary>
+		public string? SentryDsn { get; set; } = null;
+
+		/// <summary>
+		/// Whether to allow notifications about build failures
+		/// </summary>
+		public bool EnableAlerts { get; set; } = true;
+
+		static DeploymentSettings? s_instance;
+
+		public static DeploymentSettings Instance
+		{
+			get
+			{
+				if (s_instance == null)
+				{
+					FileReference assemblyFile = new FileReference(Assembly.GetExecutingAssembly().Location);
+					FileReference settingsFile = FileReference.Combine(assemblyFile.Directory, "Deployment.json");
+
+					if (FileReference.Exists(settingsFile))
+					{
+						byte[] data = FileReference.ReadAllBytes(settingsFile);
+						JsonSerializerOptions options = new JsonSerializerOptions { AllowTrailingCommas = true, PropertyNameCaseInsensitive = true, ReadCommentHandling = JsonCommentHandling.Skip };
+						s_instance = JsonSerializer.Deserialize<DeploymentSettings>(data, options);
+					}
+
+					s_instance ??= new DeploymentSettings();
+				}
+				return s_instance;
+			}
+		}
+	}
+}
diff --git a/Engine/Source/Programs/UnrealGameSync/UnrealGameSyncShared/UnrealGameSyncShared.csproj b/Engine/Source/Programs/UnrealGameSync/UnrealGameSyncShared/UnrealGameSyncShared.csproj
--- a/Engine/Source/Programs/UnrealGameSync/UnrealGameSyncShared/UnrealGameSyncShared.csproj
+++ b/Engine/Source/Programs/UnrealGameSync/UnrealGameSyncShared/UnrealGameSyncShared.csproj
@@ -1,4 +1,4 @@
-<Project Sdk="Microsoft.NET.Sdk">
+﻿<Project Sdk="Microsoft.NET.Sdk">
 
   <PropertyGroup>
     <TargetFramework>net6.0</TargetFramework>
@@ -6,6 +6,27 @@
     <RootNamespace>UnrealGameSync</RootNamespace>
   </PropertyGroup>
 
+  <ItemGroup>
+    <None Remove="Deployment.json" />
+  </ItemGroup>
+  
+  <Choose>
+    <When Condition="Exists('..\..\..\..\Restricted\NotForLicensees\Source\Programs\UnrealGameSync\UnrealGameSyncShared\Deployment.json')">
+      <ItemGroup>
+        <Content Include="..\..\..\..\Restricted\NotForLicensees\Source\Programs\UnrealGameSync\UnrealGameSyncShared\Deployment.json">
+          <CopyToOutputDirectory>Always</CopyToOutputDirectory>
+        </Content>
+      </ItemGroup>
+    </When>
+    <Otherwise>
+      <ItemGroup>
+        <Content Include="Deployment.json">
+          <CopyToOutputDirectory>Always</CopyToOutputDirectory>
+        </Content>
+      </ItemGroup>
+    </Otherwise>
+  </Choose>
+  
   <ItemGroup>
     <PackageReference Include="System.IO.FileSystem.DriveInfo" Version="4.3.1" />
   </ItemGroup>
