UGS: Add commands to update UGS-CLI from Horde.

diff --git a/Engine/Source/Programs/UnrealGameSync/UnrealGameSyncCmd/Program.cs b/Engine/Source/Programs/UnrealGameSync/UnrealGameSyncCmd/Program.cs
--- a/Engine/Source/Programs/UnrealGameSync/UnrealGameSyncCmd/Program.cs
+++ b/Engine/Source/Programs/UnrealGameSync/UnrealGameSyncCmd/Program.cs
@@ -20,6 +20,12 @@ using Microsoft.Extensions.Configuration;
 using UnrealGameSync;
 using System.Text;
 using System.Globalization;
+using System.Net.Http;
+using System.Text.Json;
+using System.Net.Http.Json;
+using Microsoft.Extensions.Options;
+using System.IO.Compression;
+using JetBrains.Annotations;
 
 namespace UnrealGameSyncCmd
 {
@@ -115,6 +121,12 @@ namespace UnrealGameSyncCmd
 				"ugs version",
 				"Prints the current application version"
 			),
+			new CommandInfo("install", typeof(InstallCommand), null,
+				"ugs install",
+				"Registers an alias to run ugs via 'ugs'."),
+			new CommandInfo("upgrade", typeof(UpgradeCommand), typeof(UpgradeCommandOptions),
+				"ugs upgrade",
+				"Upgrades the current installation with the latest build of UGS.")
 		};
 
 		class CommandContext
@@ -1338,17 +1350,167 @@ namespace UnrealGameSyncCmd
 			}
 		}
 
+		static string GetVersion()
+		{
+			AssemblyInformationalVersionAttribute? version = Assembly.GetExecutingAssembly().GetCustomAttribute<AssemblyInformationalVersionAttribute>();
+			return version?.InformationalVersion ?? "Unknown";
+		}
+
 		class VersionCommand : Command
 		{
 			public override Task ExecuteAsync(CommandContext context)
 			{
 				ILogger logger = context.Logger;
- 
-				AssemblyInformationalVersionAttribute? version = Assembly.GetExecutingAssembly().GetCustomAttribute<AssemblyInformationalVersionAttribute>();
-				logger.LogInformation("UnrealGameSync {Version}", version?.InformationalVersion ?? "Unknown");
+
+				string version = GetVersion();
+				logger.LogInformation("UnrealGameSync {Version}", version);
 
 				return Task.CompletedTask;
 			}
 		}
+
+		class InstallCommand : Command
+		{
+			public override async Task ExecuteAsync(CommandContext context)
+			{
+				ILogger logger = context.Logger;
+
+				DirectoryReference currentDir = new FileReference(Assembly.GetExecutingAssembly().Location).Directory;
+				if (RuntimeInformation.IsOSPlatform(OSPlatform.OSX))
+				{
+					DirectoryReference? userDir = DirectoryReference.GetSpecialFolder(Environment.SpecialFolder.UserProfile);
+					if (userDir != null)
+					{
+						FileReference configFile = FileReference.Combine(userDir, ".zshrc");
+
+						List<string> lines = new List<string>();
+						if (FileReference.Exists(configFile))
+						{
+							lines.AddRange(await FileReference.ReadAllLinesAsync(configFile));
+							lines.RemoveAll(x => Regex.IsMatch(x, @"^\s*alias\s+ugs\s*="));
+						}
+						lines.Add($"alias ugs={FileReference.Combine(currentDir, "ugs")}");
+
+						await FileReference.WriteAllLinesAsync(configFile, lines);
+						logger.LogInformation("Added 'ugs' alias to {ConfigFile}", configFile);
+					}
+				}
+			}
+		}
+
+		class UpgradeCommandOptions
+		{
+			[CommandLine("-Check")]
+			public bool Check { get; set; }
+		}
+
+		class UpgradeCommand : Command
+		{
+			class DeploymentInfo
+			{
+				public string Version { get; set; } = String.Empty;
+			}
+
+			public override async Task ExecuteAsync(CommandContext context)
+			{
+				ILogger logger = context.Logger;
+				
+				UpgradeCommandOptions options = new UpgradeCommandOptions();
+				context.Arguments.ApplyTo(options);
+				string? targetDirStr = context.Arguments.GetStringOrDefault("-TargetDir=", null);
+				context.Arguments.CheckAllArgumentsUsed(logger);
+
+				string? hordeUrl = DeploymentSettings.Instance.HordeUrl;
+				if (hordeUrl == null)
+				{
+					logger.LogError("Horde URL is not set in deployment config file. Cannot upgrade.");
+					return;
+				}
+
+				string toolName;
+				if (RuntimeInformation.IsOSPlatform(OSPlatform.Windows))
+				{
+					toolName = "ugs-mac";
+				}
+				else if (RuntimeInformation.IsOSPlatform(OSPlatform.OSX))
+				{
+					toolName = "ugs-mac";
+				}
+				else if (RuntimeInformation.IsOSPlatform(OSPlatform.Linux))
+				{
+					toolName = "ugs-linux";
+				}
+				else
+				{
+					logger.LogError("Unknown platform; cannot query for upgrade.");
+					return;
+				}
+
+				using (HttpClient httpClient = new HttpClient())
+				{
+					Uri baseUrl = new Uri(hordeUrl);
+
+					string currentVersion = GetVersion();
+
+					DeploymentInfo? deploymentInfo;
+					try
+					{
+						deploymentInfo = await httpClient.GetFromJsonAsync<DeploymentInfo>(new Uri(baseUrl, $"api/v1/tools/{toolName}/deployments"));
+					}
+					catch (Exception ex)
+					{
+						logger.LogError(ex, "Failed to query for deployment info: {Message}", ex.Message);
+						return;
+					}
+					if (deploymentInfo == null)
+					{
+						logger.LogError("Failed to query for deployment info.");
+						return;
+					}
+					if (deploymentInfo.Version.Equals(currentVersion, StringComparison.OrdinalIgnoreCase))
+					{
+						logger.LogInformation("You are running the latest version ({Version})", currentVersion);
+						return;
+					}
+					if (options.Check)
+					{
+						logger.LogWarning("A newer version of UGS is available ({NewVersion})", deploymentInfo.Version);
+						return;
+					}
+
+					logger.LogInformation("Downloading {NewVersion}...", deploymentInfo.Version);
+
+					DirectoryReference currentDir = new FileReference(Assembly.GetExecutingAssembly().Location).Directory;
+
+					DirectoryReference targetDir = (targetDirStr == null)? currentDir : DirectoryReference.Combine(currentDir, targetDirStr);
+					DirectoryReference.CreateDirectory(targetDir);
+
+					FileReference tempFile = FileReference.Combine(targetDir, "update.zip");
+					using (Stream requestStream = await httpClient.GetStreamAsync(new Uri(baseUrl, $"api/v1/tools/{toolName}?action=download")))
+					{
+						using (Stream tempFileStream = FileReference.Open(tempFile, FileMode.Create, FileAccess.Write, FileShare.None))
+						{
+							await requestStream.CopyToAsync(tempFileStream);
+						}
+					}
+
+					using (FileStream stream = FileReference.Open(tempFile, FileMode.Open, FileAccess.Read, FileShare.Read))
+					{
+						using (ZipArchive archive = new ZipArchive(stream, ZipArchiveMode.Read, true))
+						{
+							foreach (ZipArchiveEntry entry in archive.Entries)
+							{
+								FileReference targetFile = FileReference.Combine(targetDir, entry.Name);
+								if (!targetFile.IsUnderDirectory(targetDir))
+								{
+									throw new InvalidDataException("Attempt to extract file outside source directory");
+								}
+								entry.ExtractToFile_CrossPlatform(targetFile.FullName, true);
+							}
+						}
+					}
+				}
+			}
+		}
 	}
 }
diff --git a/Engine/Source/Programs/UnrealGameSync/UnrealGameSyncShared/DeploymentSettings.cs b/Engine/Source/Programs/UnrealGameSync/UnrealGameSyncShared/DeploymentSettings.cs
--- a/Engine/Source/Programs/UnrealGameSync/UnrealGameSyncShared/DeploymentSettings.cs
+++ b/Engine/Source/Programs/UnrealGameSync/UnrealGameSyncShared/DeploymentSettings.cs
@@ -20,6 +20,11 @@ namespace UnrealGameSync
 	/// </summary>
 	public partial class DeploymentSettings
 	{
+		/// <summary>
+		/// Url for the Horde server
+		/// </summary>
+		public string? HordeUrl { get; set; }
+
 		/// <summary>
 		/// SQL connection string used to connect to the database for telemetry and review data.
 		/// </summary>
