UGS: Add version strings to UGS tool uploads.

diff --git a/Engine/Source/Programs/UnrealGameSync/BuildUGS.xml b/Engine/Source/Programs/UnrealGameSync/BuildUGS.xml
--- a/Engine/Source/Programs/UnrealGameSync/BuildUGS.xml
+++ b/Engine/Source/Programs/UnrealGameSync/BuildUGS.xml
@@ -75,7 +75,7 @@
 		</Node>
 		<Node Name="Publish UnrealGameSyncCmd Mac" Requires="Build UnrealGameSyncCmd Mac">
 			<Error Message="Missing -HordeSettings=... argument" If="'$(HordeSettings)' == ''"/>
-			<DeployTool Id="ugs-mac" Settings="$(HordeSettings)" Directory="$(UgsCmdMacOutputDir)"/>
+			<DeployTool Id="ugs-mac" Version="$(InformationalVersion)" Settings="$(HordeSettings)" Directory="$(UgsCmdMacOutputDir)"/>
 		</Node>
 		
 		<Node Name="Build UnrealGameSyncCmd Linux">
@@ -84,7 +84,7 @@
 		</Node>
 		<Node Name="Publish UnrealGameSyncCmd Linux" Requires="Build UnrealGameSyncCmd Linux">
 			<Error Message="Missing -HordeSettings=... argument" If="'$(HordeSettings)' == ''"/>
-			<DeployTool Id="ugs-linux" Settings="$(HordeSettings)" Directory="$(UgsCmdMacOutputDir)"/>
+			<DeployTool Id="ugs-linux" Version="$(InformationalVersion)" Settings="$(HordeSettings)" Directory="$(UgsCmdMacOutputDir)"/>
 		</Node>
 	</Agent>
 
