UGS: Automatically check for UGS-CLI updates.

diff --git a/Engine/Source/Programs/UnrealGameSync/UnrealGameSyncCmd/Program.cs b/Engine/Source/Programs/UnrealGameSync/UnrealGameSyncCmd/Program.cs
--- a/Engine/Source/Programs/UnrealGameSync/UnrealGameSyncCmd/Program.cs
+++ b/Engine/Source/Programs/UnrealGameSync/UnrealGameSyncCmd/Program.cs
@@ -21,11 +21,8 @@ using UnrealGameSync;
 using System.Text;
 using System.Globalization;
 using System.Net.Http;
-using System.Text.Json;
 using System.Net.Http.Json;
-using Microsoft.Extensions.Options;
 using System.IO.Compression;
-using JetBrains.Annotations;
 
 namespace UnrealGameSyncCmd
 {
@@ -58,10 +55,10 @@ namespace UnrealGameSyncCmd
 			public string Name { get; }
 			public Type Type { get; }
 			public Type? OptionsType { get; }
-			public string Usage { get; }
-			public string Brief { get; }
+			public string? Usage { get; }
+			public string? Brief { get; }
 
-			public CommandInfo(string name, Type type, Type? optionsType, string usage, string brief)
+			public CommandInfo(string name, Type type, Type? optionsType, string? usage, string? brief)
 			{
 				Name = name;
 				Type = type;
@@ -121,12 +118,14 @@ namespace UnrealGameSyncCmd
 				"ugs version",
 				"Prints the current application version"
 			),
-			new CommandInfo("install", typeof(InstallCommand), null,
-				"ugs install",
-				"Registers an alias to run ugs via 'ugs'."),
+			new CommandInfo("install", typeof(InstallCommand), null, 
+				null, 
+				null
+			),
 			new CommandInfo("upgrade", typeof(UpgradeCommand), typeof(UpgradeCommandOptions),
 				"ugs upgrade",
-				"Upgrades the current installation with the latest build of UGS.")
+				"Upgrades the current installation with the latest build of UGS."
+			)
 		};
 
 		class CommandContext
@@ -181,6 +180,12 @@ namespace UnrealGameSyncCmd
 			public string? ProjectName { get; set; }
 		}
 
+		class UpdateState
+		{
+			public string? LatestVersion { get; set; }
+			public DateTime LastVersionCheck { get; set; }
+		}
+
 		public static async Task<int> Main(string[] rawArgs)
 		{
 			DirectoryReference globalConfigFolder;
@@ -250,6 +255,49 @@ namespace UnrealGameSyncCmd
 					return 1;
 				}
 
+				if (command.Type != typeof(UpgradeCommand))
+				{
+					DateTime utcNow = DateTime.UtcNow;
+					if (settings.Global.LastVersionCheck < utcNow - TimeSpan.FromDays(1.0))
+					{
+						using (CancellationTokenSource cancellationSource = new CancellationTokenSource(TimeSpan.FromSeconds(10.0)))
+						{
+							Task<string?> latestVersionTask = GetLatestVersionAsync(null, cancellationSource.Token);
+
+							Task delay = Task.Delay(TimeSpan.FromSeconds(2.0));
+							await Task.WhenAny(latestVersionTask, delay);
+
+							if (!latestVersionTask.IsCompleted)
+							{
+								logger.LogInformation("Checking for UGS updates...");
+							}
+
+							try
+							{
+								settings.Global.LatestVersion = await latestVersionTask;
+							}
+							catch (OperationCanceledException)
+							{
+								logger.LogInformation("Request timed out.");
+							}
+							catch (Exception ex)
+							{
+								logger.LogInformation(ex, "Upgrade check failed: {Message}", ex.Message);
+							}
+						}
+
+						settings.Global.LastVersionCheck = utcNow;
+						settings.Save(logger);
+					}
+
+					string version = GetVersion();
+					if (settings.Global.LatestVersion != null && !settings.Global.LatestVersion.Equals(version, StringComparison.OrdinalIgnoreCase))
+					{
+						logger.LogWarning("A newer version of UGS is available ({LatestVersion}). Run {Command} to update.", settings.Global.LatestVersion, "ugs upgrade");
+						logger.LogInformation("");
+					}
+				}
+
 				Command instance = (Command)Activator.CreateInstance(command.Type)!;
 				await instance.ExecuteAsync(new CommandContext(args, logger, loggerFactory, settings));
 				return 0;
@@ -276,9 +324,12 @@ namespace UnrealGameSyncCmd
 			Console.WriteLine("Usage:");
 			foreach (CommandInfo command in _commands)
 			{
-				Console.WriteLine();
-				ConsoleUtils.WriteLineWithWordWrap(GetUsage(command), 2, 8);
-				ConsoleUtils.WriteLineWithWordWrap(command.Brief, 4, 4);
+				if (command.Usage != null && command.Brief != null)
+				{
+					Console.WriteLine();
+					ConsoleUtils.WriteLineWithWordWrap(GetUsage(command), 2, 8);
+					ConsoleUtils.WriteLineWithWordWrap(command.Brief, 4, 4);
+				}
 			}
 		}
 
@@ -311,6 +362,72 @@ namespace UnrealGameSyncCmd
 			return result.ToString();
 		}
 
+		static string GetVersion()
+		{
+			AssemblyInformationalVersionAttribute? version = Assembly.GetExecutingAssembly().GetCustomAttribute<AssemblyInformationalVersionAttribute>();
+			return version?.InformationalVersion ?? "Unknown";
+		}
+
+		class DeploymentInfo
+		{
+			public string Version { get; set; } = String.Empty;
+		}
+
+		static string? GetUpgradeToolName()
+		{
+			if (RuntimeInformation.IsOSPlatform(OSPlatform.OSX))
+			{
+				return "ugs-mac";
+			}
+			else if (RuntimeInformation.IsOSPlatform(OSPlatform.Linux))
+			{
+				return "ugs-linux";
+			}
+			else
+			{
+				return "ugs-mac";
+			}
+		}
+
+		static async Task<string?> GetLatestVersionAsync(ILogger? logger, CancellationToken cancellationToken)
+		{
+			string? hordeUrl = DeploymentSettings.Instance.HordeUrl;
+			if (hordeUrl == null)
+			{
+				logger?.LogError("Horde URL is not set in deployment config file. Cannot upgrade.");
+				return null;
+			}
+
+			string? toolName = GetUpgradeToolName();
+			if (toolName == null)
+			{
+				logger?.LogError("Unknown platform; cannot query for upgrade.");
+				return null;
+			}
+
+			using (HttpClient httpClient = new HttpClient())
+			{
+				Uri baseUrl = new Uri(hordeUrl);
+
+				DeploymentInfo? deploymentInfo;
+				try
+				{
+					deploymentInfo = await httpClient.GetFromJsonAsync<DeploymentInfo>(new Uri(baseUrl, $"api/v1/tools/{toolName}/deployments"), cancellationToken);
+				}
+				catch (Exception ex)
+				{
+					logger?.LogError(ex, "Failed to query for deployment info: {Message}", ex.Message);
+					return null;
+				}
+				if (deploymentInfo == null)
+				{
+					logger?.LogError("Failed to query for deployment info.");
+					return null;
+				}
+				return deploymentInfo.Version;
+			}
+		}
+
 		public static UserWorkspaceSettings? ReadOptionalUserWorkspaceSettings()
 		{
 			DirectoryReference? dir = DirectoryReference.GetCurrentDirectory();
@@ -1350,12 +1467,6 @@ namespace UnrealGameSyncCmd
 			}
 		}
 
-		static string GetVersion()
-		{
-			AssemblyInformationalVersionAttribute? version = Assembly.GetExecutingAssembly().GetCustomAttribute<AssemblyInformationalVersionAttribute>();
-			return version?.InformationalVersion ?? "Unknown";
-		}
-
 		class VersionCommand : Command
 		{
 			public override Task ExecuteAsync(CommandContext context)
@@ -1406,11 +1517,6 @@ namespace UnrealGameSyncCmd
 
 		class UpgradeCommand : Command
 		{
-			class DeploymentInfo
-			{
-				public string Version { get; set; } = String.Empty;
-			}
-
 			public override async Task ExecuteAsync(CommandContext context)
 			{
 				ILogger logger = context.Logger;
@@ -1420,65 +1526,27 @@ namespace UnrealGameSyncCmd
 				string? targetDirStr = context.Arguments.GetStringOrDefault("-TargetDir=", null);
 				context.Arguments.CheckAllArgumentsUsed(logger);
 
-				string? hordeUrl = DeploymentSettings.Instance.HordeUrl;
-				if (hordeUrl == null)
-				{
-					logger.LogError("Horde URL is not set in deployment config file. Cannot upgrade.");
-					return;
-				}
+				string currentVersion = GetVersion();
 
-				string toolName;
-				if (RuntimeInformation.IsOSPlatform(OSPlatform.Windows))
+				string? latestVersion = await GetLatestVersionAsync(logger, CancellationToken.None);
+				if (latestVersion == null)
 				{
-					toolName = "ugs-mac";
-				}
-				else if (RuntimeInformation.IsOSPlatform(OSPlatform.OSX))
-				{
-					toolName = "ugs-mac";
+					return;
 				}
-				else if (RuntimeInformation.IsOSPlatform(OSPlatform.Linux))
+				if (latestVersion.Equals(currentVersion, StringComparison.OrdinalIgnoreCase))
 				{
-					toolName = "ugs-linux";
+					logger.LogInformation("You are running the latest version ({Version})", currentVersion);
+					return;
 				}
-				else
+				if (options.Check)
 				{
-					logger.LogError("Unknown platform; cannot query for upgrade.");
+					logger.LogWarning("A newer version of UGS is available ({NewVersion})", latestVersion);
 					return;
 				}
 
 				using (HttpClient httpClient = new HttpClient())
 				{
-					Uri baseUrl = new Uri(hordeUrl);
-
-					string currentVersion = GetVersion();
-
-					DeploymentInfo? deploymentInfo;
-					try
-					{
-						deploymentInfo = await httpClient.GetFromJsonAsync<DeploymentInfo>(new Uri(baseUrl, $"api/v1/tools/{toolName}/deployments"));
-					}
-					catch (Exception ex)
-					{
-						logger.LogError(ex, "Failed to query for deployment info: {Message}", ex.Message);
-						return;
-					}
-					if (deploymentInfo == null)
-					{
-						logger.LogError("Failed to query for deployment info.");
-						return;
-					}
-					if (deploymentInfo.Version.Equals(currentVersion, StringComparison.OrdinalIgnoreCase))
-					{
-						logger.LogInformation("You are running the latest version ({Version})", currentVersion);
-						return;
-					}
-					if (options.Check)
-					{
-						logger.LogWarning("A newer version of UGS is available ({NewVersion})", deploymentInfo.Version);
-						return;
-					}
-
-					logger.LogInformation("Downloading {NewVersion}...", deploymentInfo.Version);
+					Uri baseUrl = new Uri(DeploymentSettings.Instance.HordeUrl ?? String.Empty);
 
 					DirectoryReference currentDir = new FileReference(Assembly.GetExecutingAssembly().Location).Directory;
 
@@ -1486,7 +1554,7 @@ namespace UnrealGameSyncCmd
 					DirectoryReference.CreateDirectory(targetDir);
 
 					FileReference tempFile = FileReference.Combine(targetDir, "update.zip");
-					using (Stream requestStream = await httpClient.GetStreamAsync(new Uri(baseUrl, $"api/v1/tools/{toolName}?action=download")))
+					using (Stream requestStream = await httpClient.GetStreamAsync(new Uri(baseUrl, $"api/v1/tools/{GetUpgradeToolName()}?action=download")))
 					{
 						using (Stream tempFileStream = FileReference.Open(tempFile, FileMode.Create, FileAccess.Write, FileShare.None))
 						{
diff --git a/Engine/Source/Programs/UnrealGameSync/UnrealGameSyncShared/GlobalSettings.cs b/Engine/Source/Programs/UnrealGameSync/UnrealGameSyncShared/GlobalSettings.cs
--- a/Engine/Source/Programs/UnrealGameSync/UnrealGameSyncShared/GlobalSettings.cs
+++ b/Engine/Source/Programs/UnrealGameSync/UnrealGameSyncShared/GlobalSettings.cs
@@ -51,6 +51,8 @@ namespace UnrealGameSync
 
 	public class GlobalSettings
 	{
+		public string? LatestVersion { get; set; }
+		public DateTime LastVersionCheck { get; set; }
 		public PerforceSyncOptions Perforce { get; set; } = new PerforceSyncOptions();
 		public FilterSettings Filter { get; set; } = new FilterSettings();
 	}
