Allow setting executable bits for files in zip archives from Windows, and set the executable permissions on Mac/Linux builds of UGS-CLI by default.

diff --git a/Engine/Source/Programs/AutomationTool/AutomationUtils/CommandUtils.cs b/Engine/Source/Programs/AutomationTool/AutomationUtils/CommandUtils.cs
--- a/Engine/Source/Programs/AutomationTool/AutomationUtils/CommandUtils.cs
+++ b/Engine/Source/Programs/AutomationTool/AutomationUtils/CommandUtils.cs
@@ -2227,10 +2227,11 @@ namespace AutomationTool
 		/// <summary>
 		/// Creates a zip file containing the given input files
 		/// </summary>
-		/// <param name="ZipFile">Filename for the zip</param>
+		/// <param name="OutputFile">Filename for the zip</param>
 		/// <param name="BaseDirectory">Base directory to store relative paths in the zip file to</param>
 		/// <param name="Files">Files to include in the archive</param>
-		public static void ZipFiles(FileReference OutputFile, DirectoryReference BaseDirectory, IEnumerable<FileReference> Files)
+		/// <param name="ExecutableFiles">Files to flag with the executable attribute</param>
+		public static void ZipFiles(FileReference OutputFile, DirectoryReference BaseDirectory, IEnumerable<FileReference> Files, HashSet<FileReference> ExecutableFiles = null)
 		{
 			if (!DirectoryReference.Exists(OutputFile.Directory))
 			{
@@ -2251,7 +2252,19 @@ namespace AutomationTool
 					{
 						Name = Name.Replace(Path.DirectorySeparatorChar, '/');
 					}
-					ZipArchive.CreateEntryFromFile_CrossPlatform(File.FullName, Name, CompressionLevel.Fastest);
+
+					ZipArchiveEntry entry = ZipArchive.CreateEntryFromFile_CrossPlatform(File.FullName, Name, CompressionLevel.Fastest);
+					if (ExecutableFiles != null)
+					{
+						if (ExecutableFiles.Contains(File))
+						{
+							entry.ExternalAttributes |= 0b_111_111_101 << 16; // rwx rwx r-x
+						}
+						else
+						{
+							entry.ExternalAttributes |= 0b_110_110_100 << 16; // rw- rw- r--
+						}
+					}
 				}
 			}
 		}
diff --git a/Engine/Source/Programs/AutomationTool/BuildGraph/Tasks/ZipTask.cs b/Engine/Source/Programs/AutomationTool/BuildGraph/Tasks/ZipTask.cs
--- a/Engine/Source/Programs/AutomationTool/BuildGraph/Tasks/ZipTask.cs
+++ b/Engine/Source/Programs/AutomationTool/BuildGraph/Tasks/ZipTask.cs
@@ -4,6 +4,7 @@ using AutomationTool;
 using EpicGames.BuildGraph;
 using System;
 using System.Collections.Generic;
+using System.IO.Compression;
 using System.Linq;
 using System.Text;
 using System.Threading.Tasks;
@@ -11,6 +12,7 @@ using System.Xml;
 using EpicGames.Core;
 using UnrealBuildTool;
 using Microsoft.Extensions.Logging;
+using System.IO;
 
 namespace AutomationTool.Tasks
 {
@@ -31,6 +33,12 @@ namespace AutomationTool.Tasks
 		[TaskParameter(Optional = true, ValidationType = TaskParameterValidationType.FileSpec)]
 		public string Files;
 
+		/// <summary>
+		/// List of files that should have an executable bit set.
+		/// </summary>
+		[TaskParameter(Optional = true, ValidationType = TaskParameterValidationType.FileSpec)]
+		public string ExecutableFiles;
+
 		/// <summary>
 		/// The zip file to create.
 		/// </summary>
@@ -85,10 +93,21 @@ namespace AutomationTool.Tasks
 
 			// Create the zip file
 			Logger.LogInformation("Adding {NumFiles} files to {ZipFile}...", Files.Count, Parameters.ZipFile);
-			CommandUtils.ZipFiles(Parameters.ZipFile, Parameters.FromDir, Files);
+
+			HashSet<FileReference> ExecutableFiles = null;
+			if (Parameters.ExecutableFiles != null)
+			{
+				ExecutableFiles = ResolveFilespec(Parameters.FromDir, Parameters.ExecutableFiles, TagNameToFileSet);
+				foreach (FileReference ExecutableFile in Files.Intersect(ExecutableFiles))
+				{
+					Logger.LogInformation("  Executable file: {File}", ExecutableFile);
+				}
+			}
+
+			CommandUtils.ZipFiles(Parameters.ZipFile, Parameters.FromDir, Files, ExecutableFiles);
 
 			// Apply the optional tag to the produced archive
-			foreach(string TagName in FindTagNamesFromList(Parameters.Tag))
+			foreach (string TagName in FindTagNamesFromList(Parameters.Tag))
 			{
 				FindOrAddTagSet(TagNameToFileSet, TagName).Add(Parameters.ZipFile);
 			}
diff --git a/Engine/Source/Programs/UnrealGameSync/BuildUGS.xml b/Engine/Source/Programs/UnrealGameSync/BuildUGS.xml
--- a/Engine/Source/Programs/UnrealGameSync/BuildUGS.xml
+++ b/Engine/Source/Programs/UnrealGameSync/BuildUGS.xml
@@ -13,7 +13,9 @@
 	<Property Name="UgsOutputDir" Value="$(RootDir)/LocalBuilds/UnrealGameSync"/>
 	<Property Name="UgsZipOutputDir" Value="$(RootDir)/LocalBuilds/UnrealGameSyncZip"/>
 	<Property Name="UgsCmdOutputDir" Value="$(RootDir)/LocalBuilds/UnrealGameSyncCmd"/>
+	<Property Name="UgsCmdMacZipFile" Value="$(UgsCmdOutputDir)/ugs-mac.zip"/>
 	<Property Name="UgsCmdMacOutputDir" Value="$(UgsCmdOutputDir)/Mac"/>
+	<Property Name="UgsCmdLinuxZipFile" Value="$(UgsCmdOutputDir)/ugs-linux.zip"/>
 	<Property Name="UgsCmdLinuxOutputDir" Value="$(UgsCmdOutputDir)/Linux"/>
 	<Property Name="WixDir" Value="$(RootDir)/Engine/Source/ThirdParty/WiX/3.8"/>
 
@@ -72,19 +74,21 @@
 		<Node Name="Build UnrealGameSyncCmd Mac">
 			<Delete Files="$(UgsCmdMacOutputDir)/..."/>
 			<DotNet Arguments="publish &quot;$(UgsDir)/UnrealGameSyncCmd/UnrealGameSyncCmd.csproj&quot; --output=&quot;$(UgsCmdMacOutputDir)&quot; --self-contained -c Release -r osx-x64 $(VersionArguments)"/>
+			<Zip ZipFile="$(UgsCmdMacZipFile)" FromDir="$(UgsCmdMacOutputDir)" ExecutableFiles="ugs;*.dylib"/>
 		</Node>
 		<Node Name="Publish UnrealGameSyncCmd Mac" Requires="Build UnrealGameSyncCmd Mac">
 			<Error Message="Missing -HordeSettings=... argument" If="'$(HordeSettings)' == ''"/>
-			<DeployTool Id="ugs-mac" Version="$(InformationalVersion)" Settings="$(HordeSettings)" Directory="$(UgsCmdMacOutputDir)"/>
+			<DeployTool Id="ugs-mac" Version="$(InformationalVersion)" Settings="$(HordeSettings)" File="$(UgsCmdMacZipFile)"/>
 		</Node>
 		
 		<Node Name="Build UnrealGameSyncCmd Linux">
 			<Delete Files="$(UgsCmdLinuxOutputDir)/..."/>
 			<DotNet Arguments="publish &quot;$(UgsDir)/UnrealGameSyncCmd/UnrealGameSyncCmd.csproj&quot; --output=&quot;$(UgsCmdLinuxOutputDir)&quot; --self-contained -c Release -r linux-x64 $(VersionArguments)"/>
+			<Zip ZipFile="$(UgsCmdLinuxZipFile)" FromDir="$(UgsCmdLinuxOutputDir)" ExecutableFiles="ugs;*.so"/>
 		</Node>
 		<Node Name="Publish UnrealGameSyncCmd Linux" Requires="Build UnrealGameSyncCmd Linux">
 			<Error Message="Missing -HordeSettings=... argument" If="'$(HordeSettings)' == ''"/>
-			<DeployTool Id="ugs-linux" Version="$(InformationalVersion)" Settings="$(HordeSettings)" Directory="$(UgsCmdMacOutputDir)"/>
+			<DeployTool Id="ugs-linux" Version="$(InformationalVersion)" Settings="$(HordeSettings)" File="$(UgsCmdLinuxZipFile)"/>
 		</Node>
 	</Agent>
 
