UGS: Add a setup script for UGS on Mac, and fix cached settings for upgrades causing false-positive notifications.

diff --git a/Engine/Source/Programs/UnrealGameSync/BuildUGS.xml b/Engine/Source/Programs/UnrealGameSync/BuildUGS.xml
--- a/Engine/Source/Programs/UnrealGameSync/BuildUGS.xml
+++ b/Engine/Source/Programs/UnrealGameSync/BuildUGS.xml
@@ -74,7 +74,7 @@
 		<Node Name="Build UnrealGameSyncCmd Mac">
 			<Delete Files="$(UgsCmdMacOutputDir)/..."/>
 			<DotNet Arguments="publish &quot;$(UgsDir)/UnrealGameSyncCmd/UnrealGameSyncCmd.csproj&quot; --output=&quot;$(UgsCmdMacOutputDir)&quot; --self-contained -c Release -r osx-x64 $(VersionArguments)"/>
-			<Zip ZipFile="$(UgsCmdMacZipFile)" FromDir="$(UgsCmdMacOutputDir)" ExecutableFiles="ugs;*.dylib"/>
+			<Zip ZipFile="$(UgsCmdMacZipFile)" FromDir="$(UgsCmdMacOutputDir)" ExecutableFiles="ugs;*.dylib;*.sh"/>
 		</Node>
 		<Node Name="Publish UnrealGameSyncCmd Mac" Requires="Build UnrealGameSyncCmd Mac">
 			<Error Message="Missing -HordeSettings=... argument" If="'$(HordeSettings)' == ''"/>
@@ -84,7 +84,7 @@
 		<Node Name="Build UnrealGameSyncCmd Linux">
 			<Delete Files="$(UgsCmdLinuxOutputDir)/..."/>
 			<DotNet Arguments="publish &quot;$(UgsDir)/UnrealGameSyncCmd/UnrealGameSyncCmd.csproj&quot; --output=&quot;$(UgsCmdLinuxOutputDir)&quot; --self-contained -c Release -r linux-x64 $(VersionArguments)"/>
-			<Zip ZipFile="$(UgsCmdLinuxZipFile)" FromDir="$(UgsCmdLinuxOutputDir)" ExecutableFiles="ugs;*.so"/>
+			<Zip ZipFile="$(UgsCmdLinuxZipFile)" FromDir="$(UgsCmdLinuxOutputDir)" ExecutableFiles="ugs;*.so;*.sh"/>
 		</Node>
 		<Node Name="Publish UnrealGameSyncCmd Linux" Requires="Build UnrealGameSyncCmd Linux">
 			<Error Message="Missing -HordeSettings=... argument" If="'$(HordeSettings)' == ''"/>
diff --git a/Engine/Source/Programs/UnrealGameSync/UnrealGameSyncCmd/Program.cs b/Engine/Source/Programs/UnrealGameSync/UnrealGameSyncCmd/Program.cs
--- a/Engine/Source/Programs/UnrealGameSync/UnrealGameSyncCmd/Program.cs
+++ b/Engine/Source/Programs/UnrealGameSync/UnrealGameSyncCmd/Program.cs
@@ -255,10 +255,12 @@ namespace UnrealGameSyncCmd
 					return 1;
 				}
 
-				if (command.Type != typeof(UpgradeCommand))
+				if (command.Type != typeof(InstallCommand) && command.Type != typeof(UpgradeCommand))
 				{
+					string version = GetVersion();
+
 					DateTime utcNow = DateTime.UtcNow;
-					if (settings.Global.LastVersionCheck < utcNow - TimeSpan.FromDays(1.0))
+					if (settings.Global.LastVersionCheck < utcNow - TimeSpan.FromDays(1.0) || IsUpgradeAvailable(settings, version))
 					{
 						using (CancellationTokenSource cancellationSource = new CancellationTokenSource(TimeSpan.FromSeconds(10.0)))
 						{
@@ -290,10 +292,9 @@ namespace UnrealGameSyncCmd
 						settings.Save(logger);
 					}
 
-					string version = GetVersion();
-					if (settings.Global.LatestVersion != null && !settings.Global.LatestVersion.Equals(version, StringComparison.OrdinalIgnoreCase))
+					if (IsUpgradeAvailable(settings, version))
 					{
-						logger.LogWarning("A newer version of UGS is available ({LatestVersion}). Run {Command} to update.", settings.Global.LatestVersion, "ugs upgrade");
+						logger.LogWarning("A newer version of UGS is available ({LatestVersion} > {Version}). Run {Command} to update.", settings.Global.LatestVersion, version, "ugs upgrade");
 						logger.LogInformation("");
 					}
 				}
@@ -319,6 +320,11 @@ namespace UnrealGameSyncCmd
 			}
 		}
 
+		static bool IsUpgradeAvailable(GlobalSettingsFile settings, string version)
+		{
+			return settings.Global.LatestVersion != null && !settings.Global.LatestVersion.Equals(version, StringComparison.OrdinalIgnoreCase);
+		}
+
 		static void PrintHelp()
 		{
 			Console.WriteLine("Usage:");
@@ -1513,6 +1519,9 @@ namespace UnrealGameSyncCmd
 		{
 			[CommandLine("-Check")]
 			public bool Check { get; set; }
+
+			[CommandLine("-Force")]
+			public bool Force { get; set; }
 		}
 
 		class UpgradeCommand : Command
@@ -1533,7 +1542,7 @@ namespace UnrealGameSyncCmd
 				{
 					return;
 				}
-				if (latestVersion.Equals(currentVersion, StringComparison.OrdinalIgnoreCase))
+				if (latestVersion.Equals(currentVersion, StringComparison.OrdinalIgnoreCase) && !options.Force)
 				{
 					logger.LogInformation("You are running the latest version ({Version})", currentVersion);
 					return;
diff --git a/Engine/Source/Programs/UnrealGameSync/UnrealGameSyncCmd/UnrealGameSyncCmd.csproj b/Engine/Source/Programs/UnrealGameSync/UnrealGameSyncCmd/UnrealGameSyncCmd.csproj
--- a/Engine/Source/Programs/UnrealGameSync/UnrealGameSyncCmd/UnrealGameSyncCmd.csproj
+++ b/Engine/Source/Programs/UnrealGameSync/UnrealGameSyncCmd/UnrealGameSyncCmd.csproj
@@ -18,4 +18,10 @@
     <ProjectReference Include="..\UnrealGameSyncShared\UnrealGameSyncShared.csproj" />
   </ItemGroup>
 
+  <ItemGroup>
+    <None Update="setup-mac.sh">
+      <CopyToOutputDirectory>Always</CopyToOutputDirectory>
+    </None>
+  </ItemGroup>
+
 </Project>
diff --git a/Engine/Source/Programs/UnrealGameSync/UnrealGameSyncCmd/setup-mac.sh b/Engine/Source/Programs/UnrealGameSync/UnrealGameSyncCmd/setup-mac.sh
new file mode 100644
--- /dev/null
+++ b/Engine/Source/Programs/UnrealGameSync/UnrealGameSyncCmd/setup-mac.sh
@@ -0,0 +1,5 @@
+#!/bin/zsh
+
+cd "$(dirname "$0")"
+./ugs install
+source ~/.zshrc
