UGS: Fixing some Publish tasks to allow preferring stream paths over branch paths

#rnx

diff --git a/Engine/Source/Programs/UnrealGameSync/BuildUGS.xml b/Engine/Source/Programs/UnrealGameSync/BuildUGS.xml
--- a/Engine/Source/Programs/UnrealGameSync/BuildUGS.xml
+++ b/Engine/Source/Programs/UnrealGameSync/BuildUGS.xml
@@ -57,14 +57,20 @@
 			<Copy From="$(InstallerDir)/bin/Release/UnrealGameSync.msi" To="$(UgsPublishDir)/UnrealGameSync.msi" If="$(IncludeInstaller)"/>
 		</Node>
 		<Node Name="Publish Unstable" Requires="Build UnrealGameSync">
-			<Submit Description="Updating UGS from CL $(Change)" Workspace="Temp" Files="$(UgsPublishDir)/..." RootDir="$(UgsOutputDir)" Stream="$(SubmitToStream)" Branch="$(SubmitToBranch)/UnstableRelease"/>
+			<Property Name="UnstableBranch" Value="" />
+			<Property Name="UnstableBranch" Value="$(SubmitToBranch)/UnstableRelease" If="'$(SubmitToBranch)' != ''" />
+
+			<Submit Description="Updating UGS from CL $(Change)" Workspace="Temp" Files="$(UgsPublishDir)/..." RootDir="$(UgsOutputDir)" Stream="$(SubmitToStream)" Branch="$(UnstableBranch)"/>
 		</Node>
 		<Node Name="Publish Unstable Zip" Requires="Build UnrealGameSync">
 			<Zip FromDir="$(UgsOutputDir)" ZipFile="$(UgsZipPublishDir)/UnstableRelease.zip"/>
 			<Submit Description="Updating UGS from CL $(Change) (Zip)" Workspace="Temp" Files="$(UgsZipPublishDir)/UnstableRelease.zip" RootDir="$(UgsZipOutputDir)" Stream="$(SubmitToStream)" Branch="$(SubmitToBranch)"/>
 		</Node>
 		<Node Name="Publish Live" Requires="Build UnrealGameSync">
-			<Submit Description="Updating UGS from CL $(Change)" Workspace="Temp" Files="$(UgsPublishDir)/..." RootDir="$(UgsOutputDir)" Stream="$(SubmitToStream)" Branch="$(SubmitToBranch)/Release"/>
+			<Property Name="ReleaseBranch" Value="" />
+			<Property Name="ReleaseBranch" Value="$(SubmitToBranch)/Release" If="'$(SubmitToBranch)' != ''" />
+
+			<Submit Description="Updating UGS from CL $(Change)" Workspace="Temp" Files="$(UgsPublishDir)/..." RootDir="$(UgsOutputDir)" Stream="$(SubmitToStream)" Branch="$(ReleaseBranch)"/>
 		</Node>
 		<Node Name="Publish Live Zip" Requires="Build UnrealGameSync">
 			<Zip FromDir="$(UgsOutputDir)" ZipFile="$(UgsZipPublishDir)/Release.zip"/>
