UGS: Add the possibility to store badge metadata & show additional badge links on right click.

diff --git a/Engine/Source/Programs/UnrealGameSync/MetadataServer/Connectors/SqlConnector.cs b/Engine/Source/Programs/UnrealGameSync/MetadataServer/Connectors/SqlConnector.cs
--- a/Engine/Source/Programs/UnrealGameSync/MetadataServer/Connectors/SqlConnector.cs
+++ b/Engine/Source/Programs/UnrealGameSync/MetadataServer/Connectors/SqlConnector.cs
@@ -17,6 +17,8 @@ namespace MetadataServer.Connectors
 	public static class SqlConnector
 	{
 		private static string ConnectionString = System.Configuration.ConfigurationManager.ConnectionStrings["ConnectionString"].ConnectionString;
+		private static bool StoreBuildMetadata = Convert.ToBoolean(System.Configuration.ConfigurationManager.AppSettings["StoreBuildMetadata"]);
+
 		public static LatestData GetLastIds(string Project = null)
 		{
 			// Get ids going back 432 builds for the project being asked for
@@ -151,9 +153,14 @@ namespace MetadataServer.Connectors
 			string ProjectLikeString = "%" + (Project == null ? String.Empty : GetProjectStream(Project)) + "%";
 			using (MySqlConnection Connection = new MySqlConnection(ConnectionString))
 			{
+				string QueryWithMetadata =
+					"SELECT Badges.Id, Badges.ChangeNumber, Badges.BuildType, Badges.Result, Badges.Url, Projects.Name, Badges.ArchivePath, Badges.Metadata FROM ugs_db.Badges " +
+					"INNER JOIN ugs_db.Projects ON Projects.Id = Badges.ProjectId WHERE Badges.Id > @param1 AND Projects.Name LIKE @param2 ORDER BY Badges.Id";
+				string QueryWithoutMetadata =
+					"SELECT Badges.Id, Badges.ChangeNumber, Badges.BuildType, Badges.Result, Badges.Url, Projects.Name, Badges.ArchivePath FROM ugs_db.Badges " +
+					"INNER JOIN ugs_db.Projects ON Projects.Id = Badges.ProjectId WHERE Badges.Id > @param1 AND Projects.Name LIKE @param2 ORDER BY Badges.Id";
 				Connection.Open();
-				using (MySqlCommand Command = new MySqlCommand("SELECT Badges.Id, Badges.ChangeNumber, Badges.BuildType, Badges.Result, Badges.Url, Projects.Name, Badges.ArchivePath FROM ugs_db.Badges " +
-																 "INNER JOIN ugs_db.Projects ON Projects.Id = Badges.ProjectId WHERE Badges.Id > @param1 AND Projects.Name LIKE @param2 ORDER BY Badges.Id", Connection))
+				using (MySqlCommand Command = new MySqlCommand(StoreBuildMetadata ? QueryWithMetadata : QueryWithoutMetadata, Connection))
 				{ 
 					Command.Parameters.AddWithValue("@param1", LastBuildId);
 					Command.Parameters.AddWithValue("@param2", ProjectLikeString);
@@ -174,6 +181,12 @@ namespace MetadataServer.Connectors
 									ReturnedBuilds.Add(Build);
 								}
 							}
+
+							if (StoreBuildMetadata)
+							{
+								Build.Metadata = JsonConvert.DeserializeObject<BuildMetadata>(Reader.GetString(7));
+							}
+
 							LastBuildId = Math.Max(LastBuildId, Build.Id);
 						}
 					}
@@ -227,7 +240,12 @@ namespace MetadataServer.Connectors
 			{
 				Connection.Open();
 				long ProjectId = TryInsertAndGetProject(Connection, Build.Project);
-				using (MySqlCommand Command = new MySqlCommand("INSERT INTO ugs_db.Badges (ChangeNumber, BuildType, Result, URL, ArchivePath, ProjectId) VALUES (@ChangeNumber, @BuildType, @Result, @URL, @ArchivePath, @ProjectId)", Connection))
+
+				string InsertWithMetadata =
+					"INSERT INTO ugs_db.Badges (ChangeNumber, BuildType, Result, URL, ArchivePath, ProjectId, Metadata) VALUES (@ChangeNumber, @BuildType, @Result, @URL, @ArchivePath, @ProjectId, @Metadata)";
+				string InsertWithoutMetadata =
+					"INSERT INTO ugs_db.Badges (ChangeNumber, BuildType, Result, URL, ArchivePath, ProjectId) VALUES (@ChangeNumber, @BuildType, @Result, @URL, @ArchivePath, @ProjectId)";
+				using (MySqlCommand Command = new MySqlCommand(StoreBuildMetadata ? InsertWithMetadata : InsertWithoutMetadata, Connection))
 				{
 					Command.Parameters.AddWithValue("@ChangeNumber", Build.ChangeNumber);
 					Command.Parameters.AddWithValue("@BuildType", Build.BuildType);
@@ -235,6 +253,10 @@ namespace MetadataServer.Connectors
 					Command.Parameters.AddWithValue("@URL", Build.Url);
 					Command.Parameters.AddWithValue("@ArchivePath", Build.ArchivePath);
 					Command.Parameters.AddWithValue("@ProjectId", ProjectId);
+					if (StoreBuildMetadata)
+					{
+						Command.Parameters.AddWithValue("@Metadata", JsonConvert.SerializeObject(Build.Metadata ?? new BuildMetadata()));
+					}
 					Command.ExecuteNonQuery();
 				}
 			}
diff --git a/Engine/Source/Programs/UnrealGameSync/MetadataServer/Models/BuildData.cs b/Engine/Source/Programs/UnrealGameSync/MetadataServer/Models/BuildData.cs
--- a/Engine/Source/Programs/UnrealGameSync/MetadataServer/Models/BuildData.cs
+++ b/Engine/Source/Programs/UnrealGameSync/MetadataServer/Models/BuildData.cs
@@ -7,6 +7,17 @@ using System.Web;
 
 namespace MetadataServer.Models
 {
+	public class Link
+	{
+		public string Text { get; set; } = String.Empty;
+		public string Url { get; set; } = String.Empty;
+	}
+
+   	public class BuildMetadata
+	{
+		public List<Link> Links { get; set; } = new List<Link>();
+	}
+
 	public class BuildData
 	{
 		public enum BuildDataResult
@@ -25,6 +36,7 @@ namespace MetadataServer.Models
 		public string Url;
 		public string Project;
 		public string ArchivePath;
+		public BuildMetadata Metadata;
 
 		public bool IsSuccess
 		{
diff --git a/Engine/Source/Programs/UnrealGameSync/MetadataServer/Web.template.config b/Engine/Source/Programs/UnrealGameSync/MetadataServer/Web.template.config
--- a/Engine/Source/Programs/UnrealGameSync/MetadataServer/Web.template.config
+++ b/Engine/Source/Programs/UnrealGameSync/MetadataServer/Web.template.config
@@ -4,7 +4,9 @@
   http://go.microsoft.com/fwlink/?LinkId=301879
   -->
 <configuration>
-	<appSettings/>
+    <appSettings>
+        <add key="StoreBuildMetadata" value="false" />
+    </appSettings>
   <!--
     For a description of web.config changes see http://go.microsoft.com/fwlink/?LinkId=235367.
 
diff --git a/Engine/Source/Programs/UnrealGameSync/PostBadgeStatus/Program.cs b/Engine/Source/Programs/UnrealGameSync/PostBadgeStatus/Program.cs
--- a/Engine/Source/Programs/UnrealGameSync/PostBadgeStatus/Program.cs
+++ b/Engine/Source/Programs/UnrealGameSync/PostBadgeStatus/Program.cs
@@ -25,12 +25,14 @@ namespace WriteBadgeStatus
 			string RestUrl = ParseParam(Arguments, "RestUrl");
 			string Status = ParseParam(Arguments, "Status");
 			string Url = ParseParam(Arguments, "Url");
+			string Metadata = ParseParam(Arguments, "Metadata");
 
 			// Check we've got all the arguments we need (and no more)
 			if (Arguments.Count > 0 || Name == null || Change == null || Project == null || RestUrl == null || Status == null || Url == null)
 			{
 				Console.WriteLine("Syntax:");
-				Console.WriteLine("  PostBadgeStatus -Name=<Name> -Change=<CL> -Project=<DepotPath> -RestUrl=<Url> -Status=<Status> -Url=<Url>");
+				Console.WriteLine("  PostBadgeStatus -Name=<Name> -Change=<CL> -Project=<DepotPath> -RestUrl=<Url> -Status=<Status> -Url=<Url> -Metadata=<json metadata>");
+				Console.WriteLine("  Metadata is optional, but if provided, it has to contain array of links; format: {'Links': [{'Text': 'link text', 'Url': 'http://your.url'}]}");
 				return 1;
 			}
 
@@ -51,6 +53,10 @@ namespace WriteBadgeStatus
 				Console.WriteLine("Change must be Starting, Failure, Warning, Success, or Skipped!");
 				return 1;
 			}
+			if (!string.IsNullOrWhiteSpace(Metadata))
+			{
+				Build.Metadata = new JavaScriptSerializer().Deserialize<object>(Metadata);
+			}
 			int NumRetries = 0;
 			while (true)
 			{
@@ -107,6 +113,7 @@ namespace WriteBadgeStatus
 			public string Url;
 			public string Project;
 			public string ArchivePath;
+			public object Metadata;
 
 			public bool IsSuccess
 			{
diff --git a/Engine/Source/Programs/UnrealGameSync/UnrealGameSync/Controls/WorkspaceControl.cs b/Engine/Source/Programs/UnrealGameSync/UnrealGameSync/Controls/WorkspaceControl.cs
--- a/Engine/Source/Programs/UnrealGameSync/UnrealGameSync/Controls/WorkspaceControl.cs
+++ b/Engine/Source/Programs/UnrealGameSync/UnrealGameSync/Controls/WorkspaceControl.cs
@@ -85,26 +85,29 @@ namespace UnrealGameSync
 			public int _height;
 			public Color _backgroundColor;
 			public Color _hoverBackgroundColor;
-			public Action? _clickHandler;
+			public Action? _leftClickHandler;
+			public Action<Control, Point>? _rightClickHandler;
 			public string? _toolTip;
 
 			public BadgeInfo(string label, string group, Color badgeColor)
-				: this(label, group, null, badgeColor, badgeColor, null)
+				: this(label, group, null, badgeColor, badgeColor, null, null)
 			{
 			}
 
-			public BadgeInfo(string label, string? group, string? uniqueId, Color backgroundColor, Color hoverBackgroundColor, Action? clickHandler)
+			public BadgeInfo(string label, string? group, string? uniqueId, Color backgroundColor,
+				Color hoverBackgroundColor, Action? leftClickHandler, Action<Control, Point>? rightClickHandler)
 			{
 				_label = label;
 				_group = group;
 				_uniqueId = uniqueId;
 				_backgroundColor = backgroundColor;
 				_hoverBackgroundColor = hoverBackgroundColor;
-				_clickHandler = clickHandler;
+				_leftClickHandler = leftClickHandler;
+				_rightClickHandler = rightClickHandler;
 			}
 
 			public BadgeInfo(BadgeInfo other)
-				: this(other._label, other._group, other._uniqueId, other._backgroundColor, other._hoverBackgroundColor, other._clickHandler)
+				: this(other._label, other._group, other._uniqueId, other._backgroundColor, other._hoverBackgroundColor, other._leftClickHandler, other._rightClickHandler)
 			{
 				_offset = other._offset;
 				_width = other._width;
@@ -2576,7 +2579,7 @@ namespace UnrealGameSync
 										clickHandler = () => SafeProcessStart(expandedUrl, expandedArguments);
 									}
 
-									badges.Add(new BadgeInfo(ReplaceRegexMatches(name, matchResult), group, uniqueId, badgeColor, hoverBadgeColor, clickHandler));
+									badges.Add(new BadgeInfo(ReplaceRegexMatches(name, matchResult), group, uniqueId, badgeColor, hoverBadgeColor, clickHandler, null));
 								}
 							}
 						}
@@ -2781,12 +2784,7 @@ namespace UnrealGameSync
 					BadgeData? badgeData;
 					badgeNameToBuildData.TryGetValue(badgeNameAndGroup.Key, out badgeData);
 
-					BadgeInfo badgeInfo = CreateBadge(changeNumber, badgeNameAndGroup.Key, badgeNameAndGroup.Value, badgeData);
-					if (_maxBuildBadgeChars != -1 && badgeInfo._label.Length > _maxBuildBadgeChars)
-					{
-						badgeInfo._toolTip = badgeInfo._label;
-						badgeInfo._label = badgeInfo._label.Substring(0, _maxBuildBadgeChars);
-					}
+					BadgeInfo badgeInfo = CreateBadge(changeNumber, badgeNameAndGroup.Key, badgeNameAndGroup.Value, badgeData, _maxBuildBadgeChars);
 					badges.Add(badgeInfo);
 				}
 			}
@@ -2796,7 +2794,7 @@ namespace UnrealGameSync
 			return badges;
 		}
 
-		private static BadgeInfo CreateBadge(int changeNumber, string badgeName, string badgeGroup, BadgeData? badgeData)
+		private BadgeInfo CreateBadge(int changeNumber, string badgeName, string badgeGroup, BadgeData? badgeData, int maxBadgeChars = -1)
 		{
 			string badgeLabel = badgeName;
 			Color badgeColor = Color.FromArgb(0, Color.White);
@@ -2809,18 +2807,62 @@ namespace UnrealGameSync
 
 			Color hoverBadgeColor = Color.FromArgb(badgeColor.A, Math.Min(badgeColor.R + 32, 255), Math.Min(badgeColor.G + 32, 255), Math.Min(badgeColor.B + 32, 255));
 
-			Action? clickHandler;
+			Action? leftClickHandler;
 			if (badgeData == null || String.IsNullOrEmpty(badgeData.Url))
 			{
-				clickHandler = null;
+				leftClickHandler = null;
+			}
+			else
+			{
+				leftClickHandler = () => SafeProcessStart(badgeData.Url);
+			}
+
+			Action<Control, Point>? rightClickHandler;
+			if (badgeData?.Metadata?.Links == null || !badgeData.Metadata.Links.Any())
+			{
+				rightClickHandler = null;
 			}
 			else
 			{
-				clickHandler = () => SafeProcessStart(badgeData.Url);
+				rightClickHandler = (control, point) =>
+				{
+					BadgeContextMenu.Items.Clear();
+
+					foreach (Link additionalUrl in badgeData.Metadata.Links)
+					{
+						ToolStripMenuItem item = new ToolStripMenuItem();
+						item.Text = additionalUrl.Text;
+    					item.Click += (sender, e) => SafeProcessStart(additionalUrl.Url);
+    					BadgeContextMenu.Items.Add(item);
+					}
+
+					BadgeContextMenu.Show(control, point);
+				};
 			}
 
 			string uniqueId = String.Format("{0}:{1}", changeNumber, badgeName);
-			return new BadgeInfo(badgeLabel, badgeGroup, uniqueId, badgeColor, hoverBadgeColor, clickHandler);
+			BadgeInfo badgeInfo = new BadgeInfo(badgeLabel, badgeGroup, uniqueId, badgeColor, hoverBadgeColor, leftClickHandler, rightClickHandler);
+
+			if (maxBadgeChars != -1 && badgeInfo._label.Length > maxBadgeChars)
+			{
+				badgeInfo._toolTip = badgeInfo._label;
+				badgeInfo._label = badgeInfo._label.Substring(0, maxBadgeChars);
+			}
+
+			if (badgeInfo._rightClickHandler != null)
+			{
+				string rightClickIndictor = "\U0001F4DC";
+				if (badgeInfo._toolTip == null)
+				{
+					badgeInfo._toolTip = rightClickIndictor;
+				}
+				else
+				{
+					badgeInfo._toolTip += (" " + rightClickIndictor);
+				}
+			}
+
+			return badgeInfo;
 		}
 
 		private Dictionary<string, List<BadgeInfo>> CreateCustomBadges(int changeNumber, EventSummary? summary)
@@ -4134,9 +4176,9 @@ namespace UnrealGameSync
 								foreach (BadgeInfo badge in layoutInfo._descriptionBadges)
 								{
 									Rectangle badgeBounds = badge.GetBounds(buildListLocation);
-									if (badgeBounds.Contains(args.Location) && badge._clickHandler != null)
+									if (badgeBounds.Contains(args.Location) && badge._leftClickHandler != null)
 									{
-										badge._clickHandler();
+										badge._leftClickHandler();
 										break;
 									}
 								}
@@ -4152,9 +4194,9 @@ namespace UnrealGameSync
 								buildListLocation.X = Math.Max(buildListLocation.X, hitTest.SubItem.Bounds.Left);
 
 								BadgeInfo? badgeInfo = HitTestBadge(args.Location, layoutInfo._buildBadges, buildListLocation);
-								if (badgeInfo != null && badgeInfo._clickHandler != null)
+								if (badgeInfo != null && badgeInfo._leftClickHandler != null)
 								{
-									badgeInfo._clickHandler();
+									badgeInfo._leftClickHandler();
 								}
 							}
 						}
@@ -4171,9 +4213,9 @@ namespace UnrealGameSync
 									Point listLocation = GetBadgeListLocation(badges, hitTest.SubItem.Bounds, HorizontalAlign.Center, VerticalAlignment.Middle);
 
 									BadgeInfo? badgeInfo = HitTestBadge(args.Location, badges, listLocation);
-									if (badgeInfo != null && badgeInfo._clickHandler != null)
+									if (badgeInfo != null && badgeInfo._leftClickHandler != null)
 									{
-										badgeInfo._clickHandler();
+										badgeInfo._leftClickHandler();
 									}
 								}
 							}
@@ -4199,13 +4241,38 @@ namespace UnrealGameSync
 					}
 					else
 					{
-						_contextMenuChange = (ChangesRecord)hitTest.Item.Tag;
-						// Show time related context menu items if the click was made on the time column :
-						bool isTimeColumn = (hitTest.Item.SubItems.IndexOf(hitTest.SubItem) == TimeColumn.Index);
-						ShowChangelistContextMenu(BuildList, showTimeRelatedItems: isTimeColumn, args.Location);
-					}
-				}
-			}
+						bool badgeWasClicked = false;
+						if (CISColumn.Index < hitTest.Item.SubItems.Count && hitTest.Item.SubItems[CISColumn.Index] == hitTest.SubItem)
+						{
+							ChangesRecord change = (ChangesRecord)hitTest.Item.Tag;
+							ChangeLayoutInfo layoutInfo = GetChangeLayoutInfo(change);
+							if (layoutInfo._buildBadges.Count > 0)
+							{
+								Point buildListLocation = GetBadgeListLocation(layoutInfo._buildBadges, hitTest.SubItem.Bounds, HorizontalAlign.Center, VerticalAlignment.Middle);
+								buildListLocation.X = Math.Max(buildListLocation.X, hitTest.SubItem.Bounds.Left);
+
+								BadgeInfo? badgeInfo = HitTestBadge(args.Location, layoutInfo._buildBadges, buildListLocation);
+								if (badgeInfo != null)
+								{
+									badgeWasClicked = true;
+									if (badgeInfo._rightClickHandler != null)
+									{
+										badgeInfo._rightClickHandler(BuildList, args.Location);
+									}
+								}
+							}
+						}
+
+						if (!badgeWasClicked)
+						{
+						    _contextMenuChange = (ChangesRecord)hitTest.Item.Tag;
+						    // Show time related context menu items if the click was made on the time column :
+						    bool isTimeColumn = (hitTest.Item.SubItems.IndexOf(hitTest.SubItem) == TimeColumn.Index);
+						    ShowChangelistContextMenu(BuildList, showTimeRelatedItems: isTimeColumn, args.Location);
+					    }
+				    }
+			    }
+		    }
 		}
 
 		private void ShowChangelistContextMenu(Control parentControl, bool showTimeRelatedItems, Point location, ToolStripDropDownDirection direction = ToolStripDropDownDirection.BelowRight)
diff --git a/Engine/Source/Programs/UnrealGameSync/UnrealGameSync/EventMonitor.cs b/Engine/Source/Programs/UnrealGameSync/UnrealGameSync/EventMonitor.cs
--- a/Engine/Source/Programs/UnrealGameSync/UnrealGameSync/EventMonitor.cs
+++ b/Engine/Source/Programs/UnrealGameSync/UnrealGameSync/EventMonitor.cs
@@ -71,6 +71,17 @@ namespace UnrealGameSync
 		Skipped,
 	}
 
+	class Link
+	{
+		public string Text { get; set; } = String.Empty;
+		public string Url { get; set; } = String.Empty;
+	}
+
+	class BadgeMetadata
+	{
+		public List<Link>? Links { get; set; } = new List<Link>();
+	}
+
 	class BadgeData
 	{
 		public long Id { get; set; }
@@ -79,6 +90,7 @@ namespace UnrealGameSync
 		public BadgeResult Result { get; set; }
 		public string Url { get; set; } = String.Empty;
 		public string Project { get; set; } = String.Empty;
+		public BadgeMetadata? Metadata { get; set; }
 
 		public bool IsSuccess => Result == BadgeResult.Success || Result == BadgeResult.Warning;
 
