UGS: Fix incorrect handling when cleaning workspaces in the root directory.

diff --git a/Engine/Source/Programs/UnrealGameSync/UnrealGameSync/FindFoldersToCleanTask.cs b/Engine/Source/Programs/UnrealGameSync/UnrealGameSync/FindFoldersToCleanTask.cs
--- a/Engine/Source/Programs/UnrealGameSync/UnrealGameSync/FindFoldersToCleanTask.cs
+++ b/Engine/Source/Programs/UnrealGameSync/UnrealGameSync/FindFoldersToCleanTask.cs
@@ -266,7 +266,7 @@ namespace UnrealGameSync
 			string localRootPrefix = _rootFolderToClean._directory.FullName.Replace(Path.AltDirectorySeparatorChar, Path.DirectorySeparatorChar);
 			if (!localRootPrefix.EndsWith(Path.DirectorySeparatorChar))
 			{
-				localRootPrefix.Append(Path.DirectorySeparatorChar);
+				localRootPrefix += Path.DirectorySeparatorChar;
 			}
 
 			// Query the have table and build a separate tree from it
