UGS: Fix deprecation warning.

diff --git a/Engine/Source/Programs/UnrealGameSync/UnrealGameSyncShared/WorkspaceState.cs b/Engine/Source/Programs/UnrealGameSync/UnrealGameSyncShared/WorkspaceState.cs
--- a/Engine/Source/Programs/UnrealGameSync/UnrealGameSyncShared/WorkspaceState.cs
+++ b/Engine/Source/Programs/UnrealGameSync/UnrealGameSyncShared/WorkspaceState.cs
@@ -163,7 +163,7 @@ namespace UnrealGameSync
 		public int LastBuiltChangeNumber => _inner.LastBuiltChangeNumber;
 
 		// Expanded archives in the workspace
-		public ReadOnlyHashSet<string> ExpandedArchiveTypes { get; }
+		public IReadOnlySet<string> ExpandedArchiveTypes { get; }
 
 		// The changes that we're regressing at the moment
 		public IReadOnlyList<BisectEntry> BisectChanges => _inner.BisectChanges;
@@ -171,7 +171,7 @@ namespace UnrealGameSync
 		internal ReadOnlyWorkspaceState(WorkspaceState inner)
 		{
 			_inner = inner;
-			ExpandedArchiveTypes = new ReadOnlyHashSet<string>(_inner.ExpandedArchiveTypes);
+			ExpandedArchiveTypes = new HashSet<string>(_inner.ExpandedArchiveTypes);
 		}
 
 		public ReadOnlyWorkspaceState ResetForProject(ProjectInfo projectInfo)
