UGS: Fix DPI scaling on tab control buttons.

diff --git a/Engine/Source/Programs/UnrealGameSync/UnrealGameSync/Controls/TabControl.cs b/Engine/Source/Programs/UnrealGameSync/UnrealGameSync/Controls/TabControl.cs
--- a/Engine/Source/Programs/UnrealGameSync/UnrealGameSync/Controls/TabControl.cs
+++ b/Engine/Source/Programs/UnrealGameSync/UnrealGameSync/Controls/TabControl.cs
@@ -270,7 +270,7 @@ namespace UnrealGameSync
 				{
 					TabData tab = _tabs[idx];
 					tab._textSize = TextRenderer.MeasureText(graphics, tab._name, Font, new Size(Int32.MaxValue, Int32.MaxValue), TextFormatFlags.NoPadding);
-					tab._width = TabPadding + tab._textSize.Width + (int)(TabPadding * dpiScaleX);
+					tab._width = (int)(TabPadding * dpiScaleX) + tab._textSize.Width + (int)(TabPadding * dpiScaleX);
 					if(idx == _selectedTabIdx)
 					{
 						tab._closeButtonWidth = (int)(TabCloseButtonWidth * dpiScaleX);
