PR #10334: UGS: Updating installer batch location after .net 6 migration (Contributed by [REDACTED_PERSON])

diff --git a/Engine/Source/Programs/UnrealGameSync/Installer/Build.bat b/Engine/Source/Programs/UnrealGameSync/Installer/Build.bat
--- a/Engine/Source/Programs/UnrealGameSync/Installer/Build.bat
+++ b/Engine/Source/Programs/UnrealGameSync/Installer/Build.bat
@@ -3,7 +3,7 @@ setlocal
 
 pushd %~dp0
 set WIXDIR=..\..\..\ThirdParty\WiX\3.8
-set SRCDIR=..\UnrealGameSyncLauncher\bin\Release\netcoreapp3.1
+set SRCDIR=..\UnrealGameSyncLauncher\bin\Release\net6.0-windows
 
 %WIXDIR%\heat.exe dir "%SRCDIR%" -cg UGSLauncher_Project -dr INSTALLFOLDER -scom -sreg -srd -var var.BasePath -gg -sfrag -out obj/Release/UGSLauncher.wxs -nologo 
 if errorlevel 1 goto :eof
