UGS: Return a non-zero exit code if syncing/building fails with an error.

diff --git a/Engine/Source/Programs/UnrealGameSync/UnrealGameSyncCmd/Program.cs b/Engine/Source/Programs/UnrealGameSync/UnrealGameSyncCmd/Program.cs
--- a/Engine/Source/Programs/UnrealGameSync/UnrealGameSyncCmd/Program.cs
+++ b/Engine/Source/Programs/UnrealGameSync/UnrealGameSyncCmd/Program.cs
@@ -896,12 +896,13 @@ namespace UnrealGameSyncCmd
 					}
 					logger.LogWarning("Use -Clobber to overwrite");
 				}
-				else if (result != WorkspaceUpdateResult.Success)
-				{
-					logger.LogError("{Message} (Result: {Result})", message, result);
-				}
 
 				state.Modify(x => x.SetLastSyncState(result, updateContext, message));
+
+				if (result != WorkspaceUpdateResult.Success)
+				{
+					throw new UserErrorException("{Message} (Result: {Result})", message, result);
+				}
 			}
 
 			static WorkspaceLock? CreateWorkspaceLock(DirectoryReference rootDir)
@@ -1342,7 +1343,12 @@ namespace UnrealGameSyncCmd
 				WorkspaceUpdateContext updateContext = new WorkspaceUpdateContext(state.Current.CurrentChangeNumber, WorkspaceUpdateOptions.Build, BuildConfig.Development, null, new List<ConfigObject>(), steps);
 
 				WorkspaceUpdate update = new WorkspaceUpdate(updateContext);
-				await update.ExecuteAsync(perforceClient.Settings, projectInfo, state, context.Logger, CancellationToken.None);
+				(WorkspaceUpdateResult result, string message) = await update.ExecuteAsync(perforceClient.Settings, projectInfo, state, context.Logger, CancellationToken.None);
+
+				if (result != WorkspaceUpdateResult.Success)
+				{
+					throw new UserErrorException("{Message}", message); 
+				}
 			}
 		}
 
