UGS: Move setting for auto-resolving conflicts into global settings file.

diff --git a/Engine/Source/Programs/UnrealGameSync/UnrealGameSyncShared/GlobalSettings.cs b/Engine/Source/Programs/UnrealGameSync/UnrealGameSyncShared/GlobalSettings.cs
--- a/Engine/Source/Programs/UnrealGameSync/UnrealGameSyncShared/GlobalSettings.cs
+++ b/Engine/Source/Programs/UnrealGameSync/UnrealGameSyncShared/GlobalSettings.cs
@@ -55,6 +55,7 @@ namespace UnrealGameSync
 		public DateTime LastVersionCheck { get; set; }
 		public PerforceSyncOptions Perforce { get; set; } = new PerforceSyncOptions();
 		public FilterSettings Filter { get; set; } = new FilterSettings();
+		public bool AutoResolveConflicts { get; set; } = true;
 	}
 
 	public class GlobalSettingsFile
diff --git a/Engine/Source/Programs/UnrealGameSync/UnrealGameSyncShared/UserSettings.cs b/Engine/Source/Programs/UnrealGameSync/UnrealGameSyncShared/UserSettings.cs
--- a/Engine/Source/Programs/UnrealGameSync/UnrealGameSyncShared/UserSettings.cs
+++ b/Engine/Source/Programs/UnrealGameSync/UnrealGameSyncShared/UserSettings.cs
@@ -489,7 +489,6 @@ namespace UnrealGameSync
 		public bool RunAfterSync { get; set; }
 		public bool OpenSolutionAfterSync { get; set; }
 		public bool ShowLogWindow { get; set; }
-		public bool AutoResolveConflicts { get; set; }
 		public bool ShowUnreviewedChanges { get; set; }
 		public bool ShowAutomatedChanges { get; set; }
 		public RobomergeShowChangesOption ShowRobomerge { get; set; }
@@ -629,7 +628,13 @@ namespace UnrealGameSync
 			bool syncPrecompiledEditor = (_configFile.GetValue("General.SyncPrecompiledEditor", "0") != "0");
 			OpenSolutionAfterSync = (_configFile.GetValue("General.OpenSolutionAfterSync", "0") != "0");
 			ShowLogWindow = (_configFile.GetValue("General.ShowLogWindow", false));
-			AutoResolveConflicts = (_configFile.GetValue("General.AutoResolveConflicts", "1") != "0");
+
+			string? autoResolveConflicts = _configFile.GetValue("General.AutoResolveConflicts", "");
+			if (!String.IsNullOrEmpty(autoResolveConflicts))
+			{
+				Global.AutoResolveConflicts = (autoResolveConflicts != "0");
+			}
+
 			ShowUnreviewedChanges = _configFile.GetValue("General.ShowUnreviewed", true);
 			ShowAutomatedChanges = _configFile.GetValue("General.ShowAutomated", false);
 			NextLauncherVersionCheck = _configFile.GetValue("General.NextLauncherVersionCheck", 0);
@@ -1001,7 +1006,6 @@ namespace UnrealGameSync
 			generalSection.SetValue("RunAfterSync", RunAfterSync);
 			generalSection.SetValue("OpenSolutionAfterSync", OpenSolutionAfterSync);
 			generalSection.SetValue("ShowLogWindow", ShowLogWindow);
-			generalSection.SetValue("AutoResolveConflicts", AutoResolveConflicts);
 			generalSection.SetValue("ShowUnreviewed", ShowUnreviewedChanges);
 			generalSection.SetValue("ShowAutomated", ShowAutomatedChanges);
 			generalSection.SetValue("RobomergeFilter", ShowRobomerge.ToString());
diff --git a/Engine/Source/Programs/UnrealGameSync/UnrealGameSyncShared/WorkspaceUpdate.cs b/Engine/Source/Programs/UnrealGameSync/UnrealGameSyncShared/WorkspaceUpdate.cs
--- a/Engine/Source/Programs/UnrealGameSync/UnrealGameSyncShared/WorkspaceUpdate.cs
+++ b/Engine/Source/Programs/UnrealGameSync/UnrealGameSyncShared/WorkspaceUpdate.cs
@@ -118,6 +118,10 @@ namespace UnrealGameSync
 		public static WorkspaceUpdateOptions GetOptionsFromConfig(GlobalSettings globalSettings, UserWorkspaceSettings workspaceSettings)
 		{
 			WorkspaceUpdateOptions options = 0;
+			if (globalSettings.AutoResolveConflicts)
+			{
+				options |= WorkspaceUpdateOptions.AutoResolveChanges;
+			}
 			if (workspaceSettings.Filter.AllProjects ?? globalSettings.Filter.AllProjects ?? false)
 			{
 				options |= WorkspaceUpdateOptions.SyncAllProjects | WorkspaceUpdateOptions.IncludeAllProjectsInSolution;
