UGS: Prevent exception trying to relaunch process if the installer is not available.

diff --git a/Engine/Source/Programs/UnrealGameSync/UnrealGameSync/Program.cs b/Engine/Source/Programs/UnrealGameSync/UnrealGameSync/Program.cs
--- a/Engine/Source/Programs/UnrealGameSync/UnrealGameSync/Program.cs
+++ b/Engine/Source/Programs/UnrealGameSync/UnrealGameSync/Program.cs
@@ -157,6 +157,10 @@ namespace UnrealGameSync
 				{
 					Directory.SetCurrentDirectory(Path.GetDirectoryName(updateSpawn)!);
 				}
+				else
+				{
+					updateSpawn = null;
+				}
 			}
 
 			string syncVersionFile = Path.Combine(Path.GetDirectoryName(Assembly.GetExecutingAssembly().Location!)!, "SyncVersion.txt");
