UGS: Prevent exceptions reading config files being fatal.

diff --git a/Engine/Source/Programs/UnrealGameSync/UnrealGameSyncShared/Utility.cs b/Engine/Source/Programs/UnrealGameSync/UnrealGameSyncShared/Utility.cs
--- a/Engine/Source/Programs/UnrealGameSync/UnrealGameSyncShared/Utility.cs
+++ b/Engine/Source/Programs/UnrealGameSync/UnrealGameSyncShared/Utility.cs
@@ -435,10 +435,10 @@ namespace UnrealGameSync
 		{
 			if(digest == null)
 			{
-				PerforceResponse<PrintRecord<string[]>> response = await perforce.TryPrintLinesAsync(depotPath, cancellationToken);
-				if (response.Succeeded)
+				PerforceResponse<PrintRecord<string[]>> printLinesResponse = await perforce.TryPrintLinesAsync(depotPath, cancellationToken);
+				if (printLinesResponse.Succeeded)
 				{
-					return response.Data.Contents;
+					return printLinesResponse.Data.Contents;
 				}
 				else
 				{
@@ -450,48 +450,54 @@ namespace UnrealGameSync
 			if(FileReference.Exists(cacheFile))
 			{
 				logger.LogDebug("Reading cached copy of {DepotFile} from {LocalFile}", depotPath, cacheFile);
-				string[] lines = FileReference.ReadAllLines(cacheFile);
 				try
 				{
-					FileReference.SetLastWriteTimeUtc(cacheFile, DateTime.UtcNow);
+					string[] lines = await FileReference.ReadAllLinesAsync(cacheFile, cancellationToken);
+					try
+					{
+						FileReference.SetLastWriteTimeUtc(cacheFile, DateTime.UtcNow);
+					}
+					catch (Exception ex)
+					{
+						logger.LogWarning(ex, "Exception touching cache file {LocalFile}", cacheFile);
+					}
+					return lines;
 				}
-				catch(Exception ex)
+				catch (Exception ex)
 				{
-					logger.LogWarning(ex, "Exception touching cache file {LocalFile}", cacheFile);
+					logger.LogWarning(ex, "Error while reading cache file {LocalFile}: {Message}", cacheFile, ex.Message);
 				}
-				return lines;
+			}
+
+
+			DirectoryReference.CreateDirectory(cacheFolder);
+
+			FileReference tempFile = new FileReference(String.Format("{0}.{1}.temp", cacheFile.FullName, Guid.NewGuid()));
+			PerforceResponseList<PrintRecord> printResponse = await perforce.TryPrintAsync(tempFile.FullName, depotPath, cancellationToken);
+			if (!printResponse.Succeeded)
+			{
+				return null;
 			}
 			else
 			{
-				DirectoryReference.CreateDirectory(cacheFolder);
-
-				FileReference tempFile = new FileReference(String.Format("{0}.{1}.temp", cacheFile.FullName, Guid.NewGuid()));
-				PerforceResponseList<PrintRecord> response = await perforce.TryPrintAsync(tempFile.FullName, depotPath, cancellationToken);
-				if (!response.Succeeded)
+				string[] lines = await FileReference.ReadAllLinesAsync(tempFile, cancellationToken);
+				try
 				{
-					return null;
+					FileReference.SetAttributes(tempFile, FileAttributes.Normal);
+					FileReference.SetLastWriteTimeUtc(tempFile, DateTime.UtcNow);
+					FileReference.Move(tempFile, cacheFile);
 				}
-				else
+				catch
 				{
-					string[] lines = await FileReference.ReadAllLinesAsync(tempFile, cancellationToken);
 					try
 					{
-						FileReference.SetAttributes(tempFile, FileAttributes.Normal);
-						FileReference.SetLastWriteTimeUtc(tempFile, DateTime.UtcNow);
-						FileReference.Move(tempFile, cacheFile);
+						FileReference.Delete(tempFile);
 					}
 					catch
 					{
-						try
-						{
-							FileReference.Delete(tempFile);
-						}
-						catch
-						{
-						}
 					}
-					return lines;
 				}
+				return lines;
 			}
 		}
 
