UGS: Do not attempt to open projects which have an invalid path.

diff --git a/Engine/Source/Programs/UnrealGameSync/UnrealGameSync/ProgramApplicationContext.cs b/Engine/Source/Programs/UnrealGameSync/UnrealGameSync/ProgramApplicationContext.cs
--- a/Engine/Source/Programs/UnrealGameSync/UnrealGameSync/ProgramApplicationContext.cs
+++ b/Engine/Source/Programs/UnrealGameSync/UnrealGameSync/ProgramApplicationContext.cs
@@ -194,10 +194,13 @@ namespace UnrealGameSync
 			List<(UserSelectedProjectSettings, ModalTask<OpenProjectInfo>)> startupTasks = new List<(UserSelectedProjectSettings, ModalTask<OpenProjectInfo>)>();
 			foreach (UserSelectedProjectSettings projectSettings in _settings.OpenProjects)
 			{
-				ILogger<OpenProjectInfo> logger = serviceProvider.GetRequiredService<ILogger<OpenProjectInfo>>();
-				OidcTokenManager oidcTokenManager = serviceProvider.GetRequiredService<OidcTokenManager>();
-				Task<OpenProjectInfo> startupTask = Task.Run(() => OpenProjectInfo.CreateAsync(defaultPerforceSettings, projectSettings, _settings, oidcTokenManager, logger, _startupCancellationSource.Token), _startupCancellationSource.Token);
-				startupTasks.Add((projectSettings, new ModalTask<OpenProjectInfo>(startupTask)));
+				if (IsValidProject(projectSettings, _logger))
+				{
+					ILogger<OpenProjectInfo> logger = serviceProvider.GetRequiredService<ILogger<OpenProjectInfo>>();
+					OidcTokenManager oidcTokenManager = serviceProvider.GetRequiredService<OidcTokenManager>();
+					Task<OpenProjectInfo> startupTask = Task.Run(() => OpenProjectInfo.CreateAsync(defaultPerforceSettings, projectSettings, _settings, oidcTokenManager, logger, _startupCancellationSource.Token), _startupCancellationSource.Token);
+					startupTasks.Add((projectSettings, new ModalTask<OpenProjectInfo>(startupTask)));
+				}
 			}
 			_startupTask = Task.Run(() => WaitForStartupTasks(startupTasks));
 
@@ -219,6 +222,20 @@ namespace UnrealGameSync
 			_startupWindow.FormClosed += (s, e) => OnStartupComplete(startupTasks);
 		}
 
+		static bool IsValidProject(UserSelectedProjectSettings projectSettings, ILogger logger)
+		{
+			try
+			{
+				bool valid = String.IsNullOrEmpty(projectSettings.LocalPath) || File.Exists(projectSettings.LocalPath);
+				return valid;
+			}
+			catch
+			{
+				logger.LogWarning("Invalid project: {Project}", projectSettings);
+				return false;
+			}
+		}
+
 		static async Task WaitForStartupTasks(List<(UserSelectedProjectSettings, ModalTask<OpenProjectInfo>)> startupTasks)
 		{
 			foreach ((_, ModalTask<OpenProjectInfo> modalTask) in startupTasks)
