UGS: Speculative fix for a null dereference.

diff --git a/Engine/Source/Programs/UnrealGameSync/UnrealGameSync/Controls/WorkspaceControl.cs b/Engine/Source/Programs/UnrealGameSync/UnrealGameSync/Controls/WorkspaceControl.cs
--- a/Engine/Source/Programs/UnrealGameSync/UnrealGameSync/Controls/WorkspaceControl.cs
+++ b/Engine/Source/Programs/UnrealGameSync/UnrealGameSync/Controls/WorkspaceControl.cs
@@ -1015,7 +1015,7 @@ namespace UnrealGameSync
 
 		void CancelWorkspaceUpdate()
 		{
-			if (_workspace.IsBusy() && MessageBox.Show("Are you sure you want to cancel the current operation?", "Cancel operation", MessageBoxButtons.YesNo) == DialogResult.Yes)
+			if (_workspace != null && _workspace.IsBusy() && MessageBox.Show("Are you sure you want to cancel the current operation?", "Cancel operation", MessageBoxButtons.YesNo) == DialogResult.Yes)
 			{
 				_workspace.CancelUpdate();
 
