UGS: Present a user-facing error if copying data to the clipboard fails.

diff --git a/Engine/Source/Programs/UnrealGameSync/UnrealGameSync/Controls/WorkspaceControl.cs b/Engine/Source/Programs/UnrealGameSync/UnrealGameSync/Controls/WorkspaceControl.cs
--- a/Engine/Source/Programs/UnrealGameSync/UnrealGameSync/Controls/WorkspaceControl.cs
+++ b/Engine/Source/Programs/UnrealGameSync/UnrealGameSync/Controls/WorkspaceControl.cs
@@ -4825,7 +4825,14 @@ namespace UnrealGameSync
 		{
 			if (_contextMenuChange != null)
 			{
-				System.Windows.Forms.Clipboard.SetText(_contextMenuChange.Number.ToString());
+				try
+				{
+					Clipboard.SetText(_contextMenuChange.Number.ToString());
+				}
+				catch
+				{
+					MessageBox.Show("Unable to copy data to clipboard; another process may have the clipboard locked.");
+				}
 			}
 		}
 
