UGS: Build fix.

diff --git a/Engine/Source/Programs/UnrealGameSync/UnrealGameSync/Controls/WorkspaceControl.cs b/Engine/Source/Programs/UnrealGameSync/UnrealGameSync/Controls/WorkspaceControl.cs
--- a/Engine/Source/Programs/UnrealGameSync/UnrealGameSync/Controls/WorkspaceControl.cs
+++ b/Engine/Source/Programs/UnrealGameSync/UnrealGameSync/Controls/WorkspaceControl.cs
@@ -965,10 +965,6 @@ namespace UnrealGameSync
 
 		void StartWorkspaceUpdate(WorkspaceUpdateContext context, WorkspaceUpdateCallback? callback)
 		{
-			if (_settings.AutoResolveConflicts)
-			{
-				context.Options |= WorkspaceUpdateOptions.AutoResolveChanges;
-			}
 			context.Options |= WorkspaceUpdateContext.GetOptionsFromConfig(_settings.Global, _workspaceSettings);
 
 			_updateCallback = callback;
@@ -4534,7 +4530,7 @@ namespace UnrealGameSync
 
 		private void OptionsButton_Click(object sender, EventArgs e)
 		{
-			OptionsContextMenu_AutoResolveConflicts.Checked = _settings.AutoResolveConflicts;
+			OptionsContextMenu_AutoResolveConflicts.Checked = _settings.Global.AutoResolveConflicts;
 
 			OptionsContextMenu_SyncPrecompiledBinaries.DropDownItems.Clear();
 
@@ -5044,7 +5040,7 @@ namespace UnrealGameSync
 		private void OptionsContextMenu_AutoResolveConflicts_Click(object sender, EventArgs e)
 		{
 			OptionsContextMenu_AutoResolveConflicts.Checked ^= true;
-			_settings.AutoResolveConflicts = OptionsContextMenu_AutoResolveConflicts.Checked;
+			_settings.Global.AutoResolveConflicts = OptionsContextMenu_AutoResolveConflicts.Checked;
 			_settings.Save(_logger);
 		}
 
@@ -5181,7 +5177,7 @@ namespace UnrealGameSync
 				else
 				{
 					WorkspaceUpdateOptions options = WorkspaceUpdateOptions.Sync | WorkspaceUpdateOptions.SyncArchives | WorkspaceUpdateOptions.GenerateProjectFiles | WorkspaceUpdateOptions.Build | WorkspaceUpdateOptions.ScheduledBuild;
-					if (_settings.AutoResolveConflicts)
+					if (_settings.Global.AutoResolveConflicts)
 					{
 						options |= WorkspaceUpdateOptions.AutoResolveChanges;
 					}
