UGS: Suppress upgrade prompts on Windows.

diff --git a/Engine/Source/Programs/UnrealGameSync/UnrealGameSyncCmd/Program.cs b/Engine/Source/Programs/UnrealGameSync/UnrealGameSyncCmd/Program.cs
--- a/Engine/Source/Programs/UnrealGameSync/UnrealGameSyncCmd/Program.cs
+++ b/Engine/Source/Programs/UnrealGameSync/UnrealGameSyncCmd/Program.cs
@@ -255,7 +255,8 @@ namespace UnrealGameSyncCmd
 					return 1;
 				}
 
-				if (command.Type != typeof(InstallCommand) && command.Type != typeof(UpgradeCommand))
+				// On Windows this is distributed with the GUI client, so we don't need to check for upgrades
+				if (command.Type != typeof(InstallCommand) && command.Type != typeof(UpgradeCommand) && !RuntimeInformation.IsOSPlatform(OSPlatform.Windows))
 				{
 					string version = GetVersion();
 
@@ -391,7 +392,7 @@ namespace UnrealGameSyncCmd
 			}
 			else
 			{
-				return "ugs-mac";
+				return null;
 			}
 		}
 
@@ -407,7 +408,7 @@ namespace UnrealGameSyncCmd
 			string? toolName = GetUpgradeToolName();
 			if (toolName == null)
 			{
-				logger?.LogError("Unknown platform; cannot query for upgrade.");
+				logger?.LogError("Command-line upgrades are not supported on this platform.");
 				return null;
 			}
 
