UGS: Add an alias to .bashrc on Linux.

diff --git a/Engine/Source/Programs/UnrealGameSync/UnrealGameSyncCmd/Program.cs b/Engine/Source/Programs/UnrealGameSync/UnrealGameSyncCmd/Program.cs
--- a/Engine/Source/Programs/UnrealGameSync/UnrealGameSyncCmd/Program.cs
+++ b/Engine/Source/Programs/UnrealGameSync/UnrealGameSyncCmd/Program.cs
@@ -1499,26 +1499,40 @@ namespace UnrealGameSyncCmd
 			{
 				ILogger logger = context.Logger;
 
-				DirectoryReference currentDir = new FileReference(Assembly.GetExecutingAssembly().Location).Directory;
 				if (RuntimeInformation.IsOSPlatform(OSPlatform.OSX))
 				{
 					DirectoryReference? userDir = DirectoryReference.GetSpecialFolder(Environment.SpecialFolder.UserProfile);
 					if (userDir != null)
 					{
 						FileReference configFile = FileReference.Combine(userDir, ".zshrc");
+						await AddAliasAsync(configFile, logger);
+					}
+				}
+				else if (RuntimeInformation.IsOSPlatform(OSPlatform.Linux))
+				{
+					DirectoryReference? userDir = DirectoryReference.GetSpecialFolder(Environment.SpecialFolder.UserProfile);
+					if (userDir != null)
+					{
+						FileReference configFile = FileReference.Combine(userDir, ".bashrc");
+						await AddAliasAsync(configFile, logger);
+					}
+				}
+			}
 
-						List<string> lines = new List<string>();
-						if (FileReference.Exists(configFile))
-						{
-							lines.AddRange(await FileReference.ReadAllLinesAsync(configFile));
-							lines.RemoveAll(x => Regex.IsMatch(x, @"^\s*alias\s+ugs\s*="));
-						}
-						lines.Add($"alias ugs={FileReference.Combine(currentDir, "ugs")}");
+			static async Task AddAliasAsync(FileReference configFile, ILogger logger)
+			{
+				DirectoryReference currentDir = new FileReference(Assembly.GetExecutingAssembly().Location).Directory;
 
-						await FileReference.WriteAllLinesAsync(configFile, lines);
-						logger.LogInformation("Added 'ugs' alias to {ConfigFile}", configFile);
-					}
+				List<string> lines = new List<string>();
+				if (FileReference.Exists(configFile))
+				{
+					lines.AddRange(await FileReference.ReadAllLinesAsync(configFile));
+					lines.RemoveAll(x => Regex.IsMatch(x, @"^\s*alias\s+ugs\s*="));
 				}
+				lines.Add($"alias ugs={FileReference.Combine(currentDir, "ugs")}");
+
+				await FileReference.WriteAllLinesAsync(configFile, lines);
+				logger.LogInformation("Added 'ugs' alias to {ConfigFile}", configFile);
 			}
 		}
 
