UGS: Add support for publishing the installer to Horde.

diff --git a/Engine/Source/Programs/UnrealGameSync/BuildUGS.xml b/Engine/Source/Programs/UnrealGameSync/BuildUGS.xml
--- a/Engine/Source/Programs/UnrealGameSync/BuildUGS.xml
+++ b/Engine/Source/Programs/UnrealGameSync/BuildUGS.xml
@@ -44,6 +44,10 @@
 			<Spawn Exe="$(WixDir)/candle.exe" Arguments="-dBasePath=&quot;$(UgsLauncherPublishDir)&quot; -out obj\Release\ -dConfiguration=Release -dPlatform=x64 -arch x86 -ext ..\..\..\ThirdParty\WiX\3.8\WixUtilExtension.dll Product.wxs obj/Release/UGSLauncher.wxs" WorkingDir="$(InstallerDir)"/>
 			<Spawn Exe="$(WixDir)/light.exe" Arguments="-out &quot;bin/Release/UnrealGameSync.msi&quot; -pdbout &quot;bin/Release/UnrealGameSync.wixpdb&quot; -cultures:null -ext ..\..\..\ThirdParty\WiX\3.8\WixUtilExtension.dll -sice:ICE69 &quot;obj\Release\Product.wixobj&quot; &quot;obj\Release\UGSLauncher.wixobj&quot;" WorkingDir="$(InstallerDir)"/>
 		</Node>
+		<Node Name="Publish UnrealGameSync Installer" Requires="Build UnrealGameSync Installer">
+			<Error Message="Missing -HordeSettings=... argument" If="'$(HordeSettings)' == ''"/>
+			<DeployTool Id="ugs" Version="$(InformationalVersion)" Settings="$(HordeSettings)" File="$(InstallerDir)/bin/Release/UnrealGameSync.msi"/>
+		</Node>
 		
 		<Property Name="InstallerDependency" Value=""/>
 		<Property Name="InstallerDependency" Value="Build UnrealGameSync Installer" If="$(IncludeInstaller)"/>
