Upload UGS installer as a zip file.

diff --git a/Engine/Source/Programs/UnrealGameSync/BuildUGS.xml b/Engine/Source/Programs/UnrealGameSync/BuildUGS.xml
--- a/Engine/Source/Programs/UnrealGameSync/BuildUGS.xml
+++ b/Engine/Source/Programs/UnrealGameSync/BuildUGS.xml
@@ -26,6 +26,7 @@
 	
 	<Agent Name="UnrealGameSync Agent" Type="Win64">
 		<Property Name="InstallerDir" Value="$(RootDir)/Engine/Source/Programs/UnrealGameSync/Installer"/>
+		<Property Name="InstallerZip" Value="$(RootDir)/LocalBuilds/UnrealGameSyncInstaller/ugs-installer.zip"/>
 
 		<Property Name="UgsPublishDir" Value="$(UgsOutputDir)"/>
 		<Property Name="UgsPublishDir" Value="$(UgsPublishDir)/$(SubmitPath)" If="'$(SubmitPath)' != ''"/>
@@ -43,10 +44,12 @@
 			<Spawn Exe="$(WixDir)/heat.exe" Arguments="dir &quot;$(UgsLauncherPublishDir)&quot; -cg UGSLauncher_Project -dr INSTALLFOLDER -scom -sreg -srd -var var.BasePath -gg -sfrag -out obj/Release/UGSLauncher.wxs" WorkingDir="$(InstallerDir)"/>
 			<Spawn Exe="$(WixDir)/candle.exe" Arguments="-dBasePath=&quot;$(UgsLauncherPublishDir)&quot; -out obj\Release\ -dConfiguration=Release -dPlatform=x64 -arch x86 -ext ..\..\..\ThirdParty\WiX\3.8\WixUtilExtension.dll Product.wxs obj/Release/UGSLauncher.wxs" WorkingDir="$(InstallerDir)"/>
 			<Spawn Exe="$(WixDir)/light.exe" Arguments="-out &quot;bin/Release/UnrealGameSync.msi&quot; -pdbout &quot;bin/Release/UnrealGameSync.wixpdb&quot; -cultures:null -ext ..\..\..\ThirdParty\WiX\3.8\WixUtilExtension.dll -sice:ICE69 &quot;obj\Release\Product.wixobj&quot; &quot;obj\Release\UGSLauncher.wixobj&quot;" WorkingDir="$(InstallerDir)"/>
+
+			<Zip ZipFile="$(InstallerZip)" FromDir="$(InstallerDir)/bin/Release" Files="UnrealGameSync.msi"/>
 		</Node>
 		<Node Name="Publish UnrealGameSync Installer" Requires="Build UnrealGameSync Installer">
 			<Error Message="Missing -HordeSettings=... argument" If="'$(HordeSettings)' == ''"/>
-			<DeployTool Id="ugs" Version="$(InformationalVersion)" Settings="$(HordeSettings)" File="$(InstallerDir)/bin/Release/UnrealGameSync.msi"/>
+			<DeployTool Id="ugs" Version="$(InformationalVersion)" Settings="$(HordeSettings)" File="$(InstallerZip)"/>
 		</Node>
 		
 		<Property Name="InstallerDependency" Value=""/>
