UGS: Add a registry key to opt out of checking for running processes by opening the file for write. This seems to be causing some sort of issue on user machines, where the OS hangs inside NtCreateFile. Possibly related to virus scanner?

diff --git a/Engine/Source/Programs/UnrealGameSync/UnrealGameSync/Controls/WorkspaceControl.cs b/Engine/Source/Programs/UnrealGameSync/UnrealGameSync/Controls/WorkspaceControl.cs
--- a/Engine/Source/Programs/UnrealGameSync/UnrealGameSync/Controls/WorkspaceControl.cs
+++ b/Engine/Source/Programs/UnrealGameSync/UnrealGameSync/Controls/WorkspaceControl.cs
@@ -21,6 +21,7 @@ using Microsoft.Extensions.Logging;
 using Microsoft.Extensions.DependencyInjection;
 using System.Diagnostics.CodeAnalysis;
 using EpicGames.OIDC;
+using Microsoft.Win32;
 
 namespace UnrealGameSync
 {
@@ -3157,6 +3158,12 @@ namespace UnrealGameSync
 
 		private FileReference[] GetProcessesRunningInWorkspace()
 		{
+			object? skipRunnningProcessCheck = Registry.GetValue("HKEY_CURRENT_USER\\SOFTWARE\\Epic Games\\UnrealGameSync", "SkipRunningProcessCheck", null);
+			if (skipRunnningProcessCheck != null && skipRunnningProcessCheck is int value && value != 0)
+			{
+				return Array.Empty<FileReference>();
+			}
+
 			HashSet<string> processNames = new HashSet<string>(StringComparer.InvariantCultureIgnoreCase);
 			processNames.Add("Win64\\UE4Editor.exe");
 			processNames.Add("Win64\\UE4Editor-Cmd.exe");
