UGS: Fix PCBs not being found for changelists outside of those enumerated for the current project.

Fix to change implemented in CL 23015525.

diff --git a/Engine/Source/Programs/UnrealGameSync/UnrealGameSync/Controls/WorkspaceControl.cs b/Engine/Source/Programs/UnrealGameSync/UnrealGameSync/Controls/WorkspaceControl.cs
--- a/Engine/Source/Programs/UnrealGameSync/UnrealGameSync/Controls/WorkspaceControl.cs
+++ b/Engine/Source/Programs/UnrealGameSync/UnrealGameSync/Controls/WorkspaceControl.cs
@@ -2046,7 +2046,7 @@ namespace UnrealGameSync
 						int index = _sortedChangeNumbers.BinarySearch(changeNumber);
 						if (index > 0)
 						{
-							archivePath = GetArchiveKeyForChangeNumber(archive, _sortedChangeNumbers[index - 1]);
+							archivePath = GetArchiveKeyForChangeNumber(archive, _sortedChangeNumbers[index - 1], maxChangeNumber, changeNumberToArchivePath);
 						}
 					}
 				}
