UGS: Add option for a description badge to only match a single item

diff --git a/Engine/Source/Programs/UnrealGameSync/UnrealGameSync/Controls/WorkspaceControl.cs b/Engine/Source/Programs/UnrealGameSync/UnrealGameSync/Controls/WorkspaceControl.cs
--- a/Engine/Source/Programs/UnrealGameSync/UnrealGameSync/Controls/WorkspaceControl.cs
+++ b/Engine/Source/Programs/UnrealGameSync/UnrealGameSync/Controls/WorkspaceControl.cs
@@ -2548,6 +2548,7 @@ namespace UnrealGameSync
 						string hoverColor = badgeDefinitionObject.GetValue("HoverColor", "#b0b0b0");
 						string? url = badgeDefinitionObject.GetValue("Url", null);
 						string? arguments = badgeDefinitionObject.GetValue("Arguments", null);
+						bool singleMatch = badgeDefinitionObject.GetValue("SingleMatch", false);
 						if (!String.IsNullOrEmpty(name) && !String.IsNullOrEmpty(pattern))
 						{
 							foreach (Match? matchResult in Regex.Matches(description, pattern, RegexOptions.Multiline))
@@ -2577,6 +2578,11 @@ namespace UnrealGameSync
 									}
 
 									badges.Add(new BadgeInfo(ReplaceRegexMatches(name, matchResult), group, uniqueId, badgeColor, hoverBadgeColor, clickHandler, null));
+
+									if (singleMatch)
+									{
+										break;
+									}
 								}
 							}
 						}
