Message dialogs shown at the end of syncing (unresolved files, files needing deletion etc) are now properly linked to the main UGS window.
In addition we now flash the UGS taskbar entry before showing these modal dialogs and at the end of the sync to try and get the users attention.

- The message boxes will now move to the front of the active monitor when UGS is selected, where as before you would have to select the message boxes entry in the taskbar to find it.
- Because the message box no longer has its own entry in the task bar, a user will have no way of knowing that a problem has been encoutnered. So before we display one of these messages we flash the UGS entry in the task bar to get their attention.
- Since it is also useful to know when a sync has completed (and it is something I have wanted myself) I also added a flash once syncing has completed.

diff --git a/Engine/Source/Programs/UnrealGameSync/UnrealGameSync/Controls/WorkspaceControl.cs b/Engine/Source/Programs/UnrealGameSync/UnrealGameSync/Controls/WorkspaceControl.cs
--- a/Engine/Source/Programs/UnrealGameSync/UnrealGameSync/Controls/WorkspaceControl.cs
+++ b/Engine/Source/Programs/UnrealGameSync/UnrealGameSync/Controls/WorkspaceControl.cs
@@ -200,6 +200,9 @@ namespace UnrealGameSync
 		[DllImport("uxtheme.dll", CharSet = CharSet.Unicode)]
 		static extern int SetWindowTheme(IntPtr hWnd, string pszSubAppName, string pszSubIdList);
 
+		[DllImport("user32.dll")]
+		static extern bool FlashWindow(IntPtr hwnd, bool bInvert);
+
 		const int BuildListExpandCount = 250;
 
 		readonly IWorkspaceControlOwner _owner;
@@ -1061,13 +1064,15 @@ namespace UnrealGameSync
 
 			if (result == WorkspaceUpdateResult.FilesToResolve)
 			{
-				MessageBox.Show("You have files to resolve after syncing your workspace. Please check P4.");
+				FlashWindow(ParentForm.Handle, true);
+				MessageBox.Show(ParentForm, "You have files to resolve after syncing your workspace. Please check P4.");
 			}
 			else if (result == WorkspaceUpdateResult.FilesToDelete)
 			{
 				DesiredTaskbarState = Tuple.Create(TaskbarState.Paused, 0.0f);
 				_owner.UpdateProgress();
 
+				FlashWindow(ParentForm.Handle, true);
 				using DeleteWindow window = new DeleteWindow(context.DeleteFiles);
 				if (window.ShowDialog(this) == DialogResult.OK)
 				{
@@ -1143,6 +1148,8 @@ namespace UnrealGameSync
 			{
 				LoginExpired();
 			}
+
+			FlashWindow(ParentForm.Handle, true);
 		}
 
 		void UpdateBuildListCallback()
