Fix a crash if the user selects 'Open Project...' from the context menu for the main windows tab control when the context menu was not opened over a valid tab.

- When the context menu is opened enable/disable the 'Open Project...' option depending on if we found a valid tab index or not.

diff --git a/Engine/Source/Programs/UnrealGameSync/UnrealGameSync/Forms/MainWindow.cs b/Engine/Source/Programs/UnrealGameSync/UnrealGameSync/Forms/MainWindow.cs
--- a/Engine/Source/Programs/UnrealGameSync/UnrealGameSync/Forms/MainWindow.cs
+++ b/Engine/Source/Programs/UnrealGameSync/UnrealGameSync/Forms/MainWindow.cs
@@ -542,6 +542,8 @@ namespace UnrealGameSync
 					}
 				}
 
+				TabMenu_OpenProject.Enabled = _tabMenuTabIdx != -1;
+
 				foreach(UserSelectedProjectSettings recentProject in _settings.RecentProjects)
 				{
 #pragma warning disable CA2000 // warning CA2000: Call System.IDisposable.Dispose on object created by 'new ToolStripMenuItem(recentProject.ToString(), null, new EventHandler((o, e) => TabMenu_OpenRecentProject_Click(recentProject, _tabMenuTabIdx)))' before all references to it are out of scope
