UGS: Fix incorrect determination of content changes. maxFiles was updated in the inner loop, and incorrectly affecting future iterations of the outer loop. (#10507)

diff --git a/Engine/Source/Programs/UnrealGameSync/UnrealGameSync/PerforceMonitor.cs b/Engine/Source/Programs/UnrealGameSync/UnrealGameSync/PerforceMonitor.cs
--- a/Engine/Source/Programs/UnrealGameSync/UnrealGameSync/PerforceMonitor.cs
+++ b/Engine/Source/Programs/UnrealGameSync/UnrealGameSync/PerforceMonitor.cs
@@ -465,9 +465,9 @@ namespace UnrealGameSync
 					}
 
 					// If there's something to check for, find all the content changes after this changelist
-					int maxFiles = 100;
+					const int InitialMaxFiles = 100;
 
-					List<DescribeRecord> describeRecords = await perforce.DescribeAsync(DescribeOptions.None, maxFiles, queryChangeNumberBatch.ToArray(), cancellationToken);
+					List<DescribeRecord> describeRecords = await perforce.DescribeAsync(DescribeOptions.None, InitialMaxFiles, queryChangeNumberBatch.ToArray(), cancellationToken);
 					foreach (DescribeRecord describeRecordLoop in describeRecords.OrderByDescending(x => x.Number))
 					{
 						DescribeRecord describeRecord = describeRecordLoop;
@@ -476,12 +476,13 @@ namespace UnrealGameSync
 						PerforceChangeDetails details = new PerforceChangeDetails(describeRecord, isCodeFile);
 
 						// Content only changes must be flagged accurately, because code changes invalidate precompiled binaries. Increase the number of files fetched until we can classify it correctly.
-						while (describeRecord.Files.Count >= maxFiles && !details.ContainsCode)
+						int currentMaxFiles = InitialMaxFiles;
+						while (describeRecord.Files.Count >= currentMaxFiles && !details.ContainsCode)
 						{
 							cancellationToken.ThrowIfCancellationRequested();
-							maxFiles *= 10;
+							currentMaxFiles *= 10;
 
-							List<DescribeRecord> newDescribeRecords = await perforce.DescribeAsync(DescribeOptions.None, maxFiles, new int[] { queryChangeNumber }, cancellationToken);
+							List<DescribeRecord> newDescribeRecords = await perforce.DescribeAsync(DescribeOptions.None, currentMaxFiles, new int[] { queryChangeNumber }, cancellationToken);
 							if (newDescribeRecords.Count == 0)
 							{
 								break;
