EpicGames.Perforce: Fix handling of opened command when a stream spec is opened for edit.

diff --git a/Engine/Extras/P4VUtils/Commands/BackoutCommand.cs b/Engine/Extras/P4VUtils/Commands/BackoutCommand.cs
--- a/Engine/Extras/P4VUtils/Commands/BackoutCommand.cs
+++ b/Engine/Extras/P4VUtils/Commands/BackoutCommand.cs
@@ -170,7 +170,7 @@ namespace P4VUtils.Commands
 			List<OpenedRecord> OpenedRecords = await Perforce.OpenedAsync(OpenedOptions.AllWorkspaces | OpenedOptions.ShortOutput, ExistingChangeRecord.Files.Select(x => x.DepotFile).ToArray(), CancellationToken.None).ToListAsync();
 			if (OpenedRecords.Count > 0)
 			{
-				HashSet<string> UniqueDepotFiles = (OpenedRecords.Select(x => x.DepotFile)).ToHashSet();
+				HashSet<string> UniqueDepotFiles = (OpenedRecords.Where(x => !String.IsNullOrEmpty(x.DepotFile)).Select(x => x.DepotFile!)).ToHashSet();
 				string FileListString = string.Join("\r\n", UniqueDepotFiles);
 
 				Logger.LogInformation("\r\nSome files are checked out on another workspace, please confirm that you wish to continue.\r\n");
diff --git a/Engine/Source/Programs/Shared/EpicGames.Perforce/Records/OpenedRecord.cs b/Engine/Source/Programs/Shared/EpicGames.Perforce/Records/OpenedRecord.cs
--- a/Engine/Source/Programs/Shared/EpicGames.Perforce/Records/OpenedRecord.cs
+++ b/Engine/Source/Programs/Shared/EpicGames.Perforce/Records/OpenedRecord.cs
@@ -12,9 +12,15 @@ namespace EpicGames.Perforce
 	public class OpenedRecord
 	{
 		/// <summary>
-		/// Depot path to file
+		/// The stream spec opened for edit
 		/// </summary>
-		[PerforceTag("depotFile")]
+		[PerforceTag("stream", Optional = true)]
+		public string Stream { get; set; } = String.Empty;
+
+		/// <summary>
+		/// Depot path to file. May be unset if this record describes a stream spec.
+		/// </summary>
+		[PerforceTag("depotFile", Optional = true)]
 		public string DepotFile { get; set; } = String.Empty;
 
 		/// <summary>
@@ -68,7 +74,7 @@ namespace EpicGames.Perforce
 		/// <summary>
 		/// Client with the file open
 		/// </summary>
-		[PerforceTag("client")]
+		[PerforceTag("client", Optional = true)]
 		public string Client { get; set; } = String.Empty;
 	}
 }
diff --git a/Engine/Source/Programs/UnrealGameSync/UnrealGameSyncShared/WorkspaceUpdate.cs b/Engine/Source/Programs/UnrealGameSync/UnrealGameSyncShared/WorkspaceUpdate.cs
--- a/Engine/Source/Programs/UnrealGameSync/UnrealGameSyncShared/WorkspaceUpdate.cs
+++ b/Engine/Source/Programs/UnrealGameSync/UnrealGameSyncShared/WorkspaceUpdate.cs
@@ -823,10 +823,13 @@ namespace UnrealGameSync
 								}
 
 								OpenedRecord record = response.Data;
-								if (record.Action != FileAction.Add || record.Action != FileAction.Branch || record.Action != FileAction.MoveAdd)
+								if (!String.IsNullOrEmpty(record.DepotFile) && !String.IsNullOrEmpty(record.ClientFile))
 								{
-									string relativePath = PerforceUtils.GetClientRelativePath(record.ClientFile);
-									syncFiles.Add(new SyncFile(record.DepotFile, relativePath, 0));
+									if (record.Action != FileAction.Add || record.Action != FileAction.Branch || record.Action != FileAction.MoveAdd)
+									{
+										string relativePath = PerforceUtils.GetClientRelativePath(record.ClientFile);
+										syncFiles.Add(new SyncFile(record.DepotFile, relativePath, 0));
+									}
 								}
 							}
 
