UGS: Improve performance of determining code changelist number for matching content changes.

- Now uses the same method for determining code/content changes when updating the workspace as enumerating changes for the GUI.
- Both methods share a cache for enumerating change records and creating PerforceChangeDetails objects.

diff --git a/Engine/Source/Programs/UnrealGameSync/UnrealGameSync/PerforceMonitor.cs b/Engine/Source/Programs/UnrealGameSync/UnrealGameSync/PerforceMonitor.cs
--- a/Engine/Source/Programs/UnrealGameSync/UnrealGameSync/PerforceMonitor.cs
+++ b/Engine/Source/Programs/UnrealGameSync/UnrealGameSync/PerforceMonitor.cs
@@ -299,17 +299,8 @@ namespace UnrealGameSync
 			}
 
 			// Read any new changes
-			List<ChangesRecord> newChanges;
-			if (maxChanges > CurrentMaxChanges || newestChangeNumber == -1)
-			{
-				newChanges = await perforce.GetChangesAsync(ChangesOptions.IncludeTimes | ChangesOptions.LongOutput, maxChanges, ChangeStatus.Submitted, depotPaths, cancellationToken);
-			}
-			else
-			{
-				List<string> depotPathsWithRange = depotPaths.ConvertAll(x => $"{x}@{newestChangeNumber + 1},#head");
-				//				newChanges = await perforce.GetChangesAsync(ChangesOptions.IncludeTimes | ChangesOptions.LongOutput, maxChanges, ChangeStatus.Submitted, depotPaths.Select(x => $"{x}@>{newestChangeNumber}").ToArray(), cancellationToken);
-				newChanges = await perforce.GetChangesAsync(ChangesOptions.IncludeTimes | ChangesOptions.LongOutput, clientName: null, minChangeNumber: newestChangeNumber + 1, maxChanges: maxChanges, status: ChangeStatus.Submitted, userName: null, fileSpecs: depotPathsWithRange, cancellationToken);
-			}
+			int? minChangeNumber = (maxChanges > CurrentMaxChanges || newestChangeNumber == -1)? null : (newestChangeNumber + 1);
+			List<ChangesRecord> newChanges = await Utility.EnumerateChanges(perforce, depotPaths, minChangeNumber, null, maxChanges, cancellationToken).ToListAsync(cancellationToken);
 
 			// Remove anything we already have
 			newChanges.RemoveAll(x => currentChangelists.Contains(x.Number));
@@ -416,26 +407,13 @@ namespace UnrealGameSync
 		public async Task<bool> UpdateChangeTypesAsync(IPerforceConnection perforce, CancellationToken cancellationToken)
 		{
 			// Get the filter for code changes
-			ConfigSection? projectConfigSection = LatestProjectConfigFile.FindSection("Perforce");
-
-			string[] codeRules = projectConfigSection?.GetValues("CodeFilter", (string[]?)null) ?? Array.Empty<string>();
+			string[] codeRules = Utility.GetCodeFilter(LatestProjectConfigFile);
 			if (!Enumerable.SequenceEqual(codeRules, prevCodeRules))
 			{
 				_changeDetails.Clear();
 				prevCodeRules = codeRules;
 			}
 
-			Func<string, bool>? isCodeFile = null;
-			if (codeRules.Length > 0)
-			{
-				FileFilter filter = new FileFilter(PerforceUtils.CodeExtensions.Select(x => $"*{x}"));
-				foreach (string codeRule in codeRules)
-				{
-					filter.AddRule(codeRule);
-				}
-				isCodeFile = filter.Matches;
-			}
-
 			// Find the changes we need to query
 			List<int> queryChangeNumbers = new List<int>();
 			lock(_lockObject)
@@ -454,7 +432,7 @@ namespace UnrealGameSync
 			using (CancellationTokenSource cancellationSource = new CancellationTokenSource())
 			{
 				Task notifyTask = Task.CompletedTask;
-				foreach (IReadOnlyList<int> queryChangeNumberBatch in queryChangeNumbers.OrderByDescending(x => x).Batch(10))
+				await foreach(PerforceChangeDetails details in Utility.EnumerateChangeDetails(perforce, queryChangeNumbers.ToAsyncEnumerable(), codeRules, cancellationToken))
 				{
 					cancellationToken.ThrowIfCancellationRequested();
 
@@ -464,54 +442,26 @@ namespace UnrealGameSync
 						break;
 					}
 
-					// If there's something to check for, find all the content changes after this changelist
-					const int InitialMaxFiles = 100;
-
-					List<DescribeRecord> describeRecords = await perforce.DescribeAsync(DescribeOptions.None, InitialMaxFiles, queryChangeNumberBatch.ToArray(), cancellationToken);
-					foreach (DescribeRecord describeRecordLoop in describeRecords.OrderByDescending(x => x.Number))
+					// Add this change to the cache
+					lock (_lockObject)
 					{
-						DescribeRecord describeRecord = describeRecordLoop;
-						int queryChangeNumber = describeRecord.Number;
-
-						PerforceChangeDetails details = new PerforceChangeDetails(describeRecord, isCodeFile);
-
-						// Content only changes must be flagged accurately, because code changes invalidate precompiled binaries. Increase the number of files fetched until we can classify it correctly.
-						int currentMaxFiles = InitialMaxFiles;
-						while (describeRecord.Files.Count >= currentMaxFiles && !details.ContainsCode)
-						{
-							cancellationToken.ThrowIfCancellationRequested();
-							currentMaxFiles *= 10;
-
-							List<DescribeRecord> newDescribeRecords = await perforce.DescribeAsync(DescribeOptions.None, currentMaxFiles, new int[] { queryChangeNumber }, cancellationToken);
-							if (newDescribeRecords.Count == 0)
-							{
-								break;
-							}
-
-							describeRecord = newDescribeRecords[0];
-							details = new PerforceChangeDetails(describeRecord, isCodeFile);
-						}
-
-						lock (_lockObject)
+						if (!_changeDetails.ContainsKey(details.Number))
 						{
-							if (!_changeDetails.ContainsKey(queryChangeNumber))
-							{
-								_changeDetails.Add(queryChangeNumber, details);
-							}
+							_changeDetails.Add(details.Number, details);
 						}
+					}
 
-						// Reload the config file if it changes
-						if (describeRecord.Files.Any(x => x.DepotFile.EndsWith("/UnrealGameSync.ini", StringComparison.OrdinalIgnoreCase)) && !updatedConfigFile)
-						{
-							await UpdateProjectConfigFileAsync(perforce, cancellationToken);
-							updatedConfigFile = true;
-						}
+					// Reload the config file if it changes
+					if (details.ContainsUgsConfig && !updatedConfigFile)
+					{
+						await UpdateProjectConfigFileAsync(perforce, cancellationToken);
+						updatedConfigFile = true;
+					}
 
-						// Notify the caller after a fixed period of time, in case further updates are slow to arrive
-						if (notifyTask.IsCompleted)
-						{
-							notifyTask = Task.Delay(TimeSpan.FromSeconds(5.0), cancellationSource.Token).ContinueWith(_ => _synchronizationContext.Post(_ => OnUpdateMetadata?.Invoke(), null), cancellationSource.Token, new TaskContinuationOptions(), TaskScheduler.Default);
-						}
+					// Notify the caller after a fixed period of time, in case further updates are slow to arrive
+					if (notifyTask.IsCompleted)
+					{
+						notifyTask = Task.Delay(TimeSpan.FromSeconds(5.0), cancellationSource.Token).ContinueWith(_ => _synchronizationContext.Post(_ => OnUpdateMetadata?.Invoke(), null), cancellationSource.Token, new TaskContinuationOptions(), TaskScheduler.Default);
 					}
 				}
 				cancellationSource.Cancel();
diff --git a/Engine/Source/Programs/UnrealGameSync/UnrealGameSyncShared/UnrealGameSyncShared.csproj b/Engine/Source/Programs/UnrealGameSync/UnrealGameSyncShared/UnrealGameSyncShared.csproj
--- a/Engine/Source/Programs/UnrealGameSync/UnrealGameSyncShared/UnrealGameSyncShared.csproj
+++ b/Engine/Source/Programs/UnrealGameSync/UnrealGameSyncShared/UnrealGameSyncShared.csproj
@@ -28,6 +28,7 @@
   </Choose>
   
   <ItemGroup>
+    <PackageReference Include="Microsoft.Extensions.Caching.Memory" Version="6.0.1" />
     <PackageReference Include="System.IO.FileSystem.DriveInfo" Version="4.3.1" />
   </ItemGroup>
 
diff --git a/Engine/Source/Programs/UnrealGameSync/UnrealGameSyncShared/Utility.cs b/Engine/Source/Programs/UnrealGameSync/UnrealGameSyncShared/Utility.cs
--- a/Engine/Source/Programs/UnrealGameSync/UnrealGameSyncShared/Utility.cs
+++ b/Engine/Source/Programs/UnrealGameSync/UnrealGameSyncShared/Utility.cs
@@ -2,6 +2,8 @@
 
 using EpicGames.Core;
 using EpicGames.Perforce;
+using JetBrains.Annotations;
+using Microsoft.Extensions.Caching.Memory;
 using Microsoft.Extensions.Logging;
 using System;
 using System.Collections.Generic;
@@ -9,7 +11,8 @@ using System.Diagnostics;
 using System.Diagnostics.CodeAnalysis;
 using System.Drawing;
 using System.IO;
-using System.Linq;	
+using System.Linq;
+using System.Runtime.CompilerServices;
 using System.Text;
 using System.Text.Json;
 using System.Text.Json.Serialization;
@@ -30,14 +33,17 @@ namespace UnrealGameSync
 
 	public class PerforceChangeDetails
 	{
+		public int Number { get; }
 		public string Description { get; }
 		public bool ContainsCode { get; }
 		public bool ContainsContent { get; }
+		public bool ContainsUgsConfig { get; }
 
 		public PerforceChangeDetails(DescribeRecord describeRecord, Func<string, bool>? isCodeFile = null)
 		{
 			isCodeFile ??= IsCodeFile;
 
+			Number = describeRecord.Number;
 			Description = describeRecord.Description;
 
 			// Check whether the files are code or content
@@ -52,9 +58,9 @@ namespace UnrealGameSync
 					ContainsContent = true;
 				}
 
-				if (ContainsCode && ContainsContent)
+				if (file.DepotFile.EndsWith("/UnrealGameSync.ini", StringComparison.OrdinalIgnoreCase))
 				{
-					break;
+					ContainsUgsConfig = true;
 				}
 			}
 		}
@@ -67,6 +73,271 @@ namespace UnrealGameSync
 
 	public static class Utility
 	{
+		static readonly MemoryCache s_changeCache = new MemoryCache(new MemoryCacheOptions { SizeLimit = 4096 });
+
+		class CachedChangeRecord
+		{
+			public ChangesRecord ChangesRecord { get; }
+			public int? PrevNumber { get; set; }
+
+			public CachedChangeRecord(ChangesRecord changesRecord) => ChangesRecord = changesRecord;
+		}
+
+		static string GetChangeCacheKey(int number, Sha1Hash config) => $"change: {number} ({config})";
+
+		static string GetChangeDetailsCacheKey(int number, Sha1Hash config) => $"details: {number} ({config})";
+
+		static Sha1Hash GetConfigHash(IPerforceConnection perforce, IEnumerable<string> syncPaths, IEnumerable<string> codeRules)
+		{
+			StringBuilder digest = new StringBuilder();
+			digest.AppendLine($"Server: {perforce.Settings.ServerAndPort}");
+			digest.AppendLine($"User: {perforce.Settings.UserName}");
+			digest.AppendLine($"Client: {perforce.Settings.ClientName}");
+			foreach (string syncPath in syncPaths)
+			{
+				digest.AppendLine($"Sync: {syncPath}");
+			}
+			foreach (string codeRule in codeRules)
+			{
+				digest.AppendLine($"Rule: {codeRule}");
+			}
+			return Sha1Hash.Compute(Encoding.UTF8.GetBytes(digest.ToString()));
+		}
+
+		/// <summary>
+		/// Gets the code filter from the project config file
+		/// </summary>
+		public static string[] GetCodeFilter(ConfigFile projectConfigFile)
+		{
+			ConfigSection? projectConfigSection = projectConfigFile.FindSection("Perforce");
+			return projectConfigSection?.GetValues("CodeFilter", (string[]?)null) ?? Array.Empty<string>();
+		}
+
+		/// <summary>
+		/// Creates or returns cached <see cref="PerforceChangeDetails"/> objects describing the requested chagnes
+		/// </summary>
+		public static async IAsyncEnumerable<ChangesRecord> EnumerateChanges(IPerforceConnection perforce, IEnumerable<string> syncPaths, int? minChangeNumber, int? maxChangeNumber, int? maxChanges, [EnumeratorCancellation] CancellationToken cancellationToken)
+		{
+			CachedChangeRecord? prevCachedChangeRecord = null;
+
+			Sha1Hash configHash = GetConfigHash(perforce, syncPaths, Array.Empty<string>());
+			while (maxChanges == null || maxChanges.Value > 0)
+			{
+				// If we have a maximum changelist number, see if the previous change is in the cache
+				if (maxChangeNumber != null)
+				{
+					if (minChangeNumber != null && maxChangeNumber.Value < minChangeNumber.Value)
+					{
+						yield break;
+					}
+
+					string cacheKey = GetChangeCacheKey(maxChangeNumber.Value, configHash);
+					if (s_changeCache.TryGetValue(cacheKey, out CachedChangeRecord change))
+					{
+						yield return change.ChangesRecord;
+
+						if (maxChanges != null)
+						{
+							maxChanges = maxChanges.Value - 1;
+						}
+
+						if (change.PrevNumber != null)
+						{
+							maxChangeNumber = change.PrevNumber.Value;
+						}
+						else
+						{
+							maxChangeNumber = change.ChangesRecord.Number - 1;
+						}
+
+						prevCachedChangeRecord = change;
+						continue;
+					}
+				}
+
+				// Get the search range
+				string range;
+				if (minChangeNumber == null)
+				{
+					if (maxChangeNumber == null)
+					{
+						range = "";
+					}
+					else
+					{
+						range = $"@<={maxChangeNumber.Value}";
+					}
+				}
+				else
+				{
+					if (maxChangeNumber == null)
+					{
+						range = $"@{minChangeNumber.Value},#head";
+					}
+					else
+					{
+						range = $"@{minChangeNumber.Value},{maxChangeNumber.Value}";
+					}
+				}
+
+				// Get the next batch of change numbers
+				int maxChangesForBatch = maxChanges ?? 20;
+				List<string> syncPathsWithChange = syncPaths.Select(x => $"{x}{range}").ToList();
+
+				List<ChangesRecord> changes;
+				if (minChangeNumber == null)
+				{
+					changes = await perforce.GetChangesAsync(ChangesOptions.IncludeTimes | ChangesOptions.LongOutput, maxChangesForBatch, ChangeStatus.Submitted, syncPathsWithChange, cancellationToken);
+				}
+				else
+				{
+					changes = await perforce.GetChangesAsync(ChangesOptions.IncludeTimes | ChangesOptions.LongOutput, clientName: null, minChangeNumber: minChangeNumber.Value, maxChangesForBatch, ChangeStatus.Submitted, userName:null, fileSpecs: syncPathsWithChange, cancellationToken: cancellationToken);
+				}
+
+				// Sort the changes in case we get interleaved output from multiple sync paths
+				changes = changes.OrderByDescending(x => x.Number).Take(maxChangesForBatch).ToList();
+
+				// Add all the previous change numbers to the cache
+				foreach(ChangesRecord change in changes)
+				{
+					if (prevCachedChangeRecord != null)
+					{
+						prevCachedChangeRecord.PrevNumber = change.Number;
+					}
+
+					prevCachedChangeRecord = new CachedChangeRecord(change);
+
+					string cacheKey = GetChangeCacheKey(change.Number, configHash);
+					using (ICacheEntry entry = s_changeCache.CreateEntry(cacheKey))
+					{
+						entry.SetSize(1);
+						entry.Value = prevCachedChangeRecord;
+					}
+
+					maxChangeNumber = change.Number - 1;
+				}
+
+				// Return the results
+				foreach (ChangesRecord change in changes)
+				{
+					yield return change;
+				}
+
+				// If we have run out of changes, quit now
+				if (changes.Count < maxChangesForBatch)
+				{
+					break;
+				}
+				if (maxChanges != null)
+				{
+					maxChanges = maxChanges.Value - changes.Count;
+				}
+			}
+		}
+
+		/// <summary>
+		/// Creates or returns cached <see cref="PerforceChangeDetails"/> objects describing the requested chagnes
+		/// </summary>
+		public static IAsyncEnumerable<PerforceChangeDetails> EnumerateChangeDetails(IPerforceConnection perforce, int? minChangeNumber, int? maxChangeNumber, IEnumerable<string> syncPaths, IEnumerable<string> codeRules, CancellationToken cancellationToken)
+		{
+			IAsyncEnumerable<int> changeNumbers = EnumerateChanges(perforce, syncPaths, minChangeNumber, maxChangeNumber, null, cancellationToken).Select(x => x.Number);
+			return EnumerateChangeDetails(perforce, changeNumbers, codeRules, cancellationToken);
+		}
+
+		/// <summary>
+		/// Creates or returns cached <see cref="PerforceChangeDetails"/> objects describing the requested chagnes
+		/// </summary>
+		public static async IAsyncEnumerable<PerforceChangeDetails> EnumerateChangeDetails(IPerforceConnection perforce, IAsyncEnumerable<int> changeNumbers, IEnumerable<string> codeRules, [EnumeratorCancellation] CancellationToken cancellationToken)
+		{
+			// Get a delegate that determines if a file is a code change or not
+			Func<string, bool>? isCodeFile = null;
+			if (codeRules.Any())
+			{
+				FileFilter filter = new FileFilter(PerforceUtils.CodeExtensions.Select(x => $"*{x}"));
+				foreach (string codeRule in codeRules)
+				{
+					filter.AddRule(codeRule);
+				}
+				isCodeFile = filter.Matches;
+			}
+
+			// Get the hash of the configuration
+			Sha1Hash hash = GetConfigHash(perforce, Enumerable.Empty<string>(), codeRules);
+
+			// Update them in batches
+			await using IAsyncEnumerator<int> changeNumberEnumerator = changeNumbers.GetAsyncEnumerator(cancellationToken);
+			for (; ; )
+			{
+				// Get the next batch of changes to query, up to the next change that we already have a cached value for
+				const int BatchSize = 10;
+				PerforceChangeDetails? cachedDetails = null;
+
+				List<int> changeBatch = new List<int>(BatchSize);
+				while (changeBatch.Count < BatchSize && await changeNumberEnumerator.MoveNextAsync(cancellationToken))
+				{
+					string cacheKey = GetChangeDetailsCacheKey(changeNumberEnumerator.Current, hash);
+					if (s_changeCache.TryGetValue(cacheKey, out cachedDetails))
+					{
+						break;
+					}
+					changeBatch.Add(changeNumberEnumerator.Current);
+				}
+
+				// Describe the requested changes
+				if (changeBatch.Count > 0)
+				{
+					const int InitialMaxFiles = 100;
+
+					List<DescribeRecord> describeRecords = await perforce.DescribeAsync(DescribeOptions.None, InitialMaxFiles, changeBatch.ToArray(), cancellationToken);
+					foreach (DescribeRecord describeRecordLoop in describeRecords.OrderByDescending(x => x.Number))
+					{
+						DescribeRecord describeRecord = describeRecordLoop;
+						int queryChangeNumber = describeRecord.Number;
+
+						PerforceChangeDetails details = new PerforceChangeDetails(describeRecord, isCodeFile);
+
+						// Content only changes must be flagged accurately, because code changes invalidate precompiled binaries. Increase the number of files fetched until we can classify it correctly.
+						int currentMaxFiles = InitialMaxFiles;
+						while (describeRecord.Files.Count >= currentMaxFiles && !details.ContainsCode)
+						{
+							cancellationToken.ThrowIfCancellationRequested();
+							currentMaxFiles *= 10;
+
+							List<DescribeRecord> newDescribeRecords = await perforce.DescribeAsync(DescribeOptions.None, currentMaxFiles, new int[] { queryChangeNumber }, cancellationToken);
+							if (newDescribeRecords.Count == 0)
+							{
+								break;
+							}
+
+							describeRecord = newDescribeRecords[0];
+							details = new PerforceChangeDetails(describeRecord, isCodeFile);
+						}
+
+						// Add it to the cache
+						string cacheKey = GetChangeDetailsCacheKey(describeRecord.Number, hash);
+						using (ICacheEntry entry = s_changeCache.CreateEntry(cacheKey))
+						{
+							entry.SetSize(10);
+							entry.Value = details;
+						}
+
+						// Return the value
+						yield return details;
+					}
+				}
+
+				// Return the cached value, if there was one
+				if (cachedDetails != null)
+				{
+					yield return cachedDetails;
+				}
+				else if (changeBatch.Count == 0)
+				{
+					yield break;
+				}
+			}
+		}
+
 		public static event Action<Exception>? TraceException;
 
 		static JsonSerializerOptions GetDefaultJsonSerializerOptions()
diff --git a/Engine/Source/Programs/UnrealGameSync/UnrealGameSyncShared/WorkspaceUpdate.cs b/Engine/Source/Programs/UnrealGameSync/UnrealGameSyncShared/WorkspaceUpdate.cs
--- a/Engine/Source/Programs/UnrealGameSync/UnrealGameSyncShared/WorkspaceUpdate.cs
+++ b/Engine/Source/Programs/UnrealGameSync/UnrealGameSyncShared/WorkspaceUpdate.cs
@@ -936,43 +936,30 @@ namespace UnrealGameSync
 							branchOrStreamName = PerforceUtils.GetClientOrDepotDirectoryName(files[0].DepotFile);
 						}
 
+						logger.LogInformation("");
+
 						// Get the last code change
 						int codeChangeNumber = Context.CodeChangeNumber ?? 0;
 						if (codeChangeNumber == 0)
 						{
-							string range = $"<={Context.ChangeNumber}";
+							logger.LogInformation("Finding last code change for CL {Number}...", Context.ChangeNumber);
 
-							// First, check the most recent changes for a limited subset of file types (this is a much cheaper query because of revcx optimization in the Perforce server.)
-							int optimisticRangeStart = Context.ChangeNumber - 1000;
-							if (optimisticRangeStart >= 1)
+							string[] codeRules = Utility.GetCodeFilter(Context.ProjectConfigFile);
+							await foreach (PerforceChangeDetails details in Utility.EnumerateChangeDetails(perforce, minChangeNumber: null, maxChangeNumber: Context.ChangeNumber, syncPaths, codeRules, cancellationToken))
 							{
-								string[] optimisticFileTypes = { ".cpp" };
-
-								string[] optimisticCodeFilter = optimisticFileTypes.SelectMany(x => syncPaths.Where(x => x.EndsWith("...", StringComparison.Ordinal)).Select(y => $"{y}{x}@{optimisticRangeStart},{Context.ChangeNumber}")).ToArray();
-								PerforceResponseList<ChangesRecord> optimisticChanges = await perforce.TryGetChangesAsync(ChangesOptions.None, 1, ChangeStatus.Submitted, optimisticCodeFilter, cancellationToken);
-
-								if (optimisticChanges.Succeeded && optimisticChanges.Count > 0)
+								if (details.ContainsCode)
 								{
-									int maxChange = optimisticChanges.Max(x => x.Data.Number);
-									range = $"{maxChange},{Context.ChangeNumber}";
+									codeChangeNumber = details.Number;
+									break;
 								}
 							}
-
-							// If no change found in recent changes, do the full and expensive check.
-							string[] codeFilter = PerforceUtils.CodeExtensions.SelectMany(x => syncPaths.Select(y => $"{y}{x}@{range}")).ToArray();
-							PerforceResponseList<ChangesRecord> codeChanges = await perforce.TryGetChangesAsync(ChangesOptions.None, 1, ChangeStatus.Submitted, codeFilter, cancellationToken);
-
-							if (!codeChanges.Succeeded)
-							{
-								return (WorkspaceUpdateResult.FailedToSync, $"Couldn't determine last code changelist before CL {Context.ChangeNumber}.");
-							}
-							if (codeChanges.Count == 0)
+							
+							if (codeChangeNumber == 0)
 							{
 								return (WorkspaceUpdateResult.FailedToSync, $"Could not find any code changes before CL {Context.ChangeNumber}.");
 							}
 
-							// Get the last code change number
-							codeChangeNumber = codeChanges.Max(x => x.Data.Number);
+							logger.LogInformation("Using code CL {Number}", codeChangeNumber);
 						}
 
 						// Set the version change
