UGS: Fix some DPI scaling issues with icons in the build list control.

diff --git a/Engine/Source/Programs/UnrealGameSync/UnrealGameSync/Controls/WorkspaceControl.cs b/Engine/Source/Programs/UnrealGameSync/UnrealGameSync/Controls/WorkspaceControl.cs
--- a/Engine/Source/Programs/UnrealGameSync/UnrealGameSync/Controls/WorkspaceControl.cs
+++ b/Engine/Source/Programs/UnrealGameSync/UnrealGameSync/Controls/WorkspaceControl.cs
@@ -2198,15 +2198,15 @@ namespace UnrealGameSync
 				float minX = 4 * dpiScaleX;
 				if ((summary != null && _eventMonitor.WasSyncedByCurrentUser(summary.ChangeNumber)) || (_workspace != null && _workspace.CurrentChangeNumber == change.Number))
 				{
-					e.Graphics.DrawImage(Properties.Resources.Icons, minX * dpiScaleX, iconY, PreviousSyncIcon, GraphicsUnit.Pixel);
+					e.Graphics.DrawImage(Properties.Resources.Icons, minX, iconY, PreviousSyncIcon, GraphicsUnit.Pixel);
 				}
 				else if (_workspaceSettings != null && _workspace != null && _workspace.State.AdditionalChangeNumbers.Contains(change.Number))
 				{
-					e.Graphics.DrawImage(Properties.Resources.Icons, minX * dpiScaleX, iconY, AdditionalSyncIcon, GraphicsUnit.Pixel);
+					e.Graphics.DrawImage(Properties.Resources.Icons, minX, iconY, AdditionalSyncIcon, GraphicsUnit.Pixel);
 				}
 				else if (allowSync && ((summary != null && summary.LastStarReview != null && summary.LastStarReview.Type == EventType.Starred) || _promotedChangeNumbers.Contains(change.Number)))
 				{
-					e.Graphics.DrawImage(Properties.Resources.Icons, minX * dpiScaleX, iconY, PromotedBuildIcon, GraphicsUnit.Pixel);
+					e.Graphics.DrawImage(Properties.Resources.Icons, minX, iconY, PromotedBuildIcon, GraphicsUnit.Pixel);
 				}
 				minX += PromotedBuildIcon.Width * dpiScaleX;
 
@@ -2329,12 +2329,12 @@ namespace UnrealGameSync
 				{
 					Tuple<string, float> progress = _workspace.CurrentProgress;
 
-					maxX -= CancelIcon.Width;
+					maxX -= CancelIcon.Width * dpiScaleX;
 					e.Graphics.DrawImage(Properties.Resources.Icons, maxX, iconY, CancelIcon, GraphicsUnit.Pixel);
 
 					if (!Splitter.IsLogVisible())
 					{
-						maxX -= InfoIcon.Width;
+						maxX -= InfoIcon.Width * dpiScaleX;
 						e.Graphics.DrawImage(Properties.Resources.Icons, maxX, iconY, InfoIcon, GraphicsUnit.Pixel);
 					}
 
