UGS: Fix an exception on shutdown.

diff --git a/Engine/Source/Programs/UnrealGameSync/UnrealGameSync/Controls/WorkspaceControl.cs b/Engine/Source/Programs/UnrealGameSync/UnrealGameSync/Controls/WorkspaceControl.cs
--- a/Engine/Source/Programs/UnrealGameSync/UnrealGameSync/Controls/WorkspaceControl.cs
+++ b/Engine/Source/Programs/UnrealGameSync/UnrealGameSync/Controls/WorkspaceControl.cs
@@ -1002,9 +1002,12 @@ namespace UnrealGameSync
 
 		void StateChangedCallback(ReadOnlyWorkspaceState state)
 		{
-			_changeNumberToLayoutInfo.Clear();
-			UpdateBuildList();
-			UpdateStatusPanel();
+			if (!_isDisposing)
+			{
+				_changeNumberToLayoutInfo.Clear();
+				UpdateBuildList();
+				UpdateStatusPanel();
+			}
 		}
 
 		void UpdateCompleteCallback(WorkspaceUpdateContext context, WorkspaceUpdateResult result, string resultMessage)
