UGS: Fix hang on shutdown if a workspace lock has not been released.

diff --git a/Engine/Source/Programs/UnrealGameSync/UnrealGameSyncShared/WorkspaceLock.cs b/Engine/Source/Programs/UnrealGameSync/UnrealGameSyncShared/WorkspaceLock.cs
--- a/Engine/Source/Programs/UnrealGameSync/UnrealGameSyncShared/WorkspaceLock.cs
+++ b/Engine/Source/Programs/UnrealGameSync/UnrealGameSyncShared/WorkspaceLock.cs
@@ -55,10 +55,12 @@ namespace UnrealGameSync
 			_mutex = new Mutex(false, $"{Prefix}.{_objectName}.mutex");
 
 			_acquireThread = new Thread(AcquireThread);
+			_acquireThread.Name = nameof(AcquireThread);
 			_acquireThread.IsBackground = true;
 			_acquireThread.Start();
 
 			_monitorThread = new Thread(MonitorThread);
+			_monitorThread.Name = nameof(MonitorThread);
 			_monitorThread.IsBackground = true;
 			_monitorThread.Start();
 		}
@@ -153,17 +155,25 @@ namespace UnrealGameSync
 				if (_acquireCount > 0)
 				{
 					_mutex.ReleaseMutex();
-					if (--_acquireCount == 0 && _lockedEvent != null)
+					if (--_acquireCount == 0)
 					{
-						_lockedEvent.Reset();
-						_lockedEvent.Dispose();
-						_lockedEvent = null;
+						ReleaseLockedEvent();
 					}
 				}
 			}
 			Task.Run(() => resultTcs.TrySetResult(true));
 		}
 
+		private void ReleaseLockedEvent()
+		{
+			if (_lockedEvent != null)
+			{
+				_lockedEvent.Reset();
+				_lockedEvent.Dispose();
+				_lockedEvent = null;
+			}
+		}
+
 		void AcquireThread()
 		{
 			for (; ; )
@@ -182,6 +192,8 @@ namespace UnrealGameSync
 			{
 				_mutex.ReleaseMutex();
 			}
+
+			ReleaseLockedEvent();
 		}
 
 		void MonitorThread()
