UGS: Hack to force WinForms to DPI scale tool strip menus correctly.

diff --git a/Engine/Source/Programs/UnrealGameSync/UnrealGameSync/Controls/WorkspaceControl.cs b/Engine/Source/Programs/UnrealGameSync/UnrealGameSync/Controls/WorkspaceControl.cs
--- a/Engine/Source/Programs/UnrealGameSync/UnrealGameSync/Controls/WorkspaceControl.cs
+++ b/Engine/Source/Programs/UnrealGameSync/UnrealGameSync/Controls/WorkspaceControl.cs
@@ -264,10 +264,15 @@ namespace UnrealGameSync
 		// Placeholder text that is in the control and cleared when the user starts editing.
 		const string AuthorFilterPlaceholderText = "<username>";
 
+		static readonly Bitmap s_hackToolStripMenuGutter = new Bitmap(1, 1);
+
 		public WorkspaceControl(IWorkspaceControlOwner owner, DirectoryReference appDataFolder, string? apiUrl, OpenProjectInfo openProjectInfo, IServiceProvider serviceProvider, UserSettings settings)
 		{
 			InitializeComponent();
 
+			// Winforms doesn't seem to DPI scale the left gutter for toolstrip menus correctly, but setting a custom image forces it to work correctly. (?)
+			OptionsContextMenu_ApplicationSettings.Image = s_hackToolStripMenuGutter;
+
 			_mainThreadSynchronizationContext = SynchronizationContext.Current!;
 
 			_owner = owner;
@@ -6099,6 +6104,12 @@ namespace UnrealGameSync
 
 		private void FilterButton_Click(object sender, EventArgs e)
 		{
+			// Hack for toolstrip dpi scaling
+			foreach (ToolStripMenuItem item in FilterContextMenu.Items.OfType<ToolStripMenuItem>())
+			{
+				item.Image = null;
+			}
+
 			FilterContextMenu_Default.Checked = !_settings.ShowAutomatedChanges && _projectSettings.FilterType == FilterType.None && _projectSettings.FilterBadges.Count == 0;
 
 			FilterContextMenu_Type.Checked = _projectSettings.FilterType != FilterType.None;
@@ -6146,6 +6157,16 @@ namespace UnrealGameSync
 			// Set checks if an author filter string is set
 			FilterContextMenu_Author.Checked = !String.IsNullOrEmpty(_authorFilterText) && _authorFilterText != AuthorFilterPlaceholderText;
 
+			// Hack for toolstrip dpi scaling
+			foreach (ToolStripMenuItem item in FilterContextMenu.Items.OfType<ToolStripMenuItem>())
+			{
+				if (!item.Checked)
+				{
+					item.Image = s_hackToolStripMenuGutter;
+					break;
+				}
+			}
+
 			FilterContextMenu.Show(FilterButton, new Point(0, FilterButton.Height));
 		}
 
