UGS: Fix exception on shutdown if a workspace sync is still cancelling.

diff --git a/Engine/Source/Programs/UnrealGameSync/UnrealGameSync/Workspace.cs b/Engine/Source/Programs/UnrealGameSync/UnrealGameSync/Workspace.cs
--- a/Engine/Source/Programs/UnrealGameSync/UnrealGameSync/Workspace.cs
+++ b/Engine/Source/Programs/UnrealGameSync/UnrealGameSync/Workspace.cs
@@ -87,8 +87,17 @@ namespace UnrealGameSync
 			CancelUpdate();
 			if (_prevUpdateTask != null)
 			{
+				_prevUpdateTask = _prevUpdateTask.ContinueWith(x => FinishDispose(), TaskScheduler.Default);
 				_asyncDisposer.Add(_prevUpdateTask);
 			}
+			else
+			{
+				FinishDispose();
+			}
+		}
+
+		void FinishDispose()
+		{
 			_stateWrapper.Dispose();
 			Lock.Dispose();
 		}
