UGS: Allow adding dependencies from tools to other tools.

diff --git a/Engine/Source/Programs/UnrealGameSync/UnrealGameSync/ToolUpdateMonitor.cs b/Engine/Source/Programs/UnrealGameSync/UnrealGameSync/ToolUpdateMonitor.cs
--- a/Engine/Source/Programs/UnrealGameSync/UnrealGameSync/ToolUpdateMonitor.cs
+++ b/Engine/Source/Programs/UnrealGameSync/UnrealGameSync/ToolUpdateMonitor.cs
@@ -43,8 +43,9 @@ namespace UnrealGameSync
 		public int ZipChange;
 		public string ConfigPath;
 		public int ConfigChange;
+		public HashSet<Guid> DependsOnToolIds;
 
-		public ToolDefinition(Guid id, string name, string description, string zipPath, int zipChange, string configPath, int configChange)
+		public ToolDefinition(Guid id, string name, string description, string zipPath, int zipChange, string configPath, int configChange, HashSet<Guid> dependsOnToolIds)
 		{
 			Id = id;
 			Name = name;
@@ -53,6 +54,7 @@ namespace UnrealGameSync
 			ZipChange = zipChange;
 			ConfigPath = configPath;
 			ConfigChange = configChange;
+			DependsOnToolIds = dependsOnToolIds;
 		}
 	}
 
@@ -195,9 +197,14 @@ namespace UnrealGameSync
 			}
 			Tools = newTools;
 
+			// Find all the tools which are enabled due to dependencies
+			HashSet<Guid> enabledTools = new HashSet<Guid>();
+			AddEnabledTools(Settings.EnabledTools, Tools, enabledTools);
+
+			// Configure each tool
 			foreach (ToolDefinition tool in Tools)
 			{
-				tool.Enabled = Settings.EnabledTools.Contains(tool.Id);
+				tool.Enabled = enabledTools.Contains(tool.Id);
 
 				if(!tool.Enabled)
 				{
@@ -245,6 +252,20 @@ namespace UnrealGameSync
 			_synchronizationContext.Post(_ => OnChange?.Invoke(), null);
 		}
 
+		static void AddEnabledTools(HashSet<Guid> inputToolIds, List<ToolDefinition> tools, HashSet<Guid> enabledTools)
+		{
+			if (inputToolIds.Count > 0)
+			{
+				foreach (ToolDefinition tool in tools)
+				{
+					if (inputToolIds.Contains(tool.Id) && enabledTools.Add(tool.Id))
+					{
+						AddEnabledTools(tool.DependsOnToolIds, tools, enabledTools);
+					}
+				}
+			}
+		}
+
 		async Task<ToolDefinition?> ReadToolDefinitionAsync(IPerforceConnection perforce, string depotPath, int change, CancellationToken cancellationToken)
 		{
 			PerforceResponse<PrintRecord<string[]>> response = await perforce.TryPrintLinesAsync($"{depotPath}@{change}", cancellationToken);
@@ -272,7 +293,16 @@ namespace UnrealGameSync
 			string toolConfigPath = depotPath;
 			int toolConfigChange = change;
 
-			ToolDefinition tool = new ToolDefinition(toolId, toolName, toolDescription, toolZipPath, toolZipChange, toolConfigPath, toolConfigChange);
+			HashSet<Guid> dependsOnToolIds = new HashSet<Guid>();
+			foreach (string line in configFile.GetValues("Settings.DependsOnTool", Array.Empty<string>()))
+			{
+				if (Guid.TryParse(line.Trim(), out Guid requiredToolId))
+				{
+					dependsOnToolIds.Add(requiredToolId);
+				}
+			}
+
+			ToolDefinition tool = new ToolDefinition(toolId, toolName, toolDescription, toolZipPath, toolZipChange, toolConfigPath, toolConfigChange, dependsOnToolIds);
 
 			string? installCommand = configFile.GetValue("Settings.InstallCommand", null);
 			if (!String.IsNullOrEmpty(installCommand))
