UGS: Allow specifying a list of virtual streams in the project config file which will be used for versioning metadata.

By default, UGS ignores virtual streams and uses the parent stream name, under the assumption that the virtual stream just filters stuff out.

diff --git a/Engine/Source/Programs/UnrealGameSync/UnrealGameSyncShared/WorkspaceUpdate.cs b/Engine/Source/Programs/UnrealGameSync/UnrealGameSyncShared/WorkspaceUpdate.cs
--- a/Engine/Source/Programs/UnrealGameSync/UnrealGameSyncShared/WorkspaceUpdate.cs
+++ b/Engine/Source/Programs/UnrealGameSync/UnrealGameSyncShared/WorkspaceUpdate.cs
@@ -915,7 +915,8 @@ namespace UnrealGameSync
 						if (branchOrStreamName != null)
 						{
 							// If it's a virtual stream, take the concrete parent stream instead
-							for (; ; )
+							HashSet<string> versionStreams = new HashSet<string>(Context.ProjectConfigFile.GetValues("Perforce.VersionStreams", Array.Empty<string>()), StringComparer.OrdinalIgnoreCase);
+							while (!versionStreams.Contains(branchOrStreamName))
 							{
 								StreamRecord streamSpec = await perforce.GetStreamAsync(branchOrStreamName, false, cancellationToken);
 								if (streamSpec.Type != "virtual" || streamSpec.Parent == "none" || streamSpec.Parent == null)
