UGS: Fix DPI scaling on build list menu.

diff --git a/Engine/Source/Programs/UnrealGameSync/UnrealGameSync/Controls/WorkspaceControl.cs b/Engine/Source/Programs/UnrealGameSync/UnrealGameSync/Controls/WorkspaceControl.cs
--- a/Engine/Source/Programs/UnrealGameSync/UnrealGameSync/Controls/WorkspaceControl.cs
+++ b/Engine/Source/Programs/UnrealGameSync/UnrealGameSync/Controls/WorkspaceControl.cs
@@ -4268,6 +4268,9 @@ namespace UnrealGameSync
 		{
 			if (_contextMenuChange != null)
 			{
+				BuildListContextMenu_Sync.Image = s_hackToolStripMenuGutter;
+				BuildListContextMenu_Cancel.Image = s_hackToolStripMenuGutter;
+
 				BuildListContextMenu_WithdrawReview.Visible = (_eventMonitor.GetReviewByCurrentUser(_contextMenuChange.Number) != null);
 				BuildListContextMenu_StartInvestigating.Visible = !_eventMonitor.IsUnderInvestigationByCurrentUser(_contextMenuChange.Number);
 				BuildListContextMenu_FinishInvestigating.Visible = _eventMonitor.IsUnderInvestigation(_contextMenuChange.Number);
