UGS: Fix DPI scaling on "More Tools" and "Sync To" context menus.

diff --git a/Engine/Source/Programs/UnrealGameSync/UnrealGameSync/Controls/WorkspaceControl.cs b/Engine/Source/Programs/UnrealGameSync/UnrealGameSync/Controls/WorkspaceControl.cs
--- a/Engine/Source/Programs/UnrealGameSync/UnrealGameSync/Controls/WorkspaceControl.cs
+++ b/Engine/Source/Programs/UnrealGameSync/UnrealGameSync/Controls/WorkspaceControl.cs
@@ -4053,6 +4053,7 @@ namespace UnrealGameSync
 
 		private void ShowActionsMenu(Rectangle bounds)
 		{
+			MoreToolsContextMenu_CleanWorkspace.Image = s_hackToolStripMenuGutter;
 			MoreToolsContextMenu.Show(StatusPanel, new Point(bounds.Left, bounds.Bottom), ToolStripDropDownDirection.BelowRight);
 		}
 
@@ -5662,6 +5663,8 @@ namespace UnrealGameSync
 				toolStripSeparator8
 			});
 
+			SyncContexMenu_EnterChangelist.Image = s_hackToolStripMenuGutter;
+
 			foreach (LatestChangeType changeType in GetCustomLatestChangeTypes())
 			{
 				System.Windows.Forms.ToolStripMenuItem menuItem = new System.Windows.Forms.ToolStripMenuItem();
