UGS: Show custom tool dependencies which are required but not explicitly enabled as checked but greyed out.

diff --git a/Engine/Source/Programs/UnrealGameSync/UnrealGameSync/Forms/ApplicationSettingsWindow.cs b/Engine/Source/Programs/UnrealGameSync/UnrealGameSync/Forms/ApplicationSettingsWindow.cs
--- a/Engine/Source/Programs/UnrealGameSync/UnrealGameSync/Forms/ApplicationSettingsWindow.cs
+++ b/Engine/Source/Programs/UnrealGameSync/UnrealGameSync/Forms/ApplicationSettingsWindow.cs
@@ -21,7 +21,7 @@ namespace UnrealGameSync
 				string checkFilePath = String.Format("{0}/Release/UnrealGameSync.exe", depotPath);
 
 				List<FStatRecord> fileRecords = await perforce.FStatAsync(checkFilePath, cancellationToken).ToListAsync(cancellationToken);
-				if(fileRecords.Count == 0)
+				if (fileRecords.Count == 0)
 				{
 					throw new UserErrorException($"Unable to find {checkFilePath}");
 				}
@@ -46,11 +46,36 @@ namespace UnrealGameSync
 
 		class ToolItem
 		{
+			public Guid Id => Definition.Id;
+
+			public int Index { get; }
 			public ToolDefinition Definition { get; }
+			public List<ToolItem> RequiresTools { get; } = new List<ToolItem>();
+
+			public bool Enabled { get; set; }
+			public int DependencyRefCount { get; set; }
 
-			public ToolItem(ToolDefinition definition)
+			public ToolItem(int index, ToolDefinition definition, bool enabled)
 			{
+				Index = index;
 				Definition = definition;
+				Enabled = enabled;
+			}
+
+			public CheckState GetCheckState()
+			{
+				if (Enabled)
+				{
+					return CheckState.Checked;
+				}
+				else if (DependencyRefCount > 0)
+				{
+					return CheckState.Indeterminate;
+				}
+				else
+				{
+					return CheckState.Unchecked;
+				}
 			}
 
 			public override string ToString()
@@ -78,7 +103,7 @@ namespace UnrealGameSync
 
 			AutomaticallyRunAtStartupCheckBox.Checked = IsAutomaticallyRunAtStartup();
 			KeepInTrayCheckBox.Checked = settings.KeepInTray;
-					
+
 			ServerTextBox.Text = _initialServerAndPort;
 			ServerTextBox.Select(ServerTextBox.TextLength, 0);
 			ServerTextBox.CueBanner = $"Default ({defaultPerforceSettings.ServerAndPort})";
@@ -95,7 +120,7 @@ namespace UnrealGameSync
 
 			UsePreviewBuildCheckBox.Checked = preview;
 
-			if(_initialAutomationPortNumber > 0)
+			if (_initialAutomationPortNumber > 0)
 			{
 				EnableAutomationCheckBox.Checked = true;
 				AutomationPortTextBox.Enabled = true;
@@ -108,7 +133,7 @@ namespace UnrealGameSync
 				AutomationPortTextBox.Text = AutomationServer.DefaultPortNumber.ToString();
 			}
 
-			if(_initialProtocolHandlerState == ProtocolHandlerState.Installed)
+			if (_initialProtocolHandlerState == ProtocolHandlerState.Installed)
 			{
 				EnableProtocolHandlerCheckBox.CheckState = CheckState.Checked;
 			}
@@ -121,17 +146,53 @@ namespace UnrealGameSync
 				EnableProtocolHandlerCheckBox.CheckState = CheckState.Indeterminate;
 			}
 
-			List<ToolDefinition> tools = toolUpdateMonitor.Tools;
-			foreach (ToolDefinition tool in tools)
+			List<ToolItem> toolItems = new List<ToolItem>();
+			Dictionary<Guid, ToolItem> idToToolItem = new Dictionary<Guid, ToolItem>();
+			foreach (ToolDefinition tool in toolUpdateMonitor.Tools.OrderBy(x => x.ToString()))
+			{
+				ToolItem toolItem = new ToolItem(toolItems.Count, tool, settings.EnabledTools.Contains(tool.Id));
+				idToToolItem[tool.Id] = toolItem;
+				toolItems.Add(toolItem);
+			}
+
+			HashSet<ToolItem> dependsOnToolItems = new HashSet<ToolItem>();
+			foreach (ToolItem toolItem in toolItems)
 			{
-				CustomToolsListBox.Items.Add(new ToolItem(tool), settings.EnabledTools.Contains(tool.Id));
+				dependsOnToolItems.Clear();
+				FindDependencies(toolItem, dependsOnToolItems, idToToolItem);
+				toolItem.RequiresTools.AddRange(dependsOnToolItems);
+
+				if (toolItem.Enabled)
+				{
+					foreach (ToolItem requiredTool in toolItem.RequiresTools)
+					{
+						requiredTool.DependencyRefCount++;
+					}
+				}
+			}
+
+			foreach (ToolItem toolItem in toolItems)
+			{
+				CustomToolsListBox.Items.Add(toolItem, toolItem.GetCheckState());
+			}
+		}
+
+		void FindDependencies(ToolItem toolItem, HashSet<ToolItem> dependsOnToolItems, Dictionary<Guid, ToolItem> idToToolItem)
+		{
+			foreach (Guid dependsOnToolId in toolItem.Definition.DependsOnToolIds)
+			{
+				ToolItem? dependsOnToolItem;
+				if (idToToolItem.TryGetValue(dependsOnToolId, out dependsOnToolItem) && dependsOnToolItems.Add(dependsOnToolItem))
+				{
+					FindDependencies(dependsOnToolItem, dependsOnToolItems, idToToolItem);
+				}
 			}
 		}
 
 		public static bool? ShowModal(IWin32Window owner, IPerforceSettings defaultPerforceSettings, bool preview, string originalExecutableFileName, UserSettings settings, ToolUpdateMonitor toolUpdateMonitor, ILogger<ApplicationSettingsWindow> logger)
 		{
 			using ApplicationSettingsWindow applicationSettings = new ApplicationSettingsWindow(defaultPerforceSettings, preview, originalExecutableFileName, settings, toolUpdateMonitor, logger);
-			if(applicationSettings.ShowDialog(owner) == DialogResult.OK)
+			if (applicationSettings.ShowDialog(owner) == DialogResult.OK)
 			{
 				return applicationSettings._restartPreview;
 			}
@@ -151,19 +212,19 @@ namespace UnrealGameSync
 		{
 			// Update the settings
 			string? serverAndPort = ServerTextBox.Text.Trim();
-			if(serverAndPort.Length == 0)
+			if (serverAndPort.Length == 0)
 			{
 				serverAndPort = null;
 			}
 
 			string? userName = UserNameTextBox.Text.Trim();
-			if(userName.Length == 0)
+			if (userName.Length == 0)
 			{
 				userName = null;
 			}
 
 			string? depotPath = DepotPathTextBox.Text.Trim();
-			if(depotPath.Length == 0 || depotPath == DeploymentSettings.Instance.DefaultDepotPath)
+			if (depotPath.Length == 0 || depotPath == DeploymentSettings.Instance.DefaultDepotPath)
 			{
 				depotPath = null;
 			}
@@ -171,15 +232,15 @@ namespace UnrealGameSync
 			bool preview = UsePreviewBuildCheckBox.Checked;
 
 			int automationPortNumber;
-			if(!EnableAutomationCheckBox.Checked || !Int32.TryParse(AutomationPortTextBox.Text, out automationPortNumber))
+			if (!EnableAutomationCheckBox.Checked || !Int32.TryParse(AutomationPortTextBox.Text, out automationPortNumber))
 			{
 				automationPortNumber = -1;
 			}
-			
-			if(serverAndPort != _initialServerAndPort || userName != _initialUserName || depotPath != _initialDepotPath || preview != _initialPreview || automationPortNumber != _initialAutomationPortNumber)
+
+			if (serverAndPort != _initialServerAndPort || userName != _initialUserName || depotPath != _initialDepotPath || preview != _initialPreview || automationPortNumber != _initialAutomationPortNumber)
 			{
 				// Try to log in to the new server, and check the application is there
-				if(serverAndPort != _initialServerAndPort || userName != _initialUserName || depotPath != _initialDepotPath)
+				if (serverAndPort != _initialServerAndPort || userName != _initialUserName || depotPath != _initialDepotPath)
 				{
 					PerforceSettings settings = Utility.OverridePerforceSettings(_defaultPerforceSettings, serverAndPort, userName);
 
@@ -194,7 +255,7 @@ namespace UnrealGameSync
 					}
 				}
 
-				if(MessageBox.Show("UnrealGameSync must be restarted to apply these settings.\n\nWould you like to restart now?", "Restart Required", MessageBoxButtons.OKCancel) != DialogResult.OK)
+				if (MessageBox.Show("UnrealGameSync must be restarted to apply these settings.\n\nWould you like to restart now?", "Restart Required", MessageBoxButtons.OKCancel) != DialogResult.OK)
 				{
 					return;
 				}
@@ -217,15 +278,15 @@ namespace UnrealGameSync
 			if (_settings.KeepInTray != KeepInTrayCheckBox.Checked || _settings.SyncOptions.NumThreads != ParallelSyncThreadsSpinner.Value)
 			{
 				int numThreads = (int)ParallelSyncThreadsSpinner.Value;
-				_settings.SyncOptions.NumThreads = (numThreads != PerforceSyncOptions.DefaultNumThreads)? (int?)numThreads : null;
+				_settings.SyncOptions.NumThreads = (numThreads != PerforceSyncOptions.DefaultNumThreads) ? (int?)numThreads : null;
 				_settings.KeepInTray = KeepInTrayCheckBox.Checked;
 				_settings.Save(_logger);
 			}
 
 			List<Guid> newEnabledTools = new List<Guid>();
-			foreach (ToolItem? item in CustomToolsListBox.CheckedItems)
+			foreach (ToolItem? item in CustomToolsListBox.Items)
 			{
-				if (item != null)
+				if (item != null && item.Enabled)
 				{
 					newEnabledTools.Add(item.Definition.Id);
 				}
@@ -273,5 +334,43 @@ namespace UnrealGameSync
 			using PerforceSyncSettingsWindow window = new PerforceSyncSettingsWindow(_settings, _logger);
 			window.ShowDialog();
 		}
+
+		bool _recursiveItemCheck = false;
+
+		private void CustomToolsListBox_ItemCheck(object sender, ItemCheckEventArgs e)
+		{
+			if (!_recursiveItemCheck)
+			{
+				_recursiveItemCheck = true;
+				ToolItem toolItem = (ToolItem)CustomToolsListBox.Items[e.Index];
+
+				bool newEnabled = (e.CurrentValue == CheckState.Unchecked || e.CurrentValue == CheckState.Indeterminate);
+				if (newEnabled != toolItem.Enabled)
+				{
+					toolItem.Enabled = newEnabled;
+
+					foreach (ToolItem requiredTool in toolItem.RequiresTools)
+					{
+						if (toolItem.Enabled)
+						{
+							requiredTool.DependencyRefCount++;
+						}
+						else
+						{
+							requiredTool.DependencyRefCount--;
+						}
+
+						CheckState newCheckState = requiredTool.GetCheckState();
+						if (newCheckState != CustomToolsListBox.GetItemCheckState(requiredTool.Index))
+						{
+							CustomToolsListBox.SetItemCheckState(requiredTool.Index, newCheckState);
+						}
+					}
+				}
+
+				e.NewValue = toolItem.GetCheckState();
+				_recursiveItemCheck = false;
+			}
+		}
 	}
 }
