UGS: Improve error message on sync failures/retries.

diff --git a/Engine/Source/Programs/UnrealGameSync/UnrealGameSyncShared/WorkspaceUpdate.cs b/Engine/Source/Programs/UnrealGameSync/UnrealGameSyncShared/WorkspaceUpdate.cs
--- a/Engine/Source/Programs/UnrealGameSync/UnrealGameSyncShared/WorkspaceUpdate.cs
+++ b/Engine/Source/Programs/UnrealGameSync/UnrealGameSyncShared/WorkspaceUpdate.cs
@@ -1671,17 +1671,17 @@ namespace UnrealGameSync
 				WorkspaceUpdateResult result = WorkspaceUpdateResult.FailedToSync;
 				string statusMessage = "";
 
-				int retries = context.PerforceSyncOptions!.NumSyncErrorRetries ?? PerforceSyncOptions.DefaultNumSyncErrorRetries;
-
-				while (retries >= 0 && WorkspaceUpdateResult.FailedToSync == result)
+				int maxRetries = context.PerforceSyncOptions?.NumSyncErrorRetries ?? PerforceSyncOptions.DefaultNumSyncErrorRetries;
+				for (int attempt = 0; ;attempt++)
 				{
 					// Sync the files
-					(result, statusMessage) = await StaticSyncFileRevisions(perforce, context, syncCommands, record => syncOutput(record, threadLog), cancellationToken);
-
-					if (WorkspaceUpdateResult.FailedToSync == result && --retries >= 0)
+					string? errorMessage;
+					(result, statusMessage, errorMessage) = await StaticSyncFileRevisions(perforce, context, syncCommands, record => syncOutput(record, threadLog), cancellationToken);
+					if (result != WorkspaceUpdateResult.FailedToSync || attempt >= maxRetries)
 					{
-						threadLog.LogWarning("Sync Errors occurred.  Retrying: Remaining retries {Count}", retries);
+						break;
 					}
+					threadLog.LogWarning("Sync error ({Message}); retrying... ({Count}/{MaxCount})", errorMessage ?? "unknown", attempt + 1, maxRetries);
 				}
 
 				// If it failed, try to set it on the state if nothing else has failed first
@@ -1700,7 +1700,7 @@ namespace UnrealGameSync
 			}
 		}
 
-		static async Task<(WorkspaceUpdateResult, string)> StaticSyncFileRevisions(IPerforceConnection perforce, WorkspaceUpdateContext context, List<string> syncCommands, Action<SyncRecord> syncOutput, CancellationToken cancellationToken)
+		static async Task<(WorkspaceUpdateResult, string, string?)> StaticSyncFileRevisions(IPerforceConnection perforce, WorkspaceUpdateContext context, List<string> syncCommands, Action<SyncRecord> syncOutput, CancellationToken cancellationToken)
 		{
 			// Sync them all
 			List<PerforceResponse<SyncRecord>> responses = await perforce.TrySyncAsync(SyncOptions.None, -1, syncCommands, cancellationToken).ToListAsync(cancellationToken);
@@ -1723,9 +1723,7 @@ namespace UnrealGameSync
 				}
 				else
 				{
-					return (WorkspaceUpdateResult.FailedToSync, "Aborted sync due to errors.. Currently retries on sync error is set at " +
-						((null != context.PerforceSyncOptions) ? context.PerforceSyncOptions!.NumSyncErrorRetries : 0).ToString() +
-						" in Options->Application Settings...->Advanced.  You might want to set it higher if you are on a bad connection.");
+					return (WorkspaceUpdateResult.FailedToSync, $"Aborted sync due to error ({response}). If you are on an unreliable connection, you may wish to increase the number of retries from Options > Application Settings... > Advanced.", response.ToString());
 				}
 			}
 
@@ -1750,7 +1748,7 @@ namespace UnrealGameSync
 					}
 					if (numNewFilesToClobber > 0)
 					{
-						return (WorkspaceUpdateResult.FilesToClobber, $"Cancelled sync after checking files to clobber ({numNewFilesToClobber} new files).");
+						return (WorkspaceUpdateResult.FilesToClobber, $"Cancelled sync after checking files to clobber ({numNewFilesToClobber} new files).", null);
 					}
 				}
 				foreach (string tamperedFile in tamperedFiles)
@@ -1761,14 +1759,14 @@ namespace UnrealGameSync
 						List<PerforceResponse<SyncRecord>> response = await perforce.TrySyncAsync(SyncOptions.Force, -1, tamperedFile, cancellationToken).ToListAsync(cancellationToken);
 						if (!response.Succeeded())
 						{
-							return (WorkspaceUpdateResult.FailedToSync, $"Couldn't sync {tamperedFile}.");
+							return (WorkspaceUpdateResult.FailedToSync, $"Couldn't sync {tamperedFile}.", response.ToString());
 						}
 					}
 				}
 			}
 
 			// All succeeded
-			return (WorkspaceUpdateResult.Success, "Succeeded.");
+			return (WorkspaceUpdateResult.Success, "Succeeded.", null);
 		}
 
 		public static async Task<ConfigFile> ReadProjectConfigFile(DirectoryReference localRootPath, FileReference selectedLocalFileName, ILogger logger)
