UGS: Fix UBT instances not being killed if a UGS sync is aborted.

Specifying a null ManagedProcessGroup when constructing a ManagedProcess results in not using a job object.

diff --git a/Engine/Source/Programs/UnrealGameSync/UnrealGameSyncShared/Utility.cs b/Engine/Source/Programs/UnrealGameSync/UnrealGameSyncShared/Utility.cs
--- a/Engine/Source/Programs/UnrealGameSync/UnrealGameSyncShared/Utility.cs
+++ b/Engine/Source/Programs/UnrealGameSync/UnrealGameSyncShared/Utility.cs
@@ -489,7 +489,8 @@ namespace UnrealGameSync
 
 		public static async Task<int> ExecuteProcessAsync(string fileName, string? workingDir, string commandLine, Action<string> outputLine, CancellationToken cancellationToken)
 		{
-			using (ManagedProcess newProcess = new ManagedProcess(null, fileName, commandLine, workingDir, null, null, ProcessPriorityClass.Normal))
+			using (ManagedProcessGroup newProcessGroup = new ManagedProcessGroup())
+			using (ManagedProcess newProcess = new ManagedProcess(newProcessGroup, fileName, commandLine, workingDir, null, null, ProcessPriorityClass.Normal))
 			{
 				for (; ; )
 				{
