Add an option to generate a P4CONFIG to UGS when opening a project.

#rnx

*** UGS
- Add a checkbox "GenerateP4ConfigCheckbox" to the OpenProjectWindow, which is visible if the host environment has P4CONFIG set. By default the box will be checked.
- Opening a project with this new option checked will result in UGS attempting to create a p4config file at the root of the project being opened, unless the user already has a file there in which case UGS will do nothing.
- When creating the file the three values we are interested in setting are P4PORT, P4USER and P4CLIENT. With these values we are able to make reasonable assumptions in our tools when it comes to using perforce.
- If the user has enabled P4CONFIG then it is assumed that they know what they are doing with the file. We do not validate that the existing file has sensible settings or offer to replace it.

diff --git a/Engine/Source/Programs/UnrealGameSync/UnrealGameSync/Forms/MainWindow.cs b/Engine/Source/Programs/UnrealGameSync/UnrealGameSync/Forms/MainWindow.cs
--- a/Engine/Source/Programs/UnrealGameSync/UnrealGameSync/Forms/MainWindow.cs
+++ b/Engine/Source/Programs/UnrealGameSync/UnrealGameSync/Forms/MainWindow.cs
@@ -1229,7 +1229,7 @@ namespace UnrealGameSync
 
 			PerforceSettings perforceSettings = Utility.OverridePerforceSettings(_defaultPerforceSettings, project.ServerAndPort, project.UserName);
 
-			ModalTask<OpenProjectInfo>? settingsTask = PerforceModalTask.Execute(this, "Opening Project", "Opening project, please wait...", perforceSettings, (p, c) => OpenProjectWindow.DetectSettingsAsync(p, project, _settings, _oidcTokenManager, projectLogger, c), projectLogger, taskFlags);
+			ModalTask<OpenProjectInfo>? settingsTask = PerforceModalTask.Execute(this, "Opening Project", "Opening project, please wait...", perforceSettings, (p, c) => OpenProjectWindow.DetectSettingsAsync(p, project, _settings, _oidcTokenManager, false, projectLogger, c), projectLogger, taskFlags);
 			if (settingsTask == null || settingsTask.Failed)
 			{
 				if (settingsTask != null)
diff --git a/Engine/Source/Programs/UnrealGameSync/UnrealGameSync/Forms/OpenProjectWindow.cs b/Engine/Source/Programs/UnrealGameSync/UnrealGameSync/Forms/OpenProjectWindow.cs
--- a/Engine/Source/Programs/UnrealGameSync/UnrealGameSync/Forms/OpenProjectWindow.cs
+++ b/Engine/Source/Programs/UnrealGameSync/UnrealGameSync/Forms/OpenProjectWindow.cs
@@ -51,31 +51,31 @@ namespace UnrealGameSync
 			}
 			else
 			{
-				if(!String.IsNullOrWhiteSpace(project.ServerAndPort))
+				if (!String.IsNullOrWhiteSpace(project.ServerAndPort))
 				{
 					_serverAndPortOverride = project.ServerAndPort;
 				}
-				if(!String.IsNullOrWhiteSpace(project.UserName))
+				if (!String.IsNullOrWhiteSpace(project.UserName))
 				{
 					_userNameOverride = project.UserName;
 				}
 
-				if(project.ClientPath != null && project.ClientPath.StartsWith("//", StringComparison.Ordinal))
+				if (project.ClientPath != null && project.ClientPath.StartsWith("//", StringComparison.Ordinal))
 				{
 					int slashIdx = project.ClientPath.IndexOf('/', 2);
-					if(slashIdx != -1)
+					if (slashIdx != -1)
 					{
 						WorkspaceNameTextBox.Text = project.ClientPath.Substring(2, slashIdx - 2);
 						WorkspacePathTextBox.Text = project.ClientPath.Substring(slashIdx);
 					}
 				}
 
-				if(project.LocalPath != null)
+				if (project.LocalPath != null)
 				{
 					LocalFileTextBox.Text = project.LocalPath;
 				}
 
-				if(project.Type == UserSelectedProjectType.Client)
+				if (project.Type == UserSelectedProjectType.Client)
 				{
 					WorkspaceRadioBtn.Checked = true;
 				}
@@ -88,6 +88,7 @@ namespace UnrealGameSync
 			UpdateEnabledControls();
 			UpdateServerLabel();
 			UpdateWorkspacePathBrowseButton();
+			UpdateP4ConfigCheckBox();
 			UpdateOkButton();
 		}
 
@@ -96,7 +97,7 @@ namespace UnrealGameSync
 		public static OpenProjectInfo? ShowModal(IWin32Window owner, UserSelectedProjectSettings? project, UserSettings settings, IPerforceSettings defaultPerforceSettings, IServiceProvider serviceProvider, ILogger logger)
 		{
 			using OpenProjectWindow window = new OpenProjectWindow(project, settings, defaultPerforceSettings, serviceProvider, logger);
-			if(window.ShowDialog(owner) == DialogResult.OK)
+			if (window.ShowDialog(owner) == DialogResult.OK)
 			{
 				return window._openProjectInfo;
 			}
@@ -108,7 +109,7 @@ namespace UnrealGameSync
 
 		private void UpdateEnabledControls()
 		{
-			Color workspaceTextColor = WorkspaceRadioBtn.Checked? SystemColors.ControlText : SystemColors.GrayText;
+			Color workspaceTextColor = WorkspaceRadioBtn.Checked ? SystemColors.ControlText : SystemColors.GrayText;
 			WorkspaceNameLabel.ForeColor = workspaceTextColor;
 			WorkspaceNameTextBox.ForeColor = workspaceTextColor;
 			WorkspaceNameNewBtn.ForeColor = workspaceTextColor;
@@ -117,7 +118,7 @@ namespace UnrealGameSync
 			WorkspacePathTextBox.ForeColor = workspaceTextColor;
 			WorkspacePathBrowseBtn.ForeColor = workspaceTextColor;
 
-			Color localFileTextColor = LocalFileRadioBtn.Checked? SystemColors.ControlText : SystemColors.GrayText;
+			Color localFileTextColor = LocalFileRadioBtn.Checked ? SystemColors.ControlText : SystemColors.GrayText;
 			LocalFileLabel.ForeColor = localFileTextColor;
 			LocalFileTextBox.ForeColor = localFileTextColor;
 			LocalFileBrowseBtn.ForeColor = localFileTextColor;
@@ -127,14 +128,14 @@ namespace UnrealGameSync
 
 		public static string GetServerLabelText(IPerforceSettings defaultSettings, string? serverAndPort, string? userName)
 		{
-			if(serverAndPort == null && userName == null)
+			if (serverAndPort == null && userName == null)
 			{
 				return String.Format("Using default connection settings (user '{0}' on server '{1}').", defaultSettings.UserName, defaultSettings.ServerAndPort);
 			}
 			else
 			{
 				StringBuilder text = new StringBuilder("Connecting as ");
-				if(userName == null)
+				if (userName == null)
 				{
 					text.Append("default user");
 				}
@@ -143,7 +144,7 @@ namespace UnrealGameSync
 					text.AppendFormat("user '{0}'", userName);
 				}
 				text.Append(" on ");
-				if(serverAndPort == null)
+				if (serverAndPort == null)
 				{
 					text.Append("default server.");
 				}
@@ -165,6 +166,14 @@ namespace UnrealGameSync
 			WorkspacePathBrowseBtn.Enabled = TryGetWorkspaceName(out _);
 		}
 
+		private void UpdateP4ConfigCheckBox()
+		{
+			bool isP4ConfigSet = !String.IsNullOrEmpty(PerforceEnvironment.Default.GetValue("P4CONFIG"));
+
+			GenerateP4ConfigCheckbox.Visible = isP4ConfigSet;
+			GenerateP4ConfigCheckbox.Checked = isP4ConfigSet;
+		}
+
 		private void UpdateOkButton()
 		{
 			OkBtn.Enabled = WorkspaceRadioBtn.Checked? TryGetClientPath(out _) : TryGetLocalPath(out _);
@@ -290,7 +299,7 @@ namespace UnrealGameSync
 
 				PerforceSettings newPerforceSettings = Utility.OverridePerforceSettings(Perforce, selectedProject.ServerAndPort, selectedProject.UserName);
 
-				ModalTask<OpenProjectInfo>? newOpenProjectInfo = PerforceModalTask.Execute(this, "Opening project", "Opening project, please wait...", newPerforceSettings, (x, y) => DetectSettingsAsync(x, selectedProject, _settings, oidcTokenManager, logger, y), logger);
+				ModalTask<OpenProjectInfo>? newOpenProjectInfo = PerforceModalTask.Execute(this, "Opening project", "Opening project, please wait...", newPerforceSettings, (x, y) => DetectSettingsAsync(x, selectedProject, _settings, oidcTokenManager, GenerateP4ConfigCheckbox.Checked, logger, y), logger);
 				if (newOpenProjectInfo != null && newOpenProjectInfo.Succeeded)
 				{
 					_openProjectInfo = newOpenProjectInfo.Result;
@@ -300,9 +309,9 @@ namespace UnrealGameSync
 			}
 		}
 
-		public static async Task<OpenProjectInfo> DetectSettingsAsync(IPerforceConnection perforce, UserSelectedProjectSettings selectedProject, UserSettings userSettings, OidcTokenManager oidcTokenManager, ILogger<OpenProjectInfo> logger, CancellationToken cancellationToken)
+		public static async Task<OpenProjectInfo> DetectSettingsAsync(IPerforceConnection perforce, UserSelectedProjectSettings selectedProject, UserSettings userSettings, OidcTokenManager oidcTokenManager, bool GenerateP4Config, ILogger<OpenProjectInfo> logger, CancellationToken cancellationToken)
 		{
-			OpenProjectInfo settings = await OpenProjectInfo.CreateAsync(perforce, selectedProject, userSettings, oidcTokenManager, logger, cancellationToken);
+			OpenProjectInfo settings = await OpenProjectInfo.CreateAsync(perforce, selectedProject, userSettings, oidcTokenManager, GenerateP4Config, logger, cancellationToken);
 			if (s_onDetectProjectSettings != null)
 			{
 				string? message;
diff --git a/Engine/Source/Programs/UnrealGameSync/UnrealGameSync/ProgramApplicationContext.cs b/Engine/Source/Programs/UnrealGameSync/UnrealGameSync/ProgramApplicationContext.cs
--- a/Engine/Source/Programs/UnrealGameSync/UnrealGameSync/ProgramApplicationContext.cs
+++ b/Engine/Source/Programs/UnrealGameSync/UnrealGameSync/ProgramApplicationContext.cs
@@ -198,7 +198,7 @@ namespace UnrealGameSync
 				{
 					ILogger<OpenProjectInfo> logger = serviceProvider.GetRequiredService<ILogger<OpenProjectInfo>>();
 					OidcTokenManager oidcTokenManager = serviceProvider.GetRequiredService<OidcTokenManager>();
-					Task<OpenProjectInfo> startupTask = Task.Run(() => OpenProjectInfo.CreateAsync(defaultPerforceSettings, projectSettings, _settings, oidcTokenManager, logger, _startupCancellationSource.Token), _startupCancellationSource.Token);
+					Task<OpenProjectInfo> startupTask = Task.Run(() => OpenProjectInfo.CreateAsync(defaultPerforceSettings, projectSettings, _settings, oidcTokenManager, false, logger, _startupCancellationSource.Token), _startupCancellationSource.Token);
 					startupTasks.Add((projectSettings, new ModalTask<OpenProjectInfo>(startupTask)));
 				}
 			}
diff --git a/Engine/Source/Programs/UnrealGameSync/UnrealGameSyncShared/OpenProjectInfo.cs b/Engine/Source/Programs/UnrealGameSync/UnrealGameSyncShared/OpenProjectInfo.cs
--- a/Engine/Source/Programs/UnrealGameSync/UnrealGameSyncShared/OpenProjectInfo.cs
+++ b/Engine/Source/Programs/UnrealGameSync/UnrealGameSyncShared/OpenProjectInfo.cs
@@ -26,8 +26,9 @@ namespace UnrealGameSync
 		public IReadOnlyList<string>? WorkspaceProjectStreamFilter { get; }
 		public List<KeyValuePair<FileReference, DateTime>> LocalConfigFiles { get; }
 		public OidcTokenClient? OidcTokenClient { get; }
+		public bool GenerateP4Config { get; }
 
-		public OpenProjectInfo(UserSelectedProjectSettings selectedProject, IPerforceSettings perforceSettings, ProjectInfo projectInfo, UserWorkspaceSettings workspaceSettings, WorkspaceStateWrapper workspaceStateWrapper, ConfigFile latestProjectConfigFile, ConfigFile workspaceProjectConfigFile, IReadOnlyList<string>? workspaceProjectStreamFilter, List<KeyValuePair<FileReference, DateTime>> localConfigFiles, OidcTokenClient? oidcTokenClient)
+		public OpenProjectInfo(UserSelectedProjectSettings selectedProject, IPerforceSettings perforceSettings, ProjectInfo projectInfo, UserWorkspaceSettings workspaceSettings, WorkspaceStateWrapper workspaceStateWrapper, ConfigFile latestProjectConfigFile, ConfigFile workspaceProjectConfigFile, IReadOnlyList<string>? workspaceProjectStreamFilter, List<KeyValuePair<FileReference, DateTime>> localConfigFiles, OidcTokenClient? oidcTokenClient, bool generateP4Config)
 		{
 			SelectedProject = selectedProject;
 
@@ -40,9 +41,10 @@ namespace UnrealGameSync
 			WorkspaceProjectStreamFilter = workspaceProjectStreamFilter;
 			LocalConfigFiles = localConfigFiles;
 			OidcTokenClient = oidcTokenClient;
+			GenerateP4Config = generateP4Config;
 		}
 
-		public static async Task<OpenProjectInfo> CreateAsync(IPerforceSettings defaultPerforceSettings, UserSelectedProjectSettings selectedProject, UserSettings userSettings, OidcTokenManager oidcTokenManager, ILogger<OpenProjectInfo> logger, CancellationToken cancellationToken)
+		public static async Task<OpenProjectInfo> CreateAsync(IPerforceSettings defaultPerforceSettings, UserSelectedProjectSettings selectedProject, UserSettings userSettings, OidcTokenManager oidcTokenManager, bool generateP4Config, ILogger<OpenProjectInfo> logger, CancellationToken cancellationToken)
 		{
 			PerforceSettings perforceSettings = Utility.OverridePerforceSettings(defaultPerforceSettings, selectedProject.ServerAndPort, selectedProject.UserName);
 			using IPerforceConnection perforce = await PerforceConnection.CreateAsync(perforceSettings, logger);
@@ -55,10 +57,10 @@ namespace UnrealGameSync
 			}
 
 			// Execute like a regular task
-			return await CreateAsync(perforce, selectedProject, userSettings, oidcTokenManager, logger, cancellationToken);
+			return await CreateAsync(perforce, selectedProject, userSettings, oidcTokenManager, generateP4Config, logger, cancellationToken);
 		}
 
-		public static async Task<OpenProjectInfo> CreateAsync(IPerforceConnection defaultConnection, UserSelectedProjectSettings selectedProject, UserSettings userSettings, OidcTokenManager oidcTokenManager, ILogger<OpenProjectInfo> logger, CancellationToken cancellationToken)
+		public static async Task<OpenProjectInfo> CreateAsync(IPerforceConnection defaultConnection, UserSelectedProjectSettings selectedProject, UserSettings userSettings, OidcTokenManager oidcTokenManager, bool generateP4Config, ILogger<OpenProjectInfo> logger, CancellationToken cancellationToken)
 		{
 			using IDisposable loggerScope = logger.BeginScope("Project {SelectedProject}", selectedProject.ToString());
 			logger.LogInformation("Detecting settings for {Project}", selectedProject);
@@ -234,7 +236,12 @@ namespace UnrealGameSync
 				ConfigFile workspaceProjectConfigFile = await WorkspaceUpdate.ReadProjectConfigFile(branchDirectoryName, newSelectedFileName, logger);
 				IReadOnlyList<string>? workspaceProjectStreamFilter = await WorkspaceUpdate.ReadProjectStreamFilter(perforceClient, workspaceProjectConfigFile, cancellationToken);
 
-				OpenProjectInfo workspaceSettings = new OpenProjectInfo(selectedProject, perforceSettings, projectInfo, userWorkspaceSettings, workspaceStateWrapper, latestProjectConfigFile, workspaceProjectConfigFile, workspaceProjectStreamFilter, localConfigFiles, oidcTokenClient);
+				OpenProjectInfo workspaceSettings = new OpenProjectInfo(selectedProject, perforceSettings, projectInfo, userWorkspaceSettings, workspaceStateWrapper, latestProjectConfigFile, workspaceProjectConfigFile, workspaceProjectStreamFilter, localConfigFiles, oidcTokenClient, generateP4Config);
+
+				if (generateP4Config)
+				{
+					GenerateP4ConfigFile(perforceSettings, projectInfo, logger);
+				}
 
 				return workspaceSettings;
 			}
@@ -292,5 +299,50 @@ namespace UnrealGameSync
 			}
 			return candidateClients;
 		}
+
+		static void GenerateP4ConfigFile(IPerforceSettings perforceSettings, ProjectInfo projectInfo, ILogger<OpenProjectInfo> logger)
+		{
+			string? p4ConfigName = PerforceEnvironment.Default.GetValue("P4CONFIG");
+
+			// We only create a p4config if the user has opted into the system by setting P4CONFIG manually.
+			if (String.IsNullOrEmpty(p4ConfigName))
+			{
+				logger.LogError("Could not generate a p4config file as the envvar P4CONFIG is not set!");
+				return;
+			}
+
+			string configFilePath = Path.Combine(projectInfo.LocalRootPath.ToString(), p4ConfigName);
+
+			// Do not update or replace the file if it already exists 
+			if (File.Exists(configFilePath))
+			{
+				logger.LogDebug("P4CONFIG file '{Path}' already exists.", configFilePath);
+				return;
+			}
+
+			Dictionary<string, string> settings = new Dictionary<string, string>();
+			settings.Add("P4PORT", perforceSettings.ServerAndPort);
+			settings.Add("P4USER", perforceSettings.UserName);
+			settings.Add("P4CLIENT", perforceSettings.ClientName ?? "");
+
+			List<string> Lines = new List<string>();
+			foreach (KeyValuePair<string, string> setting in settings)
+			{
+				Lines.Add(String.Format("{0}={1}", setting.Key, setting.Value));
+
+			}
+
+			try
+			{
+				File.WriteAllLines(configFilePath, Lines);
+			}
+			catch (Exception ex)
+			{
+				logger.LogError(ex, "Failed to write {Path} with exception - {Msg}", p4ConfigName, ex.Message);
+				return;
+			}
+
+			logger.LogInformation("Successfully wrote a P4CONFIG file to '{Path}'.", configFilePath);
+		}
 	}
 }
