UGS: Improve the performance of finding the last code change during syncing

#rn minor ugs/sync

diff --git a/Engine/Source/Programs/UnrealGameSync/UnrealGameSyncShared/WorkspaceUpdate.cs b/Engine/Source/Programs/UnrealGameSync/UnrealGameSyncShared/WorkspaceUpdate.cs
--- a/Engine/Source/Programs/UnrealGameSync/UnrealGameSyncShared/WorkspaceUpdate.cs
+++ b/Engine/Source/Programs/UnrealGameSync/UnrealGameSyncShared/WorkspaceUpdate.cs
@@ -945,8 +945,16 @@ namespace UnrealGameSync
 						{
 							logger.LogInformation("Finding last code change for CL {Number}...", Context.ChangeNumber);
 
+							// If we are syncing to a newer change than the last code change we found (and it is not the first sync in a workspace)
+							// go head and use the last code change we found as the minimum change in our query
+							int? minChangeNumber = null;
+							if ((Context.ChangeNumber >= state.CurrentCodeChangeNumber) && (state.CurrentCodeChangeNumber > 0))
+							{
+								minChangeNumber = state.CurrentCodeChangeNumber;
+							}
+
 							string[] codeRules = Utility.GetCodeFilter(Context.ProjectConfigFile);
-							await foreach (PerforceChangeDetails details in Utility.EnumerateChangeDetails(perforce, minChangeNumber: null, maxChangeNumber: Context.ChangeNumber, syncPaths, codeRules, cancellationToken))
+							await foreach (PerforceChangeDetails details in Utility.EnumerateChangeDetails(perforce, minChangeNumber, maxChangeNumber: Context.ChangeNumber, syncPaths, codeRules, cancellationToken))
 							{
 								if (details.ContainsCode)
 								{
