- Updated ShderConductor libs on Mac to build with new Min deployment version, to work around an issue with dylib linker in Sonoma

diff --git a/Engine/Source/Programs/UnrealGameSync/UnrealGameSyncCmd/UnrealGameSyncCmd.csproj b/Engine/Source/Programs/UnrealGameSync/UnrealGameSyncCmd/UnrealGameSyncCmd.csproj
--- a/Engine/Source/Programs/UnrealGameSync/UnrealGameSyncCmd/UnrealGameSyncCmd.csproj
+++ b/Engine/Source/Programs/UnrealGameSync/UnrealGameSyncCmd/UnrealGameSyncCmd.csproj
@@ -2,11 +2,23 @@
 
   <PropertyGroup>
     <OutputType>Exe</OutputType>
-    <TargetFramework>net6.0</TargetFramework>
+    <TargetFramework>net7.0</TargetFramework>
     <Nullable>enable</Nullable>
     <AssemblyName>ugs</AssemblyName>
+    <IsMac Condition="'$([System.Runtime.InteropServices.RuntimeInformation]::IsOSPlatform($([System.Runtime.InteropServices.OSPlatform]::OSX)))' == 'true'">true</IsMac> 
   </PropertyGroup>
 
+  <Choose>
+    <When Condition="'$(IsMac)'=='true'">
+      <PropertyGroup>
+        <PublishSingleFile>true</PublishSingleFile>
+        <SelfContained>true</SelfContained>
+        <RuntimeIdentifiers>osx-x64;osx-arm64</RuntimeIdentifiers>
+      </PropertyGroup>
+    </When>
+  </Choose>
+
+
   <ItemGroup>
     <PackageReference Include="Serilog" Version="2.10.0" />
     <PackageReference Include="Serilog.Sinks.Console" Version="4.0.1" />
