UGS: Added back project and programs status lines when UGS workspace is busy/syncing

[RN] minor ugs/ui

diff --git a/Engine/Source/Programs/UnrealGameSync/UnrealGameSync/Controls/WorkspaceControl.cs b/Engine/Source/Programs/UnrealGameSync/UnrealGameSync/Controls/WorkspaceControl.cs
--- a/Engine/Source/Programs/UnrealGameSync/UnrealGameSync/Controls/WorkspaceControl.cs
+++ b/Engine/Source/Programs/UnrealGameSync/UnrealGameSync/Controls/WorkspaceControl.cs
@@ -3342,40 +3342,114 @@ namespace UnrealGameSync
 			StatusPanel.SetContentWidth(newContentWidth);
 
 			List<StatusLine> lines = new List<StatusLine>();
-			if (_workspace.IsBusy())
+			if (_workspace.IsBusy() || _workspace.IsExternalSyncActive())
 			{
-				// Sync in progress
-				Tuple<string, float> progress = _workspace.CurrentProgress;
+				// Project
+				StatusLine projectLine = new StatusLine();
+				projectLine.AddText("Opened " + SelectedFileName.FullName);
+
+				OidcTokenClient? oidcClient = _perforceMonitor.LatestOidcTokenClient;
+				if (oidcClient != null)
+				{
+					projectLine.AddText("  |  ");
+
+					StatusLine hordeLine = new StatusLine();
+					if (oidcClient.GetStatus() == OidcStatus.Connected)
+					{
+						projectLine.AddText("Connected to Horde.");
+					}
+					else
+					{
+						projectLine.AddLink("Connect to Horde", FontStyle.Bold | FontStyle.Underline, (p, r) => DoOidcLogin());
+					}
+				}
+				lines.Add(projectLine);
+
+				// Spacer
+				lines.Add(new StatusLine() { LineHeight = 0.5f });
 
+				// Sync in progress
 				StatusLine summaryLine = new StatusLine();
-				if (_workspace.PendingChangeNumber == -1)
+				if (_workspace.IsBusy())
 				{
-					summaryLine.AddText("Working... | ");
+					Tuple<string, float> progress = _workspace.CurrentProgress;
+
+					if (_workspace.PendingChangeNumber == -1)
+					{
+						summaryLine.AddText("Working... | ");
+					}
+					else
+					{
+						summaryLine.AddText("Syncing to changelist ");
+						summaryLine.AddLink(_workspace.PendingChangeNumber.ToString(), FontStyle.Regular, () => { SelectChange(_workspace.PendingChangeNumber); });
+						summaryLine.AddText("... | ");
+					}
+					summaryLine.AddLink(Splitter.IsLogVisible() ? "Hide Log" : "Show Log", FontStyle.Bold | FontStyle.Underline, () => { ToggleLogVisibility(); });
+					summaryLine.AddText(" | ");
+					summaryLine.AddLink("Cancel", FontStyle.Bold | FontStyle.Underline, () => { CancelWorkspaceUpdate(); });
+					lines.Add(summaryLine);
+
+					StatusLine progressLine = new StatusLine();
+					progressLine.AddText(String.Format("{0}  ", progress.Item1));
+					if (progress.Item2 > 0.0f)
+					{
+						progressLine.AddProgressBar(progress.Item2);
+					}
+					lines.Add(progressLine);
 				}
 				else
 				{
-					summaryLine.AddText("Syncing to changelist ");
-					summaryLine.AddLink(_workspace.PendingChangeNumber.ToString(), FontStyle.Regular, () => { SelectChange(_workspace.PendingChangeNumber); });
-					summaryLine.AddText("... | ");
+					summaryLine.AddText("Command-line sync in progress...");
+					lines.Add(summaryLine);
 				}
-				summaryLine.AddLink(Splitter.IsLogVisible() ? "Hide Log" : "Show Log", FontStyle.Bold | FontStyle.Underline, () => { ToggleLogVisibility(); });
-				summaryLine.AddText(" | ");
-				summaryLine.AddLink("Cancel", FontStyle.Bold | FontStyle.Underline, () => { CancelWorkspaceUpdate(); });
-				lines.Add(summaryLine);
 
-				StatusLine progressLine = new StatusLine();
-				progressLine.AddText(String.Format("{0}  ", progress.Item1));
-				if (progress.Item2 > 0.0f)
+				// Programs
+				StatusLine programsLine = new StatusLine();
+
+				string[]? sdkInfoEntries;
+				if (TryGetProjectSetting(_perforceMonitor.LatestProjectConfigFile, "SdkInfo", out sdkInfoEntries))
 				{
-					progressLine.AddProgressBar(progress.Item2);
+					programsLine.AddLink("SDK Info", FontStyle.Regular, () => { ShowRequiredSdkInfo(); });
+					programsLine.AddText("  |  ");
 				}
-				lines.Add(progressLine);
-			}
-			else if (_workspace.IsExternalSyncActive())
-			{
-				StatusLine summaryLine = new StatusLine();
-				summaryLine.AddText("Command-line sync in progress...");
-				lines.Add(summaryLine);
+
+				programsLine.AddLink("Perforce", FontStyle.Regular, () => { OpenPerforce(); });
+				programsLine.AddText("  |  ");
+				if (!ShouldSyncPrecompiledEditor)
+				{
+					programsLine.AddLink("Visual Studio", FontStyle.Regular, () => { OpenSolution(); });
+					programsLine.AddText("  |  ");
+				}
+
+				List<ToolDefinition> tools = _owner.ToolUpdateMonitor.Tools;
+				foreach (ToolDefinition tool in tools)
+				{
+					if (tool.Enabled && tool.SafeWhenBusy)
+					{
+						foreach (ToolLink link in tool.StatusPanelLinks)
+						{
+							programsLine.AddLink(link.Label, FontStyle.Regular, () => RunTool(tool, link));
+							programsLine.AddText("  |  ");
+						}
+					}
+				}
+
+				programsLine.AddLink("Windows Explorer", FontStyle.Regular, () => { SafeProcessStart("explorer.exe", String.Format("\"{0}\"", SelectedFileName.Directory.FullName)); });
+
+				if (GetDefaultIssueFilter() != null)
+				{
+					programsLine.AddText("  |  ");
+					if (_userHasOpenIssues)
+					{
+						programsLine.AddBadge("Build Health", GetBuildBadgeColor(BadgeResult.Failure), (p, r) => ShowBuildHealthMenu(r));
+					}
+					else
+					{
+						programsLine.AddLink("Build Health", FontStyle.Regular, (p, r) => { ShowBuildHealthMenu(r); });
+					}
+				}
+
+				lines.Add(programsLine);
 			}
 			else
 			{
@@ -3633,10 +3707,17 @@ namespace UnrealGameSync
 			}
 
 			StatusLine? caption = null;
-			if (StreamName != null && !_workspace.IsBusy())
+			if (StreamName != null)
 			{
 				caption = new StatusLine();
-				caption.AddLink(StreamName + "\u25BE", FontStyle.Bold, (p, r) => { SelectOtherStream(r); });
+				if (_workspace.IsBusy() || _workspace.IsExternalSyncActive())
+				{
+					caption.AddText(StreamName);
+				}
+				else
+				{
+					caption.AddLink(StreamName + "\u25BE", FontStyle.Bold, (p, r) => { SelectOtherStream(r); });
+				}
 			}
 
 			StatusLine? alert = null;
diff --git a/Engine/Source/Programs/UnrealGameSync/UnrealGameSync/ToolUpdateMonitor.cs b/Engine/Source/Programs/UnrealGameSync/UnrealGameSync/ToolUpdateMonitor.cs
--- a/Engine/Source/Programs/UnrealGameSync/UnrealGameSync/ToolUpdateMonitor.cs
+++ b/Engine/Source/Programs/UnrealGameSync/UnrealGameSync/ToolUpdateMonitor.cs
@@ -44,8 +44,9 @@ namespace UnrealGameSync
 		public string ConfigPath;
 		public int ConfigChange;
 		public HashSet<Guid> DependsOnToolIds;
+		public bool SafeWhenBusy;
 
-		public ToolDefinition(Guid id, string name, string description, string zipPath, int zipChange, string configPath, int configChange, HashSet<Guid> dependsOnToolIds)
+		public ToolDefinition(Guid id, string name, string description, string zipPath, int zipChange, string configPath, int configChange, HashSet<Guid> dependsOnToolIds, bool safeWhenBusy)
 		{
 			Id = id;
 			Name = name;
@@ -55,6 +56,7 @@ namespace UnrealGameSync
 			ConfigPath = configPath;
 			ConfigChange = configChange;
 			DependsOnToolIds = dependsOnToolIds;
+			SafeWhenBusy = safeWhenBusy;
 		}
 	}
 
@@ -292,6 +294,7 @@ namespace UnrealGameSync
 			int toolZipChange = GetToolChange(toolName);
 			string toolConfigPath = depotPath;
 			int toolConfigChange = change;
+			bool toolSafeWhenBusy = configFile.GetValue("Settings.SafeWhenBusy", false);
 
 			HashSet<Guid> dependsOnToolIds = new HashSet<Guid>();
 			foreach (string line in configFile.GetValues("Settings.DependsOnTool", Array.Empty<string>()))
@@ -302,7 +305,7 @@ namespace UnrealGameSync
 				}
 			}
 
-			ToolDefinition tool = new ToolDefinition(toolId, toolName, toolDescription, toolZipPath, toolZipChange, toolConfigPath, toolConfigChange, dependsOnToolIds);
+			ToolDefinition tool = new ToolDefinition(toolId, toolName, toolDescription, toolZipPath, toolZipChange, toolConfigPath, toolConfigChange, dependsOnToolIds, toolSafeWhenBusy);
 
 			string? installCommand = configFile.GetValue("Settings.InstallCommand", null);
 			if (!String.IsNullOrEmpty(installCommand))
