UGS: increase number of returned PCB revisions

diff --git a/Engine/Source/Programs/UnrealGameSync/UnrealGameSyncShared/ArchiveInfo.cs b/Engine/Source/Programs/UnrealGameSync/UnrealGameSyncShared/ArchiveInfo.cs
--- a/Engine/Source/Programs/UnrealGameSync/UnrealGameSyncShared/ArchiveInfo.cs
+++ b/Engine/Source/Programs/UnrealGameSync/UnrealGameSyncShared/ArchiveInfo.cs
@@ -169,7 +169,7 @@ namespace UnrealGameSync
 				// Make sure the zipped binaries path exists
 				foreach (PerforceArchiveInfo newArchive in newArchives)
 				{
-					PerforceResponseList<FileLogRecord> response = await perforce.TryFileLogAsync(30, FileLogOptions.FullDescriptions, newArchive.DepotPath, cancellationToken);
+					PerforceResponseList<FileLogRecord> response = await perforce.TryFileLogAsync(128, FileLogOptions.FullDescriptions, newArchive.DepotPath, cancellationToken);
 					if (response.Succeeded)
 					{
 						// Build a new list of zipped binaries
