Add option to "Always Clobber" in UGS

Skips the "Clobber Files" dialog if enabled

diff --git a/Engine/Source/Programs/UnrealGameSync/UnrealGameSync/Controls/WorkspaceControl.cs b/Engine/Source/Programs/UnrealGameSync/UnrealGameSync/Controls/WorkspaceControl.cs
--- a/Engine/Source/Programs/UnrealGameSync/UnrealGameSync/Controls/WorkspaceControl.cs
+++ b/Engine/Source/Programs/UnrealGameSync/UnrealGameSync/Controls/WorkspaceControl.cs
@@ -1078,13 +1078,20 @@ namespace UnrealGameSync
 					ParseUncontrolledChangelistsPersistencyFile(Path.GetDirectoryName(SelectedProject.LocalPath)!, uncontrolledFiles);
 				}
 
-				using ClobberWindow window = new ClobberWindow(context.ClobberFiles, uncontrolledFiles);
-
-				if (window.ShowDialog(this) == DialogResult.OK)
+				if (_settings.Global.AlwaysClobberFiles)
 				{
 					StartWorkspaceUpdate(context, _updateCallback);
 					return;
 				}
+				else
+				{
+					using ClobberWindow window = new ClobberWindow(context.ClobberFiles, uncontrolledFiles);
+					if (window.ShowDialog(this) == DialogResult.OK)
+					{
+						StartWorkspaceUpdate(context, _updateCallback);
+						return;
+					}
+				}
 			}
 			else if (result == WorkspaceUpdateResult.FailedToCompileWithCleanWorkspace)
 			{
@@ -4613,6 +4620,7 @@ namespace UnrealGameSync
 		private void OptionsButton_Click(object sender, EventArgs e)
 		{
 			OptionsContextMenu_AutoResolveConflicts.Checked = _settings.Global.AutoResolveConflicts;
+			OptionsContextMenu_AlwaysClobberFiles.Checked = _settings.Global.AlwaysClobberFiles;
 
 			OptionsContextMenu_SyncPrecompiledBinaries.DropDownItems.Clear();
 
@@ -5126,6 +5134,13 @@ namespace UnrealGameSync
 			_settings.Save(_logger);
 		}
 
+		private void OptionsContextMenu_AlwaysClobberFiles_Click(object sender, EventArgs e)
+		{
+			OptionsContextMenu_AlwaysClobberFiles.Checked ^= true;
+			_settings.Global.AlwaysClobberFiles = OptionsContextMenu_AlwaysClobberFiles.Checked;
+			_settings.Save(_logger);
+		}
+
 		private void OptionsContextMenu_EditorArguments_Click(object sender, EventArgs e)
 		{
 			ModifyEditorArguments();
diff --git a/Engine/Source/Programs/UnrealGameSync/UnrealGameSyncShared/GlobalSettings.cs b/Engine/Source/Programs/UnrealGameSync/UnrealGameSyncShared/GlobalSettings.cs
--- a/Engine/Source/Programs/UnrealGameSync/UnrealGameSyncShared/GlobalSettings.cs
+++ b/Engine/Source/Programs/UnrealGameSync/UnrealGameSyncShared/GlobalSettings.cs
@@ -56,6 +56,7 @@ namespace UnrealGameSync
 		public PerforceSyncOptions Perforce { get; set; } = new PerforceSyncOptions();
 		public FilterSettings Filter { get; set; } = new FilterSettings();
 		public bool AutoResolveConflicts { get; set; } = true;
+		public bool AlwaysClobberFiles { get; set; } = false;
 	}
 
 	public class GlobalSettingsFile
