UGS: Do not download editor archives if already downloaded

[RN] minor ugs/sync

diff --git a/Engine/Source/Programs/UnrealGameSync/UnrealGameSync/Controls/WorkspaceControl.cs b/Engine/Source/Programs/UnrealGameSync/UnrealGameSync/Controls/WorkspaceControl.cs
--- a/Engine/Source/Programs/UnrealGameSync/UnrealGameSync/Controls/WorkspaceControl.cs
+++ b/Engine/Source/Programs/UnrealGameSync/UnrealGameSync/Controls/WorkspaceControl.cs
@@ -964,17 +964,22 @@ namespace UnrealGameSync
 			{
 				if (!context.Options.HasFlag(WorkspaceUpdateOptions.ContentOnly) && (context.CustomBuildSteps == null || context.CustomBuildSteps.Count == 0))
 				{
-					FileReference targetFile = ConfigUtils.GetEditorTargetFile(_workspace.Project, _workspace.ProjectConfigFile);
-					foreach (BuildConfig config in Enum.GetValues(typeof(BuildConfig)).OfType<BuildConfig>())
+					bool usingPrecompiledEditor = context.ArchiveTypeToArchive.TryGetValue(IArchiveInfo.EditorArchiveType, out Tuple<IArchiveInfo, string>? archiveInfo) && archiveInfo != null;
+					bool usingLastSyncedEditorArchive = archiveInfo != null && archiveInfo.Item2 == _workspace.State.LastSyncEditorArchive;
+					if (!usingPrecompiledEditor || !usingLastSyncedEditorArchive)
 					{
-						FileReference receiptFile = ConfigUtils.GetReceiptFile(_workspace.Project, _workspace.ProjectConfigFile, targetFile, config.ToString());
-						if (FileReference.Exists(receiptFile))
+						FileReference targetFile = ConfigUtils.GetEditorTargetFile(_workspace.Project, _workspace.ProjectConfigFile);
+						foreach (BuildConfig config in Enum.GetValues(typeof(BuildConfig)).OfType<BuildConfig>())
 						{
-							try
+							FileReference receiptFile = ConfigUtils.GetReceiptFile(_workspace.Project, _workspace.ProjectConfigFile, targetFile, config.ToString());
+							if (FileReference.Exists(receiptFile))
 							{
-								FileReference.Delete(receiptFile);
+								try
+								{
+									FileReference.Delete(receiptFile);
+								}
+								catch (Exception) { }
 							}
-							catch (Exception) { }
 						}
 					}
 				}
diff --git a/Engine/Source/Programs/UnrealGameSync/UnrealGameSync/Workspace.cs b/Engine/Source/Programs/UnrealGameSync/UnrealGameSync/Workspace.cs
--- a/Engine/Source/Programs/UnrealGameSync/UnrealGameSync/Workspace.cs
+++ b/Engine/Source/Programs/UnrealGameSync/UnrealGameSync/Workspace.cs
@@ -214,6 +214,7 @@ namespace UnrealGameSync
 						x.LastSyncResultMessage = null;
 						x.LastSyncTime = null;
 						x.LastSyncDurationSeconds = 0;
+						x.LastSyncEditorArchive = "0";
 					}
 					x.SetLastSyncState(result, context, statusMessage);
 				});
diff --git a/Engine/Source/Programs/UnrealGameSync/UnrealGameSyncShared/UserSettings.cs b/Engine/Source/Programs/UnrealGameSync/UnrealGameSyncShared/UserSettings.cs
--- a/Engine/Source/Programs/UnrealGameSync/UnrealGameSyncShared/UserSettings.cs
+++ b/Engine/Source/Programs/UnrealGameSync/UnrealGameSyncShared/UserSettings.cs
@@ -938,6 +938,8 @@ namespace UnrealGameSync
 				currentWorkspace.LastSyncDurationSeconds = workspaceSection.GetValue("LastSyncDuration", 0);
 				currentWorkspace.LastBuiltChangeNumber = workspaceSection.GetValue("LastBuiltChangeNumber", 0);
 
+				currentWorkspace.LastSyncEditorArchive = workspaceSection.GetValue("LastSyncEditorArchive", "0");
+
 				currentWorkspace.ExpandedArchiveTypes.Clear();
 				currentWorkspace.ExpandedArchiveTypes.UnionWith(workspaceSection.GetValues("ExpandedArchiveName", Array.Empty<string>()));
 
diff --git a/Engine/Source/Programs/UnrealGameSync/UnrealGameSyncShared/WorkspaceState.cs b/Engine/Source/Programs/UnrealGameSync/UnrealGameSyncShared/WorkspaceState.cs
--- a/Engine/Source/Programs/UnrealGameSync/UnrealGameSyncShared/WorkspaceState.cs
+++ b/Engine/Source/Programs/UnrealGameSync/UnrealGameSyncShared/WorkspaceState.cs
@@ -46,6 +46,9 @@ namespace UnrealGameSync
 		// The last successful build, regardless of whether a failed sync has happened in the meantime. Used to determine whether to force a clean due to entries in the project config file.
 		public int LastBuiltChangeNumber { get; set; }
 
+		// The path of the last synced editor archive
+		public string LastSyncEditorArchive { get; set; } = "0";
+
 		// Expanded archives in the workspace
 		public HashSet<string> ExpandedArchiveTypes { get; init; } = new HashSet<string>(StringComparer.Ordinal);
 
@@ -102,6 +105,8 @@ namespace UnrealGameSync
 
 			LastBuiltChangeNumber = other.LastBuiltChangeNumber;
 
+			LastSyncEditorArchive = other.LastSyncEditorArchive;
+
 			ExpandedArchiveTypes.Clear();
 			ExpandedArchiveTypes.UnionWith(other.ExpandedArchiveTypes);
 
@@ -162,6 +167,9 @@ namespace UnrealGameSync
 		// The last successful build, regardless of whether a failed sync has happened in the meantime. Used to determine whether to force a clean due to entries in the project config file.
 		public int LastBuiltChangeNumber => _inner.LastBuiltChangeNumber;
 
+		// The path of the last synced editor archive
+		public string LastSyncEditorArchive => _inner.LastSyncEditorArchive;
+
 		// Expanded archives in the workspace
 		public IReadOnlySet<string> ExpandedArchiveTypes { get; }
 
diff --git a/Engine/Source/Programs/UnrealGameSync/UnrealGameSyncShared/WorkspaceUpdate.cs b/Engine/Source/Programs/UnrealGameSync/UnrealGameSyncShared/WorkspaceUpdate.cs
--- a/Engine/Source/Programs/UnrealGameSync/UnrealGameSyncShared/WorkspaceUpdate.cs
+++ b/Engine/Source/Programs/UnrealGameSync/UnrealGameSyncShared/WorkspaceUpdate.cs
@@ -1137,11 +1137,15 @@ namespace UnrealGameSync
 						FileReference manifestFileName = FileReference.Combine(manifestDirectoryName, String.Format("{0}.zipmanifest", archiveType));
 						if (FileReference.Exists(manifestFileName))
 						{
-							logger.LogInformation("Removing {ArchiveType} binaries...", archiveType);
-							Progress.Set(String.Format("Removing {0} binaries...", archiveType), 0.0f);
-							ArchiveUtils.RemoveExtractedFiles(project.LocalRootPath, manifestFileName, Progress, logger);
-							FileReference.Delete(manifestFileName);
-							logger.LogInformation("");
+							bool isNotLastSyncedEditorArchive = archiveType != IArchiveInfo.EditorArchiveType || (archiveTypeAndArchive.Value != null && archiveTypeAndArchive.Value.Item2 != state.LastSyncEditorArchive);
+							if (isNotLastSyncedEditorArchive)
+							{
+								logger.LogInformation("Removing {ArchiveType} binaries...", archiveType);
+								Progress.Set(String.Format("Removing {0} binaries...", archiveType), 0.0f);
+								ArchiveUtils.RemoveExtractedFiles(project.LocalRootPath, manifestFileName, Progress, logger);
+								FileReference.Delete(manifestFileName);
+								logger.LogInformation("");
+							}
 						}
 
 						// If we have a new depot path, sync it down and extract it
@@ -1150,11 +1154,26 @@ namespace UnrealGameSync
 							IArchiveInfo archiveInfo = archiveTypeAndArchive.Value.Item1;
 							string archiveKey = archiveTypeAndArchive.Value.Item2;
 
-							logger.LogInformation("Syncing {ArchiveType} binaries...", archiveType);
-							Progress.Set(String.Format("Syncing {0} binaries...", archiveType), 0.0f);
-							if (!await archiveInfo.DownloadArchive(perforce, archiveKey, project.LocalRootPath, manifestFileName, logger, Progress, CancellationToken.None))
+							if (archiveType != IArchiveInfo.EditorArchiveType || archiveKey != state.LastSyncEditorArchive)
+							{
+								logger.LogInformation("Syncing {ArchiveType} binaries...", archiveType);
+								Progress.Set(String.Format("Syncing {0} binaries...", archiveType), 0.0f);
+								if (!await archiveInfo.DownloadArchive(perforce, archiveKey, project.LocalRootPath, manifestFileName, logger, Progress, CancellationToken.None))
+								{
+									return (WorkspaceUpdateResult.FailedToSync, $"Couldn't read {archiveKey}");
+								}
+								// Update last synced editor archive
+								if (archiveType == IArchiveInfo.EditorArchiveType)
+								{
+									state = stateMgr.Modify(x =>
+									{
+										x.LastSyncEditorArchive = archiveKey;
+									});
+								}
+							}
+							else
 							{
-								return (WorkspaceUpdateResult.FailedToSync, $"Couldn't read {archiveKey}");
+								logger.LogInformation("Skipping {ArchiveType} binaries download, already downloaded", IArchiveInfo.EditorArchiveType);
 							}
 						}
 					}
