UGS: Update PostBadgeStatus from .NET 4.5 to .NET 4.8

diff --git a/Engine/Source/Programs/UnrealGameSync/PostBadgeStatus/App.config b/Engine/Source/Programs/UnrealGameSync/PostBadgeStatus/App.config
--- a/Engine/Source/Programs/UnrealGameSync/PostBadgeStatus/App.config
+++ b/Engine/Source/Programs/UnrealGameSync/PostBadgeStatus/App.config
@@ -1,6 +1,6 @@
-﻿<?xml version="1.0" encoding="utf-8" ?>
+<?xml version="1.0" encoding="utf-8"?>
 <configuration>
     <startup> 
-        <supportedRuntime version="v4.0" sku=".NETFramework,Version=v4.5" />
+        <supportedRuntime version="v4.0" sku=".NETFramework,Version=v4.8"/>
     </startup>
-</configuration>
\ No newline at end of file
+</configuration>
diff --git a/Engine/Source/Programs/UnrealGameSync/PostBadgeStatus/PostBadgeStatus.csproj b/Engine/Source/Programs/UnrealGameSync/PostBadgeStatus/PostBadgeStatus.csproj
--- a/Engine/Source/Programs/UnrealGameSync/PostBadgeStatus/PostBadgeStatus.csproj
+++ b/Engine/Source/Programs/UnrealGameSync/PostBadgeStatus/PostBadgeStatus.csproj
@@ -9,8 +9,9 @@
     <AppDesignerFolder>Properties</AppDesignerFolder>
     <RootNamespace>PostBadgeStatus</RootNamespace>
     <AssemblyName>PostBadgeStatus</AssemblyName>
-    <TargetFrameworkVersion>v4.5</TargetFrameworkVersion>
+    <TargetFrameworkVersion>v4.8</TargetFrameworkVersion>
     <FileAlignment>512</FileAlignment>
+    <TargetFrameworkProfile />
   </PropertyGroup>
   <PropertyGroup Condition="'$(Configuration)|$(Platform)' == 'Debug|AnyCPU'">
     <DebugSymbols>true</DebugSymbols>
