UGS: Second attempt to speed up "finding last code change..."

Added a fallback in case the fast p4 query fails.

diff --git a/Engine/Source/Programs/UnrealGameSync/UnrealGameSyncShared/WorkspaceUpdate.cs b/Engine/Source/Programs/UnrealGameSync/UnrealGameSyncShared/WorkspaceUpdate.cs
--- a/Engine/Source/Programs/UnrealGameSync/UnrealGameSyncShared/WorkspaceUpdate.cs
+++ b/Engine/Source/Programs/UnrealGameSync/UnrealGameSyncShared/WorkspaceUpdate.cs
@@ -945,16 +945,41 @@ namespace UnrealGameSync
 						{
 							logger.LogInformation("Finding last code change for CL {Number}...", Context.ChangeNumber);
 
+							// If we are syncing to a newer change than the last code change we found (and it is not the first sync in a workspace)
+							// go head and use the last code change we found as the minimum change in our query
+							int? minChangeNumber = null;
+							if ((Context.ChangeNumber >= state.CurrentCodeChangeNumber) && (state.CurrentCodeChangeNumber > 0))
+							{
+								minChangeNumber = state.CurrentCodeChangeNumber;
+							}
+
 							string[] codeRules = Utility.GetCodeFilter(Context.ProjectConfigFile);
-							await foreach (PerforceChangeDetails details in Utility.EnumerateChangeDetails(perforce, minChangeNumber: null, maxChangeNumber: Context.ChangeNumber, syncPaths, codeRules, cancellationToken))
+
+							try
 							{
-								if (details.ContainsCode)
+								await foreach (PerforceChangeDetails details in Utility.EnumerateChangeDetails(perforce, minChangeNumber, maxChangeNumber: Context.ChangeNumber, syncPaths, codeRules, cancellationToken))
 								{
-									codeChangeNumber = details.Number;
-									break;
+									if (details.ContainsCode)
+									{
+										codeChangeNumber = details.Number;
+										break;
+									}
 								}
 							}
-							
+							catch(EpicGames.Perforce.PerforceException)
+							{
+								logger.LogInformation("Falling back to the slow way of finding last code change for CL {Number}...", Context.ChangeNumber);
+
+								await foreach (PerforceChangeDetails details in Utility.EnumerateChangeDetails(perforce, null, maxChangeNumber: Context.ChangeNumber, syncPaths, codeRules, cancellationToken))
+								{
+									if (details.ContainsCode)
+									{
+										codeChangeNumber = details.Number;
+										break;
+									}
+								}
+							}
+
 							if (codeChangeNumber == 0)
 							{
 								return (WorkspaceUpdateResult.FailedToSync, $"Could not find any code changes before CL {Context.ChangeNumber}.");
