[Backout] - CL28189999
Original CL Desc
-----------------------------------------------------------------
UGS: Update PostBadgeStatus from .NET 4.5 to .NET 4.8

diff --git a/Engine/Source/Programs/UnrealGameSync/PostBadgeStatus/App.config b/Engine/Source/Programs/UnrealGameSync/PostBadgeStatus/App.config
--- a/Engine/Source/Programs/UnrealGameSync/PostBadgeStatus/App.config
+++ b/Engine/Source/Programs/UnrealGameSync/PostBadgeStatus/App.config
@@ -1,6 +1,6 @@
-<?xml version="1.0" encoding="utf-8"?>
+﻿<?xml version="1.0" encoding="utf-8" ?>
 <configuration>
     <startup> 
-        <supportedRuntime version="v4.0" sku=".NETFramework,Version=v4.8"/>
+        <supportedRuntime version="v4.0" sku=".NETFramework,Version=v4.5" />
     </startup>
-</configuration>
+</configuration>
\ No newline at end of file
diff --git a/Engine/Source/Programs/UnrealGameSync/PostBadgeStatus/PostBadgeStatus.csproj b/Engine/Source/Programs/UnrealGameSync/PostBadgeStatus/PostBadgeStatus.csproj
--- a/Engine/Source/Programs/UnrealGameSync/PostBadgeStatus/PostBadgeStatus.csproj
+++ b/Engine/Source/Programs/UnrealGameSync/PostBadgeStatus/PostBadgeStatus.csproj
@@ -9,9 +9,8 @@
     <AppDesignerFolder>Properties</AppDesignerFolder>
     <RootNamespace>PostBadgeStatus</RootNamespace>
     <AssemblyName>PostBadgeStatus</AssemblyName>
-    <TargetFrameworkVersion>v4.8</TargetFrameworkVersion>
+    <TargetFrameworkVersion>v4.5</TargetFrameworkVersion>
     <FileAlignment>512</FileAlignment>
-    <TargetFrameworkProfile />
   </PropertyGroup>
   <PropertyGroup Condition="'$(Configuration)|$(Platform)' == 'Debug|AnyCPU'">
     <DebugSymbols>true</DebugSymbols>
