UGS: If we clean the workspace, clear out the last synced editor archive in the state.json

diff --git a/Engine/Source/Programs/UnrealGameSync/UnrealGameSync/Controls/WorkspaceControl.cs b/Engine/Source/Programs/UnrealGameSync/UnrealGameSync/Controls/WorkspaceControl.cs
--- a/Engine/Source/Programs/UnrealGameSync/UnrealGameSync/Controls/WorkspaceControl.cs
+++ b/Engine/Source/Programs/UnrealGameSync/UnrealGameSync/Controls/WorkspaceControl.cs
@@ -5453,7 +5453,15 @@ namespace UnrealGameSync
 			string[] combinedSyncFilter = UserSettings.GetCombinedSyncFilter(GetSyncCategories(), _settings.Global.Filter, _workspaceSettings.Filter, _perforceMonitor.LatestPerforceConfigSection());
 			List<string> syncPaths = WorkspaceUpdate.GetSyncPaths(_workspace.Project, _workspaceSettings.Filter.AllProjects ?? _settings.Global.Filter.AllProjects ?? false, combinedSyncFilter);
 
-			CleanWorkspaceWindow.DoClean(ParentForm, _perforceSettings, BranchDirectoryName, _workspace.Project.ClientRootPath, syncPaths, extraSafeToDeleteFolders.Split('\n'), extraSafeToDeleteExtensions.Split('\n'), _serviceProvider.GetRequiredService<ILogger<CleanWorkspaceWindow>>());
+			bool performedClean = CleanWorkspaceWindow.DoClean(ParentForm, _perforceSettings, BranchDirectoryName, _workspace.Project.ClientRootPath, syncPaths, extraSafeToDeleteFolders.Split('\n'), extraSafeToDeleteExtensions.Split('\n'), _serviceProvider.GetRequiredService<ILogger<CleanWorkspaceWindow>>());
+
+			if (performedClean)
+			{
+				_workspace.ModifyState(x =>
+				{
+					x.LastSyncEditorArchive = "0";
+				});
+			}
 		}
 
 		private void UpdateBuildSteps()
diff --git a/Engine/Source/Programs/UnrealGameSync/UnrealGameSync/Forms/CleanWorkspaceWindow.cs b/Engine/Source/Programs/UnrealGameSync/UnrealGameSync/Forms/CleanWorkspaceWindow.cs
--- a/Engine/Source/Programs/UnrealGameSync/UnrealGameSync/Forms/CleanWorkspaceWindow.cs
+++ b/Engine/Source/Programs/UnrealGameSync/UnrealGameSync/Forms/CleanWorkspaceWindow.cs
@@ -116,7 +116,7 @@ namespace UnrealGameSync
 			Font = new System.Drawing.Font("Segoe UI", 8.25F, System.Drawing.FontStyle.Regular, System.Drawing.GraphicsUnit.Point, ((byte)(0)));
 		}
 
-		public static void DoClean(IWin32Window owner, IPerforceSettings perforceSettings, DirectoryReference localRootPath, string clientRootPath, IReadOnlyList<string> syncPaths, string[] extraSafeToDeleteFolders, string[] extraSafeToDeleteExtensions, ILogger<CleanWorkspaceWindow> logger)
+		public static bool DoClean(IWin32Window owner, IPerforceSettings perforceSettings, DirectoryReference localRootPath, string clientRootPath, IReadOnlyList<string> syncPaths, string[] extraSafeToDeleteFolders, string[] extraSafeToDeleteExtensions, ILogger<CleanWorkspaceWindow> logger)
 		{
 			// Figure out which folders to clean
 			FolderToClean rootFolderToClean = new FolderToClean(localRootPath.ToDirectoryInfo());
@@ -125,20 +125,22 @@ namespace UnrealGameSync
 				ModalTask? result = ModalTask.Execute(owner, "Clean Workspace", "Querying files in Perforce, please wait...", x => queryWorkspace.RunAsync(x), ModalTaskFlags.None);
 				if (result == null || !result.Succeeded)
 				{
-					return;
+					return false;
 				}
 			}
 
 			// If there's nothing to delete, don't bother displaying the dialog at all
 			if(rootFolderToClean._filesToDelete.Count == 0 && rootFolderToClean._nameToSubFolder.Count == 0)
 			{
-				MessageBox.Show("You have no local files which are not in Perforce.", "Workspace Clean", MessageBoxButtons.OK);
-				return;
+				MessageBox.Show("You have no local files which are not in Perforce.", "Clean Workspace", MessageBoxButtons.OK);
+				return false;
 			}
 
 			// Populate the tree
 			using CleanWorkspaceWindow cleanWorkspace = new CleanWorkspaceWindow(perforceSettings, rootFolderToClean, extraSafeToDeleteFolders, extraSafeToDeleteExtensions, logger);
-			cleanWorkspace.ShowDialog();
+			DialogResult dialogResult = cleanWorkspace.ShowDialog();
+
+			return (dialogResult == DialogResult.OK);
 		}
 
 		private void CleanWorkspaceWindow_Load(object sender, EventArgs e)
