UGS: Fix a bug where switching streams could sometimes cause "Finding last code change..." to fail.

diff --git a/Engine/Source/Programs/UnrealGameSync/UnrealGameSync/Controls/WorkspaceControl.cs b/Engine/Source/Programs/UnrealGameSync/UnrealGameSync/Controls/WorkspaceControl.cs
--- a/Engine/Source/Programs/UnrealGameSync/UnrealGameSync/Controls/WorkspaceControl.cs
+++ b/Engine/Source/Programs/UnrealGameSync/UnrealGameSync/Controls/WorkspaceControl.cs
@@ -4108,6 +4108,13 @@ namespace UnrealGameSync
 						{
 							StatusPanel.SuspendLayout();
 							StreamChanged();
+
+							// Reset the last code change we found when switching streams
+							_workspace.ModifyState(x =>
+							{
+								x.CurrentCodeChangeNumber = -1;
+							});
+
 							StatusPanel.ResumeLayout();
 							break;
 						}
