Replaced Unreal.MachineName with System.Net.Dns.GetHostName() so that we don't need to add an extra dependency on EpicGames.Build

diff --git a/Engine/Source/Programs/UnrealGameSync/UnrealGameSync/Program.cs b/Engine/Source/Programs/UnrealGameSync/UnrealGameSync/Program.cs
--- a/Engine/Source/Programs/UnrealGameSync/UnrealGameSync/Program.cs
+++ b/Engine/Source/Programs/UnrealGameSync/UnrealGameSync/Program.cs
@@ -219,7 +219,7 @@ namespace UnrealGameSync
 						{
 							Telemetry.ActiveSink = telemetrySink;
 
-							Telemetry.SendEvent("Startup", new { User = Environment.UserName, Machine = Unreal.MachineName });
+							Telemetry.SendEvent("Startup", new { User = Environment.UserName, Machine = System.Net.Dns.GetHostName() });
 
 							AppDomain.CurrentDomain.UnhandledException += CurrentDomain_UnhandledException;
 
diff --git a/Engine/Source/Programs/UnrealGameSync/UnrealGameSyncShared/WorkspaceUpdate.cs b/Engine/Source/Programs/UnrealGameSync/UnrealGameSyncShared/WorkspaceUpdate.cs
--- a/Engine/Source/Programs/UnrealGameSync/UnrealGameSyncShared/WorkspaceUpdate.cs
+++ b/Engine/Source/Programs/UnrealGameSync/UnrealGameSyncShared/WorkspaceUpdate.cs
@@ -890,7 +890,7 @@ namespace UnrealGameSync
 
 					using (TelemetryStopwatch transferStopwatch = new TelemetryStopwatch("Workspace_Sync_TransferFiles", project.TelemetryProjectIdentifier))
 					{
-						transferStopwatch.AddData(new { MachineName = Unreal.MachineName, DomainName = Environment.UserDomainName, ServerAndPort = perforce.Settings.ServerAndPort, UserName = perforce.Settings.UserName, IncludedFiles = syncTree.TotalIncludedFiles, ExcludedFiles = syncTree.TotalExcludedFiles, Size = syncTree.TotalSize, NumThreads = Context.PerforceSyncOptions?.NumThreads ?? PerforceSyncOptions.DefaultNumThreads });
+						transferStopwatch.AddData(new { MachineName = System.Net.Dns.GetHostName(), DomainName = Environment.UserDomainName, ServerAndPort = perforce.Settings.ServerAndPort, UserName = perforce.Settings.UserName, IncludedFiles = syncTree.TotalIncludedFiles, ExcludedFiles = syncTree.TotalExcludedFiles, Size = syncTree.TotalSize, NumThreads = Context.PerforceSyncOptions?.NumThreads ?? PerforceSyncOptions.DefaultNumThreads });
 
 						(WorkspaceUpdateResult, string) syncResult = await SyncFileRevisions(perforce, "Syncing files...", Context, batchBuilder.Batches, remainingDepotPaths, Progress, logger, cancellationToken);
 						if (syncResult.Item1 != WorkspaceUpdateResult.Success)
