Added a button in UGS to force the update of the tools without needing to reboot

diff --git a/Engine/Source/Programs/UnrealGameSync/UnrealGameSync/Controls/WorkspaceControl.cs b/Engine/Source/Programs/UnrealGameSync/UnrealGameSync/Controls/WorkspaceControl.cs
--- a/Engine/Source/Programs/UnrealGameSync/UnrealGameSync/Controls/WorkspaceControl.cs
+++ b/Engine/Source/Programs/UnrealGameSync/UnrealGameSync/Controls/WorkspaceControl.cs
@@ -5471,6 +5471,12 @@ namespace UnrealGameSync
 			}
 		}
 
+		private void MoreToolsContextMenu_UpdateTools_Click(object sender, EventArgs e)
+		{
+			_logger.LogInformation("Checking for tools updates...");
+			_owner.ToolUpdateMonitor.UpdateNow();
+		}
+
 		private void UpdateBuildSteps()
 		{
 			_hasBuildSteps = false;
