UGS: Fix hang when trying to show a large number of files to remove. Performance adding items to a list control seems to degrade significantly with large numbers of items.

Also split removal up into separate requests for performance reasons.

diff --git a/Engine/Source/Programs/UnrealGameSync/UnrealGameSync/Controls/WorkspaceControl.cs b/Engine/Source/Programs/UnrealGameSync/UnrealGameSync/Controls/WorkspaceControl.cs
--- a/Engine/Source/Programs/UnrealGameSync/UnrealGameSync/Controls/WorkspaceControl.cs
+++ b/Engine/Source/Programs/UnrealGameSync/UnrealGameSync/Controls/WorkspaceControl.cs
@@ -1055,11 +1055,23 @@ namespace UnrealGameSync
 				_owner.UpdateProgress();
 
 				FlashWindow(ParentForm.Handle, true);
-				using DeleteWindow window = new DeleteWindow(context.DeleteFiles);
-				if (window.ShowDialog(this) == DialogResult.OK)
+
+				if (context.DeleteFiles.Count > 2000)
 				{
-					StartWorkspaceUpdate(context, _updateCallback);
-					return;
+					if (MessageBox.Show($"{context.DeleteFiles.Count:n0} files will be removed from your workspace by the current sync filter. Would you like to continue?", "Delete Files", MessageBoxButtons.OKCancel) == DialogResult.OK)
+					{
+						StartWorkspaceUpdate(context, _updateCallback);
+						return;
+					}
+				}
+				else
+				{
+					using DeleteWindow window = new DeleteWindow(context.DeleteFiles);
+					if (window.ShowDialog(this) == DialogResult.OK)
+					{
+						StartWorkspaceUpdate(context, _updateCallback);
+						return;
+					}
 				}
 			}
 			else if (result == WorkspaceUpdateResult.FilesToClobber)
diff --git a/Engine/Source/Programs/UnrealGameSync/UnrealGameSyncShared/WorkspaceUpdate.cs b/Engine/Source/Programs/UnrealGameSync/UnrealGameSyncShared/WorkspaceUpdate.cs
--- a/Engine/Source/Programs/UnrealGameSync/UnrealGameSyncShared/WorkspaceUpdate.cs
+++ b/Engine/Source/Programs/UnrealGameSync/UnrealGameSyncShared/WorkspaceUpdate.cs
@@ -666,6 +666,7 @@ namespace UnrealGameSync
 							}
 
 							// Remove all the files that are not included by the filter
+							const int MaxLogFiles = 1000;
 							List<string> removeDepotPaths = new List<string>();
 							foreach (HaveRecord haveFile in haveFiles)
 							{
@@ -674,7 +675,10 @@ namespace UnrealGameSync
 									FileReference fullPath = new FileReference(haveFile.Path);
 									if (MatchFilter(project, fullPath, syncPathsFilter) && !MatchFilter(project, fullPath, userFilter))
 									{
-										logger.LogInformation("  {DepotFile}", haveFile.DepotFile);
+										if (removeDepotPaths.Count <= MaxLogFiles)
+										{
+											logger.LogInformation("  {DepotFile}", haveFile.DepotFile);
+										}
 										removeDepotPaths.Add(haveFile.DepotFile);
 									}
 								}
@@ -683,6 +687,10 @@ namespace UnrealGameSync
 									// We don't actually care about this when looking for files to remove. Perforce may think that it's synced the path, and silently failed. Just ignore it.
 								}
 							}
+							if (removeDepotPaths.Count > MaxLogFiles)
+							{
+								logger.LogInformation("  ...and {NumFiles} others.", removeDepotPaths.Count - MaxLogFiles);
+							}
 
 							// Check if there are any paths outside the regular sync paths
 							if (removeDepotPaths.Count > 0 && (Context.Options & WorkspaceUpdateOptions.RemoveFilteredFiles) == 0)
@@ -1116,15 +1124,15 @@ namespace UnrealGameSync
 					// Update the current state
 					state = stateMgr.Modify(x =>
 					{
-					if (Context.Options.HasFlag(WorkspaceUpdateOptions.SyncSingleChange))
-					{
+						if (Context.Options.HasFlag(WorkspaceUpdateOptions.SyncSingleChange))
+						{
 							x.AdditionalChangeNumbers.Add(Context.ChangeNumber);
-					}
-					else
-					{
+						}
+						else
+						{
 							x.CurrentChangeNumber = Context.ChangeNumber;
 							x.CurrentCodeChangeNumber = versionChangeNumber;
-					}
+						}
 					});
 
 					// Update the timing info
@@ -1604,7 +1612,10 @@ namespace UnrealGameSync
 		static Task<(WorkspaceUpdateResult, string)> SyncFileRevisions(IPerforceConnection perforce, string prefix, WorkspaceUpdateContext context, List<string> syncCommands, HashSet<string> remainingDepotPaths, ProgressValue progress, ILogger logger, CancellationToken cancellationToken)
 		{
 			Queue<List<string>> syncCommandLists = new Queue<List<string>>();
-			syncCommandLists.Enqueue(syncCommands);
+			foreach (IReadOnlyList<string> batch in syncCommands.Batch(2000))
+			{
+				syncCommandLists.Enqueue(batch.ToList());
+			}
 			return SyncFileRevisions(perforce, prefix, context, syncCommandLists, remainingDepotPaths, progress, logger, cancellationToken);
 		}
 
@@ -1673,7 +1684,7 @@ namespace UnrealGameSync
 		{
 			lock (state)
 			{
-				string message = String.Format("{0} ({1}/{2})", prefix, state.TotalDepotPaths - state.RemainingDepotPaths.Count, state.TotalDepotPaths);
+				string message = String.Format("{0} ({1:n0}/{2:n0})", prefix, state.TotalDepotPaths - state.RemainingDepotPaths.Count, state.TotalDepotPaths);
 				float fraction = Math.Min((float)(state.TotalDepotPaths - state.RemainingDepotPaths.Count) / (float)state.TotalDepotPaths, 1.0f);
 				progress.Set(message, fraction);
 			}
@@ -1685,7 +1696,7 @@ namespace UnrealGameSync
 			{
 				state.RemainingDepotPaths.Remove(record.DepotFile.ToString());
 
-				string message = String.Format("{0} ({1}/{2})", prefix, state.TotalDepotPaths - state.RemainingDepotPaths.Count, state.TotalDepotPaths);
+				string message = String.Format("{0} ({1:n0}/{2:n0})", prefix, state.TotalDepotPaths - state.RemainingDepotPaths.Count, state.TotalDepotPaths);
 				float fraction = Math.Min((float)(state.TotalDepotPaths - state.RemainingDepotPaths.Count) / (float)state.TotalDepotPaths, 1.0f);
 				progress.Set(message, fraction);
 
