UGS: Change the default UGS build config from DebugGame to Development.  Only affects new UGS installs.

diff --git a/Engine/Source/Programs/UnrealGameSync/UnrealGameSyncShared/UserSettings.cs b/Engine/Source/Programs/UnrealGameSync/UnrealGameSyncShared/UserSettings.cs
--- a/Engine/Source/Programs/UnrealGameSync/UnrealGameSyncShared/UserSettings.cs
+++ b/Engine/Source/Programs/UnrealGameSync/UnrealGameSyncShared/UserSettings.cs
@@ -696,7 +696,7 @@ namespace UnrealGameSync
 			}
 
 			// Build configuration
-			CompiledEditorBuildConfig = _configFile.GetEnumValue("General.BuildConfig", BuildConfig.DebugGame);
+			CompiledEditorBuildConfig = _configFile.GetEnumValue("General.BuildConfig", BuildConfig.Development);
 
 			// Tab names
 			TabLabels = _configFile.GetEnumValue("General.TabLabels", TabLabels.Stream);
