UGS: Minor static analyzer fixes.

diff --git a/Engine/Source/Programs/UnrealGameSync/UnrealGameSyncShared/ConfigFile.cs b/Engine/Source/Programs/UnrealGameSync/UnrealGameSyncShared/ConfigFile.cs
--- a/Engine/Source/Programs/UnrealGameSync/UnrealGameSyncShared/ConfigFile.cs
+++ b/Engine/Source/Programs/UnrealGameSync/UnrealGameSyncShared/ConfigFile.cs
@@ -115,7 +115,7 @@ namespace UnrealGameSync
 				}
 				return token.ToString();
 			}
-			else if (ConfigSeparatorCharacters.IndexOf(text[idx], StringComparison.Ordinal) != -1)
+			else if (ConfigSeparatorCharacters.Contains(text[idx], StringComparison.Ordinal))
 			{
 				return text[idx++].ToString();
 			}
diff --git a/Engine/Source/Programs/UnrealGameSync/UnrealGameSyncShared/FileFilter.cs b/Engine/Source/Programs/UnrealGameSync/UnrealGameSyncShared/FileFilter.cs
--- a/Engine/Source/Programs/UnrealGameSync/UnrealGameSyncShared/FileFilter.cs
+++ b/Engine/Source/Programs/UnrealGameSync/UnrealGameSyncShared/FileFilter.cs
@@ -349,7 +349,7 @@ namespace UnrealGameSync
 		/// <param name="tokenIdx"></param>
 		/// <param name="currentBestNode"></param>
 		/// <returns></returns>
-		FileFilterNode FindMatchingNode(FileFilterNode currentNode, string[] tokens, int tokenIdx, FileFilterNode currentBestNode)
+		static FileFilterNode FindMatchingNode(FileFilterNode currentNode, string[] tokens, int tokenIdx, FileFilterNode currentBestNode)
 		{
 			// If we've matched all the input tokens, check if this rule is better than any other we've seen
 			if (tokenIdx == tokens.Length)
@@ -393,7 +393,7 @@ namespace UnrealGameSync
 		/// <param name="tokenIdx">Current token index</param>
 		/// <param name="currentBestRuleNumber">The highest rule number seen so far. Used to optimize tree traversals.</param>
 		/// <returns>New highest rule number</returns>
-		int HighestPossibleIncludeMatch(FileFilterNode currentNode, string[] tokens, int tokenIdx, int currentBestRuleNumber)
+		static int HighestPossibleIncludeMatch(FileFilterNode currentNode, string[] tokens, int tokenIdx, int currentBestRuleNumber)
 		{
 			// If we've matched all the input tokens, check if this rule is better than any other we've seen
 			if (tokenIdx == tokens.Length)
diff --git a/Engine/Source/Programs/UnrealGameSync/UnrealGameSyncShared/Utility.cs b/Engine/Source/Programs/UnrealGameSync/UnrealGameSyncShared/Utility.cs
--- a/Engine/Source/Programs/UnrealGameSync/UnrealGameSyncShared/Utility.cs
+++ b/Engine/Source/Programs/UnrealGameSync/UnrealGameSyncShared/Utility.cs
@@ -497,7 +497,7 @@ namespace UnrealGameSync
 					string? line = await newProcess.ReadLineAsync(cancellationToken);
 					if (line == null)
 					{
-						newProcess.WaitForExit();
+						await newProcess.WaitForExitAsync(cancellationToken);
 						return newProcess.ExitCode;
 					}
 					outputLine(line);
diff --git a/Engine/Source/Programs/UnrealGameSync/UnrealGameSyncShared/WorkspaceUpdate.cs b/Engine/Source/Programs/UnrealGameSync/UnrealGameSyncShared/WorkspaceUpdate.cs
--- a/Engine/Source/Programs/UnrealGameSync/UnrealGameSyncShared/WorkspaceUpdate.cs
+++ b/Engine/Source/Programs/UnrealGameSync/UnrealGameSyncShared/WorkspaceUpdate.cs
@@ -281,7 +281,7 @@ namespace UnrealGameSync
 				string text;
 				if (FileReference.Exists(settings.LocalProjectPath))
 				{
-					text = FileReference.ReadAllText(settings.LocalProjectPath);
+					text = await FileReference.ReadAllTextAsync(settings.LocalProjectPath, cancellationToken);
 				}
 				else
 				{
@@ -1880,7 +1880,7 @@ namespace UnrealGameSync
 		static async Task<bool> HasModifiedSourceFiles(IPerforceConnection perforce, ProjectInfo project, CancellationToken cancellationToken)
 		{
 			List<OpenedRecord> openFiles = await perforce.OpenedAsync(OpenedOptions.None, project.ClientRootPath + "/...", cancellationToken).ToListAsync(cancellationToken);
-			if (openFiles.Any(x => x.DepotFile.IndexOf("/Source/", StringComparison.OrdinalIgnoreCase) != -1))
+			if (openFiles.Any(x => x.DepotFile.Contains("/Source/", StringComparison.OrdinalIgnoreCase)))
 			{
 				return true;
 			}
@@ -2025,7 +2025,7 @@ namespace UnrealGameSync
 		{
 			try
 			{
-				if (File.Exists(localPath) && File.ReadAllText(localPath) == newText)
+				if (File.Exists(localPath) && await File.ReadAllTextAsync(localPath, cancellationToken) == newText)
 				{
 					logger.LogInformation("Ignored {FileName}; contents haven't changed", localPath);
 				}
@@ -2037,7 +2037,7 @@ namespace UnrealGameSync
 					{
 						await perforce.SyncAsync(depotPath + "#0", cancellationToken).ToListAsync(cancellationToken);
 					}
-					File.WriteAllText(localPath, newText);
+					await File.WriteAllTextAsync(localPath, newText, cancellationToken);
 					logger.LogInformation("Written {FileName}", localPath);
 				}
 				return true;
