UGS: Fix a few more static analysis warnings.

diff --git a/Engine/Source/Programs/UnrealGameSync/UnrealGameSync/FindFoldersToCleanTask.cs b/Engine/Source/Programs/UnrealGameSync/UnrealGameSync/FindFoldersToCleanTask.cs
--- a/Engine/Source/Programs/UnrealGameSync/UnrealGameSync/FindFoldersToCleanTask.cs
+++ b/Engine/Source/Programs/UnrealGameSync/UnrealGameSync/FindFoldersToCleanTask.cs
@@ -107,7 +107,7 @@ namespace UnrealGameSync
 			}
 		}
 
-		void MergeTrees(FolderToClean localFolder, PerforceHaveFolder? perforceFolder, HashSet<string> openClientPaths, string? perforceConfigFile)
+		static void MergeTrees(FolderToClean localFolder, PerforceHaveFolder? perforceFolder, HashSet<string> openClientPaths, string? perforceConfigFile)
 		{
 			if(perforceFolder == null)
 			{
@@ -192,7 +192,7 @@ namespace UnrealGameSync
 			return false;
 		}
 
-		void RemoveEmptyFolders(FolderToClean folder)
+		static void RemoveEmptyFolders(FolderToClean folder)
 		{
 			foreach(FolderToClean subFolder in folder._nameToSubFolder.Values)
 			{
diff --git a/Engine/Source/Programs/UnrealGameSync/UnrealGameSync/Forms/ApplicationSettingsWindow.cs b/Engine/Source/Programs/UnrealGameSync/UnrealGameSync/Forms/ApplicationSettingsWindow.cs
--- a/Engine/Source/Programs/UnrealGameSync/UnrealGameSync/Forms/ApplicationSettingsWindow.cs
+++ b/Engine/Source/Programs/UnrealGameSync/UnrealGameSync/Forms/ApplicationSettingsWindow.cs
@@ -177,7 +177,7 @@ namespace UnrealGameSync
 			}
 		}
 
-		void FindDependencies(ToolItem toolItem, HashSet<ToolItem> dependsOnToolItems, Dictionary<Guid, ToolItem> idToToolItem)
+		static void FindDependencies(ToolItem toolItem, HashSet<ToolItem> dependsOnToolItems, Dictionary<Guid, ToolItem> idToToolItem)
 		{
 			foreach (Guid dependsOnToolId in toolItem.Definition.DependsOnToolIds)
 			{
diff --git a/Engine/Source/Programs/UnrealGameSync/UnrealGameSync/Forms/CleanWorkspaceWindow.cs b/Engine/Source/Programs/UnrealGameSync/UnrealGameSync/Forms/CleanWorkspaceWindow.cs
--- a/Engine/Source/Programs/UnrealGameSync/UnrealGameSync/Forms/CleanWorkspaceWindow.cs
+++ b/Engine/Source/Programs/UnrealGameSync/UnrealGameSync/Forms/CleanWorkspaceWindow.cs
@@ -331,7 +331,7 @@ namespace UnrealGameSync
 			}
 		}
 
-		private void SetSelected(TreeNode parentNode, SelectionType type)
+		private static void SetSelected(TreeNode parentNode, SelectionType type)
 		{
 			TreeNodeData parentNodeData = (TreeNodeData)parentNode.Tag;
 
@@ -350,7 +350,7 @@ namespace UnrealGameSync
 			}
 		}
 
-		private void SetSelectedOnChildren(TreeNode parentNode, SelectionType type)
+		private static void SetSelectedOnChildren(TreeNode parentNode, SelectionType type)
 		{
 			TreeNodeData parentNodeData = (TreeNodeData)parentNode.Tag;
 
@@ -421,7 +421,7 @@ namespace UnrealGameSync
 			}
 		}
 
-		private void FindSelection(TreeNode node, List<FileInfo> filesToSync, List<FileInfo> filesToDelete, List<DirectoryInfo> directoriesToDelete)
+		private static void FindSelection(TreeNode node, List<FileInfo> filesToSync, List<FileInfo> filesToDelete, List<DirectoryInfo> directoriesToDelete)
 		{
 			TreeNodeData nodeData = (TreeNodeData)node.Tag;
 			if(nodeData._file != null)
diff --git a/Engine/Source/Programs/UnrealGameSync/UnrealGameSync/Forms/SelectStreamWindow.cs b/Engine/Source/Programs/UnrealGameSync/UnrealGameSync/Forms/SelectStreamWindow.cs
--- a/Engine/Source/Programs/UnrealGameSync/UnrealGameSync/Forms/SelectStreamWindow.cs
+++ b/Engine/Source/Programs/UnrealGameSync/UnrealGameSync/Forms/SelectStreamWindow.cs
@@ -134,7 +134,7 @@ namespace UnrealGameSync
 			UpdateOkButton();
 		}
 
-		private void GetExpandedNodes(TreeNodeCollection nodes, List<TreeNode> expandedNodes)
+		private static void GetExpandedNodes(TreeNodeCollection nodes, List<TreeNode> expandedNodes)
 		{
 			foreach(TreeNode? node in nodes)
 			{
@@ -190,10 +190,10 @@ namespace UnrealGameSync
 
 		private static bool IncludeNodeInFilter(StreamNode node, string[] filter)
 		{
-			return filter.All(x => node.Record.Stream.IndexOf(x, StringComparison.InvariantCultureIgnoreCase) != -1 || node.Record.Name.IndexOf(x, StringComparison.InvariantCultureIgnoreCase) != -1);
+			return filter.All(x => node.Record.Stream.Contains(x, StringComparison.InvariantCultureIgnoreCase) || node.Record.Name.Contains(x, StringComparison.InvariantCultureIgnoreCase));
 		}
 
-		private bool TryFilterTree(StreamNode node, string[] filter, [NotNullWhen(true)] out StreamNode? newNode)
+		private static bool TryFilterTree(StreamNode node, string[] filter, [NotNullWhen(true)] out StreamNode? newNode)
 		{
 			StreamNode filteredNode = new StreamNode(node.Record);
 			foreach(StreamNode childNode in node.ChildNodes)
diff --git a/Engine/Source/Programs/UnrealGameSync/UnrealGameSync/Program.cs b/Engine/Source/Programs/UnrealGameSync/UnrealGameSync/Program.cs
--- a/Engine/Source/Programs/UnrealGameSync/UnrealGameSync/Program.cs
+++ b/Engine/Source/Programs/UnrealGameSync/UnrealGameSync/Program.cs
@@ -168,7 +168,7 @@ namespace UnrealGameSync
 			{
 				try
 				{
-					SyncVersion = File.ReadAllText(syncVersionFile).Trim();
+					SyncVersion = (await File.ReadAllTextAsync(syncVersionFile)).Trim();
 				}
 				catch(Exception)
 				{
