UGS: Add a URL handler to download a build from Horde using native storage primitives.

Format is: ugs://download?server=<SERVER-URL>&ns=<NAMESPACE>&refName=<REF-NAME>

diff --git a/Engine/Source/Programs/UnrealGameSync/UnrealGameSync/Forms/DownloadProgressWindow.cs b/Engine/Source/Programs/UnrealGameSync/UnrealGameSync/Forms/DownloadProgressWindow.cs
new file mode 100644
--- /dev/null
+++ b/Engine/Source/Programs/UnrealGameSync/UnrealGameSync/Forms/DownloadProgressWindow.cs
@@ -0,0 +1,99 @@
+// Copyright Epic Games, Inc. All Rights Reserved.
+
+using System;
+using System.Collections.Generic;
+using System.ComponentModel;
+using System.Data;
+using System.Drawing;
+using System.Linq;
+using System.Text;
+using System.Threading;
+using System.Threading.Tasks;
+using System.Windows.Forms;
+
+namespace UnrealGameSync.Forms
+{
+	public partial class DownloadProgressWindow : Form
+	{
+		class ProgressReporter : IProgress<string>
+		{
+			readonly DownloadProgressWindow _owner;
+			readonly SynchronizationContext _syncContext;
+
+			public ProgressReporter(DownloadProgressWindow owner, SynchronizationContext syncContext)
+			{
+				_owner = owner;
+				_syncContext = syncContext;
+			}
+
+			public void Report(string value)
+			{
+				_syncContext.Post(_ => _owner.UpdateStatus(value), null);
+			}
+		}
+
+		readonly Func<IProgress<string>, CancellationToken, Task> _taskFunc;
+		CancellationTokenSource _cancellationSource;
+		Task _task = Task.CompletedTask;
+
+		public Task Task => _task;
+
+		private DownloadProgressWindow(Func<IProgress<string>, CancellationToken, Task> taskFunc, CancellationToken cancellationToken)
+		{
+			_taskFunc = taskFunc;
+			_cancellationSource = CancellationTokenSource.CreateLinkedTokenSource(cancellationToken);
+
+			InitializeComponent();
+			Font = new System.Drawing.Font("Segoe UI", 8.25F, System.Drawing.FontStyle.Regular, System.Drawing.GraphicsUnit.Point, ((byte)(0)));
+		}
+
+		protected override void Dispose(bool disposing)
+		{
+			if (disposing)
+			{
+				components?.Dispose();
+				_cancellationSource.Dispose();
+			}
+			base.Dispose(disposing);
+		}
+
+		public static bool Execute(Func<IProgress<string>, CancellationToken, Task> taskFunc, CancellationToken cancellationToken)
+		{
+			using DownloadProgressWindow window = new DownloadProgressWindow(taskFunc, cancellationToken);
+			window.ShowDialog();
+
+			if (window._cancellationSource.IsCancellationRequested)
+			{
+				cancellationToken.ThrowIfCancellationRequested(); // Rethrow if cancelled by caller, not if cancelled by user
+				return false;
+			}
+
+			window.Task.Wait(CancellationToken.None); // should have already finished, just rethrowing exceptions
+			return true;
+		}
+
+		private void UpdateStatus(string text)
+		{
+			if (!IsDisposed)
+			{
+				StatusMessage.Text = text;
+			}
+		}
+
+		private void DownloadProcessWindow_Load(object sender, EventArgs e)
+		{
+			SynchronizationContext syncContext = SynchronizationContext.Current!;
+			_task = Task.Run(() => _taskFunc(new ProgressReporter(this, syncContext), _cancellationSource.Token), _cancellationSource.Token);
+			_task.ContinueWith(x => syncContext.Post(y => Close(), null), TaskScheduler.Default);
+		}
+
+		private void DownloadProgressWindow_FormClosing(object sender, FormClosingEventArgs e)
+		{
+			_cancellationSource.Cancel();
+			if (_task != null && !_task.IsCompleted)
+			{
+				e.Cancel = true;
+			}
+		}
+	}
+}
diff --git a/Engine/Source/Programs/UnrealGameSync/UnrealGameSync/Forms/DownloadSettingsWindow.cs b/Engine/Source/Programs/UnrealGameSync/UnrealGameSync/Forms/DownloadSettingsWindow.cs
new file mode 100644
--- /dev/null
+++ b/Engine/Source/Programs/UnrealGameSync/UnrealGameSync/Forms/DownloadSettingsWindow.cs
@@ -0,0 +1,57 @@
+﻿// Copyright Epic Games, Inc. All Rights Reserved.
+
+using System;
+using System.Collections.Generic;
+using System.ComponentModel;
+using System.Data;
+using System.Diagnostics.CodeAnalysis;
+using System.Drawing;
+using System.IO;
+using System.Linq;
+using System.Text;
+using System.Threading.Tasks;
+using System.Windows.Forms;
+
+namespace UnrealGameSync.Forms
+{
+	public partial class DownloadSettingsWindow : Form
+	{
+		public DownloadSettingsWindow()
+		{
+			InitializeComponent();
+			Font = new System.Drawing.Font("Segoe UI", 8.25F, System.Drawing.FontStyle.Regular, System.Drawing.GraphicsUnit.Point, ((byte)(0)));
+		}
+
+		public static bool Show(string source, string defaultOutputDir, [NotNullWhen(true)] out string? outputDir)
+		{
+			using DownloadSettingsWindow window = new DownloadSettingsWindow();
+			window.SourceText.Text = source;
+			window.OutputFolderText.Text = defaultOutputDir;
+
+			if (window.ShowDialog() != DialogResult.OK)
+			{
+				outputDir = null;
+				return false;
+			}
+
+			outputDir = window.OutputFolderText.Text;
+			return true;
+		}
+
+		private void OutputFolderBrowseBtn_Click(object sender, EventArgs e)
+		{
+			using FolderBrowserDialog dialog = new FolderBrowserDialog();
+			dialog.InitialDirectory = OutputFolderText.Text;
+
+			if (dialog.ShowDialog() == DialogResult.OK)
+			{
+				OutputFolderText.Text = dialog.SelectedPath;
+			}
+		}
+
+		private void DownloadSettingsWindow_Load(object sender, EventArgs e)
+		{
+			ActiveControl = DownloadBtn;
+		}
+	}
+}
diff --git a/Engine/Source/Programs/UnrealGameSync/UnrealGameSync/Program.cs b/Engine/Source/Programs/UnrealGameSync/UnrealGameSync/Program.cs
--- a/Engine/Source/Programs/UnrealGameSync/UnrealGameSync/Program.cs
+++ b/Engine/Source/Programs/UnrealGameSync/UnrealGameSync/Program.cs
@@ -77,12 +77,9 @@ namespace UnrealGameSync
 			bool firstInstance;
 			using (Mutex instanceMutex = new Mutex(true, "UnrealGameSyncRunning", out firstInstance))
 			{
-				if (firstInstance)
-				{
-					Application.EnableVisualStyles();
-					Application.SetCompatibleTextRenderingDefault(false);
-					Application.SetHighDpiMode(HighDpiMode.PerMonitorV2);
-				}
+				Application.EnableVisualStyles();
+				Application.SetCompatibleTextRenderingDefault(false);
+				Application.SetHighDpiMode(HighDpiMode.PerMonitorV2);
 
 				using (EventWaitHandle activateEvent = new EventWaitHandle(false, EventResetMode.AutoReset, "ActivateUnrealGameSync"))
 				{
diff --git a/Engine/Source/Programs/UnrealGameSync/UnrealGameSync/UriHandlers/HordeHandler.cs b/Engine/Source/Programs/UnrealGameSync/UnrealGameSync/UriHandlers/HordeHandler.cs
new file mode 100644
--- /dev/null
+++ b/Engine/Source/Programs/UnrealGameSync/UnrealGameSync/UriHandlers/HordeHandler.cs
@@ -0,0 +1,91 @@
+// Copyright Epic Games, Inc. All Rights Reserved.
+
+using System;
+using System.Collections.Generic;
+using System.IO;
+using System.Linq;
+using System.Text;
+using System.Threading;
+using System.Threading.Tasks;
+using System.Windows.Forms;
+using EpicGames.Core;
+using EpicGames.Horde;
+using EpicGames.Horde.Storage;
+using EpicGames.Horde.Storage.Clients;
+using EpicGames.Horde.Storage.Nodes;
+using Microsoft.Extensions.DependencyInjection;
+using Microsoft.Extensions.Logging;
+using Microsoft.Extensions.Logging.Abstractions;
+using UnrealGameSync.Forms;
+
+namespace UnrealGameSync.UriHandlers
+{
+	class HordeHandler
+	{
+		class CopyProgressAdapter : IProgress<ICopyStats>
+		{
+			readonly IProgress<string> _inner;
+
+			public CopyProgressAdapter(IProgress<string> inner)
+			{
+				_inner = inner;
+			}
+
+			public void Report(ICopyStats value)
+			{
+				_inner.Report($"Copied {value.Count} files ({value.Size / (1024.0 * 1024.0):n1}mb, {value.Rate / (1024.0 * 1024.0):n1}mb/s)...");
+			}
+		}
+
+		static async Task DoSyncAsync(Uri server, NamespaceId namespaceId, RefName refName, DirectoryReference outputDir, IProgress<string> progress, CancellationToken cancellationToken)
+		{
+			DirectoryReference stateFolder = DirectoryReference.Combine(outputDir, ".ugs");
+			DirectoryReference.CreateDirectory(stateFolder);
+
+			using (ILoggerProvider loggerProvider = Logging.CreateLoggerProvider(FileReference.Combine(stateFolder, "sync.log")))
+			{
+				ServiceCollection services = new ServiceCollection();
+				services.AddLogging(builder => builder.AddProvider(loggerProvider));
+				services.AddHordeHttpClient((sp, client) => client.BaseAddress = server);
+				services.AddSingleton<BundleCache>();
+				services.AddSingleton<StorageBackendCache>();
+				services.AddSingleton<IStorageClientFactory, HttpStorageClientFactory>();
+
+				await using ServiceProvider serviceProvider = services.BuildServiceProvider();
+
+				IStorageClientFactory factory = serviceProvider.GetRequiredService<IStorageClientFactory>();
+				using IStorageClient storageClient = factory.CreateClient(namespaceId);
+
+				progress.Report("Connecting to server...");
+
+				DirectoryNode? node = await storageClient.ReadRefAsync<DirectoryNode>(refName, cancellationToken: cancellationToken);
+
+				progress.Report("Starting...");
+				await node.CopyToDirectoryAsync(outputDir.ToDirectoryInfo(), new CopyProgressAdapter(progress), serviceProvider.GetRequiredService<ILogger<HordeHandler>>(), cancellationToken);
+			}
+		}
+
+		[UriHandler(true)]
+		public static UriResult Download(string server, string ns, string refName)
+		{
+			string defaultOutputDir = Directory.GetCurrentDirectory();
+
+			string? outputDir;
+			if (!DownloadSettingsWindow.Show(refName, defaultOutputDir, out outputDir))
+			{
+				return new UriResult() { Success = true };
+			}
+
+			try
+			{
+				DownloadProgressWindow.Execute((p, ctx) => DoSyncAsync(new Uri(server), new NamespaceId(ns), new RefName(refName), new DirectoryReference(outputDir), p, ctx), CancellationToken.None);
+			}
+			catch (Exception ex)
+			{
+				MessageBox.Show($"An error occurred while downloading the specified item.\n\n{ex}");
+			}
+
+			return new UriResult() { Success = true };
+		}
+	}
+}
