UGS: Move global settings for application updates into UnrealGameSyncShared.

diff --git a/Engine/Source/Programs/UnrealGameSync/UnrealGameSync/Automation/P4Automation.cs b/Engine/Source/Programs/UnrealGameSync/UnrealGameSync/Automation/P4Automation.cs
--- a/Engine/Source/Programs/UnrealGameSync/UnrealGameSync/Automation/P4Automation.cs
+++ b/Engine/Source/Programs/UnrealGameSync/UnrealGameSync/Automation/P4Automation.cs
@@ -14,14 +14,10 @@ namespace UnrealGameSync
 		public static IPerforceSettings GetConnectionSettings()
 		{
 			// Read the settings
-			string? serverAndPort = null;
-			string? userName = null;
-			string? depotPathSettings = null;
-			bool preview = false;
+			LauncherSettings launcherSettings = new LauncherSettings();
+			launcherSettings.Read();
 
-			GlobalPerforceSettings.ReadGlobalPerforceSettings(ref serverAndPort, ref userName, ref depotPathSettings, ref preview);
-
-			return Utility.OverridePerforceSettings(PerforceSettings.Default, serverAndPort, userName);
+			return Utility.OverridePerforceSettings(PerforceSettings.Default, launcherSettings.PerforceServerAndPort, launcherSettings.PerforceUserName);
 		}
 
 		public static Task<string> PrintToTempFile(IPerforceConnection? connection, string depotPath, ILogger logger)
diff --git a/Engine/Source/Programs/UnrealGameSync/UnrealGameSync/AutomationServer.cs b/Engine/Source/Programs/UnrealGameSync/UnrealGameSync/AutomationServer.cs
--- a/Engine/Source/Programs/UnrealGameSync/UnrealGameSync/AutomationServer.cs
+++ b/Engine/Source/Programs/UnrealGameSync/UnrealGameSync/AutomationServer.cs
@@ -179,7 +179,7 @@ namespace UnrealGameSync
 		{
 			if(portNumber <= 0)
 			{
-				GlobalPerforceSettings.DeleteRegistryKey(Registry.CurrentUser, "Software\\Epic Games\\UnrealGameSync", "AutomationPort");
+				Utility.DeleteRegistryKey(Registry.CurrentUser, "Software\\Epic Games\\UnrealGameSync", "AutomationPort");
 			}
 			else
 			{
diff --git a/Engine/Source/Programs/UnrealGameSync/UnrealGameSync/Forms/ApplicationSettingsWindow.cs b/Engine/Source/Programs/UnrealGameSync/UnrealGameSync/Forms/ApplicationSettingsWindow.cs
--- a/Engine/Source/Programs/UnrealGameSync/UnrealGameSync/Forms/ApplicationSettingsWindow.cs
+++ b/Engine/Source/Programs/UnrealGameSync/UnrealGameSync/Forms/ApplicationSettingsWindow.cs
@@ -95,8 +95,13 @@ namespace UnrealGameSync
 			_toolUpdateMonitor = toolUpdateMonitor;
 			_logger = logger;
 
-			GlobalPerforceSettings.ReadGlobalPerforceSettings(ref _initialServerAndPort, ref _initialUserName, ref _initialDepotPath, ref preview);
-			_initialPreview = preview;
+			LauncherSettings launcherSettings = new LauncherSettings();
+			launcherSettings.Read();
+
+			_initialServerAndPort = launcherSettings.PerforceServerAndPort;
+			_initialUserName = launcherSettings.PerforceUserName;
+			_initialDepotPath = launcherSettings.PerforceDepotPath;
+			_initialPreview = launcherSettings.PreviewBuild;
 
 			_initialAutomationPortNumber = AutomationServer.GetPortNumber();
 			_initialProtocolHandlerState = ProtocolHandlerUtils.GetState();
@@ -211,25 +216,27 @@ namespace UnrealGameSync
 		private void OkBtn_Click(object sender, EventArgs e)
 		{
 			// Update the settings
-			string? serverAndPort = ServerTextBox.Text.Trim();
-			if (serverAndPort.Length == 0)
+			LauncherSettings launcherSettings = new LauncherSettings();
+
+			launcherSettings.PerforceServerAndPort = ServerTextBox.Text.Trim();
+			if (launcherSettings.PerforceServerAndPort.Length == 0)
 			{
-				serverAndPort = null;
+				launcherSettings.PerforceServerAndPort = null;
 			}
 
-			string? userName = UserNameTextBox.Text.Trim();
-			if (userName.Length == 0)
+			launcherSettings.PerforceUserName = UserNameTextBox.Text.Trim();
+			if (launcherSettings.PerforceUserName.Length == 0)
 			{
-				userName = null;
+				launcherSettings.PerforceUserName = null;
 			}
 
-			string? depotPath = DepotPathTextBox.Text.Trim();
-			if (depotPath.Length == 0 || depotPath == DeploymentSettings.Instance.DefaultDepotPath)
+			launcherSettings.PerforceDepotPath = DepotPathTextBox.Text.Trim();
+			if (launcherSettings.PerforceDepotPath.Length == 0 || launcherSettings.PerforceDepotPath == DeploymentSettings.Instance.DefaultDepotPath)
 			{
-				depotPath = null;
+				launcherSettings.PerforceDepotPath = null;
 			}
 
-			bool preview = UsePreviewBuildCheckBox.Checked;
+			launcherSettings.PreviewBuild = UsePreviewBuildCheckBox.Checked;
 
 			int automationPortNumber;
 			if (!EnableAutomationCheckBox.Checked || !Int32.TryParse(AutomationPortTextBox.Text, out automationPortNumber))
@@ -237,14 +244,14 @@ namespace UnrealGameSync
 				automationPortNumber = -1;
 			}
 
-			if (serverAndPort != _initialServerAndPort || userName != _initialUserName || depotPath != _initialDepotPath || preview != _initialPreview || automationPortNumber != _initialAutomationPortNumber)
+			if (launcherSettings.PerforceServerAndPort != _initialServerAndPort || launcherSettings.PerforceUserName != _initialUserName || launcherSettings.PerforceDepotPath != _initialDepotPath || launcherSettings.PreviewBuild != _initialPreview || automationPortNumber != _initialAutomationPortNumber)
 			{
 				// Try to log in to the new server, and check the application is there
-				if (serverAndPort != _initialServerAndPort || userName != _initialUserName || depotPath != _initialDepotPath)
+				if (launcherSettings.PerforceServerAndPort != _initialServerAndPort || launcherSettings.PerforceUserName != _initialUserName || launcherSettings.PerforceDepotPath != _initialDepotPath)
 				{
-					PerforceSettings settings = Utility.OverridePerforceSettings(_defaultPerforceSettings, serverAndPort, userName);
+					PerforceSettings settings = Utility.OverridePerforceSettings(_defaultPerforceSettings, launcherSettings.PerforceServerAndPort, launcherSettings.PerforceUserName);
 
-					string? testDepotPath = depotPath ?? DeploymentSettings.Instance.DefaultDepotPath;
+					string? testDepotPath = launcherSettings.PerforceDepotPath ?? DeploymentSettings.Instance.DefaultDepotPath;
 					if (testDepotPath != null)
 					{
 						ModalTask? task = PerforceModalTask.Execute(this, "Checking connection", "Checking connection, please wait...", settings, (p, c) => PerforceTestConnectionTask.RunAsync(p, testDepotPath, c), _logger);
@@ -261,7 +268,7 @@ namespace UnrealGameSync
 				}
 
 				_restartPreview = UsePreviewBuildCheckBox.Checked;
-				GlobalPerforceSettings.SaveGlobalPerforceSettings(serverAndPort, userName, depotPath, preview);
+				launcherSettings.Save();
 				AutomationServer.SetPortNumber(automationPortNumber);
 			}
 
diff --git a/Engine/Source/Programs/UnrealGameSync/UnrealGameSync/GlobalSettings.cs b/Engine/Source/Programs/UnrealGameSync/UnrealGameSync/GlobalSettings.cs
deleted file mode 100644
diff --git a/Engine/Source/Programs/UnrealGameSync/UnrealGameSync/Program.cs b/Engine/Source/Programs/UnrealGameSync/UnrealGameSync/Program.cs
--- a/Engine/Source/Programs/UnrealGameSync/UnrealGameSync/Program.cs
+++ b/Engine/Source/Programs/UnrealGameSync/UnrealGameSync/Program.cs
@@ -115,11 +115,8 @@ namespace UnrealGameSync
 
 		static async Task InnerMainAsync(Mutex instanceMutex, EventWaitHandle activateEvent, string[] args)
 		{
-			string? serverAndPort = null;
-			string? userName = null;
-			string? baseUpdatePath = null;
-			bool previewSetting = false;
-			GlobalPerforceSettings.ReadGlobalPerforceSettings(ref serverAndPort, ref userName, ref baseUpdatePath, ref previewSetting);
+			LauncherSettings launcherSettings = new LauncherSettings();
+			launcherSettings.Read();
 
 			List<string> remainingArgs = new List<string>(args);
 
@@ -198,18 +195,18 @@ namespace UnrealGameSync
 					string sessionId = Guid.NewGuid().ToString();
 					logger.LogInformation("SessionId: {SessionId}", sessionId);
 
-					if (serverAndPort == null || userName == null)
+					if (launcherSettings.PerforceServerAndPort == null || launcherSettings.PerforceUserName == null)
 					{
 						logger.LogInformation("Missing server settings; finding defaults.");
-						serverAndPort ??= PerforceSettings.Default.ServerAndPort;
-						userName ??= PerforceSettings.Default.UserName;
-						GlobalPerforceSettings.SaveGlobalPerforceSettings(serverAndPort, userName, baseUpdatePath, previewSetting);
+						launcherSettings.PerforceServerAndPort ??= PerforceSettings.Default.ServerAndPort;
+						launcherSettings.PerforceUserName ??= PerforceSettings.Default.UserName;
+						launcherSettings.Save();
 					}
 
 					ILogger telemetryLogger = loggerProvider.CreateLogger("Telemetry");
 					telemetryLogger.LogInformation("Creating telemetry sink for session {SessionId}", sessionId);
 
-					using (ITelemetrySink telemetrySink = CreateTelemetrySink(userName, sessionId, telemetryLogger))
+					using (ITelemetrySink telemetrySink = CreateTelemetrySink(launcherSettings.PerforceUserName, sessionId, telemetryLogger))
 					{
 						ITelemetrySink? prevTelemetrySink = Telemetry.ActiveSink;
 						try
@@ -220,7 +217,7 @@ namespace UnrealGameSync
 
 							AppDomain.CurrentDomain.UnhandledException += CurrentDomain_UnhandledException;
 
-							IPerforceSettings defaultSettings = new PerforceSettings(serverAndPort, userName) { PreferNativeClient = true };
+							IPerforceSettings defaultSettings = new PerforceSettings(launcherSettings.PerforceServerAndPort, launcherSettings.PerforceUserName) { PreferNativeClient = true };
 
 							ProtocolHandlerUtils.InstallQuiet(logger);
 
diff --git a/Engine/Source/Programs/UnrealGameSync/UnrealGameSyncCmd/Program.cs b/Engine/Source/Programs/UnrealGameSync/UnrealGameSyncCmd/Program.cs
--- a/Engine/Source/Programs/UnrealGameSync/UnrealGameSyncCmd/Program.cs
+++ b/Engine/Source/Programs/UnrealGameSync/UnrealGameSyncCmd/Program.cs
@@ -1076,7 +1076,7 @@ namespace UnrealGameSyncCmd
 
 		class ConfigCommand : Command
 		{
-			public override Task ExecuteAsync(CommandContext context)
+			public override async Task ExecuteAsync(CommandContext context)
 			{
 				ILogger logger = context.Logger;
 
@@ -1090,7 +1090,7 @@ namespace UnrealGameSyncCmd
 					{
 						if (editor != null)
 						{
-							editor.WaitForExit();
+							await editor.WaitForExitAsync();
 						}
 					}
 				}
@@ -1105,8 +1105,6 @@ namespace UnrealGameSyncCmd
 
 					logger.LogInformation("Updated {ConfigFile}", settings.ConfigFile);
 				}
-				
-				return Task.CompletedTask;
 			}
 		}
 
diff --git a/Engine/Source/Programs/UnrealGameSync/UnrealGameSyncLauncher/Forms/SettingsWindow.cs b/Engine/Source/Programs/UnrealGameSync/UnrealGameSyncLauncher/Forms/SettingsWindow.cs
--- a/Engine/Source/Programs/UnrealGameSync/UnrealGameSyncLauncher/Forms/SettingsWindow.cs
+++ b/Engine/Source/Programs/UnrealGameSync/UnrealGameSyncLauncher/Forms/SettingsWindow.cs
@@ -60,36 +60,37 @@ namespace UnrealGameSyncLauncher
 		private void ConnectBtn_Click(object sender, EventArgs e)
 		{
 			// Update the settings
-			string? serverAndPort = ServerTextBox.Text.Trim();
-			if(serverAndPort.Length == 0)
+			LauncherSettings launcherSettings = new LauncherSettings();
+
+			launcherSettings.PerforceServerAndPort = ServerTextBox.Text.Trim();
+			if(launcherSettings.PerforceServerAndPort.Length == 0)
 			{
-				serverAndPort = null;
+				launcherSettings.PerforceServerAndPort = null;
 			}
 
-			string? userName = UserNameTextBox.Text.Trim();
-			if(userName.Length == 0)
+			launcherSettings.PerforceUserName = UserNameTextBox.Text.Trim();
+			if(launcherSettings.PerforceUserName.Length == 0)
 			{
-				userName = null;
+				launcherSettings.PerforceUserName = null;
 			}
 
-			string? depotPath = DepotPathTextBox.Text.Trim();
-			if(depotPath.Length == 0)
+			launcherSettings.PerforceDepotPath = DepotPathTextBox.Text.Trim();
+			if(launcherSettings.PerforceDepotPath.Length == 0)
 			{
-				depotPath = null;
+				launcherSettings.PerforceDepotPath = null;
 			}
 
-			bool preview = UsePreviewBuildCheckBox.Checked;
-
-			GlobalPerforceSettings.SaveGlobalPerforceSettings(serverAndPort, userName, depotPath, preview);
+			launcherSettings.PreviewBuild = UsePreviewBuildCheckBox.Checked;
+			launcherSettings.Save();
 
 			PerforceSettings perforceSettings = new PerforceSettings(PerforceSettings.Default);
-			if (!String.IsNullOrEmpty(serverAndPort))
+			if (!String.IsNullOrEmpty(launcherSettings.PerforceServerAndPort))
 			{
-				perforceSettings.ServerAndPort = serverAndPort;
+				perforceSettings.ServerAndPort = launcherSettings.PerforceServerAndPort;
 			}
-			if (!String.IsNullOrEmpty(userName))
+			if (!String.IsNullOrEmpty(launcherSettings.PerforceUserName))
 			{
-				perforceSettings.UserName = userName;
+				perforceSettings.UserName = launcherSettings.PerforceUserName;
 			}
 			perforceSettings.PreferNativeClient = true;
 
@@ -97,12 +98,12 @@ namespace UnrealGameSyncLauncher
 			CaptureLogger logger = new CaptureLogger();
 
 			// Create the task for connecting to this server
-			ModalTask? task = PerforceModalTask.Execute(this, "Updating", "Checking for updates, please wait...", perforceSettings, (p, c) => _syncAndRun(p, depotPath, preview, logger, c), logger);
+			ModalTask? task = PerforceModalTask.Execute(this, "Updating", "Checking for updates, please wait...", perforceSettings, (p, c) => _syncAndRun(p, launcherSettings.PerforceDepotPath, launcherSettings.PreviewBuild, logger, c), logger);
 			if (task != null)
 			{
 				if(task.Succeeded)
 				{
-					GlobalPerforceSettings.SaveGlobalPerforceSettings(serverAndPort, userName, depotPath, preview);
+					launcherSettings.Save();
 					DialogResult = DialogResult.OK;
 					Close();
 				}
diff --git a/Engine/Source/Programs/UnrealGameSync/UnrealGameSyncLauncher/Program.cs b/Engine/Source/Programs/UnrealGameSync/UnrealGameSyncLauncher/Program.cs
--- a/Engine/Source/Programs/UnrealGameSync/UnrealGameSyncLauncher/Program.cs
+++ b/Engine/Source/Programs/UnrealGameSync/UnrealGameSyncLauncher/Program.cs
@@ -44,17 +44,16 @@ namespace UnrealGameSyncLauncher
 				bool preview = args.Contains("-unstable", StringComparer.InvariantCultureIgnoreCase) || args.Contains("-preview", StringComparer.InvariantCultureIgnoreCase);
 
 				// Read the settings
-				string? serverAndPort = null;
-				string? userName = null;
-				string? depotPath = DeploymentSettings.Instance.DefaultDepotPath;
-				GlobalPerforceSettings.ReadGlobalPerforceSettings(ref serverAndPort, ref userName, ref depotPath, ref preview);
+				LauncherSettings launcherSettings = new LauncherSettings();
+				launcherSettings.PerforceDepotPath = DeploymentSettings.Instance.DefaultDepotPath;
+				launcherSettings.Read();
 
 				// If the shift key is held down, immediately show the settings window
 				Task SyncAndRunWrapper(IPerforceConnection perforce, string? depotPath, bool preview, ILogger logWriter, CancellationToken cancellationToken) => SyncAndRun(perforce, depotPath, preview, args, instanceMutex, logWriter, cancellationToken);
 				if ((Control.ModifierKeys & Keys.Shift) != 0)
 				{
 					// Show the settings window immediately
-					using SettingsWindow updateError = new SettingsWindow(null, null, serverAndPort, userName, depotPath, preview, SyncAndRunWrapper);
+					using SettingsWindow updateError = new SettingsWindow(null, null, launcherSettings.PerforceServerAndPort, launcherSettings.PerforceUserName, launcherSettings.PerforceDepotPath, preview, SyncAndRunWrapper);
 					if(updateError.ShowDialog() == DialogResult.OK)
 					{
 						return 0;
@@ -65,9 +64,9 @@ namespace UnrealGameSyncLauncher
 					// Try to do a sync with the current settings first
 					CaptureLogger logger = new CaptureLogger();
 
-					IPerforceSettings settings = new PerforceSettings(PerforceSettings.Default) { PreferNativeClient = true }.MergeWith(newServerAndPort: serverAndPort, newUserName: userName);
+					IPerforceSettings settings = new PerforceSettings(PerforceSettings.Default) { PreferNativeClient = true }.MergeWith(newServerAndPort: launcherSettings.PerforceServerAndPort, newUserName: launcherSettings.PerforceUserName);
 
-					ModalTask? task = PerforceModalTask.Execute(null, "Updating", "Checking for updates, please wait...", settings, (p, c) => SyncAndRun(p, depotPath, preview, args, instanceMutex, logger, c), logger);
+					ModalTask? task = PerforceModalTask.Execute(null, "Updating", "Checking for updates, please wait...", settings, (p, c) => SyncAndRun(p, launcherSettings.PerforceDepotPath, preview, args, instanceMutex, logger, c), logger);
 					if (task == null)
 					{
 						logger.LogInformation("Canceled by user");
@@ -77,7 +76,7 @@ namespace UnrealGameSyncLauncher
 						return 0;
 					}
 
-					using SettingsWindow updateError = new SettingsWindow("Unable to update UnrealGameSync from Perforce. Verify that your connection settings are correct.", logger.Render(Environment.NewLine), serverAndPort, userName, depotPath, preview, SyncAndRunWrapper);
+					using SettingsWindow updateError = new SettingsWindow("Unable to update UnrealGameSync from Perforce. Verify that your connection settings are correct.", logger.Render(Environment.NewLine), launcherSettings.PerforceServerAndPort, launcherSettings.PerforceUserName, launcherSettings.PerforceDepotPath, preview, SyncAndRunWrapper);
 					if(updateError.ShowDialog() == DialogResult.OK)
 					{
 						return 0;
@@ -144,7 +143,7 @@ namespace UnrealGameSyncLauncher
 						{
 							throw new UserErrorException($"Couldn't delete contents of {applicationFolder} (retried {numRetries} times).");
 						}
-						Thread.Sleep(500);
+						await Task.Delay(500, cancellationToken);
 					}
 
 					// Find all the files in the sync path at this changelist
@@ -251,7 +250,7 @@ namespace UnrealGameSyncLauncher
 
 		static string QuoteArgument(string arg)
 		{
-			if(arg.IndexOf(' ', StringComparison.Ordinal) != -1 && !arg.StartsWith("\"", StringComparison.Ordinal))
+			if(arg.Contains(' ', StringComparison.Ordinal) && !arg.StartsWith("\"", StringComparison.Ordinal))
 			{
 				return String.Format("\"{0}\"", arg);
 			}
diff --git a/Engine/Source/Programs/UnrealGameSync/UnrealGameSyncLauncher/UnrealGameSyncLauncher.csproj b/Engine/Source/Programs/UnrealGameSync/UnrealGameSyncLauncher/UnrealGameSyncLauncher.csproj
--- a/Engine/Source/Programs/UnrealGameSync/UnrealGameSyncLauncher/UnrealGameSyncLauncher.csproj
+++ b/Engine/Source/Programs/UnrealGameSync/UnrealGameSyncLauncher/UnrealGameSyncLauncher.csproj
@@ -38,9 +38,6 @@
     <Compile Include="..\UnrealGameSync\PerforceModalTask.cs">
       <Link>UnrealGameSync\PerforceModalTask.cs</Link>
     </Compile>
-    <Compile Include="..\UnrealGameSync\GlobalSettings.cs">
-      <Link>UnrealGameSync\GlobalSettings.cs</Link>
-    </Compile>
   </ItemGroup>
   <ItemGroup>
     <ProjectReference Include="..\..\Shared\EpicGames.Core\EpicGames.Core.csproj" />
diff --git a/Engine/Source/Programs/UnrealGameSync/UnrealGameSyncShared/LauncherSettings.cs b/Engine/Source/Programs/UnrealGameSync/UnrealGameSyncShared/LauncherSettings.cs
new file mode 100644
--- /dev/null
+++ b/Engine/Source/Programs/UnrealGameSync/UnrealGameSyncShared/LauncherSettings.cs
@@ -0,0 +1,116 @@
+// Copyright Epic Games, Inc. All Rights Reserved.
+
+using Microsoft.Win32;
+using System;
+using System.Linq;
+using System.Runtime.InteropServices;
+using System.Runtime.Versioning;
+using System.Text.RegularExpressions;
+
+namespace UnrealGameSync
+{
+	public class LauncherSettings
+	{
+		public string? PerforceServerAndPort { get; set; }
+		public string? PerforceUserName { get; set; }
+		public string? PerforceDepotPath { get; set; }
+
+		public bool PreviewBuild { get; set; }
+
+		public void Read()
+		{
+			if (RuntimeInformation.IsOSPlatform(OSPlatform.Windows))
+			{
+				ReadFromRegistry();
+			}
+		}
+
+		[SupportedOSPlatform("windows")]
+		void ReadFromRegistry()
+		{
+			using (RegistryKey? key = Registry.CurrentUser.OpenSubKey("SOFTWARE\\Epic Games\\UnrealGameSync", false))
+			{
+				if (key != null)
+				{
+					PerforceServerAndPort = key.GetValue("ServerAndPort", PerforceServerAndPort) as string;
+					PerforceUserName = key.GetValue("UserName", PerforceUserName) as string;
+					PerforceDepotPath = key.GetValue("DepotPath", PerforceDepotPath) as string;
+					PreviewBuild = ((key.GetValue("Preview", PreviewBuild? 1 : 0) as int?) ?? 0) != 0;
+
+					// Fix corrupted depot path string
+					if (PerforceDepotPath != null)
+					{
+						Match match = Regex.Match(PerforceDepotPath, "^(.*)/(Release|UnstableRelease)/\\.\\.\\.@.*$");
+						if (match.Success)
+						{
+							PerforceDepotPath = match.Groups[1].Value;
+							Save();
+						}
+					}
+				}
+			}
+		}
+
+		public bool Save()
+		{
+			try
+			{
+				if (RuntimeInformation.IsOSPlatform(OSPlatform.Windows))
+				{
+					SaveToRegistry();
+				}
+				return true;
+			}
+			catch
+			{
+				return false;
+			}
+		}
+
+		[SupportedOSPlatform("windows")]
+		void SaveToRegistry()
+		{
+			using (RegistryKey key = Registry.CurrentUser.CreateSubKey("SOFTWARE\\Epic Games\\UnrealGameSync"))
+			{
+				// Delete this legacy setting
+				Utility.DeleteRegistryKey(key, "Server");
+
+				if (String.IsNullOrEmpty(PerforceServerAndPort))
+				{
+					Utility.DeleteRegistryKey(key, "ServerAndPort");
+				}
+				else
+				{
+					key.SetValue("ServerAndPort", PerforceServerAndPort);
+				}
+
+				if (String.IsNullOrEmpty(PerforceUserName))
+				{
+					Utility.DeleteRegistryKey(key, "UserName");
+				}
+				else
+				{
+					key.SetValue("UserName", PerforceUserName);
+				}
+
+				if (String.IsNullOrEmpty(PerforceDepotPath) || (DeploymentSettings.Instance.DefaultDepotPath != null && String.Equals(PerforceDepotPath, DeploymentSettings.Instance.DefaultDepotPath, StringComparison.OrdinalIgnoreCase)))
+				{
+					Utility.DeleteRegistryKey(key, "DepotPath");
+				}
+				else
+				{
+					key.SetValue("DepotPath", PerforceDepotPath);
+				}
+
+				if (PreviewBuild)
+				{
+					key.SetValue("Preview", 1);
+				}
+				else
+				{
+					Utility.DeleteRegistryKey(key, "Preview");
+				}
+			}
+		}
+	}
+}
diff --git a/Engine/Source/Programs/UnrealGameSync/UnrealGameSyncShared/Utility.cs b/Engine/Source/Programs/UnrealGameSync/UnrealGameSyncShared/Utility.cs
--- a/Engine/Source/Programs/UnrealGameSync/UnrealGameSyncShared/Utility.cs
+++ b/Engine/Source/Programs/UnrealGameSync/UnrealGameSyncShared/Utility.cs
@@ -5,6 +5,7 @@ using EpicGames.Perforce;
 using JetBrains.Annotations;
 using Microsoft.Extensions.Caching.Memory;
 using Microsoft.Extensions.Logging;
+using Microsoft.Win32;
 using System;
 using System.Collections.Generic;
 using System.Diagnostics;
@@ -13,6 +14,7 @@ using System.Drawing;
 using System.IO;
 using System.Linq;
 using System.Runtime.CompilerServices;
+using System.Runtime.Versioning;
 using System.Text;
 using System.Text.Json;
 using System.Text.Json.Serialization;
@@ -893,5 +895,33 @@ namespace UnrealGameSync
 			startInfo.UseShellExecute = true;
 			using Process? _ = Process.Start(startInfo);
 		}
+
+		[SupportedOSPlatform("windows")]
+		public static void DeleteRegistryKey(RegistryKey rootKey, string keyName, string valueName)
+		{
+			using (RegistryKey? key = rootKey.OpenSubKey(keyName, true))
+			{
+				if (key != null)
+				{
+					DeleteRegistryKey(key, valueName);
+				}
+			}
+		}
+
+		[SupportedOSPlatform("windows")]
+		public static void DeleteRegistryKey(RegistryKey key, string name)
+		{
+			string[] valueNames = key.GetValueNames();
+			if (valueNames.Any(x => String.Equals(x, name, StringComparison.OrdinalIgnoreCase)))
+			{
+				try
+				{
+					key.DeleteValue(name);
+				}
+				catch
+				{
+				}
+			}
+		}
 	}
 }
