Horde: Add a build step to publish UGS binaries to Horde.

diff --git a/Engine/Source/Programs/UnrealGameSync/BuildUGS.xml b/Engine/Source/Programs/UnrealGameSync/BuildUGS.xml
--- a/Engine/Source/Programs/UnrealGameSync/BuildUGS.xml
+++ b/Engine/Source/Programs/UnrealGameSync/BuildUGS.xml
@@ -62,7 +62,10 @@
 			<DotNet Arguments="publish &quot;$(UgsDir)/UnrealGameSyncCmd/UnrealGameSyncCmd.csproj&quot; --output=&quot;$(UgsPublishDir)&quot; --self-contained -c Release -r win-x64 $(VersionArguments)"/>
 			<DotNet Arguments="publish &quot;$(UgsDir)/UnrealGameSync/UnrealGameSync.csproj&quot; --output=&quot;$(UgsPublishDir)&quot; --self-contained -c Release -r win-x64 $(VersionArguments)"/>
 			<Copy From="$(InstallerDir)/bin/Release/UnrealGameSync.msi" To="$(UgsPublishDir)/UnrealGameSync.msi" If="$(IncludeInstaller)"/>
+			<Zip FromDir="$(UgsOutputDir)" ZipFile="$(UgsZipPublishDir)/Ugs.zip"/>
 		</Node>
+
+		<!-- Windows (Preview) -->
 		<Node Name="Publish Unstable" Requires="Build UnrealGameSync">
 			<Property Name="UnstableBranch" Value="" />
 			<Property Name="UnstableBranch" Value="$(SubmitToBranch)/UnstableRelease" If="'$(SubmitToBranch)' != ''" />
@@ -73,6 +76,8 @@
 			<Zip FromDir="$(UgsOutputDir)" ZipFile="$(UgsZipPublishDir)/UnstableRelease.zip"/>
 			<Submit Description="Updating UGS from CL $(Change) (Zip)" Workspace="Temp" Files="$(UgsZipPublishDir)/UnstableRelease.zip" RootDir="$(UgsZipOutputDir)" Stream="$(SubmitToStream)" Branch="$(SubmitToBranch)"/>
 		</Node>
+
+		<!-- Windows (Regular) -->
 		<Node Name="Publish Live" Requires="Build UnrealGameSync">
 			<Property Name="ReleaseBranch" Value="" />
 			<Property Name="ReleaseBranch" Value="$(SubmitToBranch)/Release" If="'$(SubmitToBranch)' != ''" />
@@ -83,7 +88,12 @@
 			<Zip FromDir="$(UgsOutputDir)" ZipFile="$(UgsZipPublishDir)/Release.zip"/>
 			<Submit Description="Updating UGS from CL $(Change) (Zip)" Workspace="Temp" Files="$(UgsZipPublishDir)/Release.zip" RootDir="$(UgsZipOutputDir)" Stream="$(SubmitToStream)" Branch="$(SubmitToBranch)"/>
 		</Node>
+		<Node Name="Publish Live Horde" Requires="Build UnrealGameSync">
+			<Error Message="Missing -HordeSettings=... argument" If="'$(HordeSettings)' == ''"/>
+			<DeployTool Id="ugs-win" Version="$(InformationalVersion)" Settings="$(HordeSettings)" File="$(UgsZipPublishDir)/Ugs.zip"/>
+		</Node>
 		
+		<!-- Mac -->
 		<Node Name="Build UnrealGameSyncCmd Mac">
 			<Delete Files="$(UgsCmdMacOutputDir)/..."/>
 			<DotNet Arguments="publish &quot;$(UgsDir)/UnrealGameSyncCmd/UnrealGameSyncCmd.csproj&quot; --output=&quot;$(UgsCmdMacOutputDir)&quot; --self-contained -c Release -r osx-x64 $(VersionArguments)"/>
@@ -94,6 +104,7 @@
 			<DeployTool Id="ugs-mac" Version="$(InformationalVersion)" Settings="$(HordeSettings)" File="$(UgsCmdMacZipFile)"/>
 		</Node>
 		
+		<!-- Linux -->
 		<Node Name="Build UnrealGameSyncCmd Linux">
 			<Delete Files="$(UgsCmdLinuxOutputDir)/..."/>
 			<DotNet Arguments="publish &quot;$(UgsDir)/UnrealGameSyncCmd/UnrealGameSyncCmd.csproj&quot; --output=&quot;$(UgsCmdLinuxOutputDir)&quot; --self-contained -c Release -r linux-x64 $(VersionArguments)"/>
