UGS: Various static analysis fixes.

diff --git a/Engine/Source/Programs/UnrealGameSync/UnrealGameSync/AsyncDisposer.cs b/Engine/Source/Programs/UnrealGameSync/UnrealGameSync/AsyncDisposer.cs
--- a/Engine/Source/Programs/UnrealGameSync/UnrealGameSync/AsyncDisposer.cs
+++ b/Engine/Source/Programs/UnrealGameSync/UnrealGameSync/AsyncDisposer.cs
@@ -3,7 +3,6 @@
 using Microsoft.Extensions.Logging;
 using System;
 using System.Collections.Generic;
-using System.Threading;
 using System.Threading.Tasks;
 
 namespace UnrealGameSync
diff --git a/Engine/Source/Programs/UnrealGameSync/UnrealGameSync/Controls/WorkspaceControl.cs b/Engine/Source/Programs/UnrealGameSync/UnrealGameSync/Controls/WorkspaceControl.cs
--- a/Engine/Source/Programs/UnrealGameSync/UnrealGameSync/Controls/WorkspaceControl.cs
+++ b/Engine/Source/Programs/UnrealGameSync/UnrealGameSync/Controls/WorkspaceControl.cs
@@ -2085,8 +2085,6 @@ namespace UnrealGameSync
 		/// Determines if the ini specified INI filter is a subset of the InBadges passed
 		/// BadgeGroupSyncFilter - is used as the ini 
 		/// </summary>
-		/// <param name="inBadges"></param>
-		/// <returns></returns>
 		private static bool DoRequiredBadgesExist(List<string> requiredBadgeList, List<BadgeData> inBadges)
 		{
 			Dictionary<string, BadgeData> inBadgeDictionary = new Dictionary<string, BadgeData>();
diff --git a/Engine/Source/Programs/UnrealGameSync/UnrealGameSync/EventMonitor.cs b/Engine/Source/Programs/UnrealGameSync/UnrealGameSync/EventMonitor.cs
--- a/Engine/Source/Programs/UnrealGameSync/UnrealGameSync/EventMonitor.cs
+++ b/Engine/Source/Programs/UnrealGameSync/UnrealGameSync/EventMonitor.cs
@@ -881,7 +881,7 @@ namespace UnrealGameSync
 				_logger.LogInformation("Polling for events...");
 
 				//////////////
-				/// Initial Ids 
+				// Initial Ids 
 				//////////////
 				if (_apiVersion == 0)
 				{
@@ -933,7 +933,7 @@ namespace UnrealGameSync
 				else
 				{
 					//////////////
-					/// Builds
+					// Builds
 					//////////////
 					List<BadgeData> builds = await RestApi.GetAsync<List<BadgeData>>($"{_apiUrl}/api/build?project={_project}&lastbuildid={_latestIds.LastBuildId}", cancellationToken);
 					foreach (BadgeData build in builds)
@@ -943,12 +943,12 @@ namespace UnrealGameSync
 					}
 
 					//////////////////////////
-					/// Throttled Requests
+					// Throttled Requests
 					//////////////////////////
 					if (fireThrottledRequests)
 					{
 						//////////////
-						/// Reviews 
+						// Reviews 
 						//////////////
 						List<EventData> events = await RestApi.GetAsync<List<EventData>>($"{_apiUrl}/api/event?project={_project}&lasteventid={_latestIds.LastEventId}", cancellationToken);
 						foreach (EventData review in events)
@@ -958,7 +958,7 @@ namespace UnrealGameSync
 						}
 
 						//////////////
-						/// Comments 
+						// Comments 
 						//////////////
 						List<CommentData> comments = await RestApi.GetAsync<List<CommentData>>($"{_apiUrl}/api/comment?project={_project}&lastcommentid={_latestIds.LastCommentId}", cancellationToken);
 						foreach (CommentData comment in comments)
diff --git a/Engine/Source/Programs/UnrealGameSync/UnrealGameSync/Forms/DownloadProgressWindow.cs b/Engine/Source/Programs/UnrealGameSync/UnrealGameSync/Forms/DownloadProgressWindow.cs
--- a/Engine/Source/Programs/UnrealGameSync/UnrealGameSync/Forms/DownloadProgressWindow.cs
+++ b/Engine/Source/Programs/UnrealGameSync/UnrealGameSync/Forms/DownloadProgressWindow.cs
@@ -1,12 +1,6 @@
 // Copyright Epic Games, Inc. All Rights Reserved.
 
 using System;
-using System.Collections.Generic;
-using System.ComponentModel;
-using System.Data;
-using System.Drawing;
-using System.Linq;
-using System.Text;
 using System.Threading;
 using System.Threading.Tasks;
 using System.Windows.Forms;
@@ -33,7 +27,7 @@ namespace UnrealGameSync.Forms
 		}
 
 		readonly Func<IProgress<string>, CancellationToken, Task> _taskFunc;
-		CancellationTokenSource _cancellationSource;
+		readonly CancellationTokenSource _cancellationSource;
 		Task _task = Task.CompletedTask;
 
 		public Task Task => _task;
diff --git a/Engine/Source/Programs/UnrealGameSync/UnrealGameSync/Forms/DownloadSettingsWindow.cs b/Engine/Source/Programs/UnrealGameSync/UnrealGameSync/Forms/DownloadSettingsWindow.cs
--- a/Engine/Source/Programs/UnrealGameSync/UnrealGameSync/Forms/DownloadSettingsWindow.cs
+++ b/Engine/Source/Programs/UnrealGameSync/UnrealGameSync/Forms/DownloadSettingsWindow.cs
@@ -1,15 +1,7 @@
 ﻿// Copyright Epic Games, Inc. All Rights Reserved.
 
 using System;
-using System.Collections.Generic;
-using System.ComponentModel;
-using System.Data;
 using System.Diagnostics.CodeAnalysis;
-using System.Drawing;
-using System.IO;
-using System.Linq;
-using System.Text;
-using System.Threading.Tasks;
 using System.Windows.Forms;
 
 namespace UnrealGameSync.Forms
diff --git a/Engine/Source/Programs/UnrealGameSync/UnrealGameSync/Forms/FailedToDeleteWindow.cs b/Engine/Source/Programs/UnrealGameSync/UnrealGameSync/Forms/FailedToDeleteWindow.cs
--- a/Engine/Source/Programs/UnrealGameSync/UnrealGameSync/Forms/FailedToDeleteWindow.cs
+++ b/Engine/Source/Programs/UnrealGameSync/UnrealGameSync/Forms/FailedToDeleteWindow.cs
@@ -1,13 +1,5 @@
 // Copyright Epic Games, Inc. All Rights Reserved.
 
-using System;
-using System.Collections.Generic;
-using System.ComponentModel;
-using System.Data;
-using System.Drawing;
-using System.Linq;
-using System.Text;
-using System.Threading.Tasks;
 using System.Windows.Forms;
 
 namespace UnrealGameSync
diff --git a/Engine/Source/Programs/UnrealGameSync/UnrealGameSync/Forms/SyncFilterSyntax.cs b/Engine/Source/Programs/UnrealGameSync/UnrealGameSync/Forms/SyncFilterSyntax.cs
--- a/Engine/Source/Programs/UnrealGameSync/UnrealGameSync/Forms/SyncFilterSyntax.cs
+++ b/Engine/Source/Programs/UnrealGameSync/UnrealGameSync/Forms/SyncFilterSyntax.cs
@@ -1,13 +1,5 @@
 // Copyright Epic Games, Inc. All Rights Reserved.
 
-using System;
-using System.Collections.Generic;
-using System.ComponentModel;
-using System.Data;
-using System.Drawing;
-using System.Linq;
-using System.Text;
-using System.Threading.Tasks;
 using System.Windows.Forms;
 
 namespace UnrealGameSync
diff --git a/Engine/Source/Programs/UnrealGameSync/UnrealGameSync/PerforceMonitor.cs b/Engine/Source/Programs/UnrealGameSync/UnrealGameSync/PerforceMonitor.cs
--- a/Engine/Source/Programs/UnrealGameSync/UnrealGameSync/PerforceMonitor.cs
+++ b/Engine/Source/Programs/UnrealGameSync/UnrealGameSync/PerforceMonitor.cs
@@ -10,7 +10,6 @@ using System.Collections.Generic;
 using System.Diagnostics;
 using System.Diagnostics.CodeAnalysis;
 using System.Linq;
-using System.Text;
 using System.Threading;
 using System.Threading.Tasks;
 
diff --git a/Engine/Source/Programs/UnrealGameSync/UnrealGameSync/Program.cs b/Engine/Source/Programs/UnrealGameSync/UnrealGameSync/Program.cs
--- a/Engine/Source/Programs/UnrealGameSync/UnrealGameSync/Program.cs
+++ b/Engine/Source/Programs/UnrealGameSync/UnrealGameSync/Program.cs
@@ -11,7 +11,6 @@ using System;
 using System.Collections.Generic;
 using System.Diagnostics.CodeAnalysis;
 using System.IO;
-using System.Net;
 using System.Reflection;
 using System.Threading;
 using System.Threading.Tasks;
diff --git a/Engine/Source/Programs/UnrealGameSync/UnrealGameSync/ProgramApplicationContext.cs b/Engine/Source/Programs/UnrealGameSync/UnrealGameSync/ProgramApplicationContext.cs
--- a/Engine/Source/Programs/UnrealGameSync/UnrealGameSync/ProgramApplicationContext.cs
+++ b/Engine/Source/Programs/UnrealGameSync/UnrealGameSync/ProgramApplicationContext.cs
@@ -15,8 +15,6 @@ using System.Threading;
 using System.Threading.Tasks;
 using System.Windows.Forms;
 using EpicGames.OIDC;
-using UnrealGameSync.Forms;
-using Microsoft.Extensions.Logging.Abstractions;
 
 #nullable enable
 
diff --git a/Engine/Source/Programs/UnrealGameSync/UnrealGameSync/UriHandlers/HordeHandler.cs b/Engine/Source/Programs/UnrealGameSync/UnrealGameSync/UriHandlers/HordeHandler.cs
--- a/Engine/Source/Programs/UnrealGameSync/UnrealGameSync/UriHandlers/HordeHandler.cs
+++ b/Engine/Source/Programs/UnrealGameSync/UnrealGameSync/UriHandlers/HordeHandler.cs
@@ -1,10 +1,7 @@
 // Copyright Epic Games, Inc. All Rights Reserved.
 
 using System;
-using System.Collections.Generic;
 using System.IO;
-using System.Linq;
-using System.Text;
 using System.Threading;
 using System.Threading.Tasks;
 using System.Windows.Forms;
@@ -15,7 +12,6 @@ using EpicGames.Horde.Storage.Clients;
 using EpicGames.Horde.Storage.Nodes;
 using Microsoft.Extensions.DependencyInjection;
 using Microsoft.Extensions.Logging;
-using Microsoft.Extensions.Logging.Abstractions;
 using UnrealGameSync.Forms;
 
 namespace UnrealGameSync.UriHandlers
diff --git a/Engine/Source/Programs/UnrealGameSync/UnrealGameSyncAutomation/Properties/AssemblyInfo.cs b/Engine/Source/Programs/UnrealGameSync/UnrealGameSyncAutomation/Properties/AssemblyInfo.cs
--- a/Engine/Source/Programs/UnrealGameSync/UnrealGameSyncAutomation/Properties/AssemblyInfo.cs
+++ b/Engine/Source/Programs/UnrealGameSync/UnrealGameSyncAutomation/Properties/AssemblyInfo.cs
@@ -1,7 +1,6 @@
 // Copyright Epic Games, Inc. All Rights Reserved.
 
 using System.Reflection;
-using System.Runtime.CompilerServices;
 using System.Runtime.InteropServices;
 
 // General Information about an assembly is controlled through the following
diff --git a/Engine/Source/Programs/UnrealGameSync/UnrealGameSyncShared/ConfigFile.cs b/Engine/Source/Programs/UnrealGameSync/UnrealGameSyncShared/ConfigFile.cs
--- a/Engine/Source/Programs/UnrealGameSync/UnrealGameSyncShared/ConfigFile.cs
+++ b/Engine/Source/Programs/UnrealGameSync/UnrealGameSyncShared/ConfigFile.cs
@@ -11,7 +11,6 @@ using System.Linq;
 using System.Text;
 using System.Text.Json;
 using System.Text.Json.Serialization;
-using System.Threading.Tasks;
 
 namespace UnrealGameSync
 {
diff --git a/Engine/Source/Programs/UnrealGameSync/UnrealGameSyncShared/FileFilter.cs b/Engine/Source/Programs/UnrealGameSync/UnrealGameSyncShared/FileFilter.cs
--- a/Engine/Source/Programs/UnrealGameSync/UnrealGameSyncShared/FileFilter.cs
+++ b/Engine/Source/Programs/UnrealGameSync/UnrealGameSyncShared/FileFilter.cs
@@ -78,7 +78,7 @@ namespace UnrealGameSync
 		/// <summary>
 		/// Adds an include or exclude rule to the filter
 		/// </summary>
-		/// <param name="Pattern">Pattern to match. See CreateRegex() for details.</param>
+		/// <param name="rule">Pattern to match. See CreateRegex() for details.</param>
 		public void AddRule(string rule)
 		{
 			if (rule.StartsWith("-", StringComparison.Ordinal))
@@ -94,7 +94,8 @@ namespace UnrealGameSync
 		/// <summary>
 		/// Adds an include or exclude rule to the filter. The rule may be 
 		/// </summary>
-		/// <param name="Pattern">Pattern to match. See CreateRegex() for details.</param>
+		/// <param name="rule">Pattern to match. See CreateRegex() for details.</param>
+		/// <param name="allowTags"></param>
 		public void AddRule(string rule, params string[] allowTags)
 		{
 			string cleanRule = rule.Trim();
@@ -123,7 +124,7 @@ namespace UnrealGameSync
 		/// <summary>
 		/// Adds several rules to the filter
 		/// </summary>
-		/// <param name="Patterns">List of patterns to match.</param>
+		/// <param name="rules">List of patterns to match.</param>
 		public void AddRules(IEnumerable<string> rules)
 		{
 			foreach(string rule in rules)
@@ -136,8 +137,6 @@ namespace UnrealGameSync
 		/// Adds several rules in the given lines. Rules may be prefixed with conditions of the syntax {Key=Value, Key2=Value2}, which
 		/// will be evaluated using variables in the given dictionary before being added.
 		/// </summary>
-		/// <param name="Lines"></param>
-		/// <param name="Variables">Lookup for variables to test against</param>
 		public void AddRules(IEnumerable<string> rules, params string[] tags)
 		{
 			foreach(string rule in rules)
@@ -149,9 +148,6 @@ namespace UnrealGameSync
 		/// <summary>
 		/// Reads a configuration file split into sections
 		/// </summary>
-		/// <param name="Filter"></param>
-		/// <param name="RulesFileName"></param>
-		/// <param name="Conditions"></param>
 		public void ReadRulesFromFile(string fileName, string sectionName, params string[] allowTags)
 		{
 			bool inSection = false;
@@ -218,7 +214,7 @@ namespace UnrealGameSync
 		/// Adds an include or exclude rule to the filter
 		/// </summary>
 		/// <param name="pattern">The pattern which the rule should match</param>
-		/// <param name="bInclude">Whether to include or exclude files matching this rule</param>
+		/// <param name="type">Whether to include or exclude files matching this rule</param>
 		public void AddRule(string pattern, FileFilterType type)
 		{
 			string normalizedPattern = pattern.Replace('\\', '/');
diff --git a/Engine/Source/Programs/UnrealGameSync/UnrealGameSyncShared/LauncherSettings.cs b/Engine/Source/Programs/UnrealGameSync/UnrealGameSyncShared/LauncherSettings.cs
--- a/Engine/Source/Programs/UnrealGameSync/UnrealGameSyncShared/LauncherSettings.cs
+++ b/Engine/Source/Programs/UnrealGameSync/UnrealGameSyncShared/LauncherSettings.cs
@@ -2,7 +2,6 @@
 
 using Microsoft.Win32;
 using System;
-using System.Linq;
 using System.Runtime.InteropServices;
 using System.Runtime.Versioning;
 using System.Text.RegularExpressions;
diff --git a/Engine/Source/Programs/UnrealGameSync/UnrealGameSyncShared/Utility.cs b/Engine/Source/Programs/UnrealGameSync/UnrealGameSyncShared/Utility.cs
--- a/Engine/Source/Programs/UnrealGameSync/UnrealGameSyncShared/Utility.cs
+++ b/Engine/Source/Programs/UnrealGameSync/UnrealGameSyncShared/Utility.cs
@@ -2,7 +2,6 @@
 
 using EpicGames.Core;
 using EpicGames.Perforce;
-using JetBrains.Annotations;
 using Microsoft.Extensions.Caching.Memory;
 using Microsoft.Extensions.Logging;
 using Microsoft.Win32;
@@ -587,8 +586,6 @@ namespace UnrealGameSync
 		/// <summary>
 		/// Determines if a project is an enterprise project
 		/// </summary>
-		/// <param name="FileName">Path to the project file</param>
-		/// <returns>True if the given filename is an enterprise project</returns>
 		public static bool IsEnterpriseProjectFromText(string text)
 		{
 			try
@@ -743,7 +740,6 @@ namespace UnrealGameSync
 				}
 			}
 
-
 			DirectoryReference.CreateDirectory(cacheFolder);
 
 			FileReference tempFile = new FileReference(String.Format("{0}.{1}.temp", cacheFile.FullName, Guid.NewGuid()));
diff --git a/Engine/Source/Programs/UnrealGameSync/UnrealGameSyncShared/WorkspaceState.cs b/Engine/Source/Programs/UnrealGameSync/UnrealGameSyncShared/WorkspaceState.cs
--- a/Engine/Source/Programs/UnrealGameSync/UnrealGameSyncShared/WorkspaceState.cs
+++ b/Engine/Source/Programs/UnrealGameSync/UnrealGameSyncShared/WorkspaceState.cs
@@ -5,7 +5,6 @@ using System;
 using System.Collections.Generic;
 using System.IO;
 using System.Linq;
-using System.Reflection;
 using System.Threading;
 
 namespace UnrealGameSync
@@ -235,7 +234,6 @@ namespace UnrealGameSync
 		/// <summary>
 		/// Constructor
 		/// </summary>
-		/// <param name="file"></param>
 		public WorkspaceStateWrapper(DirectoryReference rootDir, Func<WorkspaceState> createNew)
 		{
 			RootDir = rootDir;
