UGS: Add an "Analyze" configuration to the UGS solution that builds all projects with static analysis enabled. Disable analyzers during regular builds to improve build times.

diff --git a/Engine/Source/Programs/UnrealGameSync/.editorconfig b/Engine/Source/Programs/UnrealGameSync/.editorconfig
--- a/Engine/Source/Programs/UnrealGameSync/.editorconfig
+++ b/Engine/Source/Programs/UnrealGameSync/.editorconfig
@@ -9,6 +9,9 @@ csharp_style_prefer_switch_expression = false
 csharp_style_implicit_object_creation_when_type_is_apparent = false
 dotnet_remove_unnecessary_suppression_exclusions = all
 
+# CA1591: Missing comment for XMLDOC
+dotnet_diagnostic.CS1591.severity = none
+
 # Sort using and Import directives with System.* appearing first
 dotnet_sort_system_directives_first = true
 dotnet_separate_import_directive_groups = false
diff --git a/Engine/Source/Programs/UnrealGameSync/UnrealGameSync/UnrealGameSync.csproj b/Engine/Source/Programs/UnrealGameSync/UnrealGameSync/UnrealGameSync.csproj
--- a/Engine/Source/Programs/UnrealGameSync/UnrealGameSync/UnrealGameSync.csproj
+++ b/Engine/Source/Programs/UnrealGameSync/UnrealGameSync/UnrealGameSync.csproj
@@ -14,7 +14,18 @@
     <ApplicationManifest>Properties\app.manifest</ApplicationManifest>
     <Nullable>enable</Nullable>
     <NeutralLanguage>en-US</NeutralLanguage>
+    <RunAnalyzersDuringBuild>False</RunAnalyzersDuringBuild>
+    <Configurations>Debug;Release;Analyze</Configurations>
   </PropertyGroup>
+
+  <PropertyGroup Condition=" '$(Configuration)|$(Platform)' == 'Analyze|AnyCPU' ">
+    <RunAnalyzersDuringBuild>True</RunAnalyzersDuringBuild>
+    <TreatWarningsAsErrors>false</TreatWarningsAsErrors>
+    <EnforceCodeStyleInBuild>True</EnforceCodeStyleInBuild>
+    <GenerateDocumentationFile>True</GenerateDocumentationFile>
+    <DocumentationFile></DocumentationFile>
+  </PropertyGroup>
+
   <PropertyGroup Condition="'$(Configuration)|$(Platform)'=='Release|AnyCPU'">
     <DefineConstants>TRACE;DEBUG;WITH_TELEMETRY</DefineConstants>
   </PropertyGroup>
diff --git a/Engine/Source/Programs/UnrealGameSync/UnrealGameSyncAutomation/UnrealGameSyncAutomation.csproj b/Engine/Source/Programs/UnrealGameSync/UnrealGameSyncAutomation/UnrealGameSyncAutomation.csproj
--- a/Engine/Source/Programs/UnrealGameSync/UnrealGameSyncAutomation/UnrealGameSyncAutomation.csproj
+++ b/Engine/Source/Programs/UnrealGameSync/UnrealGameSyncAutomation/UnrealGameSyncAutomation.csproj
@@ -1,4 +1,4 @@
-﻿<?xml version="1.0" encoding="utf-8"?>
+<?xml version="1.0" encoding="utf-8"?>
 <Project Sdk="Microsoft.NET.Sdk.WindowsDesktop">
   <PropertyGroup>
     <OutputType>WinExe</OutputType>
@@ -7,5 +7,15 @@
     <GenerateAssemblyInfo>false</GenerateAssemblyInfo>
     <RootNamespace>UnrealGameSyncAutomation</RootNamespace>
     <AssemblyName>UnrealGameSyncAutomation</AssemblyName>
+    <RunAnalyzersDuringBuild>False</RunAnalyzersDuringBuild>
+    <Configurations>Debug;Release;Analyze</Configurations>
+  </PropertyGroup>
+
+  <PropertyGroup Condition=" '$(Configuration)|$(Platform)' == 'Analyze|AnyCPU' ">
+    <RunAnalyzersDuringBuild>True</RunAnalyzersDuringBuild>
+    <TreatWarningsAsErrors>false</TreatWarningsAsErrors>
+    <EnforceCodeStyleInBuild>True</EnforceCodeStyleInBuild>
+    <GenerateDocumentationFile>True</GenerateDocumentationFile>
+    <DocumentationFile></DocumentationFile>
   </PropertyGroup>
 </Project>
\ No newline at end of file
diff --git a/Engine/Source/Programs/UnrealGameSync/UnrealGameSyncCmd/UnrealGameSyncCmd.csproj b/Engine/Source/Programs/UnrealGameSync/UnrealGameSyncCmd/UnrealGameSyncCmd.csproj
--- a/Engine/Source/Programs/UnrealGameSync/UnrealGameSyncCmd/UnrealGameSyncCmd.csproj
+++ b/Engine/Source/Programs/UnrealGameSync/UnrealGameSyncCmd/UnrealGameSyncCmd.csproj
@@ -5,6 +5,16 @@
     <TargetFramework>net6.0</TargetFramework>
     <Nullable>enable</Nullable>
     <AssemblyName>ugs</AssemblyName>
+    <RunAnalyzersDuringBuild>False</RunAnalyzersDuringBuild>
+    <Configurations>Debug;Release;Analyze</Configurations>
+  </PropertyGroup>
+
+  <PropertyGroup Condition=" '$(Configuration)|$(Platform)' == 'Analyze|AnyCPU' ">
+    <RunAnalyzersDuringBuild>True</RunAnalyzersDuringBuild>
+    <TreatWarningsAsErrors>false</TreatWarningsAsErrors>
+    <EnforceCodeStyleInBuild>True</EnforceCodeStyleInBuild>
+    <GenerateDocumentationFile>True</GenerateDocumentationFile>
+    <DocumentationFile></DocumentationFile>
   </PropertyGroup>
 
   <ItemGroup>
diff --git a/Engine/Source/Programs/UnrealGameSync/UnrealGameSyncLauncher/UnrealGameSyncLauncher.csproj b/Engine/Source/Programs/UnrealGameSync/UnrealGameSyncLauncher/UnrealGameSyncLauncher.csproj
--- a/Engine/Source/Programs/UnrealGameSync/UnrealGameSyncLauncher/UnrealGameSyncLauncher.csproj
+++ b/Engine/Source/Programs/UnrealGameSync/UnrealGameSyncLauncher/UnrealGameSyncLauncher.csproj
@@ -14,7 +14,16 @@
     <RuntimeIdentifiers>win10-x64</RuntimeIdentifiers>
     <Nullable>enable</Nullable>
     <NeutralLanguage>en-US</NeutralLanguage>
+    <RunAnalyzersDuringBuild>False</RunAnalyzersDuringBuild>
+    <Configurations>Debug;Release;Analyze</Configurations>
   </PropertyGroup>
+
+  <PropertyGroup Condition=" '$(Configuration)|$(Platform)' == 'Analyze|AnyCPU' ">
+    <RunAnalyzersDuringBuild>True</RunAnalyzersDuringBuild>
+    <TreatWarningsAsErrors>false</TreatWarningsAsErrors>
+    <EnforceCodeStyleInBuild>True</EnforceCodeStyleInBuild>
+  </PropertyGroup>
+
   <ItemGroup>
     <Compile Include="..\UnrealGameSync\Forms\ModalTaskWindow.cs">
       <Link>UnrealGameSync\ModalTaskWindow.cs</Link>
diff --git a/Engine/Source/Programs/UnrealGameSync/UnrealGameSyncShared/UnrealGameSyncShared.csproj b/Engine/Source/Programs/UnrealGameSync/UnrealGameSyncShared/UnrealGameSyncShared.csproj
--- a/Engine/Source/Programs/UnrealGameSync/UnrealGameSyncShared/UnrealGameSyncShared.csproj
+++ b/Engine/Source/Programs/UnrealGameSync/UnrealGameSyncShared/UnrealGameSyncShared.csproj
@@ -1,9 +1,19 @@
-﻿<Project Sdk="Microsoft.NET.Sdk">
+<Project Sdk="Microsoft.NET.Sdk">
 
   <PropertyGroup>
     <TargetFramework>net6.0</TargetFramework>
     <Nullable>enable</Nullable>
     <RootNamespace>UnrealGameSync</RootNamespace>
+    <RunAnalyzersDuringBuild>False</RunAnalyzersDuringBuild>
+    <Configurations>Debug;Release;Analyze</Configurations>
+  </PropertyGroup>
+
+  <PropertyGroup Condition=" '$(Configuration)|$(Platform)' == 'Analyze|AnyCPU' ">
+    <RunAnalyzersDuringBuild>True</RunAnalyzersDuringBuild>
+    <TreatWarningsAsErrors>false</TreatWarningsAsErrors>
+    <EnforceCodeStyleInBuild>True</EnforceCodeStyleInBuild>
+    <GenerateDocumentationFile>True</GenerateDocumentationFile>
+    <DocumentationFile></DocumentationFile>
   </PropertyGroup>
 
   <ItemGroup>
