UGS: Add support for updating UGS from Horde. Still defaults to using Perforce.

When using Horde,

diff --git a/Engine/Source/Programs/UnrealGameSync/UnrealGameSync/Forms/ApplicationSettingsWindow.cs b/Engine/Source/Programs/UnrealGameSync/UnrealGameSync/Forms/ApplicationSettingsWindow.cs
--- a/Engine/Source/Programs/UnrealGameSync/UnrealGameSync/Forms/ApplicationSettingsWindow.cs
+++ b/Engine/Source/Programs/UnrealGameSync/UnrealGameSync/Forms/ApplicationSettingsWindow.cs
@@ -14,6 +14,14 @@ namespace UnrealGameSync
 {
 	partial class ApplicationSettingsWindow : Form
 	{
+		public enum Result
+		{
+			Cancel,
+			Ok,
+			Restart,
+			RestartAndConfigureUpdate,
+		}
+
 		static class PerforceTestConnectionTask
 		{
 			public static async Task RunAsync(IPerforceConnection perforce, string depotPath, CancellationToken cancellationToken)
@@ -33,17 +41,13 @@ namespace UnrealGameSync
 		readonly UserSettings _settings;
 		readonly ILogger _logger;
 
-		readonly string? _initialServerAndPort;
-		readonly string? _initialUserName;
-		readonly string? _initialDepotPath;
-		readonly bool _initialPreview;
 		readonly int _initialAutomationPortNumber;
 		readonly ProtocolHandlerState _initialProtocolHandlerState;
 
-		bool? _restartPreview;
-
 		readonly ToolUpdateMonitor _toolUpdateMonitor;
 
+		Result _result = Result.Ok;
+
 		class ToolItem
 		{
 			public Guid Id => Definition.Id;
@@ -98,33 +102,26 @@ namespace UnrealGameSync
 			LauncherSettings launcherSettings = new LauncherSettings();
 			launcherSettings.Read();
 
-			_initialServerAndPort = launcherSettings.PerforceServerAndPort;
-			_initialUserName = launcherSettings.PerforceUserName;
-			_initialDepotPath = launcherSettings.PerforceDepotPath;
-			_initialPreview = launcherSettings.PreviewBuild;
-
 			_initialAutomationPortNumber = AutomationServer.GetPortNumber();
 			_initialProtocolHandlerState = ProtocolHandlerUtils.GetState();
 
 			AutomaticallyRunAtStartupCheckBox.Checked = IsAutomaticallyRunAtStartup();
 			KeepInTrayCheckBox.Checked = settings.KeepInTray;
 
-			ServerTextBox.Text = _initialServerAndPort;
+			HordeServerTextBox.Text = launcherSettings.HordeServer;
+			HordeServerTextBox.Select(HordeServerTextBox.TextLength, 0);
+			HordeServerTextBox.CueBanner = launcherSettings.HordeServer ?? String.Empty;
+
+			ServerTextBox.Text = launcherSettings.PerforceServerAndPort;
 			ServerTextBox.Select(ServerTextBox.TextLength, 0);
 			ServerTextBox.CueBanner = $"Default ({defaultPerforceSettings.ServerAndPort})";
 
-			UserNameTextBox.Text = _initialUserName;
+			UserNameTextBox.Text = launcherSettings.PerforceUserName;
 			UserNameTextBox.Select(UserNameTextBox.TextLength, 0);
 			UserNameTextBox.CueBanner = $"Default ({defaultPerforceSettings.UserName})";
 
 			ParallelSyncThreadsSpinner.Value = Math.Max(Math.Min(settings.SyncOptions.NumThreads ?? PerforceSyncOptions.DefaultNumThreads, ParallelSyncThreadsSpinner.Maximum), ParallelSyncThreadsSpinner.Minimum);
 
-			DepotPathTextBox.Text = _initialDepotPath;
-			DepotPathTextBox.Select(DepotPathTextBox.TextLength, 0);
-			DepotPathTextBox.CueBanner = DeploymentSettings.Instance.DefaultDepotPath ?? String.Empty;
-
-			UsePreviewBuildCheckBox.Checked = preview;
-
 			if (_initialAutomationPortNumber > 0)
 			{
 				EnableAutomationCheckBox.Checked = true;
@@ -194,17 +191,11 @@ namespace UnrealGameSync
 			}
 		}
 
-		public static bool? ShowModal(IWin32Window owner, IPerforceSettings defaultPerforceSettings, bool preview, string originalExecutableFileName, UserSettings settings, ToolUpdateMonitor toolUpdateMonitor, ILogger<ApplicationSettingsWindow> logger)
+		public static Result ShowModal(IWin32Window owner, IPerforceSettings defaultPerforceSettings, bool preview, string originalExecutableFileName, UserSettings settings, ToolUpdateMonitor toolUpdateMonitor, ILogger<ApplicationSettingsWindow> logger)
 		{
 			using ApplicationSettingsWindow applicationSettings = new ApplicationSettingsWindow(defaultPerforceSettings, preview, originalExecutableFileName, settings, toolUpdateMonitor, logger);
-			if (applicationSettings.ShowDialog(owner) == DialogResult.OK)
-			{
-				return applicationSettings._restartPreview;
-			}
-			else
-			{
-				return null;
-			}
+			applicationSettings.ShowDialog(owner);
+			return applicationSettings._result;
 		}
 
 		private static bool IsAutomaticallyRunAtStartup()
@@ -213,10 +204,34 @@ namespace UnrealGameSync
 			return (key?.GetValue("UnrealGameSync") != null);
 		}
 
+		private void UpdateSettingsBtn_Click(object sender, EventArgs e)
+		{
+			if (MessageBox.Show("UnrealGameSync must be restarted to configure update settings.\n\nWould you like to restart now?", "Restart Required", MessageBoxButtons.OKCancel) != DialogResult.OK)
+			{
+				return;
+			}
+
+			ApplySettings(Result.RestartAndConfigureUpdate);
+		}
+
 		private void OkBtn_Click(object sender, EventArgs e)
 		{
+			ApplySettings(Result.Ok);
+		}
+
+		private void ApplySettings(Result result)
+		{
+			LauncherSettings originalLauncherSettings = new LauncherSettings();
+			originalLauncherSettings.Read();
+
 			// Update the settings
-			LauncherSettings launcherSettings = new LauncherSettings();
+			LauncherSettings launcherSettings = new LauncherSettings(originalLauncherSettings);
+
+			launcherSettings.HordeServer = HordeServerTextBox.Text.Trim();
+			if (launcherSettings.HordeServer.Length == 0 || String.Equals(launcherSettings.HordeServer, DeploymentSettings.Instance.HordeUrl, StringComparison.OrdinalIgnoreCase))
+			{
+				launcherSettings.HordeServer = null;
+			}
 
 			launcherSettings.PerforceServerAndPort = ServerTextBox.Text.Trim();
 			if (launcherSettings.PerforceServerAndPort.Length == 0)
@@ -230,44 +245,28 @@ namespace UnrealGameSync
 				launcherSettings.PerforceUserName = null;
 			}
 
-			launcherSettings.PerforceDepotPath = DepotPathTextBox.Text.Trim();
-			if (launcherSettings.PerforceDepotPath.Length == 0 || launcherSettings.PerforceDepotPath == DeploymentSettings.Instance.DefaultDepotPath)
-			{
-				launcherSettings.PerforceDepotPath = null;
-			}
-
-			launcherSettings.PreviewBuild = UsePreviewBuildCheckBox.Checked;
-
 			int automationPortNumber;
 			if (!EnableAutomationCheckBox.Checked || !Int32.TryParse(AutomationPortTextBox.Text, out automationPortNumber))
 			{
 				automationPortNumber = -1;
 			}
 
-			if (launcherSettings.PerforceServerAndPort != _initialServerAndPort || launcherSettings.PerforceUserName != _initialUserName || launcherSettings.PerforceDepotPath != _initialDepotPath || launcherSettings.PreviewBuild != _initialPreview || automationPortNumber != _initialAutomationPortNumber)
+			if (!String.Equals(launcherSettings.HordeServer, originalLauncherSettings.HordeServer, StringComparison.OrdinalIgnoreCase) ||
+				!String.Equals(launcherSettings.PerforceServerAndPort, originalLauncherSettings.PerforceServerAndPort, StringComparison.OrdinalIgnoreCase) ||
+				!String.Equals(launcherSettings.PerforceUserName, originalLauncherSettings.PerforceUserName, StringComparison.OrdinalIgnoreCase))
 			{
-				// Try to log in to the new server, and check the application is there
-				if (launcherSettings.PerforceServerAndPort != _initialServerAndPort || launcherSettings.PerforceUserName != _initialUserName || launcherSettings.PerforceDepotPath != _initialDepotPath)
+				if (result == Result.Ok)
 				{
-					PerforceSettings settings = Utility.OverridePerforceSettings(_defaultPerforceSettings, launcherSettings.PerforceServerAndPort, launcherSettings.PerforceUserName);
-
-					string? testDepotPath = launcherSettings.PerforceDepotPath ?? DeploymentSettings.Instance.DefaultDepotPath;
-					if (testDepotPath != null)
+					if (MessageBox.Show("UnrealGameSync must be restarted to apply these settings.\n\nWould you like to restart now?", "Restart Required", MessageBoxButtons.OKCancel) != DialogResult.OK)
 					{
-						ModalTask? task = PerforceModalTask.Execute(this, "Checking connection", "Checking connection, please wait...", settings, (p, c) => PerforceTestConnectionTask.RunAsync(p, testDepotPath, c), _logger);
-						if (task == null || !task.Succeeded)
-						{
-							return;
-						}
+						return;
+					}
+					else
+					{
+						result = Result.Restart;
 					}
 				}
 
-				if (MessageBox.Show("UnrealGameSync must be restarted to apply these settings.\n\nWould you like to restart now?", "Restart Required", MessageBoxButtons.OKCancel) != DialogResult.OK)
-				{
-					return;
-				}
-
-				_restartPreview = UsePreviewBuildCheckBox.Checked;
 				launcherSettings.Save();
 				AutomationServer.SetPortNumber(automationPortNumber);
 			}
@@ -321,13 +320,13 @@ namespace UnrealGameSync
 				}
 			}
 
-			DialogResult = DialogResult.OK;
+			_result = result;
 			Close();
 		}
 
 		private void CancelBtn_Click(object sender, EventArgs e)
 		{
-			DialogResult = DialogResult.Cancel;
+			_result = Result.Cancel;
 			Close();
 		}
 
diff --git a/Engine/Source/Programs/UnrealGameSync/UnrealGameSync/Forms/MainWindow.cs b/Engine/Source/Programs/UnrealGameSync/UnrealGameSync/Forms/MainWindow.cs
--- a/Engine/Source/Programs/UnrealGameSync/UnrealGameSync/Forms/MainWindow.cs
+++ b/Engine/Source/Programs/UnrealGameSync/UnrealGameSync/Forms/MainWindow.cs
@@ -1526,10 +1526,14 @@ namespace UnrealGameSync
 
 		public void ModifyApplicationSettings()
 		{
-			bool? relaunchPreview = ApplicationSettingsWindow.ShowModal(this, _defaultPerforceSettings, _preview, _originalExecutableFileName, _settings, ToolUpdateMonitor, _serviceProvider.GetRequiredService<ILogger<ApplicationSettingsWindow>>());
-			if(relaunchPreview.HasValue)
+			ApplicationSettingsWindow.Result result = ApplicationSettingsWindow.ShowModal(this, _defaultPerforceSettings, _preview, _originalExecutableFileName, _settings, ToolUpdateMonitor, _serviceProvider.GetRequiredService<ILogger<ApplicationSettingsWindow>>());
+			if (result == ApplicationSettingsWindow.Result.Restart)
 			{
-				_updateMonitor.TriggerUpdate(UpdateType.UserInitiated, relaunchPreview);
+				_updateMonitor.TriggerUpdate(UpdateType.UserInitiated, false);
+			}
+			else if(result == ApplicationSettingsWindow.Result.RestartAndConfigureUpdate)
+			{
+				_updateMonitor.TriggerUpdate(UpdateType.UserInitiated, true);
 			}
 
 			for (int idx = 0; idx < TabControl.GetTabCount(); idx++)
diff --git a/Engine/Source/Programs/UnrealGameSync/UnrealGameSync/Program.cs b/Engine/Source/Programs/UnrealGameSync/UnrealGameSync/Program.cs
--- a/Engine/Source/Programs/UnrealGameSync/UnrealGameSync/Program.cs
+++ b/Engine/Source/Programs/UnrealGameSync/UnrealGameSync/Program.cs
@@ -1,6 +1,7 @@
 // Copyright Epic Games, Inc. All Rights Reserved.
 
 using EpicGames.Core;
+using EpicGames.Horde;
 using EpicGames.OIDC;
 using EpicGames.Perforce;
 using Microsoft.Extensions.DependencyInjection;
@@ -180,6 +181,7 @@ namespace UnrealGameSync
 				services.AddSingleton<IAsyncDisposer, AsyncDisposer>();
 				services.AddSingleton(sp => TokenStoreFactory.CreateTokenStore());
 				services.AddSingleton<OidcTokenManager>();
+				services.AddHordeHttpClient(x => x.BaseAddress = new Uri(launcherSettings.HordeServer ?? "http://localhost:5000"));
 
 				await using (ServiceProvider serviceProvider = services.BuildServiceProvider())
 				{
@@ -220,7 +222,7 @@ namespace UnrealGameSync
 
 							ProtocolHandlerUtils.InstallQuiet(logger);
 
-							using (UpdateMonitor updateMonitor = new UpdateMonitor(defaultSettings, updatePath, serviceProvider))
+							using (UpdateMonitor updateMonitor = CreateUpdateMonitor(launcherSettings, defaultSettings, updatePath, serviceProvider))
 							{
 								using ProgramApplicationContext context = new ProgramApplicationContext(defaultSettings, updateMonitor, DeploymentSettings.Instance.ApiUrl, dataFolder, activateEvent, restoreState, updateSpawn, projectFileName, preview, serviceProvider, uri);
 								Application.Run(context);
@@ -228,8 +230,7 @@ namespace UnrealGameSync
 								if (updateMonitor.IsUpdateAvailable && updateSpawn != null)
 								{
 									instanceMutex.Close();
-									bool launchPreview = updateMonitor.RelaunchPreview ?? preview;
-									Utility.SpawnProcess(updateSpawn, "-restorestate" + (launchPreview ? " -unstable" : ""));
+									Utility.SpawnProcess(updateSpawn, "-restorestate" + (updateMonitor.OpenSettings ? " -settings" : ""));
 								}
 							}
 						}
@@ -247,6 +248,18 @@ namespace UnrealGameSync
 			}
 		}
 
+		private static UpdateMonitor CreateUpdateMonitor(LauncherSettings launcherSettings, IPerforceSettings defaultSettings, string? updatePath, IServiceProvider serviceProvider)
+		{
+			if (launcherSettings.UpdateSource == LauncherUpdateSource.Horde)
+			{
+				return new HordeUpdateMonitor(SyncVersion ?? String.Empty, serviceProvider);
+			}
+			else
+			{
+				return new PerforceUpdateMonitor(defaultSettings, updatePath, serviceProvider);
+			}
+		}
+
 		private static void TraceException(Exception ex, ILogger logger)
 		{
 			if (DeploymentSettings.Instance.SentryDsn != null)
diff --git a/Engine/Source/Programs/UnrealGameSync/UnrealGameSync/UpdateMonitor.cs b/Engine/Source/Programs/UnrealGameSync/UnrealGameSync/UpdateMonitor.cs
--- a/Engine/Source/Programs/UnrealGameSync/UnrealGameSync/UpdateMonitor.cs
+++ b/Engine/Source/Programs/UnrealGameSync/UnrealGameSync/UpdateMonitor.cs
@@ -1,5 +1,8 @@
 // Copyright Epic Games, Inc. All Rights Reserved.
 
+using EpicGames.Core;
+using EpicGames.Horde;
+using EpicGames.Horde.Tools;
 using EpicGames.Perforce;
 using Microsoft.Extensions.DependencyInjection;
 using Microsoft.Extensions.Logging;
@@ -15,24 +18,110 @@ namespace UnrealGameSync
 		UserInitiated,
 	}
 
-	class UpdateMonitor : IDisposable
+	abstract class UpdateMonitor : IDisposable
 	{
-		Task? _workerTask;
-#pragma warning disable CA2213 // warning CA2213: 'UpdateMonitor' contains field '_cancellationSource' that is of IDisposable type 'CancellationTokenSource', but it is never disposed. Change the Dispose method on 'UpdateMonitor' to call Close or Dispose on this field.
-		readonly CancellationTokenSource _cancellationSource = new CancellationTokenSource();
-#pragma warning restore CA2213
-		readonly ILogger _logger;
-		readonly IAsyncDisposer _asyncDisposer;
+		public bool IsUpdateAvailable
+		{
+			get;
+			private set;
+		}
 
 		public Action<UpdateType>? OnUpdateAvailable;
 
-		public bool? RelaunchPreview
+		public bool OpenSettings
 		{
 			get;
 			private set;
 		}
 
-		public UpdateMonitor(IPerforceSettings perforceSettings, string? watchPath, IServiceProvider serviceProvider)
+		public abstract void Dispose();
+
+		public void TriggerUpdate(UpdateType updateType, bool openSettings)
+		{
+			OpenSettings = openSettings;
+			IsUpdateAvailable = true;
+			if (OnUpdateAvailable != null)
+			{
+				OnUpdateAvailable(updateType);
+			}
+		}
+	}
+
+	class HordeUpdateMonitor : UpdateMonitor
+	{
+		readonly ToolId _toolId;
+		readonly string _currentVersion;
+		readonly IServiceProvider _serviceProvider;
+		readonly BackgroundTask _backgroundTask;
+		readonly ILogger _logger;
+
+		public HordeUpdateMonitor(string currentVersion, IServiceProvider serviceProvider)
+			: this(DeploymentSettings.Instance.HordeToolId, currentVersion, serviceProvider)
+		{
+		}
+
+		public HordeUpdateMonitor(ToolId toolId, string currentVersion, IServiceProvider serviceProvider)
+		{
+			_toolId = toolId;
+			_currentVersion = currentVersion;
+			_serviceProvider = serviceProvider;
+			_logger = serviceProvider.GetRequiredService<ILogger<HordeUpdateMonitor>>();
+			_backgroundTask = BackgroundTask.StartNew(ctx => CheckForUpdatesLoopAsync(ctx));
+		}
+
+		public override void Dispose()
+		{
+			Task.Run(() => _backgroundTask.DisposeAsync()).Wait();
+		}
+
+		public async Task CheckForUpdatesLoopAsync(CancellationToken cancellationToken)
+		{
+			for (; ; )
+			{
+				// Check if there's a new build available on the server
+				HordeHttpClient hordeHttpClient = _serviceProvider.GetRequiredService<HordeHttpClient>();
+				try
+				{
+					GetToolResponse response = await hordeHttpClient.GetToolAsync(_toolId, cancellationToken);
+					if (response.Deployments.Count == 0)
+					{
+						_logger.LogWarning("No deployments on Horde server for tool {ToolId}", _toolId);
+					}
+					else
+					{
+						string latestUrl = new Uri(hordeHttpClient.BaseUrl, $"api/v1/tools/{_toolId}/deployments/{response.Deployments[^1].Id}").ToString();
+						if (String.Equals(latestUrl, _currentVersion, StringComparison.OrdinalIgnoreCase))
+						{
+							_logger.LogInformation("Currently running latest version ({LatestUrl})", latestUrl);
+						}
+						else
+						{
+							_logger.LogInformation("Triggering update request {CurrentUrl} -> {LatestUrl}", _currentVersion, latestUrl);
+							TriggerUpdate(UpdateType.Background, false);
+						}
+					}
+				}
+				catch (Exception ex)
+				{
+					_logger.LogError(ex, "Error while checking for tool updates: {Message}", ex.Message);
+				}
+
+				// Wait a while before checking again
+				await Task.Delay(TimeSpan.FromMinutes(5.0), cancellationToken);
+			}
+		}
+	}
+
+	class PerforceUpdateMonitor : UpdateMonitor 
+	{
+		Task? _workerTask;
+#pragma warning disable CA2213 // warning CA2213: 'UpdateMonitor' contains field '_cancellationSource' that is of IDisposable type 'CancellationTokenSource', but it is never disposed. Change the Dispose method on 'UpdateMonitor' to call Close or Dispose on this field.
+		readonly CancellationTokenSource _cancellationSource = new CancellationTokenSource();
+#pragma warning restore CA2213
+		readonly ILogger _logger;
+		readonly IAsyncDisposer _asyncDisposer;
+
+		public PerforceUpdateMonitor(IPerforceSettings perforceSettings, string? watchPath, IServiceProvider serviceProvider)
 		{
 			_logger = serviceProvider.GetRequiredService<ILogger<UpdateMonitor>>();
 			_asyncDisposer = serviceProvider.GetRequiredService<IAsyncDisposer>();
@@ -44,7 +133,7 @@ namespace UnrealGameSync
 			}
 		}
 
-		public void Dispose()
+		public override void Dispose()
 		{
 			OnUpdateAvailable = null;
 
@@ -56,12 +145,6 @@ namespace UnrealGameSync
 			}
 		}
 
-		public bool IsUpdateAvailable
-		{
-			get;
-			private set;
-		}
-
 		async Task PollForUpdatesAsync(IPerforceSettings perforceSettings, string watchPath, CancellationToken cancellationToken)
 		{
 			for (; ; )
@@ -76,7 +159,7 @@ namespace UnrealGameSync
 					PerforceResponseList<ChangesRecord> changes = await perforce.TryGetChangesAsync(ChangesOptions.None, -1, ChangeStatus.Submitted, watchPath, cancellationToken);
 					if (changes.Succeeded && changes.Data.Count > 0)
 					{
-						TriggerUpdate(UpdateType.Background, null);
+						TriggerUpdate(UpdateType.Background, false);
 					}
 				}
 				catch (PerforceException ex)
@@ -94,15 +177,5 @@ namespace UnrealGameSync
 				}
 			}
 		}
-
-		public void TriggerUpdate(UpdateType updateType, bool? relaunchPreview)
-		{
-			RelaunchPreview = relaunchPreview;
-			IsUpdateAvailable = true;
-			if(OnUpdateAvailable != null)
-			{
-				OnUpdateAvailable(updateType);
-			}
-		}
 	}
 }
diff --git a/Engine/Source/Programs/UnrealGameSync/UnrealGameSyncLauncher/Forms/SettingsWindow.cs b/Engine/Source/Programs/UnrealGameSync/UnrealGameSyncLauncher/Forms/SettingsWindow.cs
--- a/Engine/Source/Programs/UnrealGameSync/UnrealGameSyncLauncher/Forms/SettingsWindow.cs
+++ b/Engine/Source/Programs/UnrealGameSync/UnrealGameSyncLauncher/Forms/SettingsWindow.cs
@@ -17,27 +17,36 @@ namespace UnrealGameSyncLauncher
 		[DllImport("user32.dll")]
 		private static extern IntPtr SendMessage(IntPtr hWnd, int msg, int wParam, [MarshalAs(UnmanagedType.LPWStr)] string lParam);
 
-		public delegate Task SyncAndRunDelegate(IPerforceConnection perforce, string? depotPath, bool preview, ILogger logWriter, CancellationToken cancellationToken);
+		public delegate Task SyncAndRunDelegate(IPerforceConnection? perforce, LauncherSettings settings, ILogger logWriter, CancellationToken cancellationToken);
 
 		const int EmSetcuebanner = 0x1501;
 
+		LauncherSettings _settings;
 		string? _logText;
 		readonly SyncAndRunDelegate _syncAndRun;
 
-		public SettingsWindow(string? prompt, string? logText, string? serverAndPort, string? userName, string? depotPath, bool preview, SyncAndRunDelegate syncAndRun)
+		public SettingsWindow(string? prompt, string? logText, LauncherSettings settings, SyncAndRunDelegate syncAndRun)
 		{
 			InitializeComponent();
+			Font = new System.Drawing.Font("Segoe UI", 8.25F, System.Drawing.FontStyle.Regular, System.Drawing.GraphicsUnit.Point, ((byte)(0)));
 
-			if(prompt != null)
+			if (prompt != null)
 			{
 				PromptLabel.Text = prompt;
 			}
 
 			_logText = logText;
-			ServerTextBox.Text = serverAndPort ?? String.Empty;
-			UserNameTextBox.Text = userName ?? String.Empty;
-			DepotPathTextBox.Text = depotPath ?? String.Empty;
-			UsePreviewBuildCheckBox.Checked = preview;
+			_settings = settings;
+
+			string defaultHordeServer = DeploymentSettings.Instance.HordeUrl ?? String.Empty;
+
+			HordeRadioBtn.Checked = settings.UpdateSource == LauncherUpdateSource.Horde;
+			HordeServerTextBox.Text = String.IsNullOrEmpty(settings.HordeServer) ? defaultHordeServer : settings.HordeServer;
+			ServerTextBox.Text = settings.PerforceServerAndPort ?? String.Empty;
+			UserNameTextBox.Text = settings.PerforceUserName ?? String.Empty;
+			DepotPathTextBox.Text = settings.PerforceDepotPath ?? String.Empty;
+			UsePreviewBuildCheckBox.Checked = settings.PreviewBuild;
+
 			_syncAndRun = syncAndRun;
 
 			ViewLogBtn.Visible = logText != null;
@@ -59,51 +68,48 @@ namespace UnrealGameSyncLauncher
 
 		private void ConnectBtn_Click(object sender, EventArgs e)
 		{
-			// Update the settings
-			LauncherSettings launcherSettings = new LauncherSettings();
-
-			launcherSettings.PerforceServerAndPort = ServerTextBox.Text.Trim();
-			if(launcherSettings.PerforceServerAndPort.Length == 0)
-			{
-				launcherSettings.PerforceServerAndPort = null;
-			}
-
-			launcherSettings.PerforceUserName = UserNameTextBox.Text.Trim();
-			if(launcherSettings.PerforceUserName.Length == 0)
+			_settings.UpdateSource = HordeRadioBtn.Checked ? LauncherUpdateSource.Horde : LauncherUpdateSource.Perforce;
+			_settings.HordeServer = HordeServerTextBox.Text.Trim();
+			if (String.Equals(_settings.HordeServer, DeploymentSettings.Instance.HordeUrl, StringComparison.OrdinalIgnoreCase))
 			{
-				launcherSettings.PerforceUserName = null;
+				_settings.HordeServer = null;
 			}
+			_settings.PerforceServerAndPort = ServerTextBox.Text.Trim();
+			_settings.PerforceUserName = UserNameTextBox.Text.Trim();
+			_settings.PerforceDepotPath = DepotPathTextBox.Text.Trim();
+			_settings.PreviewBuild = UsePreviewBuildCheckBox.Checked;
+			_settings.Save();
 
-			launcherSettings.PerforceDepotPath = DepotPathTextBox.Text.Trim();
-			if(launcherSettings.PerforceDepotPath.Length == 0)
-			{
-				launcherSettings.PerforceDepotPath = null;
-			}
-
-			launcherSettings.PreviewBuild = UsePreviewBuildCheckBox.Checked;
-			launcherSettings.Save();
+			// Create the P4 connection
+			CaptureLogger logger = new CaptureLogger();
 
-			PerforceSettings perforceSettings = new PerforceSettings(PerforceSettings.Default);
-			if (!String.IsNullOrEmpty(launcherSettings.PerforceServerAndPort))
+			// Create the task for connecting to this server
+			ModalTask? task;
+			if (_settings.UpdateSource == LauncherUpdateSource.Horde)
 			{
-				perforceSettings.ServerAndPort = launcherSettings.PerforceServerAndPort;
+				task = ModalTask.Execute(this, "Updating", "Checking for updates, please wait...", c => _syncAndRun(null, _settings, logger, c));
 			}
-			if (!String.IsNullOrEmpty(launcherSettings.PerforceUserName))
+			else
 			{
-				perforceSettings.UserName = launcherSettings.PerforceUserName;
-			}
-			perforceSettings.PreferNativeClient = true;
+				PerforceSettings perforceSettings = new PerforceSettings(PerforceSettings.Default);
+				if (!String.IsNullOrEmpty(_settings.PerforceServerAndPort))
+				{
+					perforceSettings.ServerAndPort = _settings.PerforceServerAndPort;
+				}
+				if (!String.IsNullOrEmpty(_settings.PerforceUserName))
+				{
+					perforceSettings.UserName = _settings.PerforceUserName;
+				}
+				perforceSettings.PreferNativeClient = true;
 
-			// Create the P4 connection
-			CaptureLogger logger = new CaptureLogger();
+				task = PerforceModalTask.Execute(this, "Updating", "Checking for updates, please wait...", perforceSettings, (p, c) => _syncAndRun(p, _settings, logger, c), logger);
+			}
 
-			// Create the task for connecting to this server
-			ModalTask? task = PerforceModalTask.Execute(this, "Updating", "Checking for updates, please wait...", perforceSettings, (p, c) => _syncAndRun(p, launcherSettings.PerforceDepotPath, launcherSettings.PreviewBuild, logger, c), logger);
 			if (task != null)
 			{
-				if(task.Succeeded)
+				if (task.Succeeded)
 				{
-					launcherSettings.Save();
+					_settings.Save();
 					DialogResult = DialogResult.OK;
 					Close();
 				}
@@ -119,5 +125,17 @@ namespace UnrealGameSyncLauncher
 			DialogResult = DialogResult.Cancel;
 			Close();
 		}
+
+		private void HordeRadioBtn_CheckedChanged(object sender, EventArgs e)
+		{
+			HordeGroupBox.Enabled = HordeRadioBtn.Checked;
+			PerforceGroupBox.Enabled = PerforceRadioBtn.Checked;
+		}
+
+		private void PerforceRadioBtn_CheckedChanged(object sender, EventArgs e)
+		{
+			HordeGroupBox.Enabled = HordeRadioBtn.Checked;
+			PerforceGroupBox.Enabled = PerforceRadioBtn.Checked;
+		}
 	}
 }
diff --git a/Engine/Source/Programs/UnrealGameSync/UnrealGameSyncLauncher/Program.cs b/Engine/Source/Programs/UnrealGameSync/UnrealGameSyncLauncher/Program.cs
--- a/Engine/Source/Programs/UnrealGameSync/UnrealGameSyncLauncher/Program.cs
+++ b/Engine/Source/Programs/UnrealGameSync/UnrealGameSyncLauncher/Program.cs
@@ -1,7 +1,13 @@
 // Copyright Epic Games, Inc. All Rights Reserved.
 
 using EpicGames.Core;
+using EpicGames.Horde;
+using EpicGames.Horde.Storage;
+using EpicGames.Horde.Storage.Clients;
+using EpicGames.Horde.Storage.Nodes;
+using EpicGames.Horde.Tools;
 using EpicGames.Perforce;
+using Microsoft.Extensions.DependencyInjection;
 using Microsoft.Extensions.Logging;
 using System;
 using System.Collections.Generic;
@@ -10,7 +16,9 @@ using System.Diagnostics.CodeAnalysis;
 using System.IO;
 using System.IO.Compression;
 using System.Linq;
+using System.Net.Http;
 using System.Reflection;
+using System.Reflection.Metadata;
 using System.Text;
 using System.Threading;
 using System.Threading.Tasks;
@@ -42,6 +50,7 @@ namespace UnrealGameSyncLauncher
 
 				// Figure out if we should sync the unstable build by default
 				bool preview = args.Contains("-unstable", StringComparer.InvariantCultureIgnoreCase) || args.Contains("-preview", StringComparer.InvariantCultureIgnoreCase);
+				bool openSettings = args.Contains("-settings", StringComparer.OrdinalIgnoreCase);
 
 				// Read the settings
 				LauncherSettings launcherSettings = new LauncherSettings();
@@ -49,11 +58,11 @@ namespace UnrealGameSyncLauncher
 				launcherSettings.Read();
 
 				// If the shift key is held down, immediately show the settings window
-				Task SyncAndRunWrapper(IPerforceConnection perforce, string? depotPath, bool preview, ILogger logWriter, CancellationToken cancellationToken) => SyncAndRun(perforce, depotPath, preview, args, instanceMutex, logWriter, cancellationToken);
-				if ((Control.ModifierKeys & Keys.Shift) != 0)
+				Task SyncAndRunWrapper(IPerforceConnection? perforce, LauncherSettings settings, ILogger logWriter, CancellationToken cancellationToken) => SyncAndRun(perforce, settings, args, instanceMutex, logWriter, cancellationToken);
+				if ((Control.ModifierKeys & Keys.Shift) != 0 || openSettings)
 				{
 					// Show the settings window immediately
-					using SettingsWindow updateError = new SettingsWindow(null, null, launcherSettings.PerforceServerAndPort, launcherSettings.PerforceUserName, launcherSettings.PerforceDepotPath, preview, SyncAndRunWrapper);
+					using SettingsWindow updateError = new SettingsWindow(null, null, launcherSettings, SyncAndRunWrapper);
 					if(updateError.ShowDialog() == DialogResult.OK)
 					{
 						return 0;
@@ -64,9 +73,17 @@ namespace UnrealGameSyncLauncher
 					// Try to do a sync with the current settings first
 					CaptureLogger logger = new CaptureLogger();
 
-					IPerforceSettings settings = new PerforceSettings(PerforceSettings.Default) { PreferNativeClient = true }.MergeWith(newServerAndPort: launcherSettings.PerforceServerAndPort, newUserName: launcherSettings.PerforceUserName);
+					ModalTask? task;
+					if (launcherSettings.UpdateSource == LauncherUpdateSource.Horde)
+					{
+						task = ModalTask.Execute(null, "Updating", "Checking for updates, please wait...", c => SyncAndRunWrapper(null, launcherSettings, logger, c));
+					}
+					else
+					{
+						IPerforceSettings settings = new PerforceSettings(PerforceSettings.Default) { PreferNativeClient = true }.MergeWith(newServerAndPort: launcherSettings.PerforceServerAndPort, newUserName: launcherSettings.PerforceUserName);
+						task = PerforceModalTask.Execute(null, "Updating", "Checking for updates, please wait...", settings, (p, c) => SyncAndRunWrapper(p, launcherSettings, logger, c), logger);
+					}
 
-					ModalTask? task = PerforceModalTask.Execute(null, "Updating", "Checking for updates, please wait...", settings, (p, c) => SyncAndRun(p, launcherSettings.PerforceDepotPath, preview, args, instanceMutex, logger, c), logger);
 					if (task == null)
 					{
 						logger.LogInformation("Canceled by user");
@@ -76,7 +93,7 @@ namespace UnrealGameSyncLauncher
 						return 0;
 					}
 
-					using SettingsWindow updateError = new SettingsWindow("Unable to update UnrealGameSync from Perforce. Verify that your connection settings are correct.", logger.Render(Environment.NewLine), launcherSettings.PerforceServerAndPort, launcherSettings.PerforceUserName, launcherSettings.PerforceDepotPath, preview, SyncAndRunWrapper);
+					using SettingsWindow updateError = new SettingsWindow("Unable to update UnrealGameSync from Perforce. Verify that your connection settings are correct.", logger.Render(Environment.NewLine), launcherSettings, SyncAndRunWrapper);
 					if(updateError.ShowDialog() == DialogResult.OK)
 					{
 						return 0;
@@ -89,35 +106,19 @@ namespace UnrealGameSyncLauncher
 		// Values of the Perforce "action" field that means that the file should no longer be synced
 		public static readonly string[] s_deleteActions = { "delete", "move/delete", "purge", "archive" };
 
-		public static async Task SyncAndRun(IPerforceConnection perforce, string? baseDepotPath, bool preview, string[] args, Mutex instanceMutex, ILogger logger, CancellationToken cancellationToken)
+		class LoggerProviderAdapter : ILoggerProvider
 		{
-			try
-			{
-				if (String.IsNullOrEmpty(baseDepotPath))
-				{
-					throw new UserErrorException($"Invalid setting for sync path");
-				}
+			readonly ILogger _logger;
 
-				string baseDepotPathPrefix = baseDepotPath.TrimEnd('/');
-
-				// Find the most recent changelist
-				string syncPath = baseDepotPathPrefix + (preview ? "/UnstableRelease.zip" : "/Release.zip");
-				List<ChangesRecord> changes = await perforce.GetChangesAsync(ChangesOptions.None, 1, ChangeStatus.Submitted, syncPath, cancellationToken);
-				if (changes.Count == 0)
-				{
-					syncPath = baseDepotPathPrefix + (preview ? "/UnstableRelease/..." : "/Release/...");
-					changes = await perforce.GetChangesAsync(ChangesOptions.None, 1, ChangeStatus.Submitted, syncPath, cancellationToken);
-#pragma warning disable CA1508 // warning CA1508: 'changes.Count == 0' is always 'true'. Remove or refactor the condition(s) to avoid dead code.
-					if (changes.Count == 0)
-					{
-						throw new UserErrorException($"Unable to find any UGS binaries under {syncPath}");
-					}
-#pragma warning restore CA1508
-				}
-
-				int requiredChangeNumber = changes[0].Number;
-				logger.LogInformation("Syncing from {SyncPath}", syncPath);
+			public LoggerProviderAdapter(ILogger logger) => _logger = logger;
+			public ILogger CreateLogger(string categoryName) => _logger;
+			public void Dispose() { }
+		}
 
+		public static async Task SyncAndRun(IPerforceConnection? perforce, LauncherSettings launcherSettings, string[] args, Mutex instanceMutex, ILogger logger, CancellationToken cancellationToken)
+		{
+			try
+			{
 				// Create the target folder
 				string applicationFolder = Path.Combine(Environment.GetFolderPath(Environment.SpecialFolder.LocalApplicationData), "UnrealGameSync", "Latest");
 				if (!SafeCreateDirectory(applicationFolder))
@@ -125,76 +126,160 @@ namespace UnrealGameSyncLauncher
 					throw new UserErrorException($"Couldn't create directory: {applicationFolder}");
 				}
 
+				// Check the application exists
+				string applicationExe = Path.Combine(applicationFolder, "UnrealGameSync.exe");
+
 				// Read the current version
 				string syncVersionFile = Path.Combine(applicationFolder, "SyncVersion.txt");
-				string requiredSyncText = String.Format("{0}\n{1}@{2}", perforce.Settings.ServerAndPort ?? "", syncPath, requiredChangeNumber);
 
-				// Check the application exists
-				string applicationExe = Path.Combine(applicationFolder, "UnrealGameSync.exe");
+				string? syncText = null;
+				if (File.Exists(applicationExe) && File.Exists(syncVersionFile))
+				{
+					TryReadAllText(syncVersionFile, out syncText);
+				}
 
-				// Check if the version has changed
-				string? syncText;
-				if (!File.Exists(syncVersionFile) || !File.Exists(applicationExe) || !TryReadAllText(syncVersionFile, out syncText) || syncText != requiredSyncText)
+				// New command line for launching the updated executable
+				string? updatePath = null;
+				if (launcherSettings.UpdateSource == LauncherUpdateSource.Horde)
 				{
-					// Try to delete the directory contents. Retry for a while, in case we've been spawned by an application in this folder to do an update.
-					for (int numRetries = 0; !SafeDeleteDirectoryContents(applicationFolder); numRetries++)
+					Uri hordeServerUrl = new Uri(launcherSettings.HordeServer ?? "https://horde/");
+
+					ServiceCollection services = new ServiceCollection();
+					services.AddLogging(builder => builder.AddProvider(new LoggerProviderAdapter(logger)));
+					services.AddHordeHttpClient((sp, client) => client.BaseAddress = hordeServerUrl);
+
+					await using ServiceProvider serviceProvider = services.BuildServiceProvider();
+					HordeHttpClient httpClient = serviceProvider.GetRequiredService<HordeHttpClient>();
+
+					ToolId toolId = DeploymentSettings.Instance.HordeToolId;
+
+					GetToolResponse tool = await httpClient.GetToolAsync(toolId, cancellationToken: cancellationToken);
+					GetToolDeploymentResponse deployment = tool.Deployments[^1];
+					Uri deploymentUri = new Uri(hordeServerUrl, $"api/v1/tools/{toolId}/deployments/{deployment.Id}");
+
+					string requiredSyncText = deploymentUri.ToString();
+					if (syncText == null || syncText != requiredSyncText)
 					{
-						if (numRetries > 20)
+						// Delete the output directory
+						await SafeDeleteDirectoryContentsWithRetryAsync(applicationFolder, cancellationToken);
+
+						// Download and extract the zip file
+						string zipFile = Path.Combine(applicationFolder, "update.zip");
+						using (Stream requestStream = await httpClient.GetToolDeploymentZipAsync(toolId, deployment.Id, cancellationToken))
+						{
+							using (Stream tempFileStream = File.Open(zipFile, FileMode.Create, FileAccess.Write, FileShare.None))
+							{
+								await requestStream.CopyToAsync(tempFileStream, cancellationToken);
+							}
+						}
+						ZipFile.ExtractToDirectory(zipFile, applicationFolder);
+						File.Delete(zipFile);
+
+						// Update the version
+						if (!TryWriteAllText(syncVersionFile, requiredSyncText))
 						{
-							throw new UserErrorException($"Couldn't delete contents of {applicationFolder} (retried {numRetries} times).");
+							throw new UserErrorException("Couldn't write sync text to {SyncVersionFile}");
 						}
-						await Task.Delay(500, cancellationToken);
+					}
+				}
+				else if (launcherSettings.UpdateSource == LauncherUpdateSource.Perforce)
+				{
+					bool preview = launcherSettings.PreviewBuild;
+					if (perforce == null)
+					{
+						throw new UserErrorException("No Perforce connection");
 					}
 
-					// Find all the files in the sync path at this changelist
-					List<FStatRecord> fileRecords = await perforce.FStatAsync(FStatOptions.None, $"{syncPath}@{requiredChangeNumber}", cancellationToken).ToListAsync(cancellationToken);
-					if (fileRecords.Count == 0)
+					string? baseDepotPath = launcherSettings.PerforceDepotPath;
+					if (String.IsNullOrEmpty(baseDepotPath))
 					{
-						throw new UserErrorException($"Couldn't find any matching files for {syncPath}@{requiredChangeNumber}");
+						throw new UserErrorException($"Invalid setting for sync path");
 					}
 
-					// Sync all the files in this list to the same directory structure under the application folder
-					string depotPathPrefix = syncPath.Substring(0, syncPath.LastIndexOf('/') + 1);
-					foreach (FStatRecord fileRecord in fileRecords)
+					string baseDepotPathPrefix = baseDepotPath.TrimEnd('/');
+
+					// Find the most recent changelist
+					string syncPath = baseDepotPathPrefix + (preview ? "/UnstableRelease.zip" : "/Release.zip");
+					List<ChangesRecord> changes = await perforce.GetChangesAsync(ChangesOptions.None, 1, ChangeStatus.Submitted, syncPath, cancellationToken);
+					if (changes.Count == 0)
 					{
-						// Skip deleted files
-						if (Array.IndexOf(s_deleteActions, fileRecord.Action) != -1)
+						syncPath = baseDepotPathPrefix + (preview ? "/UnstableRelease/..." : "/Release/...");
+						changes = await perforce.GetChangesAsync(ChangesOptions.None, 1, ChangeStatus.Submitted, syncPath, cancellationToken);
+#pragma warning disable CA1508 // warning CA1508: 'changes.Count == 0' is always 'true'. Remove or refactor the condition(s) to avoid dead code.
+						if (changes.Count == 0)
 						{
-							continue;
+							throw new UserErrorException($"Unable to find any UGS binaries under {syncPath}");
 						}
+#pragma warning restore CA1508
+					}
 
-						if (fileRecord.DepotFile == null)
+					int requiredChangeNumber = changes[0].Number;
+					logger.LogInformation("Syncing from {SyncPath}", syncPath);
+
+					// Check if the version has changed
+					string requiredSyncText = String.Format("{0}\n{1}@{2}", perforce.Settings.ServerAndPort ?? "", syncPath, requiredChangeNumber);
+					if (syncText == null || syncText != requiredSyncText)
+					{
+						// Delete the output directory
+						await SafeDeleteDirectoryContentsWithRetryAsync(applicationFolder, cancellationToken);
+
+						// Find all the files in the sync path at this changelist
+						List<FStatRecord> fileRecords = await perforce.FStatAsync(FStatOptions.None, $"{syncPath}@{requiredChangeNumber}", cancellationToken).ToListAsync(cancellationToken);
+						if (fileRecords.Count == 0)
 						{
-							throw new UserErrorException("Missing depot path for returned file");
+							throw new UserErrorException($"Couldn't find any matching files for {syncPath}@{requiredChangeNumber}");
 						}
 
-						string localPath = Path.Combine(applicationFolder, fileRecord.DepotFile.Substring(depotPathPrefix.Length).Replace('/', Path.DirectorySeparatorChar));
-						if (!SafeCreateDirectory(Path.GetDirectoryName(localPath)!))
+						// Sync all the files in this list to the same directory structure under the application folder
+						string depotPathPrefix = syncPath.Substring(0, syncPath.LastIndexOf('/') + 1);
+						foreach (FStatRecord fileRecord in fileRecords)
 						{
-							throw new UserErrorException($"Couldn't create folder {Path.GetDirectoryName(localPath)}");
+							// Skip deleted files
+							if (Array.IndexOf(s_deleteActions, fileRecord.Action) != -1)
+							{
+								continue;
+							}
+
+							if (fileRecord.DepotFile == null)
+							{
+								throw new UserErrorException("Missing depot path for returned file");
+							}
+
+							string localPath = Path.Combine(applicationFolder, fileRecord.DepotFile.Substring(depotPathPrefix.Length).Replace('/', Path.DirectorySeparatorChar));
+							if (!SafeCreateDirectory(Path.GetDirectoryName(localPath)!))
+							{
+								throw new UserErrorException($"Couldn't create folder {Path.GetDirectoryName(localPath)}");
+							}
+
+							await perforce.PrintAsync(localPath, fileRecord.DepotFile, cancellationToken);
 						}
 
-						await perforce.PrintAsync(localPath, fileRecord.DepotFile, cancellationToken);
-					}
+						// If it was a zip file, extract it
+						if (syncPath.EndsWith(".zip", StringComparison.OrdinalIgnoreCase))
+						{
+							string localPath = Path.Combine(applicationFolder, syncPath.Substring(depotPathPrefix.Length).Replace('/', Path.DirectorySeparatorChar));
+							ZipFile.ExtractToDirectory(localPath, applicationFolder);
+						}
 
-					// If it was a zip file, extract it
-					if (syncPath.EndsWith(".zip", StringComparison.OrdinalIgnoreCase))
-					{
-						string localPath = Path.Combine(applicationFolder, syncPath.Substring(depotPathPrefix.Length).Replace('/', Path.DirectorySeparatorChar));
-						ZipFile.ExtractToDirectory(localPath, applicationFolder);
-					}
+						// Check the application exists
+						if (!File.Exists(applicationExe))
+						{
+							throw new UserErrorException($"Application was not synced from Perforce. Check that UnrealGameSync exists at {syncPath}/UnrealGameSync.exe, and you have access to it.");
+						}
 
-					// Check the application exists
-					if (!File.Exists(applicationExe))
-					{
-						throw new UserErrorException($"Application was not synced from Perforce. Check that UnrealGameSync exists at {syncPath}/UnrealGameSync.exe, and you have access to it.");
+						// Update the version
+						if (!TryWriteAllText(syncVersionFile, requiredSyncText))
+						{
+							throw new UserErrorException("Couldn't write sync text to {SyncVersionFile}");
+						}
 					}
 
-					// Update the version
-					if (!TryWriteAllText(syncVersionFile, requiredSyncText))
-					{
-						throw new UserErrorException("Couldn't write sync text to {SyncVersionFile}");
-					}
+					// Argument for updating in the future
+					updatePath = $"{syncPath}@>{requiredChangeNumber}";
+				}
+				else
+				{
+					throw new NotSupportedException("Invalid sync type");
 				}
 				logger.LogInformation("");
 
@@ -209,21 +294,28 @@ namespace UnrealGameSyncLauncher
                     }
                 }
 
-				StringBuilder newCommandLine = new StringBuilder(String.Format("-updatepath=\"{0}@>{1}\" -updatespawn=\"{2}\"{3}", syncPath, requiredChangeNumber, originalExecutable, preview ? " -unstable" : ""));
-				foreach (string arg in args)
+				// Create the new argument list
+				List<string> newArguments = new List<string>(args);
+				newArguments.Add($"-updatespawn={originalExecutable}");
+				if (updatePath != null)
 				{
-					newCommandLine.AppendFormat(" {0}", QuoteArgument(arg));
+					newArguments.Add($"-updatepath={updatePath}");
+				}
+				if (launcherSettings.PreviewBuild)
+				{
+					newArguments.Add("-unstable");
 				}
 
 				// Release the mutex now so that the new application can start up
 				instanceMutex.Close();
 
 				// Spawn the application
-				logger.LogInformation("Spawning {App} with command line: {CmdLine}", applicationExe, newCommandLine.ToString());
+				string newCommandLine = CommandLineArguments.Join(newArguments);
+				logger.LogInformation("Spawning {App} with command line: {CmdLine}", applicationExe, newCommandLine);
 				using (Process childProcess = new Process())
 				{
 					childProcess.StartInfo.FileName = applicationExe;
-					childProcess.StartInfo.Arguments = newCommandLine.ToString();
+					childProcess.StartInfo.Arguments = newCommandLine;
 					childProcess.StartInfo.UseShellExecute = false;
 					childProcess.StartInfo.CreateNoWindow = false;
 					if (!childProcess.Start())
@@ -334,5 +426,18 @@ namespace UnrealGameSyncLauncher
 				return false;
 			}
 		}
+
+		static async Task SafeDeleteDirectoryContentsWithRetryAsync(string directoryName, CancellationToken cancellationToken)
+		{
+			// Try to delete the directory contents. Retry for a while, in case we've been spawned by an application in this folder to do an update.
+			for (int numRetries = 0; !SafeDeleteDirectoryContents(directoryName); numRetries++)
+			{
+				if (numRetries > 20)
+				{
+					throw new UserErrorException($"Couldn't delete contents of {directoryName} (retried {numRetries} times).");
+				}
+				await Task.Delay(500, cancellationToken);
+			}
+		}
 	}
 }
diff --git a/Engine/Source/Programs/UnrealGameSync/UnrealGameSyncLauncher/UnrealGameSyncLauncher.csproj b/Engine/Source/Programs/UnrealGameSync/UnrealGameSyncLauncher/UnrealGameSyncLauncher.csproj
--- a/Engine/Source/Programs/UnrealGameSync/UnrealGameSyncLauncher/UnrealGameSyncLauncher.csproj
+++ b/Engine/Source/Programs/UnrealGameSync/UnrealGameSyncLauncher/UnrealGameSyncLauncher.csproj
@@ -50,6 +50,7 @@
   </ItemGroup>
   <ItemGroup>
     <ProjectReference Include="..\..\Shared\EpicGames.Core\EpicGames.Core.csproj" />
+    <ProjectReference Include="..\..\Shared\EpicGames.Horde\EpicGames.Horde.csproj" />
     <ProjectReference Include="..\..\Shared\EpicGames.Perforce\EpicGames.Perforce.csproj" />
     <ProjectReference Include="..\UnrealGameSyncShared\UnrealGameSyncShared.csproj" />
   </ItemGroup>
diff --git a/Engine/Source/Programs/UnrealGameSync/UnrealGameSyncShared/DeploymentSettings.cs b/Engine/Source/Programs/UnrealGameSync/UnrealGameSyncShared/DeploymentSettings.cs
--- a/Engine/Source/Programs/UnrealGameSync/UnrealGameSyncShared/DeploymentSettings.cs
+++ b/Engine/Source/Programs/UnrealGameSync/UnrealGameSyncShared/DeploymentSettings.cs
@@ -4,6 +4,8 @@
 #pragma warning disable CA1802 // warning CA1802: Field 'EnableAlerts' is declared as 'readonly' but is initialized with a constant value. Mark this field as 'const' instead.
 
 using EpicGames.Core;
+using EpicGames.Horde;
+using EpicGames.Horde.Tools;
 using Microsoft.Extensions.Logging;
 using Microsoft.Extensions.Logging.Abstractions;
 using System.Collections.Generic;
@@ -25,6 +27,11 @@ namespace UnrealGameSync
 		/// </summary>
 		public string? HordeUrl { get; set; }
 
+		/// <summary>
+		/// Identifier for the tool to sync from UGS
+		/// </summary>
+		public ToolId HordeToolId { get; set; } = new ToolId("ugs-win");
+
 		/// <summary>
 		/// SQL connection string used to connect to the database for telemetry and review data.
 		/// </summary>
diff --git a/Engine/Source/Programs/UnrealGameSync/UnrealGameSyncShared/LauncherSettings.cs b/Engine/Source/Programs/UnrealGameSync/UnrealGameSyncShared/LauncherSettings.cs
--- a/Engine/Source/Programs/UnrealGameSync/UnrealGameSyncShared/LauncherSettings.cs
+++ b/Engine/Source/Programs/UnrealGameSync/UnrealGameSyncShared/LauncherSettings.cs
@@ -8,14 +8,39 @@ using System.Text.RegularExpressions;
 
 namespace UnrealGameSync
 {
+	public enum LauncherUpdateSource
+	{
+		Perforce = 0,
+		Horde = 1,
+	}
+
 	public class LauncherSettings
 	{
+		public LauncherUpdateSource UpdateSource { get; set; }
+
+		public string? HordeServer { get; set; }
+
 		public string? PerforceServerAndPort { get; set; }
 		public string? PerforceUserName { get; set; }
 		public string? PerforceDepotPath { get; set; }
 
 		public bool PreviewBuild { get; set; }
 
+		public LauncherSettings()
+		{
+			HordeServer = DeploymentSettings.Instance.HordeUrl;
+			PerforceDepotPath = DeploymentSettings.Instance.DefaultDepotPath;
+		}
+
+		public LauncherSettings(LauncherSettings other)
+		{
+			PerforceServerAndPort = other.PerforceServerAndPort;
+			PerforceUserName = other.PerforceUserName;
+			PerforceDepotPath = other.PerforceDepotPath;
+
+			PreviewBuild = other.PreviewBuild;
+		}
+
 		public void Read()
 		{
 			if (RuntimeInformation.IsOSPlatform(OSPlatform.Windows))
@@ -31,6 +56,13 @@ namespace UnrealGameSync
 			{
 				if (key != null)
 				{
+					LauncherUpdateSource updateSource;
+					if (Enum.TryParse(key.GetValue("UpdateSource", UpdateSource) as string, out updateSource))
+					{
+						UpdateSource = updateSource;
+					}
+
+					HordeServer = key.GetValue("HordeServer", HordeServer) as string;
 					PerforceServerAndPort = key.GetValue("ServerAndPort", PerforceServerAndPort) as string;
 					PerforceUserName = key.GetValue("UserName", PerforceUserName) as string;
 					PerforceDepotPath = key.GetValue("DepotPath", PerforceDepotPath) as string;
@@ -74,32 +106,11 @@ namespace UnrealGameSync
 				// Delete this legacy setting
 				Utility.DeleteRegistryKey(key, "Server");
 
-				if (String.IsNullOrEmpty(PerforceServerAndPort))
-				{
-					Utility.DeleteRegistryKey(key, "ServerAndPort");
-				}
-				else
-				{
-					key.SetValue("ServerAndPort", PerforceServerAndPort);
-				}
-
-				if (String.IsNullOrEmpty(PerforceUserName))
-				{
-					Utility.DeleteRegistryKey(key, "UserName");
-				}
-				else
-				{
-					key.SetValue("UserName", PerforceUserName);
-				}
-
-				if (String.IsNullOrEmpty(PerforceDepotPath) || (DeploymentSettings.Instance.DefaultDepotPath != null && String.Equals(PerforceDepotPath, DeploymentSettings.Instance.DefaultDepotPath, StringComparison.OrdinalIgnoreCase)))
-				{
-					Utility.DeleteRegistryKey(key, "DepotPath");
-				}
-				else
-				{
-					key.SetValue("DepotPath", PerforceDepotPath);
-				}
+				SaveRegistryValue(key, "Source", UpdateSource.ToString(), LauncherUpdateSource.Perforce.ToString());
+				SaveRegistryValue(key, "HordeServer", HordeServer, DeploymentSettings.Instance.HordeUrl);
+				SaveRegistryValue(key, "ServerAndPort", PerforceServerAndPort, null);
+				SaveRegistryValue(key, "UserName", PerforceUserName, null);
+				SaveRegistryValue(key, "DepotPath", PerforceDepotPath, DeploymentSettings.Instance.DefaultDepotPath);
 
 				if (PreviewBuild)
 				{
@@ -111,5 +122,18 @@ namespace UnrealGameSync
 				}
 			}
 		}
+
+		[SupportedOSPlatform("windows")]
+		static void SaveRegistryValue(RegistryKey key, string name, string? value, string? defaultValue)
+		{
+			if (String.IsNullOrEmpty(value) || String.Equals(value, defaultValue, StringComparison.OrdinalIgnoreCase))
+			{
+				Utility.DeleteRegistryKey(key, name);
+			}
+			else
+			{
+				key.SetValue(name, value);
+			}
+		}
 	}
 }
diff --git a/Engine/Source/Programs/UnrealGameSync/UnrealGameSyncShared/UnrealGameSyncShared.csproj b/Engine/Source/Programs/UnrealGameSync/UnrealGameSyncShared/UnrealGameSyncShared.csproj
--- a/Engine/Source/Programs/UnrealGameSync/UnrealGameSyncShared/UnrealGameSyncShared.csproj
+++ b/Engine/Source/Programs/UnrealGameSync/UnrealGameSyncShared/UnrealGameSyncShared.csproj
@@ -44,6 +44,7 @@
 
   <ItemGroup>
     <ProjectReference Include="..\..\Shared\EpicGames.Core\EpicGames.Core.csproj" />
+    <ProjectReference Include="..\..\Shared\EpicGames.Horde\EpicGames.Horde.csproj" />
     <ProjectReference Include="..\..\Shared\EpicGames.OIDC\EpicGames.OIDC.csproj" />
     <ProjectReference Include="..\..\Shared\EpicGames.Perforce\EpicGames.Perforce.csproj" />
   </ItemGroup>
