UGS: Merge UnrealGameSync and UnrealGameSyncLauncher.

Having one application for launching and running UGS makes for simpler and easier deployment, and allows easier code sharing between the tools. By default, debug builds of UGS will launch directly into the UI, while release builds will do the standard upgrade check. Passing -noupgradecheck will force it to launch directly, and passing -settings will open the update configuration UI.

diff --git a/Engine/Source/Programs/UnrealGameSync/BuildUGS.xml b/Engine/Source/Programs/UnrealGameSync/BuildUGS.xml
--- a/Engine/Source/Programs/UnrealGameSync/BuildUGS.xml
+++ b/Engine/Source/Programs/UnrealGameSync/BuildUGS.xml
@@ -26,7 +26,7 @@
 	
 	<Agent Name="UnrealGameSync Agent" Type="Win64">
 		<Property Name="InstallerDir" Value="$(RootDir)/Engine/Source/Programs/UnrealGameSync/Installer"/>
-		<Property Name="InstallerZip" Value="$(RootDir)/LocalBuilds/UnrealGameSyncInstaller/ugs-installer.zip"/>
+		<Property Name="InstallerPublishDir" Value="$(RootDir)/LocalBuilds/UnrealGameSyncInstaller"/>
 
 		<Property Name="UgsPublishDir" Value="$(UgsOutputDir)"/>
 		<Property Name="UgsPublishDir" Value="$(UgsPublishDir)/$(SubmitPath)" If="'$(SubmitPath)' != ''"/>
@@ -34,35 +34,29 @@
 		<Property Name="UgsZipPublishDir" Value="$(UgsZipOutputDir)"/>
 		<Property Name="UgsZipPublishDir" Value="$(UgsZipPublishDir)/$(SubmitPath)" If="'$(SubmitPath)' != ''"/>
 
-		<Node Name="Build UnrealGameSync Installer">
+		<Node Name="Build UnrealGameSync">
+			<Delete Files="$(UgsOutputDir)/..."/>
+			<DotNet Arguments="publish &quot;$(UgsDir)/UnrealGameSyncCmd/UnrealGameSyncCmd.csproj&quot; --output=&quot;$(UgsPublishDir)&quot; --no-self-contained -c Release -r osx-x64 $(VersionArguments)"/>
+			<DotNet Arguments="publish &quot;$(UgsDir)/UnrealGameSyncCmd/UnrealGameSyncCmd.csproj&quot; --output=&quot;$(UgsPublishDir)&quot; --no-self-contained -c Release -r linux-x64 $(VersionArguments)"/>
+			<DotNet Arguments="publish &quot;$(UgsDir)/UnrealGameSyncCmd/UnrealGameSyncCmd.csproj&quot; --output=&quot;$(UgsPublishDir)&quot; --self-contained -c Release -r win-x64 $(VersionArguments)"/>
+			<DotNet Arguments="publish &quot;$(UgsDir)/UnrealGameSync/UnrealGameSync.csproj&quot; --output=&quot;$(UgsPublishDir)&quot; --self-contained -c Release -r win-x64 $(VersionArguments)"/>
+		</Node>
+		<Node Name="Build UnrealGameSync Installer" Requires="Build UnrealGameSync">
 			<Property Name="WixDir" Value="$(RootDir)/Engine/Source/ThirdParty/WiX/3.8"/>
-			<Property Name="UgsLauncherPublishDir" Value="$(UgsDir)/UnrealGameSyncLauncher/bin/Release/netcoreapp3.1/win-x64/publish"/>
-
-			<Delete Files="$(UgsLauncherPublishDir)/..."/>
-			<DotNet Arguments="publish &quot;$(UgsDir)/UnrealGameSyncLauncher/UnrealGameSyncLauncher.csproj&quot; --output=&quot;$(UgsLauncherPublishDir)&quot; --self-contained -c Release -r win-x64 $(VersionArguments)"/>
-
-			<Spawn Exe="$(WixDir)/heat.exe" Arguments="dir &quot;$(UgsLauncherPublishDir)&quot; -cg UGSLauncher_Project -dr INSTALLFOLDER -scom -sreg -srd -var var.BasePath -gg -sfrag -out obj/Release/UGSLauncher.wxs" WorkingDir="$(InstallerDir)"/>
-			<Spawn Exe="$(WixDir)/candle.exe" Arguments="-dBasePath=&quot;$(UgsLauncherPublishDir)&quot; -out obj\Release\ -dConfiguration=Release -dPlatform=x64 -arch x86 -ext ..\..\..\ThirdParty\WiX\3.8\WixUtilExtension.dll Product.wxs obj/Release/UGSLauncher.wxs" WorkingDir="$(InstallerDir)"/>
+			
+			<Spawn Exe="$(WixDir)/heat.exe" Arguments="dir &quot;$(UgsPublishDir)&quot; -cg UGSLauncher_Project -dr INSTALLFOLDER -scom -sreg -srd -var var.BasePath -gg -sfrag -out obj/Release/UGSLauncher.wxs" WorkingDir="$(InstallerDir)"/>
+			<Spawn Exe="$(WixDir)/candle.exe" Arguments="-dBasePath=&quot;$(UgsPublishDir)&quot; -out obj\Release\ -dConfiguration=Release -dPlatform=x64 -arch x86 -ext ..\..\..\ThirdParty\WiX\3.8\WixUtilExtension.dll Product.wxs obj/Release/UGSLauncher.wxs" WorkingDir="$(InstallerDir)"/>
 			<Spawn Exe="$(WixDir)/light.exe" Arguments="-out &quot;bin/Release/UnrealGameSync.msi&quot; -pdbout &quot;bin/Release/UnrealGameSync.wixpdb&quot; -cultures:null -ext ..\..\..\ThirdParty\WiX\3.8\WixUtilExtension.dll -sice:ICE69 &quot;obj\Release\Product.wixobj&quot; &quot;obj\Release\UGSLauncher.wixobj&quot;" WorkingDir="$(InstallerDir)"/>
 
-			<Zip ZipFile="$(InstallerZip)" FromDir="$(InstallerDir)/bin/Release" Files="UnrealGameSync.msi"/>
-		</Node>
-		<Node Name="Publish UnrealGameSync Installer" Requires="Build UnrealGameSync Installer">
-			<Error Message="Missing -HordeSettings=... argument" If="'$(HordeSettings)' == ''"/>
-			<DeployTool Id="ugs" Version="$(InformationalVersion)" Settings="$(HordeSettings)" File="$(InstallerZip)"/>
+			<Copy From="$(InstallerDir)/bin/Release/UnrealGameSync.msi" To="$(InstallerPublishDir)/UnrealGameSync.msi"/>
 		</Node>
-		
+
 		<Property Name="InstallerDependency" Value=""/>
 		<Property Name="InstallerDependency" Value="Build UnrealGameSync Installer" If="$(IncludeInstaller)"/>
-		
-		<Node Name="Build UnrealGameSync" Requires="$(InstallerDependency)">
-			<Delete Files="$(UgsOutputDir)/..."/>
-			<DotNet Arguments="publish &quot;$(UgsDir)/UnrealGameSyncCmd/UnrealGameSyncCmd.csproj&quot; --output=&quot;$(UgsPublishDir)&quot; --no-self-contained -c Release -r osx-x64 $(VersionArguments)"/>
-			<DotNet Arguments="publish &quot;$(UgsDir)/UnrealGameSyncCmd/UnrealGameSyncCmd.csproj&quot; --output=&quot;$(UgsPublishDir)&quot; --no-self-contained -c Release -r linux-x64 $(VersionArguments)"/>
-			<DotNet Arguments="publish &quot;$(UgsDir)/UnrealGameSyncCmd/UnrealGameSyncCmd.csproj&quot; --output=&quot;$(UgsPublishDir)&quot; --self-contained -c Release -r win-x64 $(VersionArguments)"/>
-			<DotNet Arguments="publish &quot;$(UgsDir)/UnrealGameSync/UnrealGameSync.csproj&quot; --output=&quot;$(UgsPublishDir)&quot; --self-contained -c Release -r win-x64 $(VersionArguments)"/>
+
+		<Node Name="Build UnrealGameSync Zip" Requires="Build UnrealGameSync;$(InstallerDependency)">
 			<Copy From="$(InstallerDir)/bin/Release/UnrealGameSync.msi" To="$(UgsPublishDir)/UnrealGameSync.msi" If="$(IncludeInstaller)"/>
-			<Zip FromDir="$(UgsOutputDir)" ZipFile="$(UgsZipPublishDir)/Ugs.zip"/>
+			<Zip FromDir="$(UgsPublishDir)" ZipFile="$(UgsZipPublishDir)/Ugs.zip"/>
 		</Node>
 
 		<!-- Windows (Preview) -->
diff --git a/Engine/Source/Programs/UnrealGameSync/UnrealGameSync/Forms/LogWindow.cs b/Engine/Source/Programs/UnrealGameSync/UnrealGameSync/Forms/LogWindow.cs
new file mode 100644
--- /dev/null
+++ b/Engine/Source/Programs/UnrealGameSync/UnrealGameSync/Forms/LogWindow.cs
@@ -0,0 +1,17 @@
+// Copyright Epic Games, Inc. All Rights Reserved.
+
+using System.Windows.Forms;
+
+namespace UnrealGameSync
+{
+	public partial class LogWindow : Form
+	{
+		public LogWindow(string text)
+		{
+			InitializeComponent();
+
+			LogTextBox.Text = text;
+			LogTextBox.Select(text.Length, 0);
+		}
+	}
+}
diff --git a/Engine/Source/Programs/UnrealGameSync/UnrealGameSync/Forms/UpdateSettingsWindow.cs b/Engine/Source/Programs/UnrealGameSync/UnrealGameSync/Forms/UpdateSettingsWindow.cs
new file mode 100644
--- /dev/null
+++ b/Engine/Source/Programs/UnrealGameSync/UnrealGameSync/Forms/UpdateSettingsWindow.cs
@@ -0,0 +1,141 @@
+// Copyright Epic Games, Inc. All Rights Reserved.
+
+using EpicGames.Core;
+using EpicGames.Perforce;
+using Microsoft.Extensions.Logging;
+using System;
+using System.Runtime.InteropServices;
+using System.Threading;
+using System.Threading.Tasks;
+using System.Windows.Forms;
+using UnrealGameSync;
+
+namespace UnrealGameSync
+{
+	partial class UpdateSettingsWindow : Form
+	{
+		[DllImport("user32.dll")]
+		private static extern IntPtr SendMessage(IntPtr hWnd, int msg, int wParam, [MarshalAs(UnmanagedType.LPWStr)] string lParam);
+
+		public delegate Task SyncAndRunDelegate(IPerforceConnection? perforce, LauncherSettings settings, ILogger logWriter, CancellationToken cancellationToken);
+
+		const int EmSetcuebanner = 0x1501;
+
+		LauncherSettings _settings;
+		string? _logText;
+		readonly SyncAndRunDelegate _syncAndRun;
+
+		public UpdateSettingsWindow(string? prompt, string? logText, LauncherSettings settings, SyncAndRunDelegate syncAndRun)
+		{
+			InitializeComponent();
+			Font = new System.Drawing.Font("Segoe UI", 8.25F, System.Drawing.FontStyle.Regular, System.Drawing.GraphicsUnit.Point, ((byte)(0)));
+
+			if (prompt != null)
+			{
+				PromptLabel.Text = prompt;
+			}
+
+			_logText = logText;
+			_settings = settings;
+
+			string defaultHordeServer = DeploymentSettings.Instance.HordeUrl ?? String.Empty;
+
+			HordeRadioBtn.Checked = settings.UpdateSource == LauncherUpdateSource.Horde;
+			HordeServerTextBox.Text = String.IsNullOrEmpty(settings.HordeServer) ? defaultHordeServer : settings.HordeServer;
+			ServerTextBox.Text = settings.PerforceServerAndPort ?? String.Empty;
+			UserNameTextBox.Text = settings.PerforceUserName ?? String.Empty;
+			DepotPathTextBox.Text = settings.PerforceDepotPath ?? String.Empty;
+			UsePreviewBuildCheckBox.Checked = settings.PreviewBuild;
+
+			_syncAndRun = syncAndRun;
+
+			ViewLogBtn.Visible = logText != null;
+		}
+
+		protected override void OnLoad(EventArgs e)
+		{
+			base.OnLoad(e);
+
+			SendMessage(ServerTextBox.Handle, EmSetcuebanner, 1, "Default Server");
+			SendMessage(UserNameTextBox.Handle, EmSetcuebanner, 1, "Default User");
+		}
+
+		private void ViewLogBtn_Click(object sender, EventArgs e)
+		{
+			using LogWindow log = new LogWindow(_logText ?? String.Empty);
+			log.ShowDialog(this);
+		}
+
+		private void ConnectBtn_Click(object sender, EventArgs e)
+		{
+			_settings.UpdateSource = HordeRadioBtn.Checked ? LauncherUpdateSource.Horde : LauncherUpdateSource.Perforce;
+			_settings.HordeServer = HordeServerTextBox.Text.Trim();
+			if (String.Equals(_settings.HordeServer, DeploymentSettings.Instance.HordeUrl, StringComparison.OrdinalIgnoreCase))
+			{
+				_settings.HordeServer = null;
+			}
+			_settings.PerforceServerAndPort = ServerTextBox.Text.Trim();
+			_settings.PerforceUserName = UserNameTextBox.Text.Trim();
+			_settings.PerforceDepotPath = DepotPathTextBox.Text.Trim();
+			_settings.PreviewBuild = UsePreviewBuildCheckBox.Checked;
+			_settings.Save();
+
+			// Create the P4 connection
+			CaptureLogger logger = new CaptureLogger();
+
+			// Create the task for connecting to this server
+			ModalTask? task;
+			if (_settings.UpdateSource == LauncherUpdateSource.Horde)
+			{
+				task = ModalTask.Execute(this, "Updating", "Checking for updates, please wait...", c => _syncAndRun(null, _settings, logger, c));
+			}
+			else
+			{
+				PerforceSettings perforceSettings = new PerforceSettings(PerforceSettings.Default);
+				if (!String.IsNullOrEmpty(_settings.PerforceServerAndPort))
+				{
+					perforceSettings.ServerAndPort = _settings.PerforceServerAndPort;
+				}
+				if (!String.IsNullOrEmpty(_settings.PerforceUserName))
+				{
+					perforceSettings.UserName = _settings.PerforceUserName;
+				}
+				perforceSettings.PreferNativeClient = true;
+
+				task = PerforceModalTask.Execute(this, "Updating", "Checking for updates, please wait...", perforceSettings, (p, c) => _syncAndRun(p, _settings, logger, c), logger);
+			}
+
+			if (task != null)
+			{
+				if (task.Succeeded)
+				{
+					_settings.Save();
+					DialogResult = DialogResult.OK;
+					Close();
+				}
+				PromptLabel.Text = task.Error;
+			}
+
+			_logText = logger.Render(Environment.NewLine);
+			ViewLogBtn.Visible = true;
+		}
+
+		private void CancelBtn_Click(object sender, EventArgs e)
+		{
+			DialogResult = DialogResult.Cancel;
+			Close();
+		}
+
+		private void HordeRadioBtn_CheckedChanged(object sender, EventArgs e)
+		{
+			HordeGroupBox.Enabled = HordeRadioBtn.Checked;
+			PerforceGroupBox.Enabled = PerforceRadioBtn.Checked;
+		}
+
+		private void PerforceRadioBtn_CheckedChanged(object sender, EventArgs e)
+		{
+			HordeGroupBox.Enabled = HordeRadioBtn.Checked;
+			PerforceGroupBox.Enabled = PerforceRadioBtn.Checked;
+		}
+	}
+}
diff --git a/Engine/Source/Programs/UnrealGameSync/UnrealGameSync/Launcher.cs b/Engine/Source/Programs/UnrealGameSync/UnrealGameSync/Launcher.cs
new file mode 100644
--- /dev/null
+++ b/Engine/Source/Programs/UnrealGameSync/UnrealGameSync/Launcher.cs
@@ -0,0 +1,431 @@
+// Copyright Epic Games, Inc. All Rights Reserved.
+
+using EpicGames.Core;
+using EpicGames.Horde;
+using EpicGames.Horde.Storage;
+using EpicGames.Horde.Storage.Clients;
+using EpicGames.Horde.Storage.Nodes;
+using EpicGames.Horde.Tools;
+using EpicGames.Perforce;
+using Microsoft.Extensions.DependencyInjection;
+using Microsoft.Extensions.Logging;
+using System;
+using System.Collections.Generic;
+using System.Diagnostics;
+using System.Diagnostics.CodeAnalysis;
+using System.IO;
+using System.IO.Compression;
+using System.Linq;
+using System.Net.Http;
+using System.Reflection;
+using System.Reflection.Metadata;
+using System.Text;
+using System.Threading;
+using System.Threading.Tasks;
+using System.Windows.Forms;
+using UnrealGameSync;
+
+namespace UnrealGameSync
+{
+	static class Launcher
+	{
+		public static bool SyncAndRunLatest(Mutex instanceMutex, string[] args)
+		{
+			// Figure out if we should sync the unstable build by default
+			bool preview = args.Contains("-unstable", StringComparer.InvariantCultureIgnoreCase) || args.Contains("-preview", StringComparer.InvariantCultureIgnoreCase);
+			bool openSettings = args.Contains("-settings", StringComparer.OrdinalIgnoreCase);
+
+#if DEBUG
+			// For debug builds, just launch into the regular UGS application.
+			if (!openSettings)
+			{
+				return false;
+			}
+#endif
+
+			// Read the settings
+			LauncherSettings launcherSettings = new LauncherSettings();
+			launcherSettings.PerforceDepotPath = DeploymentSettings.Instance.DefaultDepotPath;
+			launcherSettings.Read();
+
+			// If the shift key is held down, immediately show the settings window
+			Task SyncAndRunWrapper(IPerforceConnection? perforce, LauncherSettings settings, ILogger logWriter, CancellationToken cancellationToken) => SyncAndRun(perforce, settings, args, instanceMutex, logWriter, cancellationToken);
+			if ((Control.ModifierKeys & Keys.Shift) != 0 || openSettings)
+			{
+				// Show the settings window immediately
+				using UpdateSettingsWindow updateWindow = new UpdateSettingsWindow(null, null, launcherSettings, SyncAndRunWrapper);
+				return updateWindow.ShowDialog() != DialogResult.Cancel;
+			}
+			else
+			{
+				// Try to do a sync with the current settings first
+				CaptureLogger logger = new CaptureLogger();
+
+				ModalTask? task;
+				if (launcherSettings.UpdateSource == LauncherUpdateSource.Horde)
+				{
+					task = ModalTask.Execute(null, "Updating", "Checking for updates, please wait...", c => SyncAndRunWrapper(null, launcherSettings, logger, c));
+				}
+				else
+				{
+					IPerforceSettings settings = new PerforceSettings(PerforceSettings.Default) { PreferNativeClient = true }.MergeWith(newServerAndPort: launcherSettings.PerforceServerAndPort, newUserName: launcherSettings.PerforceUserName);
+					task = PerforceModalTask.Execute(null, "Updating", "Checking for updates, please wait...", settings, (p, c) => SyncAndRunWrapper(p, launcherSettings, logger, c), logger);
+				}
+
+				if (task == null)
+				{
+					logger.LogInformation("Canceled by user");
+				}
+				else if (task.Succeeded)
+				{
+					return true;
+				}
+
+				using UpdateSettingsWindow updateWindow = new UpdateSettingsWindow("Unable to update UnrealGameSync from Perforce. Verify that your connection settings are correct.", logger.Render(Environment.NewLine), launcherSettings, SyncAndRunWrapper);
+				return updateWindow.DialogResult != DialogResult.Cancel;
+			}
+		}
+
+		// Values of the Perforce "action" field that means that the file should no longer be synced
+		public static readonly string[] s_deleteActions = { "delete", "move/delete", "purge", "archive" };
+
+		class LoggerProviderAdapter : ILoggerProvider
+		{
+			readonly ILogger _logger;
+
+			public LoggerProviderAdapter(ILogger logger) => _logger = logger;
+			public ILogger CreateLogger(string categoryName) => _logger;
+			public void Dispose() { }
+		}
+
+		public static async Task SyncAndRun(IPerforceConnection? perforce, LauncherSettings launcherSettings, string[] args, Mutex instanceMutex, ILogger logger, CancellationToken cancellationToken)
+		{
+			try
+			{
+				// Create the target folder
+				string applicationFolder = Path.Combine(Environment.GetFolderPath(Environment.SpecialFolder.LocalApplicationData), "UnrealGameSync", "Latest");
+				if (!SafeCreateDirectory(applicationFolder))
+				{
+					throw new UserErrorException($"Couldn't create directory: {applicationFolder}");
+				}
+
+				// Check the application exists
+				string applicationExe = Path.Combine(applicationFolder, "UnrealGameSync.exe");
+
+				// Read the current version
+				string syncVersionFile = Path.Combine(applicationFolder, "SyncVersion.txt");
+
+				string? syncText = null;
+				if (File.Exists(applicationExe) && File.Exists(syncVersionFile))
+				{
+					TryReadAllText(syncVersionFile, out syncText);
+				}
+
+				// New command line for launching the updated executable
+				string? updatePath = null;
+				if (launcherSettings.UpdateSource == LauncherUpdateSource.Horde)
+				{
+					if (String.IsNullOrEmpty(launcherSettings.HordeServer))
+					{
+						throw new Exception("No Horde server is configured");
+					}
+
+					Uri hordeServerUrl = new Uri(launcherSettings.HordeServer);
+
+					ServiceCollection services = new ServiceCollection();
+					services.AddLogging(builder => builder.AddProvider(new LoggerProviderAdapter(logger)));
+					services.AddHordeHttpClient((sp, client) => client.BaseAddress = hordeServerUrl);
+
+					await using ServiceProvider serviceProvider = services.BuildServiceProvider();
+					HordeHttpClient httpClient = serviceProvider.GetRequiredService<HordeHttpClient>();
+
+					ToolId toolId = DeploymentSettings.Instance.HordeToolId;
+
+					GetToolResponse tool = await httpClient.GetToolAsync(toolId, cancellationToken: cancellationToken);
+					GetToolDeploymentResponse deployment = tool.Deployments[^1];
+					Uri deploymentUri = new Uri(hordeServerUrl, $"api/v1/tools/{toolId}/deployments/{deployment.Id}");
+
+					string requiredSyncText = deploymentUri.ToString();
+					if (syncText == null || syncText != requiredSyncText)
+					{
+						// Delete the output directory
+						await SafeDeleteDirectoryContentsWithRetryAsync(applicationFolder, cancellationToken);
+
+						// Download and extract the zip file
+						string zipFile = Path.Combine(applicationFolder, "update.zip");
+						using (Stream requestStream = await httpClient.GetToolDeploymentZipAsync(toolId, deployment.Id, cancellationToken))
+						{
+							using (Stream tempFileStream = File.Open(zipFile, FileMode.Create, FileAccess.Write, FileShare.None))
+							{
+								await requestStream.CopyToAsync(tempFileStream, cancellationToken);
+							}
+						}
+						ZipFile.ExtractToDirectory(zipFile, applicationFolder);
+						File.Delete(zipFile);
+
+						// Update the version
+						if (!TryWriteAllText(syncVersionFile, requiredSyncText))
+						{
+							throw new UserErrorException("Couldn't write sync text to {SyncVersionFile}");
+						}
+					}
+				}
+				else if (launcherSettings.UpdateSource == LauncherUpdateSource.Perforce)
+				{
+					bool preview = launcherSettings.PreviewBuild;
+					if (perforce == null)
+					{
+						throw new UserErrorException("No Perforce connection");
+					}
+
+					string? baseDepotPath = launcherSettings.PerforceDepotPath;
+					if (String.IsNullOrEmpty(baseDepotPath))
+					{
+						throw new UserErrorException($"Invalid setting for sync path");
+					}
+
+					string baseDepotPathPrefix = baseDepotPath.TrimEnd('/');
+
+					// Find the most recent changelist
+					string syncPath = baseDepotPathPrefix + (preview ? "/UnstableRelease.zip" : "/Release.zip");
+					List<ChangesRecord> changes = await perforce.GetChangesAsync(ChangesOptions.None, 1, ChangeStatus.Submitted, syncPath, cancellationToken);
+					if (changes.Count == 0)
+					{
+						syncPath = baseDepotPathPrefix + (preview ? "/UnstableRelease/..." : "/Release/...");
+						changes = await perforce.GetChangesAsync(ChangesOptions.None, 1, ChangeStatus.Submitted, syncPath, cancellationToken);
+#pragma warning disable CA1508 // warning CA1508: 'changes.Count == 0' is always 'true'. Remove or refactor the condition(s) to avoid dead code.
+						if (changes.Count == 0)
+						{
+							throw new UserErrorException($"Unable to find any UGS binaries under {syncPath}");
+						}
+#pragma warning restore CA1508
+					}
+
+					int requiredChangeNumber = changes[0].Number;
+					logger.LogInformation("Syncing from {SyncPath}", syncPath);
+
+					// Check if the version has changed
+					string requiredSyncText = String.Format("{0}\n{1}@{2}", perforce.Settings.ServerAndPort ?? "", syncPath, requiredChangeNumber);
+					if (syncText == null || syncText != requiredSyncText)
+					{
+						// Delete the output directory
+						await SafeDeleteDirectoryContentsWithRetryAsync(applicationFolder, cancellationToken);
+
+						// Find all the files in the sync path at this changelist
+						List<FStatRecord> fileRecords = await perforce.FStatAsync(FStatOptions.None, $"{syncPath}@{requiredChangeNumber}", cancellationToken).ToListAsync(cancellationToken);
+						if (fileRecords.Count == 0)
+						{
+							throw new UserErrorException($"Couldn't find any matching files for {syncPath}@{requiredChangeNumber}");
+						}
+
+						// Sync all the files in this list to the same directory structure under the application folder
+						string depotPathPrefix = syncPath.Substring(0, syncPath.LastIndexOf('/') + 1);
+						foreach (FStatRecord fileRecord in fileRecords)
+						{
+							// Skip deleted files
+							if (Array.IndexOf(s_deleteActions, fileRecord.Action) != -1)
+							{
+								continue;
+							}
+
+							if (fileRecord.DepotFile == null)
+							{
+								throw new UserErrorException("Missing depot path for returned file");
+							}
+
+							string localPath = Path.Combine(applicationFolder, fileRecord.DepotFile.Substring(depotPathPrefix.Length).Replace('/', Path.DirectorySeparatorChar));
+							if (!SafeCreateDirectory(Path.GetDirectoryName(localPath)!))
+							{
+								throw new UserErrorException($"Couldn't create folder {Path.GetDirectoryName(localPath)}");
+							}
+
+							await perforce.PrintAsync(localPath, fileRecord.DepotFile, cancellationToken);
+						}
+
+						// If it was a zip file, extract it
+						if (syncPath.EndsWith(".zip", StringComparison.OrdinalIgnoreCase))
+						{
+							string localPath = Path.Combine(applicationFolder, syncPath.Substring(depotPathPrefix.Length).Replace('/', Path.DirectorySeparatorChar));
+							ZipFile.ExtractToDirectory(localPath, applicationFolder);
+						}
+
+						// Check the application exists
+						if (!File.Exists(applicationExe))
+						{
+							throw new UserErrorException($"Application was not synced from Perforce. Check that UnrealGameSync exists at {syncPath}/UnrealGameSync.exe, and you have access to it.");
+						}
+
+						// Update the version
+						if (!TryWriteAllText(syncVersionFile, requiredSyncText))
+						{
+							throw new UserErrorException("Couldn't write sync text to {SyncVersionFile}");
+						}
+					}
+
+					// Argument for updating in the future
+					updatePath = $"{syncPath}@>{requiredChangeNumber}";
+				}
+				else
+				{
+					throw new NotSupportedException("Invalid sync type");
+				}
+				logger.LogInformation("");
+
+				// Build the command line for the synced application, including the sync path to monitor for updates
+				string originalExecutable = Assembly.GetEntryAssembly()!.Location;
+                if (Path.GetExtension(originalExecutable).Equals(".dll", StringComparison.OrdinalIgnoreCase))
+                {
+                    string newExecutable = Path.ChangeExtension(originalExecutable, ".exe");
+                    if (File.Exists(newExecutable))
+                    {
+                        originalExecutable = newExecutable;
+                    }
+                }
+
+				// Create the new argument list
+				List<string> newArguments = new List<string>(args);
+				newArguments.Add($"-updatespawn={originalExecutable}");
+				if (updatePath != null)
+				{
+					newArguments.Add($"-updatepath={updatePath}");
+				}
+				if (launcherSettings.PreviewBuild)
+				{
+					newArguments.Add("-unstable");
+				}
+
+				// Release the mutex now so that the new application can start up
+				instanceMutex.Close();
+
+				// Spawn the application
+				string newCommandLine = CommandLineArguments.Join(newArguments);
+				logger.LogInformation("Spawning {App} with command line: {CmdLine}", applicationExe, newCommandLine);
+				using (Process childProcess = new Process())
+				{
+					childProcess.StartInfo.FileName = applicationExe;
+					childProcess.StartInfo.Arguments = newCommandLine;
+					childProcess.StartInfo.UseShellExecute = false;
+					childProcess.StartInfo.CreateNoWindow = false;
+					if (!childProcess.Start())
+					{
+						throw new UserErrorException("Failed to start process");
+					}
+				}
+			}
+			catch (UserErrorException ex)
+			{
+				logger.LogError("{Message}", ex.Message);
+				throw;
+			}
+			catch (Exception ex)
+			{
+				logger.LogError(ex, "Error while syncing application.");
+				foreach (string line in ex.ToString().Split('\n'))
+				{
+					logger.LogError("{Line}", line);
+				}
+				throw;
+			}
+		}
+
+		static string QuoteArgument(string arg)
+		{
+			if(arg.Contains(' ', StringComparison.Ordinal) && !arg.StartsWith("\"", StringComparison.Ordinal))
+			{
+				return String.Format("\"{0}\"", arg);
+			}
+			else
+			{
+				return arg;
+			}
+		}
+
+		static bool TryReadAllText(string fileName, [NotNullWhen(true)] out string? text)
+		{
+			try
+			{
+				text = File.ReadAllText(fileName);
+				return true;
+			}
+			catch(Exception)
+			{
+				text = null;
+				return false;
+			}
+		}
+
+		static bool TryWriteAllText(string fileName, string text)
+		{
+			try
+			{
+				File.WriteAllText(fileName, text);
+				return true;
+			}
+			catch(Exception)
+			{
+				return false;
+			}
+		}
+
+		static bool SafeCreateDirectory(string directoryName)
+		{
+			try
+			{
+				Directory.CreateDirectory(directoryName);
+				return true;
+			}
+			catch(Exception)
+			{
+				return false;
+			}
+		}
+
+		static bool SafeDeleteDirectory(string directoryName)
+		{
+			try
+			{
+				Directory.Delete(directoryName, true);
+				return true;
+			}
+			catch(Exception)
+			{
+				return false;
+			}
+		}
+
+		static bool SafeDeleteDirectoryContents(string directoryName)
+		{
+			try
+			{
+				DirectoryInfo directory = new DirectoryInfo(directoryName);
+				foreach(FileInfo childFile in directory.EnumerateFiles("*", SearchOption.AllDirectories))
+				{
+					childFile.Attributes &= ~FileAttributes.ReadOnly;
+					childFile.Delete();
+				}
+				foreach(DirectoryInfo childDirectory in directory.EnumerateDirectories())
+				{
+					SafeDeleteDirectory(childDirectory.FullName);
+				}
+				return true;
+			}
+			catch(Exception)
+			{
+				return false;
+			}
+		}
+
+		static async Task SafeDeleteDirectoryContentsWithRetryAsync(string directoryName, CancellationToken cancellationToken)
+		{
+			// Try to delete the directory contents. Retry for a while, in case we've been spawned by an application in this folder to do an update.
+			for (int numRetries = 0; !SafeDeleteDirectoryContents(directoryName); numRetries++)
+			{
+				if (numRetries > 20)
+				{
+					throw new UserErrorException($"Couldn't delete contents of {directoryName} (retried {numRetries} times).");
+				}
+				await Task.Delay(500, cancellationToken);
+			}
+		}
+	}
+}
diff --git a/Engine/Source/Programs/UnrealGameSync/UnrealGameSync/Program.cs b/Engine/Source/Programs/UnrealGameSync/UnrealGameSync/Program.cs
--- a/Engine/Source/Programs/UnrealGameSync/UnrealGameSync/Program.cs
+++ b/Engine/Source/Programs/UnrealGameSync/UnrealGameSync/Program.cs
@@ -12,6 +12,7 @@ using System;
 using System.Collections.Generic;
 using System.Diagnostics.CodeAnalysis;
 using System.IO;
+using System.Linq;
 using System.Reflection;
 using System.Threading;
 using System.Threading.Tasks;
@@ -74,6 +75,18 @@ namespace UnrealGameSync
 				TaskScheduler.UnobservedTaskException += Application_UnobservedException_Sentry;
 			}
 
+			try
+			{
+				RealMain(args);
+			}
+			catch (Exception ex)
+			{
+				CaptureException(ex);
+			}
+		}
+
+		static void RealMain(string[] args)
+		{
 			bool firstInstance;
 			using (Mutex instanceMutex = new Mutex(true, "UnrealGameSyncRunning", out firstInstance))
 			{
@@ -83,15 +96,22 @@ namespace UnrealGameSync
 
 				using (EventWaitHandle activateEvent = new EventWaitHandle(false, EventResetMode.AutoReset, "ActivateUnrealGameSync"))
 				{
-					// handle any url passed in, possibly exiting
+					// Check for a newer version of the application
+					if (!args.Contains("-NoUpdateCheck", StringComparer.OrdinalIgnoreCase) && Launcher.SyncAndRunLatest(instanceMutex, args))
+					{
+						return;
+					}
+
+					// Handle any url passed in, possibly exiting
 					if (UriHandler.ProcessCommandLine(args, firstInstance, activateEvent))
 					{
 						return;
 					}
 
+					// Launch the application proper
 					if (firstInstance)
 					{
-						GuardedInnerMainAsync(instanceMutex, activateEvent, args);
+						InnerMainAsync(instanceMutex, activateEvent, args).GetAwaiter().GetResult();
 					}
 					else
 					{
@@ -101,18 +121,6 @@ namespace UnrealGameSync
 			}
 		}
 
-		static void GuardedInnerMainAsync(Mutex instanceMutex, EventWaitHandle activateEvent, string[] args)
-		{
-			try
-			{
-				InnerMainAsync(instanceMutex, activateEvent, args).GetAwaiter().GetResult();
-			}
-			catch (Exception ex)
-			{
-				CaptureException(ex);
-			}
-		}
-
 		static async Task InnerMainAsync(Mutex instanceMutex, EventWaitHandle activateEvent, string[] args)
 		{
 			LauncherSettings launcherSettings = new LauncherSettings();
diff --git a/Engine/Source/Programs/UnrealGameSync/UnrealGameSyncLauncher/Forms/LogWindow.cs b/Engine/Source/Programs/UnrealGameSync/UnrealGameSyncLauncher/Forms/LogWindow.cs
deleted file mode 100644
diff --git a/Engine/Source/Programs/UnrealGameSync/UnrealGameSyncLauncher/Forms/SettingsWindow.cs b/Engine/Source/Programs/UnrealGameSync/UnrealGameSyncLauncher/Forms/SettingsWindow.cs
deleted file mode 100644
diff --git a/Engine/Source/Programs/UnrealGameSync/UnrealGameSyncLauncher/Program.cs b/Engine/Source/Programs/UnrealGameSync/UnrealGameSyncLauncher/Program.cs
deleted file mode 100644
diff --git a/Engine/Source/Programs/UnrealGameSync/UnrealGameSyncLauncher/Properties/AssemblyInfo.cs b/Engine/Source/Programs/UnrealGameSync/UnrealGameSyncLauncher/Properties/AssemblyInfo.cs
deleted file mode 100644
diff --git a/Engine/Source/Programs/UnrealGameSync/UnrealGameSyncLauncher/Properties/Settings.settings b/Engine/Source/Programs/UnrealGameSync/UnrealGameSyncLauncher/Properties/Settings.settings
deleted file mode 100644
diff --git a/Engine/Source/Programs/UnrealGameSync/UnrealGameSyncLauncher/UnrealGameSyncLauncher.csproj b/Engine/Source/Programs/UnrealGameSync/UnrealGameSyncLauncher/UnrealGameSyncLauncher.csproj
deleted file mode 100644
diff --git a/Engine/Source/Programs/UnrealGameSync/UnrealGameSyncLauncher/app.manifest b/Engine/Source/Programs/UnrealGameSync/UnrealGameSyncLauncher/app.manifest
deleted file mode 100644
