UGS: Refactor tool manager to make it easier to add different sources.

* State of synced tools is now stored in a JSON file, includes uninstall commands to allow correctly removing tools that are no longer available.
* State is now read immediately on startup, without having to fetch data from Perforce before buttons are available.
* Fixed an issue where tools may not update until next change is detected in tools depot path.

diff --git a/Engine/Source/Programs/Shared/EpicGames.Core/FileTransaction.cs b/Engine/Source/Programs/Shared/EpicGames.Core/FileTransaction.cs
--- a/Engine/Source/Programs/Shared/EpicGames.Core/FileTransaction.cs
+++ b/Engine/Source/Programs/Shared/EpicGames.Core/FileTransaction.cs
@@ -51,6 +51,25 @@ namespace EpicGames.Core
 			}
 		}
 
+		/// <summary>
+		/// Reads all data from a file into a byte array
+		/// </summary>
+		/// <param name="file">File to read from</param>
+		/// <returns>The data that was read</returns>
+		public static byte[]? ReadAllBytes(FileReference file)
+		{
+			using (Stream? stream = OpenRead(file))
+			{
+				if (stream != null)
+				{
+					byte[] data = new byte[stream.Length];
+					stream.ReadFixedLengthBytes(data);
+					return data;
+				}
+			}
+			return null;
+		}
+
 		/// <summary>
 		/// Reads all data from a file into a byte array
 		/// </summary>
diff --git a/Engine/Source/Programs/Shared/EpicGames.Core/StreamUtils.cs b/Engine/Source/Programs/Shared/EpicGames.Core/StreamUtils.cs
--- a/Engine/Source/Programs/Shared/EpicGames.Core/StreamUtils.cs
+++ b/Engine/Source/Programs/Shared/EpicGames.Core/StreamUtils.cs
@@ -34,6 +34,24 @@ namespace EpicGames.Core
 			return length;
 		}
 
+		/// <summary>
+		/// Reads a fixed amount of data from a stream, throwing an exception if the entire buffer cannot be read.
+		/// </summary>
+		/// <param name="stream">Stream to read from</param>
+		/// <param name="buffer">Buffer to receive the output data</param>
+		public static void ReadFixedLengthBytes(this Stream stream, Span<byte> buffer)
+		{
+			for (int offset = 0; offset < buffer.Length;)
+			{
+				int readLength = stream.Read(buffer.Slice(offset));
+				if (readLength == 0)
+				{
+					throw new EndOfStreamException($"Unexpected end of stream while trying to read {buffer.Length} bytes.");
+				}
+				offset += readLength;
+			}
+		}
+
 		/// <summary>
 		/// Reads a fixed amount of data from a stream, throwing an exception if the entire buffer cannot be read.
 		/// </summary>
diff --git a/Engine/Source/Programs/UnrealGameSync/UnrealGameSync/Controls/WorkspaceControl.cs b/Engine/Source/Programs/UnrealGameSync/UnrealGameSync/Controls/WorkspaceControl.cs
--- a/Engine/Source/Programs/UnrealGameSync/UnrealGameSync/Controls/WorkspaceControl.cs
+++ b/Engine/Source/Programs/UnrealGameSync/UnrealGameSync/Controls/WorkspaceControl.cs
@@ -3325,7 +3325,7 @@ namespace UnrealGameSync
 			SafeProcessStart("p4v.exe", commandLine.ToString());
 		}
 
-		void RunTool(ToolDefinition tool, ToolLink link)
+		void RunTool(ToolInfo tool, ToolLink link)
 		{
 			DirectoryReference? toolDir = _owner.ToolUpdateMonitor.GetToolPath(tool.Name);
 			if (toolDir != null)
@@ -3443,10 +3443,10 @@ namespace UnrealGameSync
 					programsLine.AddText("  |  ");
 				}
 
-				List<ToolDefinition> tools = _owner.ToolUpdateMonitor.Tools;
-				foreach (ToolDefinition tool in tools)
+				IReadOnlyList<ToolInfo> tools = _owner.ToolUpdateMonitor.GetEnabledTools();
+				foreach (ToolInfo tool in tools)
 				{
-					if (tool.Enabled && tool.SafeWhenBusy)
+					if (tool.SafeWhenBusy)
 					{
 						foreach (ToolLink link in tool.StatusPanelLinks)
 						{
@@ -3635,16 +3635,13 @@ namespace UnrealGameSync
 					programsLine.AddText("  |  ");
 				}
 
-				List<ToolDefinition> tools = _owner.ToolUpdateMonitor.Tools;
-				foreach (ToolDefinition tool in tools)
+				IReadOnlyList<ToolInfo> tools = _owner.ToolUpdateMonitor.GetEnabledTools();
+				foreach (ToolInfo tool in tools)
 				{
-					if (tool.Enabled)
+					foreach (ToolLink link in tool.StatusPanelLinks)
 					{
-						foreach (ToolLink link in tool.StatusPanelLinks)
-						{
-							programsLine.AddLink(link.Label, FontStyle.Regular, () => RunTool(tool, link));
-							programsLine.AddText("  |  ");
-						}
+						programsLine.AddLink(link.Label, FontStyle.Regular, () => RunTool(tool, link));
+						programsLine.AddText("  |  ");
 					}
 				}
 
diff --git a/Engine/Source/Programs/UnrealGameSync/UnrealGameSync/Forms/ApplicationSettingsWindow.cs b/Engine/Source/Programs/UnrealGameSync/UnrealGameSync/Forms/ApplicationSettingsWindow.cs
--- a/Engine/Source/Programs/UnrealGameSync/UnrealGameSync/Forms/ApplicationSettingsWindow.cs
+++ b/Engine/Source/Programs/UnrealGameSync/UnrealGameSync/Forms/ApplicationSettingsWindow.cs
@@ -53,13 +53,13 @@ namespace UnrealGameSync
 			public Guid Id => Definition.Id;
 
 			public int Index { get; }
-			public ToolDefinition Definition { get; }
+			public ToolInfo Definition { get; }
 			public List<ToolItem> RequiresTools { get; } = new List<ToolItem>();
 
 			public bool Enabled { get; set; }
 			public int DependencyRefCount { get; set; }
 
-			public ToolItem(int index, ToolDefinition definition, bool enabled)
+			public ToolItem(int index, ToolInfo definition, bool enabled)
 			{
 				Index = index;
 				Definition = definition;
@@ -150,7 +150,7 @@ namespace UnrealGameSync
 
 			List<ToolItem> toolItems = new List<ToolItem>();
 			Dictionary<Guid, ToolItem> idToToolItem = new Dictionary<Guid, ToolItem>();
-			foreach (ToolDefinition tool in toolUpdateMonitor.Tools.OrderBy(x => x.ToString()))
+			foreach (ToolInfo tool in toolUpdateMonitor.GetTools().OrderBy(x => x.Name))
 			{
 				ToolItem toolItem = new ToolItem(toolItems.Count, tool, settings.EnabledTools.Contains(tool.Id));
 				idToToolItem[tool.Id] = toolItem;
@@ -228,7 +228,7 @@ namespace UnrealGameSync
 			LauncherSettings launcherSettings = new LauncherSettings(originalLauncherSettings);
 
 			launcherSettings.HordeServer = HordeServerTextBox.Text.Trim();
-			if (launcherSettings.HordeServer.Length == 0 || String.Equals(launcherSettings.HordeServer, DeploymentSettings.Instance.HordeUrl, StringComparison.OrdinalIgnoreCase))
+			if (launcherSettings.HordeServer.Length == 0)
 			{
 				launcherSettings.HordeServer = null;
 			}
diff --git a/Engine/Source/Programs/UnrealGameSync/UnrealGameSync/ToolUpdateMonitor.cs b/Engine/Source/Programs/UnrealGameSync/UnrealGameSync/ToolUpdateMonitor.cs
--- a/Engine/Source/Programs/UnrealGameSync/UnrealGameSync/ToolUpdateMonitor.cs
+++ b/Engine/Source/Programs/UnrealGameSync/UnrealGameSync/ToolUpdateMonitor.cs
@@ -1,6 +1,7 @@
 // Copyright Epic Games, Inc. All Rights Reserved.
 
 using EpicGames.Core;
+using EpicGames.Horde.Tools;
 using EpicGames.Perforce;
 using Microsoft.Extensions.DependencyInjection;
 using Microsoft.Extensions.Logging;
@@ -9,6 +10,7 @@ using System.Collections.Generic;
 using System.Diagnostics;
 using System.IO;
 using System.Linq;
+using System.Text.Json;
 using System.Threading;
 using System.Threading.Tasks;
 
@@ -16,12 +18,13 @@ using System.Threading.Tasks;
 
 namespace UnrealGameSync
 {
+	[DebuggerDisplay("{Label}")]
 	class ToolLink
 	{
-		public string Label;
-		public string FileName;
-		public string? Arguments;
-		public string? WorkingDir;
+		public string Label { get; set; }
+		public string FileName { get; set; }
+		public string? Arguments { get; set; }
+		public string? WorkingDir { get; set; }
 
 		public ToolLink(string label, string fileName)
 		{
@@ -30,33 +33,30 @@ namespace UnrealGameSync
 		}
 	}
 
-	class ToolDefinition
+	[DebuggerDisplay("{Name}")]
+	class ToolInfo
 	{
-		public Guid Id;
-		public string Name;
-		public string Description;
-		public bool Enabled;
-		public Func<ILogger, CancellationToken, Task>? InstallAction;
-		public Func<ILogger, CancellationToken, Task>? UninstallAction;
-		public List<ToolLink> StatusPanelLinks = new List<ToolLink>();
-		public string ZipPath;
-		public int ZipChange;
-		public string ConfigPath;
-		public int ConfigChange;
-		public HashSet<Guid> DependsOnToolIds;
-		public bool SafeWhenBusy;
-
-		public ToolDefinition(Guid id, string name, string description, string zipPath, int zipChange, string configPath, int configChange, HashSet<Guid> dependsOnToolIds, bool safeWhenBusy)
+		public Guid Id { get; set; }
+		public string Name { get; set; }
+		public string Description { get; set; }
+		public string? InstallCommand { get; set; }
+		public string? UninstallCommand { get; set; }
+		public List<ToolLink> StatusPanelLinks { get; set; } = new List<ToolLink>();
+		public HashSet<Guid> DependsOnToolIds { get; set; } = new HashSet<Guid>();
+		public bool SafeWhenBusy { get; set; }
+		public string Revision { get; set; }
+
+		public ToolInfo(Guid id, string name, string description, string revision)
 		{
 			Id = id;
 			Name = name;
 			Description = description;
-			ZipPath = zipPath;
-			ZipChange = zipChange;
-			ConfigPath = configPath;
-			ConfigChange = configChange;
-			DependsOnToolIds = dependsOnToolIds;
-			SafeWhenBusy = safeWhenBusy;
+			Revision = revision;
+		}
+
+		public ToolInfo Clone()
+		{
+			return (ToolInfo)MemberwiseClone();
 		}
 	}
 
@@ -69,9 +69,16 @@ namespace UnrealGameSync
 		Task? _workerTask;
 		readonly AsyncEvent _wakeEvent;
 		readonly ILogger _logger;
-		public List<ToolDefinition> Tools { get; private set; } = new List<ToolDefinition>();
-		readonly int _lastChange = -1;
 		readonly IAsyncDisposer _asyncDisposer;
+		readonly FileReference _enabledToolsFile;
+
+		bool _readLegacyConfig;
+
+		Dictionary<string, ToolInfo> _perforceTools = new Dictionary<string, ToolInfo>(StringComparer.Ordinal);
+		int _perforceToolsChange = -1;
+
+		IReadOnlyList<ToolInfo> _tools = Array.Empty<ToolInfo>();
+		IReadOnlyList<ToolInfo> _enabledTools = Array.Empty<ToolInfo>();
 
 		IPerforceSettings PerforceSettings { get; }
 		DirectoryReference ToolsDir { get; }
@@ -90,6 +97,23 @@ namespace UnrealGameSync
 			_asyncDisposer = serviceProvider.GetRequiredService<IAsyncDisposer>();
 
 			DirectoryReference.CreateDirectory(ToolsDir);
+			_enabledToolsFile = FileReference.Combine(ToolsDir, "tools.json");
+
+			if (FileReference.Exists(_enabledToolsFile))
+			{
+				try
+				{
+					Load();
+				}
+				catch (Exception ex)
+				{
+					_logger.LogError(ex, "Unable to read previous tools list: {Message}", ex.Message);
+				}
+			}
+			else
+			{
+				_readLegacyConfig = true;
+			}
 
 			_wakeEvent = new AsyncEvent();
 		}
@@ -114,6 +138,12 @@ namespace UnrealGameSync
 			}
 		}
 
+		public IReadOnlyList<ToolInfo> GetTools()
+			=> _tools;
+
+		public IReadOnlyList<ToolInfo> GetEnabledTools()
+			=> _enabledTools;
+
 		DirectoryReference GetToolPathInternal(string toolName)
 		{
 			return DirectoryReference.Combine(ToolsDir, toolName, "Current");
@@ -121,7 +151,7 @@ namespace UnrealGameSync
 
 		public string? GetToolName(Guid toolId)
 		{
-			foreach (ToolDefinition tool in Tools)
+			foreach (ToolInfo tool in _tools)
 			{
 				if(tool.Id == toolId)
 				{
@@ -133,7 +163,8 @@ namespace UnrealGameSync
 
 		public DirectoryReference? GetToolPath(string toolName)
 		{
-			if (GetToolChange(toolName) != 0)
+			IReadOnlyList<ToolInfo> enabledTools = _enabledTools;
+			if (enabledTools.Any(x => String.Equals(x.Name, toolName, StringComparison.OrdinalIgnoreCase)))
 			{
 				return GetToolPathInternal(toolName);
 			}
@@ -156,7 +187,7 @@ namespace UnrealGameSync
 
 				try
 				{
-					await PollForUpdatesOnce(_logger, cancellationToken);
+					await PollForUpdatesOnce(cancellationToken);
 				}
 				catch(Exception ex)
 				{
@@ -168,119 +199,136 @@ namespace UnrealGameSync
 			}
 		}
 
-		async Task PollForUpdatesOnce(ILogger logger, CancellationToken cancellationToken)
+		async Task PollForUpdatesOnce(CancellationToken cancellationToken)
 		{
-			using IPerforceConnection perforce = await PerforceConnection.CreateAsync(PerforceSettings, logger);
+			using IPerforceConnection perforce = await PerforceConnection.CreateAsync(PerforceSettings, _logger);
 
-			List<ChangesRecord> changes = await perforce.GetChangesAsync(ChangesOptions.None, 1, ChangeStatus.Submitted, $"{DeploymentSettings.Instance.ToolsDepotPath}/...", cancellationToken);
-			if (changes.Count == 0 || changes[0].Number == _lastChange)
+			// Update all the available tools
+			List<ToolInfo> tools = new List<ToolInfo>();
+			await ReadPerforceTools(perforce, tools, cancellationToken);
+			_tools = tools;
+
+			// When upgrading from older UGS versions, read the legacy sync CL from plain-text config files
+			if (_readLegacyConfig)
 			{
-				return;
+				await ReadLegacyConfigAsync(tools, cancellationToken);
+				_readLegacyConfig = false;
 			}
 
-			List<FStatRecord> fileRecords = await perforce.FStatAsync($"{DeploymentSettings.Instance.ToolsDepotPath}/...", cancellationToken).ToListAsync(cancellationToken);
+			// Find all the tools which are enabled, including those enabled due to dependencies from other tools
+			HashSet<Guid> enabledToolIds = new HashSet<Guid>();
+			FindEnabledTools(Settings.EnabledTools, tools, enabledToolIds);
 
-			// Update the tools list
-			List<ToolDefinition> newTools = new List<ToolDefinition>();
-			foreach (FStatRecord fileRecord in fileRecords)
+			// Install or update any new tools
+			bool hasChanged = false;
+			foreach (ToolInfo toolInfo in _tools)
 			{
-				if (fileRecord.DepotFile != null && fileRecord.DepotFile.EndsWith(".ini", StringComparison.OrdinalIgnoreCase))
+				if (enabledToolIds.Contains(toolInfo.Id))
 				{
-					ToolDefinition? tool = Tools.FirstOrDefault(x => x.ConfigPath.Equals(fileRecord.DepotFile, StringComparison.Ordinal));
-					if (tool == null || tool.ConfigChange != fileRecord.HeadChange)
+					ToolInfo? existingTool = _enabledTools.FirstOrDefault(x => x.Id == toolInfo.Id);
+					if (existingTool == null || !String.Equals(existingTool.Revision, toolInfo.Revision, StringComparison.OrdinalIgnoreCase))
 					{
-						tool = await ReadToolDefinitionAsync(perforce, fileRecord.DepotFile, fileRecord.HeadChange, cancellationToken);
-					}
-					if (tool != null)
-					{
-						newTools.Add(tool);
+						await UpdateToolAsync(perforce, toolInfo, cancellationToken);
+						hasChanged = true;
 					}
 				}
 			}
-			Tools = newTools;
 
-			// Find all the tools which are enabled due to dependencies
-			HashSet<Guid> enabledTools = new HashSet<Guid>();
-			AddEnabledTools(Settings.EnabledTools, Tools, enabledTools);
-
-			// Configure each tool
-			foreach (ToolDefinition tool in Tools)
+			// Remove any tools which we no longer need
+			for (int idx = _enabledTools.Count - 1; idx >= 0; idx--)
 			{
-				tool.Enabled = enabledTools.Contains(tool.Id);
-
-				if(!tool.Enabled)
+				ToolInfo tool = _enabledTools[idx];
+				if (!enabledToolIds.Contains(tool.Id))
 				{
-					continue;
-				}
-
-				List<FStatRecord> toolFileRecords = fileRecords.Where(x => String.Equals(x.DepotFile, tool.ZipPath, StringComparison.OrdinalIgnoreCase)).ToList();
-				if (toolFileRecords.Count == 0)
-				{
-					continue;
-				}
-
-				int headChange = toolFileRecords.Max(x => x.HeadChange);
-				if (headChange == GetToolChange(tool.Name))
-				{
-					continue;
+					await RemoveToolAsync(tool, cancellationToken);
+					hasChanged = true;
 				}
+			}
 
-				List<FStatRecord> syncFileRecords = toolFileRecords.Where(x => x.Action != FileAction.Delete && x.Action != FileAction.MoveDelete).ToList();
-				try
-				{
-					await UpdateToolAsync(perforce, tool.Name, headChange, syncFileRecords, tool.InstallAction, logger, cancellationToken);
-				}
-				catch (Exception ex)
-				{
-					logger.LogError(ex, "Exception while updating tool");
-				}
+			// Notify the main window if anything changed
+			if (hasChanged)
+			{
+				_synchronizationContext.Post(_ => OnChange?.Invoke(), null);
 			}
+		}
 
-			foreach (ToolDefinition tool in Tools)
+		static void FindEnabledTools(HashSet<Guid> inputToolIds, List<ToolInfo> tools, HashSet<Guid> enabledToolIds)
+		{
+			if (inputToolIds.Count > 0)
 			{
-				if (!tool.Enabled && GetToolChange(tool.Name) != 0)
+				foreach (ToolInfo tool in tools)
 				{
-					try
-					{
-						await RemoveToolAsync(tool.Name, tool.UninstallAction, logger, cancellationToken);
-					}
-					catch (Exception ex)
+					if (inputToolIds.Contains(tool.Id) && enabledToolIds.Add(tool.Id))
 					{
-						logger.LogError(ex, "Exception while removing tool");
+						FindEnabledTools(tool.DependsOnToolIds, tools, enabledToolIds);
 					}
 				}
 			}
-
-			_synchronizationContext.Post(_ => OnChange?.Invoke(), null);
 		}
 
-		static void AddEnabledTools(HashSet<Guid> inputToolIds, List<ToolDefinition> tools, HashSet<Guid> enabledTools)
+		async Task ReadPerforceTools(IPerforceConnection perforce, List<ToolInfo> tools, CancellationToken cancellationToken)
 		{
-			if (inputToolIds.Count > 0)
+			if (String.IsNullOrEmpty(DeploymentSettings.Instance.ToolsDepotPath))
 			{
-				foreach (ToolDefinition tool in tools)
+				return;
+			}
+
+			List<ChangesRecord> changes = await perforce.GetChangesAsync(ChangesOptions.None, 1, ChangeStatus.Submitted, $"{DeploymentSettings.Instance.ToolsDepotPath}/...", cancellationToken);
+			if (changes.Count > 0 && changes[0].Number != _perforceToolsChange)
+			{
+				Dictionary<string, ToolInfo> newPerforceTools = new Dictionary<string, ToolInfo>(StringComparer.Ordinal);
+
+				List<FStatRecord> fileRecords = await perforce.FStatAsync($"{DeploymentSettings.Instance.ToolsDepotPath}/...", cancellationToken).ToListAsync(cancellationToken);
+				foreach (FStatRecord fileRecord in fileRecords)
 				{
-					if (inputToolIds.Contains(tool.Id) && enabledTools.Add(tool.Id))
+					if (fileRecord.DepotFile != null && fileRecord.DepotFile.EndsWith(".ini", StringComparison.OrdinalIgnoreCase))
 					{
-						AddEnabledTools(tool.DependsOnToolIds, tools, enabledTools);
+						string zipFile = fileRecord.DepotFile.Substring(0, fileRecord.DepotFile.Length - 4) + ".zip";
+						FStatRecord? zipRecord = fileRecords.FirstOrDefault(x => String.Equals(x.DepotFile, zipFile, StringComparison.OrdinalIgnoreCase));
+
+						if (zipRecord != null)
+						{
+							string toolRevision = $"{zipFile}@{Math.Max(fileRecord.HeadChange, zipRecord.HeadChange)}";
+
+							ToolInfo? toolInfo;
+							if (!_perforceTools.TryGetValue(toolRevision, out toolInfo))
+							{
+								toolInfo = await ReadToolDefinitionAsync(perforce, $"{fileRecord.DepotFile}@{fileRecord.HeadChange}", toolRevision, cancellationToken);
+							}
+							if (toolInfo != null)
+							{
+								newPerforceTools.Add(toolRevision, toolInfo);
+							}
+						}
 					}
 				}
+
+				_perforceTools = newPerforceTools;
+				_perforceToolsChange = changes[0].Number;
 			}
+
+			tools.AddRange(_perforceTools.Values);
 		}
 
-		async Task<ToolDefinition?> ReadToolDefinitionAsync(IPerforceConnection perforce, string depotPath, int change, CancellationToken cancellationToken)
+		async Task<ToolInfo?> ReadToolDefinitionAsync(IPerforceConnection perforce, string iniRevision, string toolRevision, CancellationToken cancellationToken)
 		{
-			PerforceResponse<PrintRecord<string[]>> response = await perforce.TryPrintLinesAsync($"{depotPath}@{change}", cancellationToken);
+			PerforceResponse<PrintRecord<string[]>> response = await perforce.TryPrintLinesAsync(iniRevision, cancellationToken);
 			if (!response.Succeeded || response.Data.Contents == null)
 			{
 				return null;
 			}
 
-			int nameIdx = depotPath.LastIndexOf('/') + 1;
-			int extensionIdx = depotPath.LastIndexOf('.');
+			int nameIdx = iniRevision.LastIndexOf('/') + 1;
+			int extensionIdx = iniRevision.LastIndexOf('.');
+			string defaultToolName = iniRevision.Substring(nameIdx, extensionIdx - nameIdx);
 
+			return ParseToolInfo(response.Data.Contents, defaultToolName, toolRevision);
+		}
+
+		static ToolInfo? ParseToolInfo(string[] contents, string defaultToolName, string revision)
+		{
 			ConfigFile configFile = new ConfigFile();
-			configFile.Parse(response.Data.Contents);
+			configFile.Parse(contents);
 
 			string? id = configFile.GetValue("Settings.Id", null);
 			if (id == null || !Guid.TryParse(id, out Guid toolId))
@@ -288,45 +336,22 @@ namespace UnrealGameSync
 				return null;
 			}
 
-			string toolName = configFile.GetValue("Settings.Name", depotPath.Substring(nameIdx, extensionIdx - nameIdx));
+			string toolName = configFile.GetValue("Settings.Name", defaultToolName);
 			string toolDescription = configFile.GetValue("Settings.Description", toolName);
-			string toolZipPath = depotPath.Substring(0, extensionIdx) + ".zip";
-			int toolZipChange = GetToolChange(toolName);
-			string toolConfigPath = depotPath;
-			int toolConfigChange = change;
-			bool toolSafeWhenBusy = configFile.GetValue("Settings.SafeWhenBusy", false);
 
-			HashSet<Guid> dependsOnToolIds = new HashSet<Guid>();
+			ToolInfo tool = new ToolInfo(toolId, toolName, toolDescription, revision);
+			tool.InstallCommand = configFile.GetValue("Settings.InstallCommand", null);
+			tool.UninstallCommand = configFile.GetValue("Settings.UninstallCommand", null);
+			tool.SafeWhenBusy = configFile.GetValue("Settings.SafeWhenBusy", false);
+
 			foreach (string line in configFile.GetValues("Settings.DependsOnTool", Array.Empty<string>()))
 			{
 				if (Guid.TryParse(line.Trim(), out Guid requiredToolId))
 				{
-					dependsOnToolIds.Add(requiredToolId);
+					tool.DependsOnToolIds.Add(requiredToolId);
 				}
 			}
 
-			ToolDefinition tool = new ToolDefinition(toolId, toolName, toolDescription, toolZipPath, toolZipChange, toolConfigPath, toolConfigChange, dependsOnToolIds, toolSafeWhenBusy);
-
-			string? installCommand = configFile.GetValue("Settings.InstallCommand", null);
-			if (!String.IsNullOrEmpty(installCommand))
-			{
-				tool.InstallAction = async (logger, cancellationToken) =>
-				{
-					logger.LogInformation("Running install action: {Command}", installCommand);
-					await RunCommandAsync(tool.Name, installCommand, logger, cancellationToken);
-				};
-			}
-
-			string? uninstallCommand = configFile.GetValue("Settings.UninstallCommand", null);
-			if (!String.IsNullOrEmpty(uninstallCommand))
-			{
-				tool.UninstallAction = async (logger, cancellationToken) =>
-				{
-					logger.LogInformation("Running unininstall action: {Command}", uninstallCommand);
-					await RunCommandAsync(tool.Name, uninstallCommand, logger, cancellationToken);
-				};
-			}
-
 			string[] statusPanelLinks = configFile.GetValues("Settings.StatusPanelLinks", Array.Empty<string>());
 			foreach (string statusPanelLink in statusPanelLinks)
 			{
@@ -347,7 +372,7 @@ namespace UnrealGameSync
 			return tool;
 		}
 
-		async Task RunCommandAsync(string toolName, string command, ILogger logger, CancellationToken cancellationToken)
+		async Task RunCommandAsync(string toolName, string command, CancellationToken cancellationToken)
 		{
 			DirectoryReference toolPath = GetToolPathInternal(toolName);
 
@@ -361,30 +386,37 @@ namespace UnrealGameSync
 				commandArgs = command.Substring(spaceIdx + 1);
 			}
 
-			int exitCode = await Utility.ExecuteProcessAsync(FileReference.Combine(toolPath, commandExe).FullName, toolPath.FullName, commandArgs, line => logger.LogInformation("{ToolName}> {Line}", toolName, line), cancellationToken);
-			logger.LogInformation("{ToolName}> Exit code {ExitCode})", toolName, exitCode);
+			int exitCode = await Utility.ExecuteProcessAsync(FileReference.Combine(toolPath, commandExe).FullName, toolPath.FullName, commandArgs, line => _logger.LogInformation("{ToolName}> {Line}", toolName, line), cancellationToken);
+			_logger.LogInformation("{ToolName}> Exit code {ExitCode})", toolName, exitCode);
 		}
 
-		async Task RemoveToolAsync(string toolName, Func<ILogger, CancellationToken, Task>? uninstallAction, ILogger logger, CancellationToken cancellationToken)
+		async Task RemoveToolAsync(ToolInfo tool, CancellationToken cancellationToken)
 		{
-			logger.LogInformation("Removing {ToolName}", toolName);
-			DirectoryReference? toolPath = GetToolPath(toolName);
-
-			if (uninstallAction != null)
+			try
 			{
-				logger.LogInformation("Running uninstall...");
-				await uninstallAction(logger, cancellationToken);
-			}
+				_logger.LogInformation("Removing {ToolName}", tool.Name);
+				DirectoryReference? toolPath = GetToolPath(tool.Name);
 
-			SetToolChange(toolName, null);
+				if (!String.IsNullOrEmpty(tool.UninstallCommand))
+				{
+					_logger.LogInformation("Running unininstall action: {Command}", tool.UninstallCommand);
+					await RunCommandAsync(tool.Name, tool.UninstallCommand, cancellationToken);
+				}
 
-			if (toolPath != null)
+				await SetToolRevisionAsync(tool.Name, null, cancellationToken);
+
+				if (toolPath != null)
+				{
+					_logger.LogInformation("Removing {ToolPath}", toolPath);
+					TryDeleteDirectory(toolPath);
+				}
+
+				_logger.LogInformation("{ToolName} has been removed successfully", tool.Name);
+			}
+			catch (Exception ex)
 			{
-				logger.LogInformation("Removing {ToolPath}", toolPath);
-				TryDeleteDirectory(toolPath, logger);
+				_logger.LogError(ex, "Exception while removing tool {ToolId}: {Message}", tool.Id, ex.Message);
 			}
-
-			logger.LogInformation("{ToolName} has been removed successfully", toolName);
 		}
 
 		static void ForceDeleteDirectory(DirectoryReference directoryName)
@@ -400,7 +432,7 @@ namespace UnrealGameSync
 			}
 		}
 
-		static bool TryDeleteDirectory(DirectoryReference directoryName, ILogger logger)
+		bool TryDeleteDirectory(DirectoryReference directoryName)
 		{
 			try
 			{
@@ -409,62 +441,76 @@ namespace UnrealGameSync
 			}
 			catch(Exception ex)
 			{
-				logger.LogWarning(ex, "Unable to delete directory {DirectoryName}", directoryName);
+				_logger.LogWarning(ex, "Unable to delete directory {DirectoryName}", directoryName);
 				return false;
 			}
 		}
 
-		async Task<bool> UpdateToolAsync(IPerforceConnection perforce, string toolName, int change, List<FStatRecord> records, Func<ILogger, CancellationToken, Task>? installAction, ILogger logger, CancellationToken cancellationToken)
+		async Task<bool> UpdateToolAsync(IPerforceConnection perforce, ToolInfo tool, CancellationToken cancellationToken)
 		{
-			DirectoryReference toolDir = DirectoryReference.Combine(ToolsDir, toolName);
-			DirectoryReference.CreateDirectory(toolDir);
-
-			foreach (DirectoryReference existingDir in DirectoryReference.EnumerateDirectories(toolDir, "Prev-*"))
+			try
 			{
-				TryDeleteDirectory(existingDir, logger);
-			}
+				DirectoryReference toolDir = DirectoryReference.Combine(ToolsDir, tool.Name);
+				DirectoryReference.CreateDirectory(toolDir);
 
-			DirectoryReference nextToolDir = DirectoryReference.Combine(toolDir, "Next");
-			ForceDeleteDirectory(nextToolDir);
-			DirectoryReference.CreateDirectory(nextToolDir);
+				foreach (DirectoryReference existingDir in DirectoryReference.EnumerateDirectories(toolDir, "Prev-*"))
+				{
+					TryDeleteDirectory(existingDir);
+				}
 
-			DirectoryReference nextToolZipsDir = DirectoryReference.Combine(nextToolDir, ".zips");
-			DirectoryReference.CreateDirectory(nextToolZipsDir);
+				DirectoryReference nextToolDir = DirectoryReference.Combine(toolDir, "Next");
+				ForceDeleteDirectory(nextToolDir);
+				DirectoryReference.CreateDirectory(nextToolDir);
 
-			for (int idx = 0; idx < records.Count; idx++)
-			{
-				FileReference zipFile = FileReference.Combine(nextToolZipsDir, String.Format("{0}.{1}.zip", toolName, idx));
+				DirectoryReference nextToolZipsDir = DirectoryReference.Combine(nextToolDir, ".zips");
+				DirectoryReference.CreateDirectory(nextToolZipsDir);
 
-				PerforceResponseList<PrintRecord> response = await perforce.TryPrintAsync(zipFile.FullName, $"{records[idx].DepotFile}#{records[idx].HeadRevision}", cancellationToken);
-				if(!response.Succeeded || !FileReference.Exists(zipFile))
+				if (tool.Revision.StartsWith("//", StringComparison.Ordinal))
 				{
-					logger.LogError("Unable to print {DepotFile}", records[idx].DepotFile);
+					// Read it from Perforce
+					FileReference zipFile = FileReference.Combine(nextToolZipsDir, $"{tool.Name}.zip");
+
+					PerforceResponseList<PrintRecord> response = await perforce.TryPrintAsync(zipFile.FullName, tool.Revision, cancellationToken);
+					if (!response.Succeeded || !FileReference.Exists(zipFile))
+					{
+						_logger.LogError("Unable to print {DepotFile}", tool.Revision);
+						return false;
+					}
+
+					ArchiveUtils.ExtractFiles(zipFile, nextToolDir, null, new ProgressValue(), _logger);
+				}
+				else
+				{
+					// Unknown format?
 					return false;
 				}
 
-				ArchiveUtils.ExtractFiles(zipFile, nextToolDir, null, new ProgressValue(), logger);
-			}
+				DirectoryReference currentToolDir = DirectoryReference.Combine(toolDir, "Current");
+				if (DirectoryReference.Exists(currentToolDir))
+				{
+					DirectoryReference prevDirectoryName = DirectoryReference.Combine(toolDir, String.Format("Prev-{0:X16}", Stopwatch.GetTimestamp()));
+					Directory.Move(currentToolDir.FullName, prevDirectoryName.FullName);
+					await SetToolRevisionAsync(tool.Name, null, cancellationToken);
+					TryDeleteDirectory(prevDirectoryName);
+				}
 
-			DirectoryReference currentToolDir = DirectoryReference.Combine(toolDir, "Current");
-			if (DirectoryReference.Exists(currentToolDir))
-			{
-				DirectoryReference prevDirectoryName = DirectoryReference.Combine(toolDir, String.Format("Prev-{0:X16}", Stopwatch.GetTimestamp()));
-				Directory.Move(currentToolDir.FullName, prevDirectoryName.FullName);
-				SetToolChange(toolName, null);
-				TryDeleteDirectory(prevDirectoryName, logger);
-			}
+				Directory.Move(nextToolDir.FullName, currentToolDir.FullName);
 
-			Directory.Move(nextToolDir.FullName, currentToolDir.FullName);
+				if (!String.IsNullOrEmpty(tool.InstallCommand))
+				{
+					_logger.LogInformation("Running install action: {Command}", tool.InstallCommand);
+					await RunCommandAsync(tool.Name, tool.InstallCommand, cancellationToken);
+				}
 
-			if (installAction != null)
+				await SetToolRevisionAsync(tool.Name, tool, cancellationToken);
+				_logger.LogInformation("Updated {ToolName} to {Revision}", tool.Name, tool.Revision);
+				return true;
+			}
+			catch (Exception ex)
 			{
-				logger.LogInformation("Running installer...");
-				await installAction.Invoke(logger, cancellationToken);
+				_logger.LogError(ex, "Exception while updating tool: {Message}", ex.Message);
+				return false;
 			}
-
-			SetToolChange(toolName, change);
-			logger.LogInformation("Updated {ToolName} to change {Change}", toolName, change);
-			return true;
 		}
 
 		FileReference GetConfigFilePath(string toolName)
@@ -472,6 +518,43 @@ namespace UnrealGameSync
 			return FileReference.Combine(ToolsDir, toolName, toolName + ".ini");
 		}
 
+		async Task ReadLegacyConfigAsync(IReadOnlyList<ToolInfo> tools, CancellationToken cancellationToken)
+		{
+			try
+			{
+				List<FileReference> deleteFiles = new List<FileReference>();
+
+				List<ToolInfo> enabledTools = new List<ToolInfo>();
+				foreach (ToolInfo tool in tools)
+				{
+					int change = GetToolChange(tool.Name);
+					if (change != 0)
+					{
+						int atIdx = tool.Revision.LastIndexOf('@');
+						if (atIdx != -1)
+						{
+							ToolInfo clone = tool.Clone();
+							clone.Revision = tool.Revision.Substring(0, atIdx + 1) + change.ToString();
+							enabledTools.Add(clone);
+
+							deleteFiles.Add(GetConfigFilePath(tool.Name));
+						}
+					}
+				}
+
+				await SaveAsync(enabledTools, cancellationToken);
+
+				foreach (FileReference deleteFile in deleteFiles)
+				{
+					FileReference.Delete(deleteFile);
+				}
+			}
+			catch (Exception ex)
+			{
+				_logger.LogError(ex, "Unable to upgrade tools from legacy config format: {Message}", ex.Message);
+			}
+		}
+
 		int GetToolChange(string toolName)
 		{
 			try
@@ -490,19 +573,61 @@ namespace UnrealGameSync
 			return 0;
 		}
 
-		void SetToolChange(string toolName, int? change)
+		async Task SetToolRevisionAsync(string toolName, ToolInfo? toolInfo, CancellationToken cancellationToken)
 		{
-			FileReference configFilePath = GetConfigFilePath(toolName);
-			if (change.HasValue)
+			List<ToolInfo> enabledTools = new List<ToolInfo>(_enabledTools);
+
+			int toolIndex = enabledTools.FindIndex(x => String.Equals(x.Name, toolName, StringComparison.OrdinalIgnoreCase));
+			if (toolIndex != -1)
 			{
-				ConfigFile configFile = new ConfigFile();
-				configFile.SetValue("Settings.Change", change.Value);
-				configFile.Save(configFilePath);
+				if (toolInfo == null)
+				{
+					enabledTools.RemoveAt(toolIndex);
+				}
 			}
 			else
 			{
-				FileReference.Delete(configFilePath);
+				if (toolInfo != null)
+				{
+					enabledTools.Add(toolInfo);
+				}
 			}
+
+			await SaveAsync(enabledTools, cancellationToken);
+			_enabledTools = enabledTools;
+		}
+
+		class State
+		{
+			public List<ToolInfo> Tools { get; set; } = new List<ToolInfo>();
+		}
+
+		void Load()
+		{
+			if (FileReference.Exists(_enabledToolsFile))
+			{
+				byte[]? data = FileTransaction.ReadAllBytes(_enabledToolsFile);
+				if (data != null)
+				{
+					State? state = JsonSerializer.Deserialize<State>(data, new JsonSerializerOptions { PropertyNameCaseInsensitive = true });
+					if (state != null)
+					{
+						_enabledTools = state.Tools;
+					}
+				}
+			}
+		}
+
+		async Task SaveAsync(List<ToolInfo> enabledTools, CancellationToken cancellationToken)
+		{
+			JsonSerializerOptions options = new JsonSerializerOptions();
+			options.DefaultIgnoreCondition = System.Text.Json.Serialization.JsonIgnoreCondition.WhenWritingNull;
+			options.PropertyNamingPolicy = JsonNamingPolicy.CamelCase;
+			options.WriteIndented = true;
+
+			State state = new State { Tools = enabledTools };
+			byte[] data = JsonSerializer.SerializeToUtf8Bytes(state, options);
+			await FileTransaction.WriteAllBytesAsync(_enabledToolsFile, data, cancellationToken);
 		}
 	}
 }
