UGS: Support for enumerating/downloading tools from Horde.

diff --git a/Engine/Source/Programs/UnrealGameSync/UnrealGameSync/Controls/WorkspaceControl.cs b/Engine/Source/Programs/UnrealGameSync/UnrealGameSync/Controls/WorkspaceControl.cs
--- a/Engine/Source/Programs/UnrealGameSync/UnrealGameSync/Controls/WorkspaceControl.cs
+++ b/Engine/Source/Programs/UnrealGameSync/UnrealGameSync/Controls/WorkspaceControl.cs
@@ -3446,9 +3446,9 @@ namespace UnrealGameSync
 				IReadOnlyList<ToolInfo> tools = _owner.ToolUpdateMonitor.GetEnabledTools();
 				foreach (ToolInfo tool in tools)
 				{
-					if (tool.SafeWhenBusy)
+					if (tool.Settings.SafeWhenBusy)
 					{
-						foreach (ToolLink link in tool.StatusPanelLinks)
+						foreach (ToolLink link in tool.Settings.StatusPanelLinks)
 						{
 							programsLine.AddLink(link.Label, FontStyle.Regular, () => RunTool(tool, link));
 							programsLine.AddText("  |  ");
@@ -3638,7 +3638,7 @@ namespace UnrealGameSync
 				IReadOnlyList<ToolInfo> tools = _owner.ToolUpdateMonitor.GetEnabledTools();
 				foreach (ToolInfo tool in tools)
 				{
-					foreach (ToolLink link in tool.StatusPanelLinks)
+					foreach (ToolLink link in tool.Settings.StatusPanelLinks)
 					{
 						programsLine.AddLink(link.Label, FontStyle.Regular, () => RunTool(tool, link));
 						programsLine.AddText("  |  ");
diff --git a/Engine/Source/Programs/UnrealGameSync/UnrealGameSync/Forms/ApplicationSettingsWindow.cs b/Engine/Source/Programs/UnrealGameSync/UnrealGameSync/Forms/ApplicationSettingsWindow.cs
--- a/Engine/Source/Programs/UnrealGameSync/UnrealGameSync/Forms/ApplicationSettingsWindow.cs
+++ b/Engine/Source/Programs/UnrealGameSync/UnrealGameSync/Forms/ApplicationSettingsWindow.cs
@@ -37,7 +37,6 @@ namespace UnrealGameSync
 		}
 
 		readonly string _originalExecutableFileName;
-		readonly IPerforceSettings _defaultPerforceSettings;
 		readonly UserSettings _settings;
 		readonly ILogger _logger;
 
@@ -88,13 +87,12 @@ namespace UnrealGameSync
 			}
 		}
 
-		private ApplicationSettingsWindow(IPerforceSettings defaultPerforceSettings, bool preview, string originalExecutableFileName, UserSettings settings, ToolUpdateMonitor toolUpdateMonitor, ILogger<ApplicationSettingsWindow> logger)
+		private ApplicationSettingsWindow(IPerforceSettings defaultPerforceSettings, string originalExecutableFileName, UserSettings settings, ToolUpdateMonitor toolUpdateMonitor, ILogger<ApplicationSettingsWindow> logger)
 		{
 			InitializeComponent();
 			Font = new System.Drawing.Font("Segoe UI", 8.25F, System.Drawing.FontStyle.Regular, System.Drawing.GraphicsUnit.Point, ((byte)(0)));
 
 			_originalExecutableFileName = originalExecutableFileName;
-			_defaultPerforceSettings = defaultPerforceSettings;
 			_settings = settings;
 			_toolUpdateMonitor = toolUpdateMonitor;
 			_logger = logger;
@@ -191,9 +189,9 @@ namespace UnrealGameSync
 			}
 		}
 
-		public static Result ShowModal(IWin32Window owner, IPerforceSettings defaultPerforceSettings, bool preview, string originalExecutableFileName, UserSettings settings, ToolUpdateMonitor toolUpdateMonitor, ILogger<ApplicationSettingsWindow> logger)
+		public static Result ShowModal(IWin32Window owner, IPerforceSettings defaultPerforceSettings, string originalExecutableFileName, UserSettings settings, ToolUpdateMonitor toolUpdateMonitor, ILogger<ApplicationSettingsWindow> logger)
 		{
-			using ApplicationSettingsWindow applicationSettings = new ApplicationSettingsWindow(defaultPerforceSettings, preview, originalExecutableFileName, settings, toolUpdateMonitor, logger);
+			using ApplicationSettingsWindow applicationSettings = new ApplicationSettingsWindow(defaultPerforceSettings, originalExecutableFileName, settings, toolUpdateMonitor, logger);
 			applicationSettings.ShowDialog(owner);
 			return applicationSettings._result;
 		}
diff --git a/Engine/Source/Programs/UnrealGameSync/UnrealGameSync/Forms/MainWindow.cs b/Engine/Source/Programs/UnrealGameSync/UnrealGameSync/Forms/MainWindow.cs
--- a/Engine/Source/Programs/UnrealGameSync/UnrealGameSync/Forms/MainWindow.cs
+++ b/Engine/Source/Programs/UnrealGameSync/UnrealGameSync/Forms/MainWindow.cs
@@ -1526,7 +1526,7 @@ namespace UnrealGameSync
 
 		public void ModifyApplicationSettings()
 		{
-			ApplicationSettingsWindow.Result result = ApplicationSettingsWindow.ShowModal(this, _defaultPerforceSettings, _preview, _originalExecutableFileName, _settings, ToolUpdateMonitor, _serviceProvider.GetRequiredService<ILogger<ApplicationSettingsWindow>>());
+			ApplicationSettingsWindow.Result result = ApplicationSettingsWindow.ShowModal(this, _defaultPerforceSettings, _originalExecutableFileName, _settings, ToolUpdateMonitor, _serviceProvider.GetRequiredService<ILogger<ApplicationSettingsWindow>>());
 			if (result == ApplicationSettingsWindow.Result.Restart)
 			{
 				_updateMonitor.TriggerUpdate(UpdateType.UserInitiated, false);
diff --git a/Engine/Source/Programs/UnrealGameSync/UnrealGameSync/Forms/UpdateSettingsWindow.cs b/Engine/Source/Programs/UnrealGameSync/UnrealGameSync/Forms/UpdateSettingsWindow.cs
--- a/Engine/Source/Programs/UnrealGameSync/UnrealGameSync/Forms/UpdateSettingsWindow.cs
+++ b/Engine/Source/Programs/UnrealGameSync/UnrealGameSync/Forms/UpdateSettingsWindow.cs
@@ -8,7 +8,6 @@ using System.Runtime.InteropServices;
 using System.Threading;
 using System.Threading.Tasks;
 using System.Windows.Forms;
-using UnrealGameSync;
 
 namespace UnrealGameSync
 {
@@ -21,7 +20,7 @@ namespace UnrealGameSync
 
 		const int EmSetcuebanner = 0x1501;
 
-		LauncherSettings _settings;
+		readonly LauncherSettings _settings;
 		string? _logText;
 		readonly SyncAndRunDelegate _syncAndRun;
 
diff --git a/Engine/Source/Programs/UnrealGameSync/UnrealGameSync/Launcher.cs b/Engine/Source/Programs/UnrealGameSync/UnrealGameSync/Launcher.cs
--- a/Engine/Source/Programs/UnrealGameSync/UnrealGameSync/Launcher.cs
+++ b/Engine/Source/Programs/UnrealGameSync/UnrealGameSync/Launcher.cs
@@ -2,9 +2,6 @@
 
 using EpicGames.Core;
 using EpicGames.Horde;
-using EpicGames.Horde.Storage;
-using EpicGames.Horde.Storage.Clients;
-using EpicGames.Horde.Storage.Nodes;
 using EpicGames.Horde.Tools;
 using EpicGames.Perforce;
 using Microsoft.Extensions.DependencyInjection;
@@ -16,14 +13,10 @@ using System.Diagnostics.CodeAnalysis;
 using System.IO;
 using System.IO.Compression;
 using System.Linq;
-using System.Net.Http;
 using System.Reflection;
-using System.Reflection.Metadata;
-using System.Text;
 using System.Threading;
 using System.Threading.Tasks;
 using System.Windows.Forms;
-using UnrealGameSync;
 
 namespace UnrealGameSync
 {
@@ -35,14 +28,6 @@ namespace UnrealGameSync
 			bool preview = args.Contains("-unstable", StringComparer.InvariantCultureIgnoreCase) || args.Contains("-preview", StringComparer.InvariantCultureIgnoreCase);
 			bool openSettings = args.Contains("-settings", StringComparer.OrdinalIgnoreCase);
 
-#if DEBUG
-			// For debug builds, just launch into the regular UGS application.
-			if (!openSettings)
-			{
-				return false;
-			}
-#endif
-
 			// Read the settings
 			LauncherSettings launcherSettings = new LauncherSettings();
 			launcherSettings.PerforceDepotPath = DeploymentSettings.Instance.DefaultDepotPath;
@@ -328,18 +313,6 @@ namespace UnrealGameSync
 			}
 		}
 
-		static string QuoteArgument(string arg)
-		{
-			if(arg.Contains(' ', StringComparison.Ordinal) && !arg.StartsWith("\"", StringComparison.Ordinal))
-			{
-				return String.Format("\"{0}\"", arg);
-			}
-			else
-			{
-				return arg;
-			}
-		}
-
 		static bool TryReadAllText(string fileName, [NotNullWhen(true)] out string? text)
 		{
 			try
diff --git a/Engine/Source/Programs/UnrealGameSync/UnrealGameSync/Program.cs b/Engine/Source/Programs/UnrealGameSync/UnrealGameSync/Program.cs
--- a/Engine/Source/Programs/UnrealGameSync/UnrealGameSync/Program.cs
+++ b/Engine/Source/Programs/UnrealGameSync/UnrealGameSync/Program.cs
@@ -96,8 +96,14 @@ namespace UnrealGameSync
 
 				using (EventWaitHandle activateEvent = new EventWaitHandle(false, EventResetMode.AutoReset, "ActivateUnrealGameSync"))
 				{
+#if DEBUG
+					bool runUpdateCheck = args.Contains("-UpdateCheck", StringComparer.OrdinalIgnoreCase);
+#else
+					bool runUpdateCheck = !args.Contains("-NoUpdateCheck", StringComparer.OrdinalIgnoreCase);
+#endif
+
 					// Check for a newer version of the application
-					if (!args.Contains("-NoUpdateCheck", StringComparer.OrdinalIgnoreCase) && Launcher.SyncAndRunLatest(instanceMutex, args))
+					if (runUpdateCheck && Launcher.SyncAndRunLatest(instanceMutex, args))
 					{
 						return;
 					}
@@ -111,7 +117,7 @@ namespace UnrealGameSync
 					// Launch the application proper
 					if (firstInstance)
 					{
-						InnerMainAsync(instanceMutex, activateEvent, args).GetAwaiter().GetResult();
+						InnerMainAsync(instanceMutex, activateEvent, args, runUpdateCheck).GetAwaiter().GetResult();
 					}
 					else
 					{
@@ -121,7 +127,7 @@ namespace UnrealGameSync
 			}
 		}
 
-		static async Task InnerMainAsync(Mutex instanceMutex, EventWaitHandle activateEvent, string[] args)
+		static async Task InnerMainAsync(Mutex instanceMutex, EventWaitHandle activateEvent, string[] args, bool runUpdateCheck)
 		{
 			LauncherSettings launcherSettings = new LauncherSettings();
 			launcherSettings.Read();
@@ -189,7 +195,11 @@ namespace UnrealGameSync
 				services.AddSingleton<IAsyncDisposer, AsyncDisposer>();
 				services.AddSingleton(sp => TokenStoreFactory.CreateTokenStore());
 				services.AddSingleton<OidcTokenManager>();
-				services.AddHordeHttpClient(x => x.BaseAddress = new Uri(launcherSettings.HordeServer ?? "http://localhost:5000"));
+
+				if (launcherSettings.HordeServer != null)
+				{
+					services.AddHordeHttpClient(x => x.BaseAddress = new Uri(launcherSettings.HordeServer));
+				}
 
 				await using (ServiceProvider serviceProvider = services.BuildServiceProvider())
 				{
@@ -230,7 +240,7 @@ namespace UnrealGameSync
 
 							ProtocolHandlerUtils.InstallQuiet(logger);
 
-							using (UpdateMonitor updateMonitor = CreateUpdateMonitor(launcherSettings, defaultSettings, updatePath, serviceProvider))
+							await using (UpdateMonitor updateMonitor = CreateUpdateMonitor(launcherSettings, defaultSettings, updatePath, runUpdateCheck, serviceProvider))
 							{
 								using ProgramApplicationContext context = new ProgramApplicationContext(defaultSettings, updateMonitor, DeploymentSettings.Instance.ApiUrl, dataFolder, activateEvent, restoreState, updateSpawn, projectFileName, preview, serviceProvider, uri);
 								Application.Run(context);
@@ -256,9 +266,13 @@ namespace UnrealGameSync
 			}
 		}
 
-		private static UpdateMonitor CreateUpdateMonitor(LauncherSettings launcherSettings, IPerforceSettings defaultSettings, string? updatePath, IServiceProvider serviceProvider)
+		private static UpdateMonitor CreateUpdateMonitor(LauncherSettings launcherSettings, IPerforceSettings defaultSettings, string? updatePath, bool runUpdateCheck, IServiceProvider serviceProvider)
 		{
-			if (launcherSettings.UpdateSource == LauncherUpdateSource.Horde)
+			if (!runUpdateCheck)
+			{
+				return new NullUpdateMonitor();
+			}
+			else if (launcherSettings.UpdateSource == LauncherUpdateSource.Horde)
 			{
 				return new HordeUpdateMonitor(SyncVersion ?? String.Empty, serviceProvider);
 			}
diff --git a/Engine/Source/Programs/UnrealGameSync/UnrealGameSync/ProgramApplicationContext.cs b/Engine/Source/Programs/UnrealGameSync/UnrealGameSync/ProgramApplicationContext.cs
--- a/Engine/Source/Programs/UnrealGameSync/UnrealGameSync/ProgramApplicationContext.cs
+++ b/Engine/Source/Programs/UnrealGameSync/UnrealGameSync/ProgramApplicationContext.cs
@@ -25,7 +25,7 @@ namespace UnrealGameSync
 		readonly SynchronizationContext _mainThreadSynchronizationContext;
 
 		readonly IPerforceSettings _defaultPerforceSettings;
-		UpdateMonitor _updateMonitor;
+		readonly UpdateMonitor _updateMonitor;
 		readonly string? _apiUrl;
 		readonly DirectoryReference _dataFolder;
 		readonly DirectoryReference _cacheFolder;
@@ -344,12 +344,6 @@ namespace UnrealGameSync
 		{
 			base.Dispose(disposing);
 
-			if(_updateMonitor != null)
-			{
-				_updateMonitor.Dispose();
-				_updateMonitor = null!;
-			}
-
 			if(_activationListener != null)
 			{
 				_activationListener.OnActivate -= OnActivationListenerAsyncCallback;
diff --git a/Engine/Source/Programs/UnrealGameSync/UnrealGameSync/ToolUpdateMonitor.cs b/Engine/Source/Programs/UnrealGameSync/UnrealGameSync/ToolUpdateMonitor.cs
--- a/Engine/Source/Programs/UnrealGameSync/UnrealGameSync/ToolUpdateMonitor.cs
+++ b/Engine/Source/Programs/UnrealGameSync/UnrealGameSync/ToolUpdateMonitor.cs
@@ -1,6 +1,7 @@
 // Copyright Epic Games, Inc. All Rights Reserved.
 
 using EpicGames.Core;
+using EpicGames.Horde;
 using EpicGames.Horde.Tools;
 using EpicGames.Perforce;
 using Microsoft.Extensions.DependencyInjection;
@@ -10,6 +11,7 @@ using System.Collections.Generic;
 using System.Diagnostics;
 using System.IO;
 using System.Linq;
+using System.Text;
 using System.Text.Json;
 using System.Threading;
 using System.Threading.Tasks;
@@ -33,17 +35,23 @@ namespace UnrealGameSync
 		}
 	}
 
+	// Settings in the UgsTool.json file in the root of a Horde tool
+	class ToolSettings
+	{
+		public string? InstallCommand { get; set; }
+		public string? UninstallCommand { get; set; }
+		public List<ToolLink> StatusPanelLinks { get; set; } = new List<ToolLink>();
+		public bool SafeWhenBusy { get; set; }
+	}
+
 	[DebuggerDisplay("{Name}")]
 	class ToolInfo
 	{
 		public Guid Id { get; set; }
 		public string Name { get; set; }
 		public string Description { get; set; }
-		public string? InstallCommand { get; set; }
-		public string? UninstallCommand { get; set; }
-		public List<ToolLink> StatusPanelLinks { get; set; } = new List<ToolLink>();
 		public HashSet<Guid> DependsOnToolIds { get; set; } = new HashSet<Guid>();
-		public bool SafeWhenBusy { get; set; }
+		public ToolSettings Settings { get; set; } = new ToolSettings();
 		public string Revision { get; set; }
 
 		public ToolInfo(Guid id, string name, string description, string revision)
@@ -71,6 +79,7 @@ namespace UnrealGameSync
 		readonly ILogger _logger;
 		readonly IAsyncDisposer _asyncDisposer;
 		readonly FileReference _enabledToolsFile;
+		readonly IServiceProvider _serviceProvider;
 
 		bool _readLegacyConfig;
 
@@ -95,6 +104,7 @@ namespace UnrealGameSync
 			Settings = settings;
 			_logger = serviceProvider.GetRequiredService<ILogger<ToolUpdateMonitor>>();
 			_asyncDisposer = serviceProvider.GetRequiredService<IAsyncDisposer>();
+			_serviceProvider = serviceProvider;
 
 			DirectoryReference.CreateDirectory(ToolsDir);
 			_enabledToolsFile = FileReference.Combine(ToolsDir, "tools.json");
@@ -205,7 +215,17 @@ namespace UnrealGameSync
 
 			// Update all the available tools
 			List<ToolInfo> tools = new List<ToolInfo>();
-			await ReadPerforceTools(perforce, tools, cancellationToken);
+			if (!String.IsNullOrEmpty(DeploymentSettings.Instance.ToolsDepotPath))
+			{
+				await ReadPerforceToolsAsync(perforce, tools, cancellationToken);
+			}
+			using (HordeHttpClient? hordeHttpClient = _serviceProvider.GetService<HordeHttpClient>())
+			{
+				if (hordeHttpClient != null)
+				{
+					await ReadHordeToolsAsync(hordeHttpClient, tools, cancellationToken);
+				}
+			}
 			_tools = tools;
 
 			// When upgrading from older UGS versions, read the legacy sync CL from plain-text config files
@@ -266,13 +286,8 @@ namespace UnrealGameSync
 			}
 		}
 
-		async Task ReadPerforceTools(IPerforceConnection perforce, List<ToolInfo> tools, CancellationToken cancellationToken)
+		async Task ReadPerforceToolsAsync(IPerforceConnection perforce, List<ToolInfo> tools, CancellationToken cancellationToken)
 		{
-			if (String.IsNullOrEmpty(DeploymentSettings.Instance.ToolsDepotPath))
-			{
-				return;
-			}
-
 			List<ChangesRecord> changes = await perforce.GetChangesAsync(ChangesOptions.None, 1, ChangeStatus.Submitted, $"{DeploymentSettings.Instance.ToolsDepotPath}/...", cancellationToken);
 			if (changes.Count > 0 && changes[0].Number != _perforceToolsChange)
 			{
@@ -310,7 +325,22 @@ namespace UnrealGameSync
 			tools.AddRange(_perforceTools.Values);
 		}
 
-		async Task<ToolInfo?> ReadToolDefinitionAsync(IPerforceConnection perforce, string iniRevision, string toolRevision, CancellationToken cancellationToken)
+		static async Task ReadHordeToolsAsync(HordeHttpClient hordeHttpClient, List<ToolInfo> tools, CancellationToken cancellationToken)
+		{
+			GetToolsSummaryResponse toolsResponse = await hordeHttpClient.GetToolsAsync(cancellationToken);
+			foreach (GetToolSummaryResponse toolResponse in toolsResponse.Tools)
+			{
+				if (toolResponse.ShowInUgs)
+				{
+					IoHash hash = IoHash.Compute(Encoding.UTF8.GetBytes($"horde:{toolResponse.Id}"));
+					Guid guid = new Guid(hash.ToByteArray().AsSpan(0, 16));
+					ToolInfo toolInfo = new ToolInfo(guid, toolResponse.Name, toolResponse.Description, $"{hordeHttpClient},{toolResponse.Id},{toolResponse.DeploymentId}");
+					tools.Add(toolInfo);
+				}
+			}
+		}
+
+		static async Task<ToolInfo?> ReadToolDefinitionAsync(IPerforceConnection perforce, string iniRevision, string toolRevision, CancellationToken cancellationToken)
 		{
 			PerforceResponse<PrintRecord<string[]>> response = await perforce.TryPrintLinesAsync(iniRevision, cancellationToken);
 			if (!response.Succeeded || response.Data.Contents == null)
@@ -340,9 +370,10 @@ namespace UnrealGameSync
 			string toolDescription = configFile.GetValue("Settings.Description", toolName);
 
 			ToolInfo tool = new ToolInfo(toolId, toolName, toolDescription, revision);
-			tool.InstallCommand = configFile.GetValue("Settings.InstallCommand", null);
-			tool.UninstallCommand = configFile.GetValue("Settings.UninstallCommand", null);
-			tool.SafeWhenBusy = configFile.GetValue("Settings.SafeWhenBusy", false);
+			tool.Settings = new ToolSettings();
+			tool.Settings.InstallCommand = configFile.GetValue("Settings.InstallCommand", null);
+			tool.Settings.UninstallCommand = configFile.GetValue("Settings.UninstallCommand", null);
+			tool.Settings.SafeWhenBusy = configFile.GetValue("Settings.SafeWhenBusy", false);
 
 			foreach (string line in configFile.GetValues("Settings.DependsOnTool", Array.Empty<string>()))
 			{
@@ -365,7 +396,7 @@ namespace UnrealGameSync
 					ToolLink link = new ToolLink(label, fileName);
 					link.Arguments = obj.GetValue("Arguments", null);
 					link.WorkingDir = obj.GetValue("WorkingDir", null);
-					tool.StatusPanelLinks.Add(link);
+					tool.Settings.StatusPanelLinks.Add(link);
 				}
 			}
 
@@ -397,10 +428,10 @@ namespace UnrealGameSync
 				_logger.LogInformation("Removing {ToolName}", tool.Name);
 				DirectoryReference? toolPath = GetToolPath(tool.Name);
 
-				if (!String.IsNullOrEmpty(tool.UninstallCommand))
+				if (!String.IsNullOrEmpty(tool.Settings?.UninstallCommand))
 				{
-					_logger.LogInformation("Running unininstall action: {Command}", tool.UninstallCommand);
-					await RunCommandAsync(tool.Name, tool.UninstallCommand, cancellationToken);
+					_logger.LogInformation("Running unininstall action: {Command}", tool.Settings.UninstallCommand);
+					await RunCommandAsync(tool.Name, tool.Settings.UninstallCommand, cancellationToken);
 				}
 
 				await SetToolRevisionAsync(tool.Name, null, cancellationToken);
@@ -465,24 +496,50 @@ namespace UnrealGameSync
 				DirectoryReference nextToolZipsDir = DirectoryReference.Combine(nextToolDir, ".zips");
 				DirectoryReference.CreateDirectory(nextToolZipsDir);
 
+				FileReference zipFile = FileReference.Combine(nextToolZipsDir, $"{tool.Name}.zip");
 				if (tool.Revision.StartsWith("//", StringComparison.Ordinal))
 				{
 					// Read it from Perforce
-					FileReference zipFile = FileReference.Combine(nextToolZipsDir, $"{tool.Name}.zip");
-
 					PerforceResponseList<PrintRecord> response = await perforce.TryPrintAsync(zipFile.FullName, tool.Revision, cancellationToken);
 					if (!response.Succeeded || !FileReference.Exists(zipFile))
 					{
 						_logger.LogError("Unable to print {DepotFile}", tool.Revision);
 						return false;
 					}
-
 					ArchiveUtils.ExtractFiles(zipFile, nextToolDir, null, new ProgressValue(), _logger);
 				}
 				else
 				{
-					// Unknown format?
-					return false;
+					using HordeHttpClient? hordeHttpClient = _serviceProvider.GetService<HordeHttpClient>();
+					if (hordeHttpClient != null)
+					{
+						string[] fields = tool.Revision.Split(',');
+						if (fields.Length != 3)
+						{
+							_logger.LogError("Unexpected format for Horde revision ('{Revision}')", tool.Revision);
+							return false;
+						}
+
+						using (FileStream stream = FileReference.Open(zipFile, FileMode.Create, FileAccess.Write, FileShare.None))
+						{
+							await using Stream sourceStream = await hordeHttpClient.GetToolDeploymentZipAsync(new ToolId(fields[1]), new ToolDeploymentId(BinaryId.Parse(fields[2])), cancellationToken);
+							await sourceStream.CopyToAsync(stream, cancellationToken);
+						}
+
+						ArchiveUtils.ExtractFiles(zipFile, nextToolDir, null, new ProgressValue(), _logger);
+
+						FileReference settingsFile = FileReference.Combine(nextToolDir, "UgsTool.json");
+						if (FileReference.Exists(settingsFile))
+						{
+							byte[] data = await FileReference.ReadAllBytesAsync(settingsFile, cancellationToken);
+							tool.Settings = JsonSerializer.Deserialize<ToolSettings>(data, new JsonSerializerOptions { PropertyNameCaseInsensitive = true }) ?? tool.Settings;
+						}
+					}
+					else
+					{
+						_logger.LogError("Unknown source for {Revision}", tool.Revision);
+						return false;
+					}
 				}
 
 				DirectoryReference currentToolDir = DirectoryReference.Combine(toolDir, "Current");
@@ -496,10 +553,10 @@ namespace UnrealGameSync
 
 				Directory.Move(nextToolDir.FullName, currentToolDir.FullName);
 
-				if (!String.IsNullOrEmpty(tool.InstallCommand))
+				if (!String.IsNullOrEmpty(tool.Settings?.InstallCommand))
 				{
-					_logger.LogInformation("Running install action: {Command}", tool.InstallCommand);
-					await RunCommandAsync(tool.Name, tool.InstallCommand, cancellationToken);
+					_logger.LogInformation("Running install action: {Command}", tool.Settings.InstallCommand);
+					await RunCommandAsync(tool.Name, tool.Settings.InstallCommand, cancellationToken);
 				}
 
 				await SetToolRevisionAsync(tool.Name, tool, cancellationToken);
diff --git a/Engine/Source/Programs/UnrealGameSync/UnrealGameSync/UpdateMonitor.cs b/Engine/Source/Programs/UnrealGameSync/UnrealGameSync/UpdateMonitor.cs
--- a/Engine/Source/Programs/UnrealGameSync/UnrealGameSync/UpdateMonitor.cs
+++ b/Engine/Source/Programs/UnrealGameSync/UnrealGameSync/UpdateMonitor.cs
@@ -18,7 +18,7 @@ namespace UnrealGameSync
 		UserInitiated,
 	}
 
-	abstract class UpdateMonitor : IDisposable
+	abstract class UpdateMonitor : IAsyncDisposable
 	{
 		public bool IsUpdateAvailable
 		{
@@ -34,7 +34,7 @@ namespace UnrealGameSync
 			private set;
 		}
 
-		public abstract void Dispose();
+		public abstract ValueTask DisposeAsync();
 
 		public void TriggerUpdate(UpdateType updateType, bool openSettings)
 		{
@@ -47,6 +47,11 @@ namespace UnrealGameSync
 		}
 	}
 
+	class NullUpdateMonitor : UpdateMonitor
+	{
+		public override ValueTask DisposeAsync() => default;
+	}
+
 	class HordeUpdateMonitor : UpdateMonitor
 	{
 		readonly ToolId _toolId;
@@ -69,9 +74,9 @@ namespace UnrealGameSync
 			_backgroundTask = BackgroundTask.StartNew(ctx => CheckForUpdatesLoopAsync(ctx));
 		}
 
-		public override void Dispose()
+		public override async ValueTask DisposeAsync()
 		{
-			Task.Run(() => _backgroundTask.DisposeAsync()).Wait();
+			await _backgroundTask.DisposeAsync();
 		}
 
 		public async Task CheckForUpdatesLoopAsync(CancellationToken cancellationToken)
@@ -119,12 +124,10 @@ namespace UnrealGameSync
 		readonly CancellationTokenSource _cancellationSource = new CancellationTokenSource();
 #pragma warning restore CA2213
 		readonly ILogger _logger;
-		readonly IAsyncDisposer _asyncDisposer;
 
 		public PerforceUpdateMonitor(IPerforceSettings perforceSettings, string? watchPath, IServiceProvider serviceProvider)
 		{
 			_logger = serviceProvider.GetRequiredService<ILogger<UpdateMonitor>>();
-			_asyncDisposer = serviceProvider.GetRequiredService<IAsyncDisposer>();
 
 			if(watchPath != null)
 			{
@@ -133,15 +136,18 @@ namespace UnrealGameSync
 			}
 		}
 
-		public override void Dispose()
+		public override async ValueTask DisposeAsync()
 		{
 			OnUpdateAvailable = null;
 
 			if (_workerTask != null)
 			{
 				_cancellationSource.Cancel();
-				_asyncDisposer.Add(_workerTask.ContinueWith(_ => _cancellationSource.Dispose(), TaskScheduler.Default));
+
+				await _workerTask;
 				_workerTask = null;
+
+				_cancellationSource.Dispose();
 			}
 		}
 
