UGS: Extended the "SDK Info" window to support adding arbitrary badge buttons.

Previously, this dialog only had fixed support for Install and Browse commands. While sufficient in most cases, on some platforms they lacked a bit of additional flexibility.

Now, we may add badges with custom names launching arbitrary programs with optional command-line parameters. For instance,

    ReleaseNotes=(Program="notepad.exe", Args="notes.txt")

Note that now actions are added in the order in which are defined in the .ini config.

The underlying implementation has been unified so Install, Browse and the new actions launch programs in exactly the same way, but the previous behavior and syntax for Install and Browse is still supported.

#rnx

diff --git a/Engine/Programs/UnrealGameSync/UnrealGameSync.ini b/Engine/Programs/UnrealGameSync/UnrealGameSync.ini
--- a/Engine/Programs/UnrealGameSync/UnrealGameSync.ini
+++ b/Engine/Programs/UnrealGameSync/UnrealGameSync.ini
@@ -52,6 +52,6 @@ SdkInstallerDir=
 +SdkInfo=(Category="Windows", Description="Visual Studio 2019 v16.11.5")
 +SdkInfo=(Category="Windows", Description="Windows SDK 10.0.18362.0")
 +SdkInfo=(Category="Windows", Description="LLVM Clang 11.1.0")
-+SdkInfo=(Category="Linux", Description="v22_clang-16.0.6-centos7", Browse="$(SdkInstallerDir)\\Linux\\HostWin64", Install="$(SdkInstallerDir)\\Linux\\HostWin64\\v22_clang-16.0.6-centos7.exe")
++SdkInfo=(Category="Linux", Description="v22_clang-16.0.6-centos7", Install="$(SdkInstallerDir)\\Linux\\HostWin64\\v22_clang-16.0.6-centos7.exe", Browse="$(SdkInstallerDir)\\Linux\\HostWin64")
 +SdkInfo=(Category="Mac", Description="macOS 10.15.7")
 +SdkInfo=(Category="Mac", Description="Xcode 12.4")
diff --git a/Engine/Source/Programs/UnrealGameSync/UnrealGameSync/Forms/SdkInfoWindow.cs b/Engine/Source/Programs/UnrealGameSync/UnrealGameSync/Forms/SdkInfoWindow.cs
--- a/Engine/Source/Programs/UnrealGameSync/UnrealGameSync/Forms/SdkInfoWindow.cs
+++ b/Engine/Source/Programs/UnrealGameSync/UnrealGameSync/Forms/SdkInfoWindow.cs
@@ -16,12 +16,54 @@ namespace UnrealGameSync
 {
 	public partial class SdkInfoWindow : Form
 	{
+		abstract class SdkAction
+		{
+			public string Name { get; init; }
+
+			public SdkAction(string name)
+			{
+				Name = name;
+			}
+
+			public abstract void Execute();
+		}
+
+		class SdkActionRun : SdkAction
+		{
+			public string Program { get; set; } = "";
+			public string Args    { get; set; } = "";
+
+			public SdkActionRun(string name, string program = "", string args = "")
+				: base(name)
+			{
+				Program = program;
+				Args    = args;
+			}
+
+			public override void Execute()
+			{
+				try
+				{
+					ProcessStartInfo startInfo = new ProcessStartInfo();
+					startInfo.FileName = Program;
+					startInfo.Arguments = Args;
+					startInfo.UseShellExecute = true;
+					Process.Start(startInfo);
+				}
+
+				catch(Exception ex)
+				{
+					MessageBox.Show($"Unable to run '{Program} {Args}': {ex.Message}");
+				}
+			}
+		}
+
 		class SdkItem
 		{
 			public string Category { get; }
 			public string Description { get; }
-			public string? Install { get; set; }
-			public string? Browse { get; set; }
+
+			public List<SdkAction> Actions = new List<SdkAction>();
 
 			public SdkItem(string category, string description)
 			{
@@ -79,27 +121,80 @@ namespace UnrealGameSync
 				string description = obj.GetValue("Description", "");
 				SdkItem item = new SdkItem(category, description);
 
-				item.Install = Utility.ExpandVariables(obj.GetValue("Install", ""), variables);
-				if(item.Install.Contains("$(", StringComparison.Ordinal))
-				{
-					item.Install = null;
-				}
+				List<string> keys = obj.GetKeys().ToList();
 
-				item.Browse = Utility.ExpandVariables(obj.GetValue("Browse", ""), variables);
-				if(item.Browse.Contains("$(", StringComparison.Ordinal))
+				foreach (string key in keys)
 				{
-					item.Browse = null;
-				}
+					string  keyToAdd = key;
+					string? value    = obj.GetValue(keyToAdd, "");
 
-				if(!String.IsNullOrEmpty(item.Install) && String.IsNullOrEmpty(item.Browse))
-				{
-					try
+					if (String.IsNullOrEmpty(value))
+					{
+						continue;
+					}
+
+					if (!value.StartsWith('('))
 					{
-						item.Browse = Path.GetDirectoryName(item.Install);
+						// Handle the predefined Install and Browse actions.
+
+						if (keyToAdd == "Install")
+						{
+							string? installArgument = Utility.ExpandVariables(value, variables);
+							if (!installArgument.Contains("$(", StringComparison.Ordinal))
+							{
+								item.Actions.Add(new SdkActionRun("Install", installArgument));
+
+								// If Browse is not explicitly defined, generate it automatically from the Installer path.
+								if (!keys.Contains("Browse"))
+								{
+									keyToAdd = "Browse";
+									value = Path.GetDirectoryName(installArgument);
+
+									if (String.IsNullOrEmpty(value))
+									{
+										continue;
+									}
+								}
+							}
+						}
+
+						if (keyToAdd == "Browse")
+						{
+							string? browseArgument = Utility.ExpandVariables(value, variables);
+							if (!browseArgument.Contains("$(", StringComparison.Ordinal))
+							{
+								item.Actions.Add(new SdkActionRun("Browse", "explorer.exe", browseArgument.Replace('/', '\\')));
+							}
+						}
 					}
-					catch
+					else
 					{
-						item.Browse = null;
+						// Handle generic actions i.e. ones with custom names which may launch arbitrary executables with additional parameters.
+						// For instance,
+						//		ReleaseNotes=(Program="notepad.exe", Args="notes.txt")
+						// will add a "ReleaseNotes" badge which, when clicked, will open "notes.txt" in the notepad.
+
+						ConfigObject valueObj = new ConfigObject(value);
+
+						string program = valueObj.GetValue("Program", "");
+						string args = valueObj.GetValue("Args", "");
+
+						program = Utility.ExpandVariables(program, variables);
+						if (program.Contains("$(", StringComparison.Ordinal))
+						{
+							continue;
+						}
+
+						args = Utility.ExpandVariables(args, variables);
+						if (args.Contains("$(", StringComparison.Ordinal))
+						{
+							continue;
+						}
+
+						if (!String.IsNullOrEmpty(program))
+						{
+							item.Actions.Add(new SdkActionRun(keyToAdd, program, args));
+						}
 					}
 				}
 
@@ -220,19 +315,12 @@ namespace UnrealGameSync
 
 			SdkItem sdk = (SdkItem)subItem.Tag;
 
-			Action? installAction = null;
-			if(!String.IsNullOrEmpty(sdk.Install))
+			foreach (SdkAction action in sdk.Actions)
 			{
-				installAction = () => { Install(sdk.Install); };
-			}
-			badges.Add(new BadgeInfo(uniqueIdPrefix + "_Install", "Install") { OnClick = installAction });
+				Action clickAction = () => { action.Execute(); };
 
-			Action? browseAction = null;
-			if(!String.IsNullOrEmpty(sdk.Browse))
-			{
-				browseAction = () => { Browse(sdk.Browse); };
+				badges.Add(new BadgeInfo($"{uniqueIdPrefix}_{action.Name}", action.Name) { OnClick = clickAction });
 			}
-			badges.Add(new BadgeInfo(uniqueIdPrefix + "_Browse", "Browse"){ OnClick = browseAction });
 
 			int right = subItem.Bounds.Right - 10;
 			for(int idx = badges.Count - 1; idx >= 0; idx--)
@@ -245,33 +333,6 @@ namespace UnrealGameSync
 			return badges;
 		}
 
-		private static void Browse(string directoryName)
-		{
-			try
-			{
-				Process.Start("explorer.exe", String.Format("\"{0}\"", directoryName));
-			}
-			catch(Exception ex)
-			{
-				MessageBox.Show(String.Format("Unable to open explorer to {0}: {1}", directoryName, ex.Message));
-			}
-		}
-
-		private static void Install(string fileName)
-		{
-			try
-			{
-				ProcessStartInfo startInfo = new ProcessStartInfo();
-				startInfo.FileName = fileName;
-				startInfo.UseShellExecute = true;
-				Process.Start(startInfo);
-			}
-			catch(Exception ex)
-			{
-				MessageBox.Show(String.Format("Unable to run {0}: {1}", fileName, ex.Message));
-			}
-		}
-
 		private Size GetBadgeSize(string badgeText)
 		{
 			Size labelSize = TextRenderer.MeasureText(badgeText, _badgeFont);
diff --git a/Engine/Source/Programs/UnrealGameSync/UnrealGameSyncShared/ConfigFile.cs b/Engine/Source/Programs/UnrealGameSync/UnrealGameSyncShared/ConfigFile.cs
--- a/Engine/Source/Programs/UnrealGameSync/UnrealGameSyncShared/ConfigFile.cs
+++ b/Engine/Source/Programs/UnrealGameSync/UnrealGameSyncShared/ConfigFile.cs
@@ -148,6 +148,14 @@ namespace UnrealGameSync
 			return token;
 		}
 
+		public IEnumerable<string> GetKeys()
+		{
+			for (int idx = 0; idx < Pairs.Count; idx++)
+			{
+				yield return Pairs[idx].Key;
+			}
+		}
+
 		[return: NotNullIfNotNull("defaultValue")]
 		public string? GetValue(string key, string? defaultValue = null)
 		{
