UGS: Fix HordeHandler compile error

diff --git a/Engine/Source/Programs/UnrealGameSync/UnrealGameSync/UriHandlers/HordeHandler.cs b/Engine/Source/Programs/UnrealGameSync/UnrealGameSync/UriHandlers/HordeHandler.cs
--- a/Engine/Source/Programs/UnrealGameSync/UnrealGameSync/UriHandlers/HordeHandler.cs
+++ b/Engine/Source/Programs/UnrealGameSync/UnrealGameSync/UriHandlers/HordeHandler.cs
@@ -57,7 +57,7 @@ namespace UnrealGameSync.UriHandlers
 				DirectoryNode? node = await storageClient.ReadRefAsync<DirectoryNode>(refName, cancellationToken: cancellationToken);
 
 				progress.Report("Starting...");
-				await node.CopyToDirectoryAsync(outputDir.ToDirectoryInfo(), new CopyProgressAdapter(progress), serviceProvider.GetRequiredService<ILogger<HordeHandler>>(), cancellationToken);
+				await node.CopyToDirectoryAsync(outputDir.ToDirectoryInfo(), new CopyProgressAdapter(progress), BlobSerializerOptions.Default, serviceProvider.GetRequiredService<ILogger<HordeHandler>>(), cancellationToken);
 			}
 		}
 
