Allow the current changelist context menu to work if the synced changelist is not actually part of the currently clientspec view.

#rnx

- Do not use the clientspec as part of the filespec when calling p4 changes. Since we are limiting the filespace to a single changelist this shouldn't hurt performance.

diff --git a/Engine/Source/Programs/UnrealGameSync/UnrealGameSync/Controls/WorkspaceControl.cs b/Engine/Source/Programs/UnrealGameSync/UnrealGameSync/Controls/WorkspaceControl.cs
--- a/Engine/Source/Programs/UnrealGameSync/UnrealGameSync/Controls/WorkspaceControl.cs
+++ b/Engine/Source/Programs/UnrealGameSync/UnrealGameSync/Controls/WorkspaceControl.cs
@@ -3577,7 +3577,8 @@ namespace UnrealGameSync
 							ChangesRecord? foundRecord = null;
 							async Task<bool> GetChangesRecordFunc(IPerforceConnection perforce, CancellationToken cancellationToken)
 							{
-								List<ChangesRecord> changes = await perforce.GetChangesAsync(ChangesOptions.None, 1, ChangeStatus.Submitted, $"//{_perforceSettings.ClientName}/...@{_workspace.CurrentChangeNumber},{_workspace.CurrentChangeNumber}", cancellationToken);
+								List<ChangesRecord> changes = await perforce.GetChangesAsync(ChangesOptions.None, 1, ChangeStatus.Submitted, $"//...@{_workspace.CurrentChangeNumber},{_workspace.CurrentChangeNumber}", cancellationToken);
+
 								if (changes.Count == 1)
 								{
 									foundRecord = changes[0];
