Add a build of UGS-CLI for Mac-Arm64.

diff --git a/Engine/Source/Programs/UnrealGameSync/BuildUGS.xml b/Engine/Source/Programs/UnrealGameSync/BuildUGS.xml
--- a/Engine/Source/Programs/UnrealGameSync/BuildUGS.xml
+++ b/Engine/Source/Programs/UnrealGameSync/BuildUGS.xml
@@ -13,8 +13,10 @@
 	<Property Name="UgsOutputDir" Value="$(RootDir)/LocalBuilds/UnrealGameSync"/>
 	<Property Name="UgsZipOutputDir" Value="$(RootDir)/LocalBuilds/UnrealGameSyncZip"/>
 	<Property Name="UgsCmdOutputDir" Value="$(RootDir)/LocalBuilds/UnrealGameSyncCmd"/>
-	<Property Name="UgsCmdMacZipFile" Value="$(UgsCmdOutputDir)/ugs-mac.zip"/>
-	<Property Name="UgsCmdMacOutputDir" Value="$(UgsCmdOutputDir)/Mac"/>
+	<Property Name="UgsCmdMacX64ZipFile" Value="$(UgsCmdOutputDir)/ugs-mac-x64.zip"/>
+	<Property Name="UgsCmdMacX64OutputDir" Value="$(UgsCmdOutputDir)/Mac-X64"/>
+	<Property Name="UgsCmdMacArm64ZipFile" Value="$(UgsCmdOutputDir)/ugs-mac-arm64.zip"/>
+	<Property Name="UgsCmdMacArm64OutputDir" Value="$(UgsCmdOutputDir)/Mac-Arm64"/>
 	<Property Name="UgsCmdLinuxZipFile" Value="$(UgsCmdOutputDir)/ugs-linux.zip"/>
 	<Property Name="UgsCmdLinuxOutputDir" Value="$(UgsCmdOutputDir)/Linux"/>
 	<Property Name="WixDir" Value="$(RootDir)/Engine/Source/ThirdParty/WiX/3.8"/>
@@ -88,14 +90,20 @@
 		</Node>
 		
 		<!-- Mac -->
-		<Node Name="Build UnrealGameSyncCmd Mac">
-			<Delete Files="$(UgsCmdMacOutputDir)/..."/>
-			<DotNet Arguments="publish &quot;$(UgsDir)/UnrealGameSyncCmd/UnrealGameSyncCmd.csproj&quot; --output=&quot;$(UgsCmdMacOutputDir)&quot; --self-contained -c Release -r osx-x64 $(VersionArguments)"/>
-			<Zip ZipFile="$(UgsCmdMacZipFile)" FromDir="$(UgsCmdMacOutputDir)" ExecutableFiles="ugs;*.dylib;*.sh"/>
+		<Node Name="Build UnrealGameSyncCmd Mac-X64">
+			<Delete Files="$(UgsCmdMacX64OutputDir)/..."/>
+			<DotNet Arguments="publish &quot;$(UgsDir)/UnrealGameSyncCmd/UnrealGameSyncCmd.csproj&quot; --output=&quot;$(UgsCmdMacX64OutputDir)&quot; --self-contained -c Release -r osx-x64 $(VersionArguments)"/>
+			<Zip ZipFile="$(UgsCmdMacX64ZipFile)" FromDir="$(UgsCmdMacX64OutputDir)" ExecutableFiles="ugs;*.dylib;*.sh"/>
+		</Node>
+		<Node Name="Build UnrealGameSyncCmd Mac-Arm64">
+			<Delete Files="$(UgsCmdMacArm64OutputDir)/..."/>
+			<DotNet Arguments="publish &quot;$(UgsDir)/UnrealGameSyncCmd/UnrealGameSyncCmd.csproj&quot; --output=&quot;$(UgsCmdMacArm64OutputDir)&quot; --self-contained -c Release -r osx-arm64 $(VersionArguments)"/>
+			<Zip ZipFile="$(UgsCmdMacArm64ZipFile)" FromDir="$(UgsCmdMacArm64OutputDir)" ExecutableFiles="ugs;*.dylib;*.sh"/>
 		</Node>
 		<Node Name="Publish UnrealGameSyncCmd Mac" Requires="Build UnrealGameSyncCmd Mac">
 			<Error Message="Missing -HordeSettings=... argument" If="'$(HordeSettings)' == ''"/>
-			<DeployTool Id="ugs-mac" Version="$(InformationalVersion)" Settings="$(HordeSettings)" File="$(UgsCmdMacZipFile)"/>
+			<DeployTool Id="ugs-mac" Version="$(InformationalVersion)" Settings="$(HordeSettings)" File="$(UgsCmdMacX64ZipFile)"/>
+			<DeployTool Id="ugs-mac-arm64" Version="$(InformationalVersion)" Settings="$(HordeSettings)" File="$(UgsCmdMacArm64ZipFile)"/>
 		</Node>
 		
 		<!-- Linux -->
