Fix dependency in UGS build script.

diff --git a/Engine/Source/Programs/UnrealGameSync/BuildUGS.xml b/Engine/Source/Programs/UnrealGameSync/BuildUGS.xml
--- a/Engine/Source/Programs/UnrealGameSync/BuildUGS.xml
+++ b/Engine/Source/Programs/UnrealGameSync/BuildUGS.xml
@@ -100,7 +100,7 @@
 			<DotNet Arguments="publish &quot;$(UgsDir)/UnrealGameSyncCmd/UnrealGameSyncCmd.csproj&quot; --output=&quot;$(UgsCmdMacArm64OutputDir)&quot; --self-contained -c Release -r osx-arm64 $(VersionArguments)"/>
 			<Zip ZipFile="$(UgsCmdMacArm64ZipFile)" FromDir="$(UgsCmdMacArm64OutputDir)" ExecutableFiles="ugs;*.dylib;*.sh"/>
 		</Node>
-		<Node Name="Publish UnrealGameSyncCmd Mac" Requires="Build UnrealGameSyncCmd Mac">
+		<Node Name="Publish UnrealGameSyncCmd Mac" Requires="Build UnrealGameSyncCmd Mac-X64;Build UnrealGameSyncCmd Mac-Arm64">
 			<Error Message="Missing -HordeSettings=... argument" If="'$(HordeSettings)' == ''"/>
 			<DeployTool Id="ugs-mac" Version="$(InformationalVersion)" Settings="$(HordeSettings)" File="$(UgsCmdMacX64ZipFile)"/>
 			<DeployTool Id="ugs-mac-arm64" Version="$(InformationalVersion)" Settings="$(HordeSettings)" File="$(UgsCmdMacArm64ZipFile)"/>
