UGS: Fix compile errors with recent EpicGames.Horde changes.

diff --git a/Engine/Source/Programs/UnrealGameSync/UnrealGameSync/UriHandlers/HordeHandler.cs b/Engine/Source/Programs/UnrealGameSync/UnrealGameSync/UriHandlers/HordeHandler.cs
--- a/Engine/Source/Programs/UnrealGameSync/UnrealGameSync/UriHandlers/HordeHandler.cs
+++ b/Engine/Source/Programs/UnrealGameSync/UnrealGameSync/UriHandlers/HordeHandler.cs
@@ -54,10 +54,10 @@ namespace UnrealGameSync.UriHandlers
 
 				progress.Report("Connecting to server...");
 
-				DirectoryNode? node = await storageClient.ReadRefAsync<DirectoryNode>(refName, cancellationToken: cancellationToken);
+				DirectoryNode? node = await storageClient.ReadRefTargetAsync<DirectoryNode>(refName, cancellationToken: cancellationToken);
 
 				progress.Report("Starting...");
-				await node.CopyToDirectoryAsync(outputDir.ToDirectoryInfo(), new CopyProgressAdapter(progress), BlobSerializerOptions.Default, serviceProvider.GetRequiredService<ILogger<HordeHandler>>(), cancellationToken);
+				await node.CopyToDirectoryAsync(outputDir.ToDirectoryInfo(), new CopyProgressAdapter(progress), serviceProvider.GetRequiredService<ILogger<HordeHandler>>(), cancellationToken);
 			}
 		}
 
diff --git a/Engine/Source/Programs/UnrealGameSync/UnrealGameSyncShared/ArchiveInfo.cs b/Engine/Source/Programs/UnrealGameSync/UnrealGameSyncShared/ArchiveInfo.cs
--- a/Engine/Source/Programs/UnrealGameSync/UnrealGameSyncShared/ArchiveInfo.cs
+++ b/Engine/Source/Programs/UnrealGameSync/UnrealGameSyncShared/ArchiveInfo.cs
@@ -91,7 +91,7 @@ namespace UnrealGameSync
 
 		public bool TryGetArchiveKeyForChangeNumber(int changeNumber, int maxChangeNumber, [NotNullWhen(true)] out string? archiveKey)
 		{
-			int idx = ChangeNumberToFileRevision.Keys.AsReadOnly().BinarySearch(changeNumber);
+			int idx = ChangeNumberToFileRevision.Keys.AsReadOnlyList().BinarySearch(changeNumber);
 			if (idx >= 0)
 			{
 				archiveKey = ChangeNumberToFileRevision.Values[idx];
