UGS: Fix tool name for CLI upgrades on Mac Arm64.

diff --git a/Engine/Source/Programs/UnrealGameSync/UnrealGameSyncCmd/Program.cs b/Engine/Source/Programs/UnrealGameSync/UnrealGameSyncCmd/Program.cs
--- a/Engine/Source/Programs/UnrealGameSync/UnrealGameSyncCmd/Program.cs
+++ b/Engine/Source/Programs/UnrealGameSync/UnrealGameSyncCmd/Program.cs
@@ -384,7 +384,14 @@ namespace UnrealGameSyncCmd
 		{
 			if (RuntimeInformation.IsOSPlatform(OSPlatform.OSX))
 			{
-				return "ugs-mac";
+				if (RuntimeInformation.OSArchitecture == Architecture.Arm64)
+				{
+					return "ugs-mac-arm64";
+				}
+				else
+				{
+					return "ugs-mac";
+				}
 			}
 			else if (RuntimeInformation.IsOSPlatform(OSPlatform.Linux))
 			{
