UGS: Merging build script changes.

diff --git a/Engine/Source/Programs/UnrealGameSync/BuildUGS.xml b/Engine/Source/Programs/UnrealGameSync/BuildUGS.xml
--- a/Engine/Source/Programs/UnrealGameSync/BuildUGS.xml
+++ b/Engine/Source/Programs/UnrealGameSync/BuildUGS.xml
@@ -8,6 +8,7 @@
 	<Option Name="SubmitPath" DefaultValue="" Description="Relative path within the stream/branch to submit binaries"/>
 	<Option Name="IncludeInstaller" DefaultValue="false" Description="Whether to include the launcher installer in the regular client distribution"/>
 	<Option Name="HordeSettings" DefaultValue="" Description="Path to a JSON file containing settings for publishing to Horde"/>
+	<Option Name="HordeSettingsMac" DefaultValue="" Description="Path to a JSON file containing settings for publishing to Horde"/>
 
 	<Property Name="UgsDir" Value="$(RootDir)/Engine/Source/Programs/UnrealGameSync"/>
 	<Property Name="UgsOutputDir" Value="$(RootDir)/LocalBuilds/UnrealGameSync"/>
@@ -89,33 +90,39 @@
 			<DeployTool Id="ugs-win" Version="$(InformationalVersion)" Settings="$(HordeSettings)" File="$(UgsZipPublishDir)/Ugs.zip"/>
 		</Node>
 		
+		<!-- Linux -->
+		<Node Name="Build UnrealGameSyncCmd Linux">
+			<Delete Files="$(UgsCmdLinuxOutputDir)/..."/>
+			<DotNet Arguments="publish &quot;$(UgsDir)/UnrealGameSyncCmd/UnrealGameSyncCmd.csproj&quot; --output=&quot;$(UgsCmdLinuxOutputDir)&quot; --self-contained -c Release -r linux-x64 $(VersionArguments)"/>
+			<Zip ZipFile="$(UgsCmdLinuxZipFile)" FromDir="$(UgsCmdLinuxOutputDir)" ExecutableFiles="ugs;*.so;*.sh"/>
+		</Node>
+		<Node Name="Publish UnrealGameSyncCmd Linux" Requires="Build UnrealGameSyncCmd Linux">
+			<Error Message="Missing -HordeSettings=... argument" If="'$(HordeSettings)' == ''"/>
+			<DeployTool Id="ugs-linux" Version="$(InformationalVersion)" Settings="$(HordeSettings)" File="$(UgsCmdLinuxZipFile)"/>
+		</Node>
+	</Agent>
+
+	<Agent Name="UnrealGameSync Mac Agent" Type="Mac">
+
 		<!-- Mac -->
 		<Node Name="Build UnrealGameSyncCmd Mac-X64">
 			<Delete Files="$(UgsCmdMacX64OutputDir)/..."/>
 			<DotNet Arguments="publish &quot;$(UgsDir)/UnrealGameSyncCmd/UnrealGameSyncCmd.csproj&quot; --output=&quot;$(UgsCmdMacX64OutputDir)&quot; --self-contained -c Release -r osx-x64 $(VersionArguments)"/>
+			<Sign Files="$(UgsCmdMacX64OutputDir)/ugs"/>
 			<Zip ZipFile="$(UgsCmdMacX64ZipFile)" FromDir="$(UgsCmdMacX64OutputDir)" ExecutableFiles="ugs;*.dylib;*.sh"/>
 		</Node>
 		<Node Name="Build UnrealGameSyncCmd Mac-Arm64">
 			<Delete Files="$(UgsCmdMacArm64OutputDir)/..."/>
 			<DotNet Arguments="publish &quot;$(UgsDir)/UnrealGameSyncCmd/UnrealGameSyncCmd.csproj&quot; --output=&quot;$(UgsCmdMacArm64OutputDir)&quot; --self-contained -c Release -r osx-arm64 $(VersionArguments)"/>
+			<Sign Files="$(UgsCmdMacArm64OutputDir)/ugs"/>
 			<Zip ZipFile="$(UgsCmdMacArm64ZipFile)" FromDir="$(UgsCmdMacArm64OutputDir)" ExecutableFiles="ugs;*.dylib;*.sh"/>
 		</Node>
 		<Node Name="Publish UnrealGameSyncCmd Mac" Requires="Build UnrealGameSyncCmd Mac-X64;Build UnrealGameSyncCmd Mac-Arm64">
-			<Error Message="Missing -HordeSettings=... argument" If="'$(HordeSettings)' == ''"/>
-			<DeployTool Id="ugs-mac" Version="$(InformationalVersion)" Settings="$(HordeSettings)" File="$(UgsCmdMacX64ZipFile)"/>
-			<DeployTool Id="ugs-mac-arm64" Version="$(InformationalVersion)" Settings="$(HordeSettings)" File="$(UgsCmdMacArm64ZipFile)"/>
-		</Node>
-		
-		<!-- Linux -->
-		<Node Name="Build UnrealGameSyncCmd Linux">
-			<Delete Files="$(UgsCmdLinuxOutputDir)/..."/>
-			<DotNet Arguments="publish &quot;$(UgsDir)/UnrealGameSyncCmd/UnrealGameSyncCmd.csproj&quot; --output=&quot;$(UgsCmdLinuxOutputDir)&quot; --self-contained -c Release -r linux-x64 $(VersionArguments)"/>
-			<Zip ZipFile="$(UgsCmdLinuxZipFile)" FromDir="$(UgsCmdLinuxOutputDir)" ExecutableFiles="ugs;*.so;*.sh"/>
-		</Node>
-		<Node Name="Publish UnrealGameSyncCmd Linux" Requires="Build UnrealGameSyncCmd Linux">
-			<Error Message="Missing -HordeSettings=... argument" If="'$(HordeSettings)' == ''"/>
-			<DeployTool Id="ugs-linux" Version="$(InformationalVersion)" Settings="$(HordeSettings)" File="$(UgsCmdLinuxZipFile)"/>
+			<Error Message="Missing -HordeSettingsMac=... argument" If="'$(HordeSettingsMac)' == ''"/>
+			<DeployTool Id="ugs-mac" Version="$(InformationalVersion)" Settings="$(HordeSettingsMac)" File="$(UgsCmdMacX64ZipFile)"/>
+			<DeployTool Id="ugs-mac-arm64" Version="$(InformationalVersion)" Settings="$(HordeSettingsMac)" File="$(UgsCmdMacArm64ZipFile)"/>
 		</Node>
+
 	</Agent>
 
 	<Aggregate Name="Publish Unstable With Zip" Requires="Publish Unstable;Publish Unstable Zip"/>
