PR #11397: Unreal Game Sync fixes

diff --git a/Engine/Source/Programs/UnrealGameSync/BuildUGS.xml b/Engine/Source/Programs/UnrealGameSync/BuildUGS.xml
--- a/Engine/Source/Programs/UnrealGameSync/BuildUGS.xml
+++ b/Engine/Source/Programs/UnrealGameSync/BuildUGS.xml
@@ -85,7 +85,7 @@
 			<Zip FromDir="$(UgsOutputDir)" ZipFile="$(UgsZipPublishDir)/Release.zip"/>
 			<Submit Description="Updating UGS from CL $(Change) (Zip)" Workspace="Temp" Files="$(UgsZipPublishDir)/Release.zip" RootDir="$(UgsZipOutputDir)" Stream="$(SubmitToStream)" Branch="$(SubmitToBranch)"/>
 		</Node>
-		<Node Name="Publish Live Horde" Requires="Build UnrealGameSync">
+		<Node Name="Publish Live Horde" Requires="Build UnrealGameSync Zip">
 			<Error Message="Missing -HordeSettings=... argument" If="'$(HordeSettings)' == ''"/>
 			<DeployTool Id="ugs-win" Version="$(InformationalVersion)" Settings="$(HordeSettings)" File="$(UgsZipPublishDir)/Ugs.zip"/>
 		</Node>
diff --git a/Engine/Source/Programs/UnrealGameSync/UnrealGameSync/IssueMonitor.cs b/Engine/Source/Programs/UnrealGameSync/UnrealGameSync/IssueMonitor.cs
--- a/Engine/Source/Programs/UnrealGameSync/UnrealGameSync/IssueMonitor.cs
+++ b/Engine/Source/Programs/UnrealGameSync/UnrealGameSync/IssueMonitor.cs
@@ -320,7 +320,7 @@ namespace UnrealGameSync
 
 		async Task PollForUpdatesAsync(CancellationToken cancellationToken)
 		{
-			while (cancellationToken.IsCancellationRequested)
+			while (!cancellationToken.IsCancellationRequested)
 			{
 				Task refreshTask = _refreshEvent.Task;
 
