UGS: Fix tools list not showing up if connections to Horde fail.

diff --git a/Engine/Source/Programs/UnrealGameSync/UnrealGameSync/ToolUpdateMonitor.cs b/Engine/Source/Programs/UnrealGameSync/UnrealGameSync/ToolUpdateMonitor.cs
--- a/Engine/Source/Programs/UnrealGameSync/UnrealGameSync/ToolUpdateMonitor.cs
+++ b/Engine/Source/Programs/UnrealGameSync/UnrealGameSync/ToolUpdateMonitor.cs
@@ -217,13 +217,27 @@ namespace UnrealGameSync
 			List<ToolInfo> tools = new List<ToolInfo>();
 			if (!String.IsNullOrEmpty(DeploymentSettings.Instance.ToolsDepotPath))
 			{
-				await ReadPerforceToolsAsync(perforce, tools, cancellationToken);
+				try
+				{
+					await ReadPerforceToolsAsync(perforce, tools, cancellationToken);
+				}
+				catch (Exception ex)
+				{
+					_logger.LogWarning(ex, "Error while polling Perforce for available tools: {Message}", ex.Message);
+				}
 			}
 			using (HordeHttpClient? hordeHttpClient = _serviceProvider.GetService<HordeHttpClient>())
 			{
 				if (hordeHttpClient != null)
 				{
-					await ReadHordeToolsAsync(hordeHttpClient, tools, cancellationToken);
+					try
+					{
+						await ReadHordeToolsAsync(hordeHttpClient, tools, cancellationToken);
+					}
+					catch (Exception ex)
+					{
+						_logger.LogWarning(ex, "Error while polling Horde for available tools: {Message}", ex.Message);
+					}
 				}
 			}
 			_tools = tools;
