UGS: Always open the launcher dialog if the -settings parameter is included, even in debug builds.

diff --git a/Engine/Source/Programs/UnrealGameSync/UnrealGameSync/Program.cs b/Engine/Source/Programs/UnrealGameSync/UnrealGameSync/Program.cs
--- a/Engine/Source/Programs/UnrealGameSync/UnrealGameSync/Program.cs
+++ b/Engine/Source/Programs/UnrealGameSync/UnrealGameSync/Program.cs
@@ -97,7 +97,7 @@ namespace UnrealGameSync
 				using (EventWaitHandle activateEvent = new EventWaitHandle(false, EventResetMode.AutoReset, "ActivateUnrealGameSync"))
 				{
 #if DEBUG
-					bool runUpdateCheck = args.Contains("-UpdateCheck", StringComparer.OrdinalIgnoreCase);
+					bool runUpdateCheck = args.Contains("-UpdateCheck", StringComparer.OrdinalIgnoreCase) || args.Contains("-Settings", StringComparer.OrdinalIgnoreCase);
 #else
 					bool runUpdateCheck = !args.Contains("-NoUpdateCheck", StringComparer.OrdinalIgnoreCase);
 #endif
