UGS: Fix Horde connection errors due to invalid configuration.

diff --git a/Engine/Source/Programs/UnrealGameSync/UnrealGameSync/Launcher.cs b/Engine/Source/Programs/UnrealGameSync/UnrealGameSync/Launcher.cs
--- a/Engine/Source/Programs/UnrealGameSync/UnrealGameSync/Launcher.cs
+++ b/Engine/Source/Programs/UnrealGameSync/UnrealGameSync/Launcher.cs
@@ -119,7 +119,7 @@ namespace UnrealGameSync
 
 					ServiceCollection services = new ServiceCollection();
 					services.AddLogging(builder => builder.AddProvider(new LoggerProviderAdapter(logger)));
-					services.AddHordeHttpClient((sp, client) => client.BaseAddress = hordeServerUrl);
+					services.AddHorde(options => options.ServerUrl = hordeServerUrl);
 
 					await using ServiceProvider serviceProvider = services.BuildServiceProvider();
 					HordeHttpClient httpClient = serviceProvider.GetRequiredService<HordeHttpClient>();
diff --git a/Engine/Source/Programs/UnrealGameSync/UnrealGameSync/Program.cs b/Engine/Source/Programs/UnrealGameSync/UnrealGameSync/Program.cs
--- a/Engine/Source/Programs/UnrealGameSync/UnrealGameSync/Program.cs
+++ b/Engine/Source/Programs/UnrealGameSync/UnrealGameSync/Program.cs
@@ -198,7 +198,7 @@ namespace UnrealGameSync
 
 				if (launcherSettings.HordeServer != null)
 				{
-					services.AddHordeHttpClient(x => x.BaseAddress = new Uri(launcherSettings.HordeServer));
+					services.AddHorde(options => options.ServerUrl = new Uri(launcherSettings.HordeServer));
 				}
 
 				await using (ServiceProvider serviceProvider = services.BuildServiceProvider())
diff --git a/Engine/Source/Programs/UnrealGameSync/UnrealGameSync/ToolUpdateMonitor.cs b/Engine/Source/Programs/UnrealGameSync/UnrealGameSync/ToolUpdateMonitor.cs
--- a/Engine/Source/Programs/UnrealGameSync/UnrealGameSync/ToolUpdateMonitor.cs
+++ b/Engine/Source/Programs/UnrealGameSync/UnrealGameSync/ToolUpdateMonitor.cs
@@ -221,7 +221,7 @@ namespace UnrealGameSync
 				{
 					await ReadPerforceToolsAsync(perforce, tools, cancellationToken);
 				}
-				catch (Exception ex)
+				catch (Exception ex) when (ex is not OperationCanceledException)
 				{
 					_logger.LogWarning(ex, "Error while polling Perforce for available tools: {Message}", ex.Message);
 				}
@@ -234,7 +234,7 @@ namespace UnrealGameSync
 					{
 						await ReadHordeToolsAsync(hordeHttpClient, tools, cancellationToken);
 					}
-					catch (Exception ex)
+					catch (Exception ex) when (ex is not OperationCanceledException)
 					{
 						_logger.LogWarning(ex, "Error while polling Horde for available tools: {Message}", ex.Message);
 					}
