UGS: Set a WindowsFormsSynchronizationContext before spawning the main async task, to ensure that all tasks on the UI thread are routed back to it.

diff --git a/Engine/Source/Programs/UnrealGameSync/UnrealGameSync/Program.cs b/Engine/Source/Programs/UnrealGameSync/UnrealGameSync/Program.cs
--- a/Engine/Source/Programs/UnrealGameSync/UnrealGameSync/Program.cs
+++ b/Engine/Source/Programs/UnrealGameSync/UnrealGameSync/Program.cs
@@ -94,6 +94,8 @@ namespace UnrealGameSync
 				Application.SetCompatibleTextRenderingDefault(false);
 				Application.SetHighDpiMode(HighDpiMode.PerMonitorV2);
 
+				SynchronizationContext.SetSynchronizationContext(new WindowsFormsSynchronizationContext());
+
 				using (EventWaitHandle activateEvent = new EventWaitHandle(false, EventResetMode.AutoReset, "ActivateUnrealGameSync"))
 				{
 #if DEBUG
