UGS: Add a deployment setting for where to update UGS from.

#okforgthub public

diff --git a/Engine/Source/Programs/UnrealGameSync/UnrealGameSync/Launcher.cs b/Engine/Source/Programs/UnrealGameSync/UnrealGameSync/Launcher.cs
--- a/Engine/Source/Programs/UnrealGameSync/UnrealGameSync/Launcher.cs
+++ b/Engine/Source/Programs/UnrealGameSync/UnrealGameSync/Launcher.cs
@@ -30,12 +30,11 @@ namespace UnrealGameSync
 
 			// Read the settings
 			LauncherSettings launcherSettings = new LauncherSettings();
-			launcherSettings.PerforceDepotPath = DeploymentSettings.Instance.DefaultDepotPath;
 			launcherSettings.Read();
 
 			// If the shift key is held down, immediately show the settings window
 			Task SyncAndRunWrapper(IPerforceConnection? perforce, LauncherSettings settings, ILogger logWriter, CancellationToken cancellationToken) => SyncAndRun(perforce, settings, args, instanceMutex, logWriter, cancellationToken);
-			if ((Control.ModifierKeys & Keys.Shift) != 0 || openSettings)
+			if ((Control.ModifierKeys & Keys.Shift) != 0 || openSettings || launcherSettings.UpdateSource == LauncherUpdateSource.Unknown)
 			{
 				// Show the settings window immediately
 				using UpdateSettingsWindow updateWindow = new UpdateSettingsWindow(null, null, launcherSettings, SyncAndRunWrapper);
diff --git a/Engine/Source/Programs/UnrealGameSync/UnrealGameSyncShared/DeploymentSettings.cs b/Engine/Source/Programs/UnrealGameSync/UnrealGameSyncShared/DeploymentSettings.cs
--- a/Engine/Source/Programs/UnrealGameSync/UnrealGameSyncShared/DeploymentSettings.cs
+++ b/Engine/Source/Programs/UnrealGameSync/UnrealGameSyncShared/DeploymentSettings.cs
@@ -13,6 +13,7 @@ using System.Diagnostics.CodeAnalysis;
 using System.IO;
 using System.Reflection;
 using System.Text.Json;
+using System.Text.Json.Serialization;
 using System.Xml;
 
 namespace UnrealGameSync
@@ -22,6 +23,11 @@ namespace UnrealGameSync
 	/// </summary>
 	public partial class DeploymentSettings
 	{
+		/// <summary>
+		/// Default update source for UGS
+		/// </summary>
+		public LauncherUpdateSource UpdateSource { get; set; }
+
 		/// <summary>
 		/// Url for the Horde server
 		/// </summary>
@@ -85,6 +91,7 @@ namespace UnrealGameSync
 					{
 						byte[] data = FileReference.ReadAllBytes(settingsFile);
 						JsonSerializerOptions options = new JsonSerializerOptions { AllowTrailingCommas = true, PropertyNameCaseInsensitive = true, ReadCommentHandling = JsonCommentHandling.Skip };
+						options.Converters.Add(new JsonStringEnumConverter());
 						s_instance = JsonSerializer.Deserialize<DeploymentSettings>(data, options);
 					}
 
diff --git a/Engine/Source/Programs/UnrealGameSync/UnrealGameSyncShared/LauncherSettings.cs b/Engine/Source/Programs/UnrealGameSync/UnrealGameSyncShared/LauncherSettings.cs
--- a/Engine/Source/Programs/UnrealGameSync/UnrealGameSyncShared/LauncherSettings.cs
+++ b/Engine/Source/Programs/UnrealGameSync/UnrealGameSyncShared/LauncherSettings.cs
@@ -10,8 +10,9 @@ namespace UnrealGameSync
 {
 	public enum LauncherUpdateSource
 	{
-		Perforce = 0,
-		Horde = 1,
+		Unknown = 0,
+		Perforce = 1,
+		Horde = 2,
 	}
 
 	public class LauncherSettings
@@ -28,6 +29,7 @@ namespace UnrealGameSync
 
 		public LauncherSettings()
 		{
+			UpdateSource = DeploymentSettings.Instance.UpdateSource;
 			HordeServer = DeploymentSettings.Instance.HordeUrl;
 			PerforceDepotPath = DeploymentSettings.Instance.DefaultDepotPath;
 		}
@@ -106,7 +108,7 @@ namespace UnrealGameSync
 				// Delete this legacy setting
 				Utility.DeleteRegistryKey(key, "Server");
 
-				SaveRegistryValue(key, "Source", UpdateSource.ToString(), LauncherUpdateSource.Perforce.ToString());
+				SaveRegistryValue(key, "Source", UpdateSource.ToString(), DeploymentSettings.Instance.UpdateSource.ToString());
 				SaveRegistryValue(key, "HordeServer", HordeServer, DeploymentSettings.Instance.HordeUrl);
 				SaveRegistryValue(key, "ServerAndPort", PerforceServerAndPort, null);
 				SaveRegistryValue(key, "UserName", PerforceUserName, null);
