UGS: Speculative fix for hang on startup. Now runs the message loop on the main thread, and creates async tasks to dispose of any IAsyncDisposables.

diff --git a/Engine/Source/Programs/UnrealGameSync/UnrealGameSync/Program.cs b/Engine/Source/Programs/UnrealGameSync/UnrealGameSync/Program.cs
--- a/Engine/Source/Programs/UnrealGameSync/UnrealGameSync/Program.cs
+++ b/Engine/Source/Programs/UnrealGameSync/UnrealGameSync/Program.cs
@@ -119,7 +119,7 @@ namespace UnrealGameSync
 					// Launch the application proper
 					if (firstInstance)
 					{
-						InnerMainAsync(instanceMutex, activateEvent, args, runUpdateCheck).GetAwaiter().GetResult();
+						InnerMain(instanceMutex, activateEvent, args, runUpdateCheck);
 					}
 					else
 					{
@@ -129,7 +129,7 @@ namespace UnrealGameSync
 			}
 		}
 
-		static async Task InnerMainAsync(Mutex instanceMutex, EventWaitHandle activateEvent, string[] args, bool runUpdateCheck)
+		static void InnerMain(Mutex instanceMutex, EventWaitHandle activateEvent, string[] args, bool runUpdateCheck)
 		{
 			LauncherSettings launcherSettings = new LauncherSettings();
 			launcherSettings.Read();
@@ -178,7 +178,7 @@ namespace UnrealGameSync
 			{
 				try
 				{
-					SyncVersion = (await File.ReadAllTextAsync(syncVersionFile)).Trim();
+					SyncVersion = File.ReadAllText(syncVersionFile).Trim();
 				}
 				catch(Exception)
 				{
@@ -203,7 +203,8 @@ namespace UnrealGameSync
 					services.AddHorde(options => options.ServerUrl = new Uri(launcherSettings.HordeServer));
 				}
 
-				await using (ServiceProvider serviceProvider = services.BuildServiceProvider())
+				ServiceProvider serviceProvider = services.BuildServiceProvider();
+				try
 				{
 					ILoggerFactory loggerFactory = serviceProvider.GetRequiredService<ILoggerFactory>();
 
@@ -242,7 +243,8 @@ namespace UnrealGameSync
 
 							ProtocolHandlerUtils.InstallQuiet(logger);
 
-							await using (UpdateMonitor updateMonitor = CreateUpdateMonitor(launcherSettings, defaultSettings, updatePath, runUpdateCheck, serviceProvider))
+							UpdateMonitor updateMonitor = CreateUpdateMonitor(launcherSettings, defaultSettings, updatePath, runUpdateCheck, serviceProvider);
+							try
 							{
 								using ProgramApplicationContext context = new ProgramApplicationContext(defaultSettings, updateMonitor, DeploymentSettings.Instance.ApiUrl, dataFolder, activateEvent, restoreState, updateSpawn, projectFileName, preview, serviceProvider, uri);
 								Application.Run(context);
@@ -253,6 +255,10 @@ namespace UnrealGameSync
 									Utility.SpawnProcess(updateSpawn, "-restorestate" + (updateMonitor.OpenSettings ? " -settings" : ""));
 								}
 							}
+							finally
+							{
+								AsyncDispose(updateMonitor);
+							}
 						}
 						catch (Exception ex)
 						{
@@ -265,9 +271,18 @@ namespace UnrealGameSync
 						}
 					}
 				}
+				finally
+				{
+					AsyncDispose(serviceProvider);
+				}
 			}
 		}
 
+		static void AsyncDispose(IAsyncDisposable disposable)
+		{
+			disposable.DisposeAsync().AsTask().ConfigureAwait(false).GetAwaiter().GetResult();
+		}
+
 		private static UpdateMonitor CreateUpdateMonitor(LauncherSettings launcherSettings, IPerforceSettings defaultSettings, string? updatePath, bool runUpdateCheck, IServiceProvider serviceProvider)
 		{
 			if (!runUpdateCheck)
