UGS: Upgrade NET versions for PostBadgeStatus and UnrealGameSyncAutomation. PostBadgeStatus is still on NET Framework due to some minor missing dependencies in NET Core.

#rnx

diff --git a/Engine/Source/Programs/UnrealGameSync/PostBadgeStatus/App.config b/Engine/Source/Programs/UnrealGameSync/PostBadgeStatus/App.config
--- a/Engine/Source/Programs/UnrealGameSync/PostBadgeStatus/App.config
+++ b/Engine/Source/Programs/UnrealGameSync/PostBadgeStatus/App.config
@@ -1,6 +1,6 @@
-﻿<?xml version="1.0" encoding="utf-8" ?>
+<?xml version="1.0" encoding="utf-8"?>
 <configuration>
     <startup> 
-        <supportedRuntime version="v4.0" sku=".NETFramework,Version=v4.5" />
+        <supportedRuntime version="v4.0" sku=".NETFramework,Version=v4.8"/>
     </startup>
-</configuration>
\ No newline at end of file
+</configuration>
diff --git a/Engine/Source/Programs/UnrealGameSync/PostBadgeStatus/PostBadgeStatus.csproj b/Engine/Source/Programs/UnrealGameSync/PostBadgeStatus/PostBadgeStatus.csproj
--- a/Engine/Source/Programs/UnrealGameSync/PostBadgeStatus/PostBadgeStatus.csproj
+++ b/Engine/Source/Programs/UnrealGameSync/PostBadgeStatus/PostBadgeStatus.csproj
@@ -9,8 +9,9 @@
     <AppDesignerFolder>Properties</AppDesignerFolder>
     <RootNamespace>PostBadgeStatus</RootNamespace>
     <AssemblyName>PostBadgeStatus</AssemblyName>
-    <TargetFrameworkVersion>v4.5</TargetFrameworkVersion>
+    <TargetFrameworkVersion>v4.8</TargetFrameworkVersion>
     <FileAlignment>512</FileAlignment>
+    <TargetFrameworkProfile />
   </PropertyGroup>
   <PropertyGroup Condition="'$(Configuration)|$(Platform)' == 'Debug|AnyCPU'">
     <DebugSymbols>true</DebugSymbols>
diff --git a/Engine/Source/Programs/UnrealGameSync/UnrealGameSyncAutomation/UnrealGameSyncAutomation.csproj b/Engine/Source/Programs/UnrealGameSync/UnrealGameSyncAutomation/UnrealGameSyncAutomation.csproj
--- a/Engine/Source/Programs/UnrealGameSync/UnrealGameSyncAutomation/UnrealGameSyncAutomation.csproj
+++ b/Engine/Source/Programs/UnrealGameSync/UnrealGameSyncAutomation/UnrealGameSyncAutomation.csproj
@@ -1,8 +1,8 @@
 <?xml version="1.0" encoding="utf-8"?>
-<Project Sdk="Microsoft.NET.Sdk.WindowsDesktop">
+<Project Sdk="Microsoft.NET.Sdk">
   <PropertyGroup>
     <OutputType>WinExe</OutputType>
-    <TargetFramework>netcoreapp3.1</TargetFramework>
+    <TargetFramework>net6.0-windows</TargetFramework>
     <UseWindowsForms>true</UseWindowsForms>
     <GenerateAssemblyInfo>false</GenerateAssemblyInfo>
     <RootNamespace>UnrealGameSyncAutomation</RootNamespace>
