Fixing a parsing bug in UGS that would cause us to split RequiredBadges (used in defining LatestChangeType entries in UnrealGameSync.ini) that would cause us to incorrectly split a badge name that contained spaces into multiple badge names.

diff --git a/Engine/Source/Programs/UnrealGameSync/UnrealGameSyncShared/UserSettings.cs b/Engine/Source/Programs/UnrealGameSync/UnrealGameSyncShared/UserSettings.cs
--- a/Engine/Source/Programs/UnrealGameSync/UnrealGameSyncShared/UserSettings.cs
+++ b/Engine/Source/Programs/UnrealGameSync/UnrealGameSyncShared/UserSettings.cs
@@ -85,7 +85,7 @@ namespace UnrealGameSync
 				changeType.Good = definitionObject.GetValue("bGood", false);
 				changeType.Starred = definitionObject.GetValue("bStarred", false);
 				changeType.FindNewestGoodContent = definitionObject.GetValue("bFindNewestGoodContent", false);
-				changeType.RequiredBadges.AddRange(definitionObject.GetValue("RequiredBadges", "").Split(new char[] { ',',' '}, StringSplitOptions.RemoveEmptyEntries));
+				changeType.RequiredBadges.AddRange(definitionObject.GetValue("RequiredBadges", "").Split(new char[] { ',' }, StringSplitOptions.RemoveEmptyEntries | StringSplitOptions.TrimEntries));
 
 				changeType.ReadFrom = definitionObject.GetValue("ReadFrom", null);
 				return true;
