UGS: Fix autoupdate not working from UGS builds distributed with Horde.

Also add separate DebugAutoUpdate/ReleaseAutoUpdate configurations for enabling autoupdate logic in UGS, making it easier to test under visual studio.

#rnx

diff --git a/Engine/Source/Programs/Horde/BuildHorde.xml b/Engine/Source/Programs/Horde/BuildHorde.xml
--- a/Engine/Source/Programs/Horde/BuildHorde.xml
+++ b/Engine/Source/Programs/Horde/BuildHorde.xml
@@ -135,8 +135,8 @@
 		</Node>
 		<Node Name="Build UnrealGameSync" Produces="#UnrealGameSyncFiles">
 			<Delete Files="$(ServerToolsStagingDir_UnrealGameSync)/..."/>
-			<DotNet Arguments="publish &quot;$(UgsDir)/UnrealGameSyncCmd/UnrealGameSyncCmd.csproj&quot; --output=&quot;$(ServerToolsStagingDir_UnrealGameSync)&quot; --self-contained -c Release -r win-x64 $(VersionArguments)"/>
-			<DotNet Arguments="publish &quot;$(UgsDir)/UnrealGameSync/UnrealGameSync.csproj&quot; --output=&quot;$(ServerToolsStagingDir_UnrealGameSync)&quot; --self-contained -c Release -r win-x64 $(VersionArguments)"/>
+			<DotNet Arguments="publish &quot;$(UgsDir)/UnrealGameSyncCmd/UnrealGameSyncCmd.csproj&quot; --output=&quot;$(ServerToolsStagingDir_UnrealGameSync)&quot; --self-contained -c ReleaseAutoUpdate -r win-x64 $(VersionArguments)"/>
+			<DotNet Arguments="publish &quot;$(UgsDir)/UnrealGameSync/UnrealGameSync.csproj&quot; --output=&quot;$(ServerToolsStagingDir_UnrealGameSync)&quot; --self-contained -c ReleaseAutoUpdate -r win-x64 $(VersionArguments)"/>
 			<Sign Files="$(ServerToolsStagingDir_UnrealGameSync)/UnrealGameSync.exe;$(ServerToolsStagingDir_UnrealGameSync)/ugs.exe;$(ServerToolsStagingDir_UnrealGameSync)/ugs.dll;$(ServerToolsStagingDir_UnrealGameSync)/UnrealGameSync*.dll;$(ServerToolsStagingDir_UnrealGameSync)/EpicGames.*.dll" If="'$(SignOutputs)' == 'true'"/>
 			<Tag Files="$(ServerToolsStagingDir_UnrealGameSync)/..." With="#UnrealGameSyncFiles"/>
 		</Node>
diff --git a/Engine/Source/Programs/UnrealGameSync/UnrealGameSync/Launcher.cs b/Engine/Source/Programs/UnrealGameSync/UnrealGameSync/Launcher.cs
--- a/Engine/Source/Programs/UnrealGameSync/UnrealGameSync/Launcher.cs
+++ b/Engine/Source/Programs/UnrealGameSync/UnrealGameSync/Launcher.cs
@@ -32,8 +32,13 @@ namespace UnrealGameSync
 			LauncherSettings launcherSettings = new LauncherSettings();
 			launcherSettings.Read();
 
+			async Task SyncAndRunWrapper(IPerforceConnection? perforce, LauncherSettings settings, ILogger logWriter, CancellationToken cancellationToken)
+			{
+				List<string> childArgs = args.Except(new[] { "-settings", "-updatecheck", "-noupdatecheck" }, StringComparer.OrdinalIgnoreCase).ToList();
+				await SyncAndRun(perforce, settings, args, instanceMutex, logWriter, cancellationToken);
+			}
+
 			// If the shift key is held down, immediately show the settings window
-			Task SyncAndRunWrapper(IPerforceConnection? perforce, LauncherSettings settings, ILogger logWriter, CancellationToken cancellationToken) => SyncAndRun(perforce, settings, args, instanceMutex, logWriter, cancellationToken);
 			if ((Control.ModifierKeys & Keys.Shift) != 0 || openSettings || launcherSettings.UpdateSource == LauncherUpdateSource.Unknown)
 			{
 				// Show the settings window immediately
@@ -255,20 +260,9 @@ namespace UnrealGameSync
 				}
 				logger.LogInformation("");
 
-				// Build the command line for the synced application, including the sync path to monitor for updates
-				string originalExecutable = Assembly.GetEntryAssembly()!.Location;
-                if (Path.GetExtension(originalExecutable).Equals(".dll", StringComparison.OrdinalIgnoreCase))
-                {
-                    string newExecutable = Path.ChangeExtension(originalExecutable, ".exe");
-                    if (File.Exists(newExecutable))
-                    {
-                        originalExecutable = newExecutable;
-                    }
-                }
-
 				// Create the new argument list
 				List<string> newArguments = new List<string>(args);
-				newArguments.Add($"-updatespawn={originalExecutable}");
+				newArguments.Add($"-updatespawn={Program.GetCurrentExecutable()}");
 				if (updatePath != null)
 				{
 					newArguments.Add($"-updatepath={updatePath}");
diff --git a/Engine/Source/Programs/UnrealGameSync/UnrealGameSync/Program.cs b/Engine/Source/Programs/UnrealGameSync/UnrealGameSync/Program.cs
--- a/Engine/Source/Programs/UnrealGameSync/UnrealGameSync/Program.cs
+++ b/Engine/Source/Programs/UnrealGameSync/UnrealGameSync/Program.cs
@@ -98,11 +98,7 @@ namespace UnrealGameSync
 
 				using (EventWaitHandle activateEvent = new EventWaitHandle(false, EventResetMode.AutoReset, "ActivateUnrealGameSync"))
 				{
-#if DEBUG
-					bool runUpdateCheck = args.Contains("-UpdateCheck", StringComparer.OrdinalIgnoreCase) || args.Contains("-Settings", StringComparer.OrdinalIgnoreCase);
-#else
-					bool runUpdateCheck = !args.Contains("-NoUpdateCheck", StringComparer.OrdinalIgnoreCase);
-#endif
+					bool runUpdateCheck = ShouldRunAutoUpdate(args);
 
 					// Check for a newer version of the application
 					if (runUpdateCheck && Launcher.SyncAndRunLatest(instanceMutex, args))
@@ -249,10 +245,10 @@ namespace UnrealGameSync
 								using ProgramApplicationContext context = new ProgramApplicationContext(defaultSettings, updateMonitor, DeploymentSettings.Instance.ApiUrl, dataFolder, activateEvent, restoreState, updateSpawn, projectFileName, preview, serviceProvider, uri);
 								Application.Run(context);
 
-								if (updateMonitor.IsUpdateAvailable && updateSpawn != null)
+								if (updateMonitor.IsUpdateAvailable)
 								{
 									instanceMutex.Close();
-									Utility.SpawnProcess(updateSpawn, "-restorestate" + (updateMonitor.OpenSettings ? " -settings" : ""));
+									Utility.SpawnProcess(updateSpawn ?? GetCurrentExecutable(), "-restorestate" + (updateMonitor.OpenSettings ? " -settings" : ""));
 								}
 							}
 							finally
@@ -283,6 +279,29 @@ namespace UnrealGameSync
 			disposable.DisposeAsync().AsTask().ConfigureAwait(false).GetAwaiter().GetResult();
 		}
 
+		public static string GetCurrentExecutable()
+		{
+			string originalExecutable = Assembly.GetEntryAssembly()!.Location;
+			if (Path.GetExtension(originalExecutable).Equals(".dll", StringComparison.OrdinalIgnoreCase))
+			{
+				string newExecutable = Path.ChangeExtension(originalExecutable, ".exe");
+				if (File.Exists(newExecutable))
+				{
+					return newExecutable;
+				}
+			}
+			return originalExecutable;
+		}
+
+		static bool ShouldRunAutoUpdate(string[] args)
+		{
+#if WITH_AUTOUPDATE
+			return !args.Contains("-NoUpdateCheck", StringComparer.OrdinalIgnoreCase);
+#else
+			return args.Contains("-UpdateCheck", StringComparer.OrdinalIgnoreCase) || args.Contains("-Settings", StringComparer.OrdinalIgnoreCase);
+#endif
+		}
+
 		private static UpdateMonitor CreateUpdateMonitor(LauncherSettings launcherSettings, IPerforceSettings defaultSettings, string? updatePath, bool runUpdateCheck, IServiceProvider serviceProvider)
 		{
 			if (!runUpdateCheck)
diff --git a/Engine/Source/Programs/UnrealGameSync/UnrealGameSync/UnrealGameSync.csproj b/Engine/Source/Programs/UnrealGameSync/UnrealGameSync/UnrealGameSync.csproj
--- a/Engine/Source/Programs/UnrealGameSync/UnrealGameSync/UnrealGameSync.csproj
+++ b/Engine/Source/Programs/UnrealGameSync/UnrealGameSync/UnrealGameSync.csproj
@@ -15,7 +15,7 @@
     <Nullable>enable</Nullable>
     <NeutralLanguage>en-US</NeutralLanguage>
     <RunAnalyzersDuringBuild>False</RunAnalyzersDuringBuild>
-    <Configurations>Debug;Release;Analyze</Configurations>
+    <Configurations>Debug;Release;Analyze;DebugAutoUpdate;ReleaseAutoUpdate</Configurations>
   </PropertyGroup>
 
   <PropertyGroup Condition=" '$(Configuration)|$(Platform)' == 'Analyze|AnyCPU' ">
@@ -32,6 +32,9 @@
   <PropertyGroup Condition="'$(Configuration)|$(Platform)'=='Debug|AnyCPU'">
     <DefineConstants>TRACE;WITH_TELEMETRY</DefineConstants>
   </PropertyGroup>
+  <PropertyGroup Condition=" '$(Configuration)' == 'DebugAutoUpdate' or '$(Configuration)' == 'ReleaseAutoUpdate'">
+    <DefineConstants>$(DefineConstants);WITH_AUTOUPDATE</DefineConstants>
+  </PropertyGroup>
   <ItemGroup>
     <PackageReference Include="EnvDTE" Version="16.7.30328.74" />
     <PackageReference Include="Sentry" Version="3.15.0" />
