Horde: Disable auto-update logic for instances spawned after an auto-update operation.

diff --git a/Engine/Source/Programs/UnrealGameSync/UnrealGameSync/Launcher.cs b/Engine/Source/Programs/UnrealGameSync/UnrealGameSync/Launcher.cs
--- a/Engine/Source/Programs/UnrealGameSync/UnrealGameSync/Launcher.cs
+++ b/Engine/Source/Programs/UnrealGameSync/UnrealGameSync/Launcher.cs
@@ -24,6 +24,12 @@ namespace UnrealGameSync
 	{
 		public static bool SyncAndRunLatest(Mutex instanceMutex, string[] args)
 		{
+			// Don't do this if we're already running as a spawned instance
+			if (args.Any(x => x.StartsWith("-updatespawn=", StringComparison.OrdinalIgnoreCase)))
+			{
+				return false;
+			}
+
 			// Figure out if we should sync the unstable build by default
 			bool preview = args.Contains("-unstable", StringComparer.InvariantCultureIgnoreCase) || args.Contains("-preview", StringComparer.InvariantCultureIgnoreCase);
 			bool openSettings = args.Contains("-settings", StringComparer.OrdinalIgnoreCase);
