UGS: Fix tools not being enumerated correctly when running against a Horde server.

diff --git a/Engine/Source/Programs/UnrealGameSync/UnrealGameSync/ToolUpdateMonitor.cs b/Engine/Source/Programs/UnrealGameSync/UnrealGameSync/ToolUpdateMonitor.cs
--- a/Engine/Source/Programs/UnrealGameSync/UnrealGameSync/ToolUpdateMonitor.cs
+++ b/Engine/Source/Programs/UnrealGameSync/UnrealGameSync/ToolUpdateMonitor.cs
@@ -130,10 +130,7 @@ namespace UnrealGameSync
 
 		public void Start()
 		{
-			if (DeploymentSettings.Instance.ToolsDepotPath != null)
-			{
-				_workerTask = Task.Run(() => PollForUpdatesAsync(_cancellationSource.Token));
-			}
+			_workerTask = Task.Run(() => PollForUpdatesAsync(_cancellationSource.Token));
 		}
 
 		public void Dispose()
@@ -348,7 +345,7 @@ namespace UnrealGameSync
 				{
 					IoHash hash = IoHash.Compute(Encoding.UTF8.GetBytes($"horde:{toolResponse.Id}"));
 					Guid guid = new Guid(hash.ToByteArray().AsSpan(0, 16));
-					ToolInfo toolInfo = new ToolInfo(guid, toolResponse.Name, toolResponse.Description, $"{hordeHttpClient},{toolResponse.Id},{toolResponse.DeploymentId}");
+					ToolInfo toolInfo = new ToolInfo(guid, toolResponse.Id.ToString(), toolResponse.Name, $"{hordeHttpClient},{toolResponse.Id},{toolResponse.DeploymentId}");
 					tools.Add(toolInfo);
 				}
 			}
