Horde: Store the URL of the default Horde server in the registry on Windows.

* Both HKLM\SOFTWARE\Epic Games\Horde\Url and HKCU\SOFTWARE\Epic Games\Horde\Url are read from (HKCU taking precedence).
* Agent installer uses this value by default.
* EpicGames.Horde uses this value by default.
* UGS reads/writes this setting from Application Settings dialog (in HKCU).
* Default server profile is created on demand, uses this value for the default server if not specified elsewhere.

#rnx

diff --git a/Engine/Source/Programs/Horde/Horde.Agent/AgentSettings.cs b/Engine/Source/Programs/Horde/Horde.Agent/AgentSettings.cs
--- a/Engine/Source/Programs/Horde/Horde.Agent/AgentSettings.cs
+++ b/Engine/Source/Programs/Horde/Horde.Agent/AgentSettings.cs
@@ -1,7 +1,9 @@
 // Copyright Epic Games, Inc. All Rights Reserved.
 
 using System.ComponentModel.DataAnnotations;
+using System.Text.Json.Serialization;
 using EpicGames.Core;
+using EpicGames.Horde;
 using EpicGames.Horde.Tools;
 using Microsoft.Extensions.Configuration;
 
@@ -346,14 +348,13 @@ namespace Horde.Agent
 		{
 			if (Server == null)
 			{
-				if (ServerProfiles.Count == 1)
-				{
-					return ServerProfiles.Values.First();
-				}
-				else
-				{
-					throw new Exception($"No server profile is specified. Pass the -server=... argument on the command line or set the 'Server' property in the agent config file to one of: {GetServerProfileNames()}");
-				}
+				Uri? defaultServerUrl = Installed ? HordeOptions.GetDefaultServerUrl() : null;
+
+				ServerProfile defaultServerProfile = new ServerProfile();
+				defaultServerProfile.Name = "Default";
+				defaultServerProfile.Environment = "Development";
+				defaultServerProfile.Url = defaultServerUrl ?? new Uri("http://localhost:5000");
+				return defaultServerProfile;
 			}
 
 			return GetServerProfile(Server);
diff --git a/Engine/Source/Programs/Horde/Horde.Agent/appsettings.Local.json b/Engine/Source/Programs/Horde/Horde.Agent/appsettings.Local.json
--- a/Engine/Source/Programs/Horde/Horde.Agent/appsettings.Local.json
+++ b/Engine/Source/Programs/Horde/Horde.Agent/appsettings.Local.json
@@ -1,15 +1,6 @@
 {
 	"Horde": 
 	{
-		"Server": "Local", // Entry in the ServerProfiles dictionary
-		"ServerProfiles":
-		{
-			"Local": 
-			{
-				"Environment": "dev",
-				"Url": "http://localhost:5000"
-			}
-		},
 		"Installed": false,
 		"EnableAwsEc2Support": false,
 		"WriteStepOutputToLogger": true
diff --git a/Engine/Source/Programs/Shared/EpicGames.Horde/HordeHttpClient.cs b/Engine/Source/Programs/Shared/EpicGames.Horde/HordeHttpClient.cs
--- a/Engine/Source/Programs/Shared/EpicGames.Horde/HordeHttpClient.cs
+++ b/Engine/Source/Programs/Shared/EpicGames.Horde/HordeHttpClient.cs
@@ -507,6 +507,9 @@ namespace EpicGames.Horde
 					}
 				}
 
+				// Try to get the default server address from the registry
+				httpClient.BaseAddress ??= HordeOptions.GetDefaultServerUrl();
+
 				// Make sure we have a base URL set
 				if (httpClient.BaseAddress == null)
 				{
diff --git a/Engine/Source/Programs/Shared/EpicGames.Horde/HordeOptions.cs b/Engine/Source/Programs/Shared/EpicGames.Horde/HordeOptions.cs
--- a/Engine/Source/Programs/Shared/EpicGames.Horde/HordeOptions.cs
+++ b/Engine/Source/Programs/Shared/EpicGames.Horde/HordeOptions.cs
@@ -1,9 +1,13 @@
 // Copyright Epic Games, Inc. All Rights Reserved.
 
 using System;
+using System.Linq;
 using System.Net.Http;
+using System.Runtime.InteropServices;
+using System.Runtime.Versioning;
 using EpicGames.Horde.Storage;
 using EpicGames.Horde.Storage.Bundles;
+using Microsoft.Win32;
 
 namespace EpicGames.Horde
 {
@@ -41,5 +45,70 @@ namespace EpicGames.Horde
 		/// Options for caching bundles 
 		/// </summary>
 		public BundleCacheOptions BundleCache { get; } = new BundleCacheOptions();
+
+		/// <summary>
+		/// Gets the default server URL for the current user
+		/// </summary>
+		/// <returns>Default URL</returns>
+		public static Uri? GetDefaultServerUrl()
+		{
+			if (RuntimeInformation.IsOSPlatform(OSPlatform.Windows))
+			{
+				string? url = null;
+				url ??= Registry.GetValue(@"HKEY_CURRENT_USER\SOFTWARE\Epic Games\Horde", "Url", null) as string;
+				url ??= Registry.GetValue(@"HKEY_LOCAL_MACHINE\SOFTWARE\Epic Games\Horde", "Url", null) as string;
+
+				if (!String.IsNullOrEmpty(url))
+				{
+					try
+					{
+						return new Uri(url);
+					}
+					catch (UriFormatException)
+					{
+					}
+				}
+			}
+			return null;
+		}
+
+		/// <summary>
+		/// Sets the default server url for the current user
+		/// </summary>
+		/// <param name="serverUrl">Horde server URL to use</param>
+		public static void SetDefaultServerUrl(Uri serverUrl)
+		{
+			if (RuntimeInformation.IsOSPlatform(OSPlatform.Windows))
+			{
+				string? newServerUrl = serverUrl.ToString();
+				string? defaultServerUrl = Registry.GetValue(@"HKEY_LOCAL_MACHINE\SOFTWARE\Epic Games\Horde", "Url", null) as string;
+
+				if (String.Equals(newServerUrl, defaultServerUrl, StringComparison.Ordinal))
+				{
+					using RegistryKey key = Registry.CurrentUser.CreateSubKey("SOFTWARE\\Epic Games\\Horde");
+					DeleteRegistryKey(key, "Url");
+				}
+				else
+				{
+					Registry.SetValue(@"HKEY_CURRENT_USER\SOFTWARE\Epic Games\Horde", "Url", serverUrl.ToString());
+				}
+			}
+		}
+
+		[SupportedOSPlatform("windows")]
+		static void DeleteRegistryKey(RegistryKey key, string name)
+		{
+			string[] valueNames = key.GetValueNames();
+			if (valueNames.Any(x => String.Equals(x, name, StringComparison.OrdinalIgnoreCase)))
+			{
+				try
+				{
+					key.DeleteValue(name);
+				}
+				catch
+				{
+				}
+			}
+		}
 	}
 }
diff --git a/Engine/Source/Programs/UnrealGameSync/UnrealGameSyncShared/LauncherSettings.cs b/Engine/Source/Programs/UnrealGameSync/UnrealGameSyncShared/LauncherSettings.cs
--- a/Engine/Source/Programs/UnrealGameSync/UnrealGameSyncShared/LauncherSettings.cs
+++ b/Engine/Source/Programs/UnrealGameSync/UnrealGameSyncShared/LauncherSettings.cs
@@ -1,10 +1,11 @@
 // Copyright Epic Games, Inc. All Rights Reserved.
 
-using Microsoft.Win32;
 using System;
 using System.Runtime.InteropServices;
 using System.Runtime.Versioning;
 using System.Text.RegularExpressions;
+using EpicGames.Horde;
+using Microsoft.Win32;
 
 namespace UnrealGameSync
 {
@@ -54,6 +55,12 @@ namespace UnrealGameSync
 		[SupportedOSPlatform("windows")]
 		void ReadFromRegistry()
 		{
+			Uri? defaultServerUrl = HordeOptions.GetDefaultServerUrl();
+			if (defaultServerUrl != null)
+			{
+				HordeServer = defaultServerUrl.ToString();
+			}
+
 			using (RegistryKey? key = Registry.CurrentUser.OpenSubKey("SOFTWARE\\Epic Games\\UnrealGameSync", false))
 			{
 				if (key != null)
@@ -64,11 +71,10 @@ namespace UnrealGameSync
 						UpdateSource = updateSource;
 					}
 
-					HordeServer = key.GetValue("HordeServer", HordeServer) as string;
 					PerforceServerAndPort = key.GetValue("ServerAndPort", PerforceServerAndPort) as string;
 					PerforceUserName = key.GetValue("UserName", PerforceUserName) as string;
 					PerforceDepotPath = key.GetValue("DepotPath", PerforceDepotPath) as string;
-					PreviewBuild = ((key.GetValue("Preview", PreviewBuild? 1 : 0) as int?) ?? 0) != 0;
+					PreviewBuild = ((key.GetValue("Preview", PreviewBuild ? 1 : 0) as int?) ?? 0) != 0;
 
 					// Fix corrupted depot path string
 					if (PerforceDepotPath != null)
@@ -103,13 +109,17 @@ namespace UnrealGameSync
 		[SupportedOSPlatform("windows")]
 		void SaveToRegistry()
 		{
+			if (!String.IsNullOrEmpty(HordeServer))
+			{
+				HordeOptions.SetDefaultServerUrl(new Uri(HordeServer));
+			}
+
 			using (RegistryKey key = Registry.CurrentUser.CreateSubKey("SOFTWARE\\Epic Games\\UnrealGameSync"))
 			{
 				// Delete this legacy setting
 				Utility.DeleteRegistryKey(key, "Server");
 
 				SaveRegistryValue(key, "Source", UpdateSource.ToString(), DeploymentSettings.Instance.UpdateSource.ToString());
-				SaveRegistryValue(key, "HordeServer", HordeServer, DeploymentSettings.Instance.HordeUrl);
 				SaveRegistryValue(key, "ServerAndPort", PerforceServerAndPort, null);
 				SaveRegistryValue(key, "UserName", PerforceUserName, null);
 				SaveRegistryValue(key, "DepotPath", PerforceDepotPath, DeploymentSettings.Instance.DefaultDepotPath);
