UGS: Use default Perforce server setting from Deployment.json.

#rnx

diff --git a/Engine/Source/Programs/UnrealGameSync/UnrealGameSyncShared/DeploymentSettings.cs b/Engine/Source/Programs/UnrealGameSync/UnrealGameSyncShared/DeploymentSettings.cs
--- a/Engine/Source/Programs/UnrealGameSync/UnrealGameSyncShared/DeploymentSettings.cs
+++ b/Engine/Source/Programs/UnrealGameSync/UnrealGameSyncShared/DeploymentSettings.cs
@@ -43,6 +43,11 @@ namespace UnrealGameSync
 		/// </summary>
 		public string? ApiUrl { get; set; }
 
+		/// <summary>
+		/// Default Perforce server to connect to
+		/// </summary>
+		public string? DefaultPerforceServer { get; set; }
+
 		/// <summary>
 		/// Servers to connect to for issue details by default
 		/// </summary>
diff --git a/Engine/Source/Programs/UnrealGameSync/UnrealGameSyncShared/LauncherSettings.cs b/Engine/Source/Programs/UnrealGameSync/UnrealGameSyncShared/LauncherSettings.cs
--- a/Engine/Source/Programs/UnrealGameSync/UnrealGameSyncShared/LauncherSettings.cs
+++ b/Engine/Source/Programs/UnrealGameSync/UnrealGameSyncShared/LauncherSettings.cs
@@ -32,6 +32,7 @@ namespace UnrealGameSync
 		{
 			UpdateSource = DeploymentSettings.Instance.UpdateSource;
 			HordeServer = DeploymentSettings.Instance.HordeUrl;
+			PerforceServerAndPort = DeploymentSettings.Instance.DefaultPerforceServer;
 			PerforceDepotPath = DeploymentSettings.Instance.DefaultDepotPath;
 		}
 
@@ -120,7 +121,7 @@ namespace UnrealGameSync
 				Utility.DeleteRegistryKey(key, "Server");
 
 				SaveRegistryValue(key, "Source", UpdateSource.ToString(), DeploymentSettings.Instance.UpdateSource.ToString());
-				SaveRegistryValue(key, "ServerAndPort", PerforceServerAndPort, null);
+				SaveRegistryValue(key, "ServerAndPort", PerforceServerAndPort, DeploymentSettings.Instance.DefaultPerforceServer);
 				SaveRegistryValue(key, "UserName", PerforceUserName, null);
 				SaveRegistryValue(key, "DepotPath", PerforceDepotPath, DeploymentSettings.Instance.DefaultDepotPath);
 
