UGS: Add the ability to download precompiled binaries from Horde

diff --git a/Engine/Source/Programs/Shared/EpicGames.Horde/HordeHttpClient.cs b/Engine/Source/Programs/Shared/EpicGames.Horde/HordeHttpClient.cs
--- a/Engine/Source/Programs/Shared/EpicGames.Horde/HordeHttpClient.cs
+++ b/Engine/Source/Programs/Shared/EpicGames.Horde/HordeHttpClient.cs
@@ -137,6 +137,16 @@ namespace EpicGames.Horde
 			return GetAsync<GetArtifactResponse>(_httpClient, $"api/v2/artifacts/{id}", cancellationToken);
 		}
 
+		/// <summary>
+		/// Gets a zip stream for a particular artifact
+		/// </summary>
+		/// <param name="id">Identifier for the artifact</param>
+		/// <param name="cancellationToken">Cancellation token for the operation</param>
+		public async Task<Stream> GetArtifactZipAsync(ArtifactId id, CancellationToken cancellationToken = default)
+		{
+			return await _httpClient.GetStreamAsync($"api/v2/artifacts/{id}/zip", cancellationToken);
+		}
+
 		/// <summary>
 		/// Finds artifacts with a set of ids or keys
 		/// </summary>
@@ -166,6 +176,39 @@ namespace EpicGames.Horde
 			return response.Artifacts;
 		}
 
+		/// <summary>
+		/// Finds artifacts with a certain type with an optional streamId
+		/// </summary>
+		/// <param name="type">Type to find</param>
+		/// <param name="streamId">Stream to look for the artifact in</param>
+		/// <param name="minChange">The minimum change number for the artifacts</param>
+		/// <param name="maxChange">The minimum change number for the artifacts</param>
+		/// <param name="cancellationToken">Cancellation token for the operation</param>
+		/// <returns>Information about all the artifacts</returns>
+		public async Task<List<GetArtifactResponse>> FindArtifactsByTypeAsync(ArtifactType type, StreamId? streamId = null, int? minChange = null, int? maxChange = null, CancellationToken cancellationToken = default)
+		{
+			QueryStringBuilder queryParams = new QueryStringBuilder();
+			queryParams.Add("type", type.ToString());
+
+			if (streamId != null)
+			{
+				queryParams.Add("streamId", streamId.ToString()!);
+			}
+
+			if (minChange != null)
+			{
+				queryParams.Add("minChange", minChange.ToString()!);
+			}
+
+			if (maxChange != null)
+			{
+				queryParams.Add("maxChange", maxChange.ToString()!);
+			}
+
+			FindArtifactsResponse response = await GetAsync<FindArtifactsResponse>(_httpClient, $"api/v2/artifacts?{queryParams}", cancellationToken);
+			return response.Artifacts;
+		}
+
 		#endregion
 
 		#region Dashboard
diff --git a/Engine/Source/Programs/UnrealGameSync/UnrealGameSync/PerforceMonitor.cs b/Engine/Source/Programs/UnrealGameSync/UnrealGameSync/PerforceMonitor.cs
--- a/Engine/Source/Programs/UnrealGameSync/UnrealGameSync/PerforceMonitor.cs
+++ b/Engine/Source/Programs/UnrealGameSync/UnrealGameSync/PerforceMonitor.cs
@@ -3,6 +3,7 @@
 using EpicGames.Core;
 using EpicGames.OIDC;
 using EpicGames.Perforce;
+using EpicGames.Horde;
 using Microsoft.Extensions.DependencyInjection;
 using Microsoft.Extensions.Logging;
 using System;
@@ -45,7 +46,7 @@ namespace UnrealGameSync
 		SortedSet<ChangesRecord> _changes = new SortedSet<ChangesRecord>(new PerforceChangeSorter());
 		readonly SortedDictionary<int, PerforceChangeDetails> _changeDetails = new SortedDictionary<int,PerforceChangeDetails>();
 		readonly SortedSet<int> _promotedChangeNumbers = new SortedSet<int>();
-		List<PerforceArchiveInfo> _archives = new List<PerforceArchiveInfo>();
+		List<BaseArchiveInfo> _archives = new List<BaseArchiveInfo>();
 		readonly AsyncEvent _refreshEvent = new AsyncEvent();
 		readonly ILogger _logger;
 		readonly bool _isEnterpriseProject;
@@ -53,6 +54,7 @@ namespace UnrealGameSync
 		readonly List<KeyValuePair<FileReference, DateTime>> _localConfigFiles;
 		readonly IAsyncDisposer _asyncDisposeTasks;
 		readonly OidcTokenManager _oidcTokenManager;
+		readonly IHordeClient _hordeClient;
 
 		string[] prevCodeRules = Array.Empty<string>();
 
@@ -88,6 +90,7 @@ namespace UnrealGameSync
 			_synchronizationContext = SynchronizationContext.Current!;
 			_cancellationSource = new CancellationTokenSource();
 			_oidcTokenManager = serviceProvider.GetRequiredService<OidcTokenManager>();
+			_hordeClient = serviceProvider.GetRequiredService<IHordeClient>();
 
 			AvailableArchives = (new List<IArchiveInfo>()).AsReadOnly();
 			LatestOidcTokenClient = oidcTokenClient;
@@ -317,9 +320,9 @@ namespace UnrealGameSync
 			if (perforceConfigSection != null && perforceConfigSection.GetValue("FindAllChangesForPCBs", false))
 			{
 				int minZippedChangeNumber = -1;
-				foreach (PerforceArchiveInfo archive in _archives)
+				foreach (BaseArchiveInfo archive in _archives)
 				{
-					foreach (int changeNumber in archive.ChangeNumberToFileRevision.Keys)
+					foreach (int changeNumber in archive.ChangeNumberToArchiveKey.Keys)
 					{
 						if (changeNumber > minZippedChangeNumber && changeNumber <= oldestChangeNumber)
 						{
@@ -376,7 +379,7 @@ namespace UnrealGameSync
 						foreach(ChangesRecord change in _changes)
 						{
 							trimmedChanges.Add(change);
-							if(trimmedChanges.Count >= maxChanges && _archives.Any(x => x.ChangeNumberToFileRevision.Count == 0 || x.ChangeNumberToFileRevision.ContainsKey(change.Number) || x.ChangeNumberToFileRevision.First().Key > change.Number))
+							if(trimmedChanges.Count >= maxChanges && _archives.Any(x => x.ChangeNumberToArchiveKey.Count == 0 || x.ChangeNumberToArchiveKey.ContainsKey(change.Number) || x.ChangeNumberToArchiveKey.First().Key > change.Number))
 							{
 								break;
 							}
@@ -497,7 +500,7 @@ namespace UnrealGameSync
 
 		async Task<bool> UpdateArchivesAsync(IPerforceConnection perforce, CancellationToken cancellationToken)
 		{
-			List<PerforceArchiveInfo> newArchives = await PerforceArchive.EnumerateAsync(perforce, LatestProjectConfigFile, _selectedProjectIdentifier, cancellationToken);
+			List<BaseArchiveInfo> newArchives = await BaseArchive.EnumerateAsync(perforce, _hordeClient, LatestProjectConfigFile, _selectedProjectIdentifier, cancellationToken);
 
 			// Check if the information has changed
 			if (!Enumerable.SequenceEqual(_archives, newArchives))
diff --git a/Engine/Source/Programs/UnrealGameSync/UnrealGameSyncCmd/Program.cs b/Engine/Source/Programs/UnrealGameSync/UnrealGameSyncCmd/Program.cs
--- a/Engine/Source/Programs/UnrealGameSync/UnrealGameSyncCmd/Program.cs
+++ b/Engine/Source/Programs/UnrealGameSync/UnrealGameSyncCmd/Program.cs
@@ -23,6 +23,8 @@ using System.Globalization;
 using System.Net.Http;
 using System.Net.Http.Json;
 using System.IO.Compression;
+using Microsoft.Extensions.DependencyInjection;
+using EpicGames.Horde;
 
 namespace UnrealGameSyncCmd
 {
@@ -135,12 +137,15 @@ namespace UnrealGameSyncCmd
 			public ILoggerFactory LoggerFactory { get; }
 			public GlobalSettingsFile UserSettings { get; }
 
-			public CommandContext(CommandLineArguments arguments, ILogger logger, ILoggerFactory loggerFactory, GlobalSettingsFile userSettings)
+			public IHordeClient _hordeClient { get; }
+
+			public CommandContext(CommandLineArguments arguments, ILogger logger, ILoggerFactory loggerFactory, GlobalSettingsFile userSettings, IHordeClient hordeClient)
 			{
 				Arguments = arguments;
 				Logger = logger;
 				LoggerFactory = loggerFactory;
 				UserSettings = userSettings;
+				_hordeClient = hordeClient;
 			}
 		}
 
@@ -227,6 +232,18 @@ namespace UnrealGameSyncCmd
 			ILogger logger = loggerFactory.CreateLogger("Main");
 			try
 			{
+				LauncherSettings launcherSettings = new LauncherSettings();
+				launcherSettings.Read();
+
+				ServiceCollection services = new ServiceCollection();
+
+				if (launcherSettings.HordeServer != null)
+				{
+					services.AddHorde(options => options.ServerUrl = new Uri(launcherSettings.HordeServer));
+				}
+				ServiceProvider serviceProvider = services.BuildServiceProvider();
+				IHordeClient hordeClient = serviceProvider.GetRequiredService<IHordeClient>();
+
 				GlobalSettingsFile settings;
 				if (RuntimeInformation.IsOSPlatform(OSPlatform.Windows))
 				{
@@ -301,7 +318,7 @@ namespace UnrealGameSyncCmd
 				}
 
 				Command instance = (Command)Activator.CreateInstance(command.Type)!;
-				await instance.ExecuteAsync(new CommandContext(args, logger, loggerFactory, settings));
+				await instance.ExecuteAsync(new CommandContext(args, logger, loggerFactory, settings, hordeClient));
 				return 0;
 			}
 			catch (UserErrorException ex)
@@ -852,15 +869,15 @@ namespace UnrealGameSyncCmd
 
 				if (syncOptions.Binaries)
 				{
-					List<PerforceArchiveInfo> archives = await PerforceArchive.EnumerateAsync(perforceClient, projectConfig, state.Current.ProjectIdentifier, CancellationToken.None);
+					List<BaseArchiveInfo> archives = await BaseArchive.EnumerateAsync(perforceClient, context._hordeClient, projectConfig, state.Current.ProjectIdentifier, CancellationToken.None);
 
-					PerforceArchiveInfo? editorArchiveInfo = archives.FirstOrDefault(x => x.Name == IArchiveInfo.EditorArchiveType);
+					BaseArchiveInfo? editorArchiveInfo = archives.FirstOrDefault(x => x.Name == IArchiveInfo.EditorArchiveType);
 					if (editorArchiveInfo == null)
 					{
 						throw new UserErrorException("No editor archives found for project");
 					}
 
-					KeyValuePair<int, string> revision = editorArchiveInfo.ChangeNumberToFileRevision.LastOrDefault(x => x.Key <= change);
+					KeyValuePair<int, string> revision = editorArchiveInfo.ChangeNumberToArchiveKey.LastOrDefault(x => x.Key <= change);
 					if (revision.Key == 0)
 					{
 						throw new UserErrorException($"No editor archives found for CL {change}");
diff --git a/Engine/Source/Programs/UnrealGameSync/UnrealGameSyncShared/ArchiveInfo.cs b/Engine/Source/Programs/UnrealGameSync/UnrealGameSyncShared/ArchiveInfo.cs
--- a/Engine/Source/Programs/UnrealGameSync/UnrealGameSyncShared/ArchiveInfo.cs
+++ b/Engine/Source/Programs/UnrealGameSync/UnrealGameSyncShared/ArchiveInfo.cs
@@ -1,6 +1,9 @@
 // Copyright Epic Games, Inc. All Rights Reserved.
 
 using EpicGames.Core;
+using EpicGames.Horde;
+using EpicGames.Horde.Artifacts;
+using EpicGames.Horde.Streams;
 using EpicGames.Perforce;
 using Microsoft.Extensions.Logging;
 using System;
@@ -27,30 +30,33 @@ namespace UnrealGameSync
 		Task<bool> DownloadArchive(IPerforceConnection perforce, string archiveKey, DirectoryReference localRootPath, FileReference manifestFileName, ILogger logger, ProgressValue progress, CancellationToken cancellationToken);
 	}
 
-	public class PerforceArchiveInfo : IArchiveInfo
+	public abstract class BaseArchiveInfo : IArchiveInfo
 	{
 		public string Name { get; }
 		public string Type { get; }
-		public string DepotPath { get; }
+		public string BasePath { get; set; }
 		public string? Target { get; }
 
-		public string BasePath => DepotPath;
+		public abstract Task<bool> DownloadArchive(IPerforceConnection perforce, string archiveKey, DirectoryReference localRootPath, FileReference manifestFileName, ILogger logger, ProgressValue progress, CancellationToken cancellationToken);
+
+		// Abstract BaseArchiveInfo Helper Functions
+		public abstract Task FindArtifacts(IPerforceConnection perforce, CancellationToken cancellationToken);
 
 		// TODO: executable/configuration?
-		public SortedList<int, string> ChangeNumberToFileRevision { get; } = new SortedList<int, string>();
+		public SortedList<int, string> ChangeNumberToArchiveKey { get; } = new SortedList<int, string>();
 
-		public PerforceArchiveInfo(string name, string type, string depotPath, string? target)
+		protected BaseArchiveInfo(string name, string type, string basePath, string? target)
 		{
 			Name = name;
 			Type = type;
-			DepotPath = depotPath;
+			BasePath = basePath;
 			Target = target;
 		}
-
+		
 		public override bool Equals(object? other)
 		{
-			PerforceArchiveInfo? otherArchive = other as PerforceArchiveInfo;
-			return otherArchive != null && Name == otherArchive.Name && Type == otherArchive.Type && DepotPath == otherArchive.DepotPath && Target == otherArchive.Target && Enumerable.SequenceEqual(ChangeNumberToFileRevision, otherArchive.ChangeNumberToFileRevision);
+			BaseArchiveInfo? otherArchive = other as BaseArchiveInfo;
+			return otherArchive != null && Name == otherArchive.Name && Type == otherArchive.Type && BasePath == otherArchive.BasePath && Target == otherArchive.Target && Enumerable.SequenceEqual(ChangeNumberToArchiveKey, otherArchive.ChangeNumberToArchiveKey);
 		}
 
 		public override int GetHashCode()
@@ -60,10 +66,10 @@ namespace UnrealGameSync
 
 		public bool Exists()
 		{
-			return ChangeNumberToFileRevision.Count > 0;
+			return ChangeNumberToArchiveKey.Count > 0;
 		}
 
-		public static bool TryParseConfigEntry(string text, [NotNullWhen(true)] out PerforceArchiveInfo? info)
+		public static bool TryParseConfigEntry(IHordeClient hordeClient, string text, [NotNullWhen(true)] out BaseArchiveInfo? info)
 		{
 			ConfigObject obj = new ConfigObject(text);
 
@@ -74,8 +80,12 @@ namespace UnrealGameSync
 				return false;
 			}
 
+			// Where to find archives, you'll have either Perforce (DepotPath) or Horde (ArchiveType)
 			string? depotPath = obj.GetValue("DepotPath", null);
-			if (depotPath == null)
+			string? archiveType = obj.GetValue("ArchiveType", null);
+
+			// We only want one of the other, not both or none
+			if (((depotPath == null) && (archiveType == null)) || ((depotPath != null) && (archiveType != null)))
 			{
 				info = null;
 				return false;
@@ -85,23 +95,37 @@ namespace UnrealGameSync
 
 			string type = obj.GetValue("Type", null) ?? name;
 
-			info = new PerforceArchiveInfo(name, type, depotPath, target);
+			string? streamName = obj.GetValue("StreamName", null);
+
+			if (depotPath != null)
+			{
+				info = new PerforceArchiveInfo(name, type, depotPath, target);
+			}
+			else if ((archiveType != null) && (streamName != null))
+			{
+				info = new HordeArchiveInfo(hordeClient, name, type, archiveType, target, streamName);
+			}
+			else
+			{
+				info = null;
+				return false;
+			}
 			return true;
 		}
 
 		public bool TryGetArchiveKeyForChangeNumber(int changeNumber, int maxChangeNumber, [NotNullWhen(true)] out string? archiveKey)
 		{
-			int idx = ChangeNumberToFileRevision.Keys.AsReadOnlyList().BinarySearch(changeNumber);
+			int idx = ChangeNumberToArchiveKey.Keys.AsReadOnlyList().BinarySearch(changeNumber);
 			if (idx >= 0)
 			{
-				archiveKey = ChangeNumberToFileRevision.Values[idx];
+				archiveKey = ChangeNumberToArchiveKey.Values[idx];
 				return true;
 			}
 
 			int nextIdx = ~idx;
-			if (nextIdx < ChangeNumberToFileRevision.Count && ChangeNumberToFileRevision.Keys[nextIdx] <= maxChangeNumber)
+			if (nextIdx < ChangeNumberToArchiveKey.Count && ChangeNumberToArchiveKey.Keys[nextIdx] <= maxChangeNumber)
 			{
-				archiveKey = ChangeNumberToFileRevision.Values[nextIdx];
+				archiveKey = ChangeNumberToArchiveKey.Values[nextIdx];
 				return true;
 			}
 
@@ -109,7 +133,37 @@ namespace UnrealGameSync
 			return false;
 		}
 
-		public async Task<bool> DownloadArchive(IPerforceConnection perforce, string archiveKey, DirectoryReference localRootPath, FileReference manifestFileName, ILogger logger, ProgressValue progress, CancellationToken cancellationToken)
+		public override string ToString()
+		{
+			return Name;
+		}
+	}
+
+	public class PerforceArchiveInfo : BaseArchiveInfo
+	{
+		public string DepotPath
+		{
+			get => BasePath;
+			set => BasePath = value;
+		}
+
+		public PerforceArchiveInfo(string name, string type, string depotPath, string? target)
+			: base(name, type, depotPath, target)
+		{
+		}
+
+		public override bool Equals(object? other)
+		{
+			PerforceArchiveInfo? otherArchive = other as PerforceArchiveInfo;
+			return otherArchive != null && Name == otherArchive.Name && Type == otherArchive.Type && DepotPath == otherArchive.DepotPath && Target == otherArchive.Target && Enumerable.SequenceEqual(ChangeNumberToArchiveKey, otherArchive.ChangeNumberToArchiveKey);
+		}
+
+		public override int GetHashCode()
+		{
+			throw new NotSupportedException();
+		}
+
+		public override async Task<bool> DownloadArchive(IPerforceConnection perforce, string archiveKey, DirectoryReference localRootPath, FileReference manifestFileName, ILogger logger, ProgressValue progress, CancellationToken cancellationToken)
 		{
 			DirectoryReference configDir = UserSettings.GetConfigDir(localRootPath);
 			UserSettings.CreateConfigDir(configDir);
@@ -118,6 +172,7 @@ namespace UnrealGameSync
 			try
 			{
 				PrintRecord record = await perforce.PrintAsync(tempZipFileName.FullName, archiveKey, cancellationToken);
+
 				if (tempZipFileName.ToFileInfo().Length == 0)
 				{
 					return false;
@@ -133,17 +188,125 @@ namespace UnrealGameSync
 			return true;
 		}
 
-		public override string ToString()
+		public override async Task FindArtifacts(IPerforceConnection perforce, CancellationToken cancellationToken)
 		{
-			return Name;
+			PerforceResponseList<FileLogRecord> response = await perforce.TryFileLogAsync(128, FileLogOptions.FullDescriptions, DepotPath, cancellationToken);
+			if (response.Succeeded)
+			{
+				// Build a new list of zipped binaries
+				foreach (FileLogRecord file in response.Data)
+				{
+					foreach (RevisionRecord revision in file.Revisions)
+					{
+						if (revision.Action != FileAction.Purge)
+						{
+							string[] tokens = revision.Description.Split(' ');
+							if (tokens[0].StartsWith("[CL", StringComparison.Ordinal) && tokens[1].EndsWith("]", StringComparison.Ordinal))
+							{
+								int originalChangeNumber;
+								if (Int32.TryParse(tokens[1].Substring(0, tokens[1].Length - 1), out originalChangeNumber) && !ChangeNumberToArchiveKey.ContainsKey(originalChangeNumber))
+								{
+									ChangeNumberToArchiveKey[originalChangeNumber] = $"{DepotPath}#{revision.RevisionNumber}";
+								}
+							}
+						}
+					}
+				}
+			}
+		}
+	}
+
+	public class HordeArchiveInfo : BaseArchiveInfo
+	{
+		public string ArchiveType
+		{
+			get => BasePath;
+			set => BasePath = value;
+		}
+
+		public string StreamName { get; set; }
+
+		private readonly IHordeClient _hordeClient;
+
+		public HordeArchiveInfo(IHordeClient hordeClient, string name, string type, string archiveType, string? target, string streamName)
+			: base(name, type, archiveType, target)
+		{
+			StreamName = streamName;
+			_hordeClient = hordeClient;
+		}
+		public override bool Equals(object? other)
+		{
+			HordeArchiveInfo? otherArchive = other as HordeArchiveInfo;
+			return otherArchive != null && Name == otherArchive.Name && Type == otherArchive.Type && ArchiveType == otherArchive.ArchiveType && Target == otherArchive.Target && StreamName == otherArchive.StreamName && Enumerable.SequenceEqual(ChangeNumberToArchiveKey, otherArchive.ChangeNumberToArchiveKey);
+		}
+
+		public override int GetHashCode()
+		{
+			throw new NotSupportedException();
+		}
+
+		public override async Task<bool> DownloadArchive(IPerforceConnection perforce, string archiveKey, DirectoryReference localRootPath, FileReference manifestFileName, ILogger logger, ProgressValue progress, CancellationToken cancellationToken)
+		{
+			DirectoryReference configDir = UserSettings.GetConfigDir(localRootPath);
+			UserSettings.CreateConfigDir(configDir);
+
+			FileReference tempZipFileName = FileReference.Combine(configDir, "archive.zip");
+			try
+			{
+				HordeHttpClient hordeHttpClient = _hordeClient.CreateHttpClient();
+
+				ArtifactId artifactId = ArtifactId.Parse(archiveKey);
+
+				using (FileStream stream = FileReference.Open(tempZipFileName, FileMode.Create, FileAccess.Write, FileShare.None))
+				{
+					await using Stream sourceStream = await hordeHttpClient.GetArtifactZipAsync(artifactId, cancellationToken);
+					await sourceStream.CopyToAsync(stream, cancellationToken);
+				}
+
+				if (tempZipFileName.ToFileInfo().Length == 0)
+				{
+					return false;
+				}
+				ArchiveUtils.ExtractFiles(tempZipFileName, localRootPath, manifestFileName, progress, logger);
+			}
+			finally
+			{
+				FileReference.SetAttributes(tempZipFileName, FileAttributes.Normal);
+				FileReference.Delete(tempZipFileName);
+			}
+
+			return true;
+		}
+
+		public override async Task FindArtifacts(IPerforceConnection perforce, CancellationToken cancellationToken)
+		{
+			try
+			{
+				HordeHttpClient hordeHttpClient = _hordeClient.CreateHttpClient();
+
+				ArtifactType artifactType = new ArtifactType(ArchiveType);
+				StreamId streamId = new StreamId(StreamName);
+				int? minChange = null;
+				int? maxChange = null;
+				List<GetArtifactResponse> artifactResponse = await hordeHttpClient.FindArtifactsByTypeAsync(artifactType, streamId, minChange, maxChange, cancellationToken);
+
+				foreach (GetArtifactResponse response in artifactResponse)
+				{
+					ChangeNumberToArchiveKey[response.Change] = response.Id.ToString();
+				}
+			}
+			catch (Exception)
+			{
+				return;
+			}
 		}
 	}
 
-	public static class PerforceArchive
+	public static class BaseArchive
 	{
-		public static async Task<List<PerforceArchiveInfo>> EnumerateAsync(IPerforceConnection perforce, ConfigFile latestProjectConfigFile, string projectIdentifier, CancellationToken cancellationToken)
+		public static async Task<List<BaseArchiveInfo>> EnumerateAsync(IPerforceConnection perforce, IHordeClient hordeClient, ConfigFile latestProjectConfigFile, string projectIdentifier, CancellationToken cancellationToken)
 		{
-			List<PerforceArchiveInfo> newArchives = new List<PerforceArchiveInfo>();
+			List<BaseArchiveInfo> newArchives = new List<BaseArchiveInfo>();
 
 			// Find all the zipped binaries under this stream
 			ConfigSection? projectConfigSection = latestProjectConfigFile.FindSection(projectIdentifier);
@@ -153,45 +316,24 @@ namespace UnrealGameSync
 				string? legacyEditorArchivePath = projectConfigSection.GetValue("ZippedBinariesPath", null);
 				if (legacyEditorArchivePath != null)
 				{
+					// Only Perforce uses the legacy method
 					newArchives.Add(new PerforceArchiveInfo("Editor", "Editor", legacyEditorArchivePath, null));
 				}
 
 				// New style
 				foreach (string archiveValue in projectConfigSection.GetValues("Archives", Array.Empty<string>()))
 				{
-					PerforceArchiveInfo? archive;
-					if (PerforceArchiveInfo.TryParseConfigEntry(archiveValue, out archive))
+					BaseArchiveInfo? archive;
+					if (BaseArchiveInfo.TryParseConfigEntry(hordeClient, archiveValue, out archive))
 					{
-						newArchives.Add(archive);
+						newArchives.Add(archive!);
 					}
 				}
 
 				// Make sure the zipped binaries path exists
-				foreach (PerforceArchiveInfo newArchive in newArchives)
+				foreach (BaseArchiveInfo newArchive in newArchives)
 				{
-					PerforceResponseList<FileLogRecord> response = await perforce.TryFileLogAsync(128, FileLogOptions.FullDescriptions, newArchive.DepotPath, cancellationToken);
-					if (response.Succeeded)
-					{
-						// Build a new list of zipped binaries
-						foreach (FileLogRecord file in response.Data)
-						{
-							foreach (RevisionRecord revision in file.Revisions)
-							{
-								if (revision.Action != FileAction.Purge)
-								{
-									string[] tokens = revision.Description.Split(' ');
-									if (tokens[0].StartsWith("[CL", StringComparison.Ordinal) && tokens[1].EndsWith("]", StringComparison.Ordinal))
-									{
-										int originalChangeNumber;
-										if (Int32.TryParse(tokens[1].Substring(0, tokens[1].Length - 1), out originalChangeNumber) && !newArchive.ChangeNumberToFileRevision.ContainsKey(originalChangeNumber))
-										{
-											newArchive.ChangeNumberToFileRevision[originalChangeNumber] = $"{newArchive.DepotPath}#{revision.RevisionNumber}";
-										}
-									}
-								}
-							}
-						}
-					}
+					await newArchive.FindArtifacts(perforce, cancellationToken);
 				}
 			}
 
