Horde: Fix tool update errors if Perforce connection cannot be created.

diff --git a/Engine/Source/Programs/UnrealGameSync/UnrealGameSync/ToolUpdateMonitor.cs b/Engine/Source/Programs/UnrealGameSync/UnrealGameSync/ToolUpdateMonitor.cs
--- a/Engine/Source/Programs/UnrealGameSync/UnrealGameSync/ToolUpdateMonitor.cs
+++ b/Engine/Source/Programs/UnrealGameSync/UnrealGameSync/ToolUpdateMonitor.cs
@@ -208,78 +208,85 @@ namespace UnrealGameSync
 
 		async Task PollForUpdatesOnce(CancellationToken cancellationToken)
 		{
-			using IPerforceConnection perforce = await PerforceConnection.CreateAsync(PerforceSettings, _logger);
-
-			// Update all the available tools
-			List<ToolInfo> tools = new List<ToolInfo>();
-			if (!String.IsNullOrEmpty(DeploymentSettings.Instance.ToolsDepotPath))
-			{
-				try
-				{
-					await ReadPerforceToolsAsync(perforce, tools, cancellationToken);
-				}
-				catch (Exception ex) when (ex is not OperationCanceledException)
-				{
-					_logger.LogWarning(ex, "Error while polling Perforce for available tools: {Message}", ex.Message);
-				}
-			}
-			using (HordeHttpClient? hordeHttpClient = _serviceProvider.GetService<HordeHttpClient>())
+			IPerforceConnection? perforce = null;
+			try
 			{
-				if (hordeHttpClient != null)
+				// Update all the available tools
+				List<ToolInfo> tools = new List<ToolInfo>();
+				if (!String.IsNullOrEmpty(DeploymentSettings.Instance.ToolsDepotPath))
 				{
 					try
 					{
-						await ReadHordeToolsAsync(hordeHttpClient, tools, cancellationToken);
+						perforce = await PerforceConnection.CreateAsync(PerforceSettings, _logger);
+						await ReadPerforceToolsAsync(perforce, tools, cancellationToken);
 					}
 					catch (Exception ex) when (ex is not OperationCanceledException)
 					{
-						_logger.LogWarning(ex, "Error while polling Horde for available tools: {Message}", ex.Message);
+						_logger.LogWarning(ex, "Error while polling Perforce for available tools: {Message}", ex.Message);
 					}
 				}
-			}
-			_tools = tools;
+				using (HordeHttpClient? hordeHttpClient = _serviceProvider.GetService<HordeHttpClient>())
+				{
+					if (hordeHttpClient != null)
+					{
+						try
+						{
+							await ReadHordeToolsAsync(hordeHttpClient, tools, cancellationToken);
+						}
+						catch (Exception ex) when (ex is not OperationCanceledException)
+						{
+							_logger.LogWarning(ex, "Error while polling Horde for available tools: {Message}", ex.Message);
+						}
+					}
+				}
+				_tools = tools;
 
-			// When upgrading from older UGS versions, read the legacy sync CL from plain-text config files
-			if (_readLegacyConfig)
-			{
-				await ReadLegacyConfigAsync(tools, cancellationToken);
-				_readLegacyConfig = false;
-			}
+				// When upgrading from older UGS versions, read the legacy sync CL from plain-text config files
+				if (_readLegacyConfig)
+				{
+					await ReadLegacyConfigAsync(tools, cancellationToken);
+					_readLegacyConfig = false;
+				}
 
-			// Find all the tools which are enabled, including those enabled due to dependencies from other tools
-			HashSet<Guid> enabledToolIds = new HashSet<Guid>();
-			FindEnabledTools(Settings.EnabledTools, tools, enabledToolIds);
+				// Find all the tools which are enabled, including those enabled due to dependencies from other tools
+				HashSet<Guid> enabledToolIds = new HashSet<Guid>();
+				FindEnabledTools(Settings.EnabledTools, tools, enabledToolIds);
 
-			// Install or update any new tools
-			bool hasChanged = false;
-			foreach (ToolInfo toolInfo in _tools)
-			{
-				if (enabledToolIds.Contains(toolInfo.Id))
+				// Install or update any new tools
+				bool hasChanged = false;
+				foreach (ToolInfo toolInfo in _tools)
 				{
-					ToolInfo? existingTool = _enabledTools.FirstOrDefault(x => x.Id == toolInfo.Id);
-					if (existingTool == null || !String.Equals(existingTool.Revision, toolInfo.Revision, StringComparison.OrdinalIgnoreCase))
+					if (enabledToolIds.Contains(toolInfo.Id))
 					{
-						await UpdateToolAsync(perforce, toolInfo, cancellationToken);
+						ToolInfo? existingTool = _enabledTools.FirstOrDefault(x => x.Id == toolInfo.Id);
+						if (existingTool == null || !String.Equals(existingTool.Revision, toolInfo.Revision, StringComparison.OrdinalIgnoreCase))
+						{
+							await UpdateToolAsync(perforce, toolInfo, cancellationToken);
+							hasChanged = true;
+						}
+					}
+				}
+
+				// Remove any tools which we no longer need
+				for (int idx = _enabledTools.Count - 1; idx >= 0; idx--)
+				{
+					ToolInfo tool = _enabledTools[idx];
+					if (!enabledToolIds.Contains(tool.Id))
+					{
+						await RemoveToolAsync(tool, cancellationToken);
 						hasChanged = true;
 					}
 				}
-			}
 
-			// Remove any tools which we no longer need
-			for (int idx = _enabledTools.Count - 1; idx >= 0; idx--)
-			{
-				ToolInfo tool = _enabledTools[idx];
-				if (!enabledToolIds.Contains(tool.Id))
+				// Notify the main window if anything changed
+				if (hasChanged)
 				{
-					await RemoveToolAsync(tool, cancellationToken);
-					hasChanged = true;
+					_synchronizationContext.Post(_ => OnChange?.Invoke(), null);
 				}
 			}
-
-			// Notify the main window if anything changed
-			if (hasChanged)
+			finally
 			{
-				_synchronizationContext.Post(_ => OnChange?.Invoke(), null);
+				perforce?.Dispose();
 			}
 		}
 
@@ -488,7 +495,7 @@ namespace UnrealGameSync
 			}
 		}
 
-		async Task<bool> UpdateToolAsync(IPerforceConnection perforce, ToolInfo tool, CancellationToken cancellationToken)
+		async Task<bool> UpdateToolAsync(IPerforceConnection? perforce, ToolInfo tool, CancellationToken cancellationToken)
 		{
 			try
 			{
@@ -508,7 +515,7 @@ namespace UnrealGameSync
 				DirectoryReference.CreateDirectory(nextToolZipsDir);
 
 				FileReference zipFile = FileReference.Combine(nextToolZipsDir, $"{tool.Name}.zip");
-				if (tool.Revision.StartsWith("//", StringComparison.Ordinal))
+				if (perforce != null && tool.Revision.StartsWith("//", StringComparison.Ordinal))
 				{
 					// Read it from Perforce
 					PerforceResponseList<PrintRecord> response = await perforce.TryPrintAsync(zipFile.FullName, tool.Revision, cancellationToken);
