UGS: Fix UGS not reopening project windows that are configured with a client specification (rather than file) if the file does not exist.

diff --git a/Engine/Source/Programs/UnrealGameSync/UnrealGameSync/ProgramApplicationContext.cs b/Engine/Source/Programs/UnrealGameSync/UnrealGameSync/ProgramApplicationContext.cs
--- a/Engine/Source/Programs/UnrealGameSync/UnrealGameSync/ProgramApplicationContext.cs
+++ b/Engine/Source/Programs/UnrealGameSync/UnrealGameSync/ProgramApplicationContext.cs
@@ -224,7 +224,7 @@ namespace UnrealGameSync
 		{
 			try
 			{
-				bool valid = String.IsNullOrEmpty(projectSettings.LocalPath) || File.Exists(projectSettings.LocalPath);
+				bool valid = projectSettings.Type != UserSelectedProjectType.Local || File.Exists(projectSettings.LocalPath);
 				return valid;
 			}
 			catch
