UGS: Add an explicit option to disable auto-update functionality and launch the current build instead, rather than doing that when the user cancels out of the update dialog.

diff --git a/Engine/Source/Programs/UnrealGameSync/UnrealGameSync/Forms/UpdateSettingsWindow.cs b/Engine/Source/Programs/UnrealGameSync/UnrealGameSync/Forms/UpdateSettingsWindow.cs
--- a/Engine/Source/Programs/UnrealGameSync/UnrealGameSync/Forms/UpdateSettingsWindow.cs
+++ b/Engine/Source/Programs/UnrealGameSync/UnrealGameSync/Forms/UpdateSettingsWindow.cs
@@ -39,7 +39,19 @@ namespace UnrealGameSync
 
 			string defaultHordeServer = DeploymentSettings.Instance.HordeUrl ?? String.Empty;
 
-			HordeRadioBtn.Checked = settings.UpdateSource != LauncherUpdateSource.Perforce;
+			if (settings.UpdateSource == LauncherUpdateSource.Perforce)
+			{
+				PerforceRadioBtn.Checked = true;
+			}
+			else if (settings.UpdateSource == LauncherUpdateSource.None)
+			{
+				DisableRadioBtn.Checked = true;
+			}
+			else
+			{
+				HordeRadioBtn.Checked = true;
+			}
+
 			HordeServerTextBox.Text = String.IsNullOrEmpty(settings.HordeServer) ? defaultHordeServer : settings.HordeServer;
 			ServerTextBox.Text = settings.PerforceServerAndPort ?? String.Empty;
 			UserNameTextBox.Text = settings.PerforceUserName ?? String.Empty;
@@ -51,6 +63,22 @@ namespace UnrealGameSync
 			ViewLogBtn.Visible = logText != null;
 		}
 
+		public LauncherResult ShowModal()
+		{
+			if (ShowDialog() == DialogResult.Cancel)
+			{
+				return LauncherResult.Exit;
+			}
+			else if (_settings.UpdateSource == LauncherUpdateSource.None)
+			{
+				return LauncherResult.Continue;
+			}
+			else
+			{
+				return LauncherResult.Exit;
+			}
+		}
+
 		protected override void OnLoad(EventArgs e)
 		{
 			base.OnLoad(e);
@@ -67,7 +95,7 @@ namespace UnrealGameSync
 
 		private void ConnectBtn_Click(object sender, EventArgs e)
 		{
-			_settings.UpdateSource = HordeRadioBtn.Checked ? LauncherUpdateSource.Horde : LauncherUpdateSource.Perforce;
+			_settings.UpdateSource = HordeRadioBtn.Checked ? LauncherUpdateSource.Horde : PerforceRadioBtn.Checked? LauncherUpdateSource.Perforce : LauncherUpdateSource.None;
 			_settings.HordeServer = HordeServerTextBox.Text.Trim();
 			if (String.Equals(_settings.HordeServer, DeploymentSettings.Instance.HordeUrl, StringComparison.OrdinalIgnoreCase))
 			{
@@ -129,12 +157,21 @@ namespace UnrealGameSync
 		{
 			HordeGroupBox.Enabled = HordeRadioBtn.Checked;
 			PerforceGroupBox.Enabled = PerforceRadioBtn.Checked;
+			DisableGroupBox.Enabled = DisableRadioBtn.Checked;
 		}
 
 		private void PerforceRadioBtn_CheckedChanged(object sender, EventArgs e)
 		{
 			HordeGroupBox.Enabled = HordeRadioBtn.Checked;
 			PerforceGroupBox.Enabled = PerforceRadioBtn.Checked;
+			DisableGroupBox.Enabled = DisableRadioBtn.Checked;
+		}
+
+		private void DisableRadioBtn_CheckedChanged(object sender, EventArgs e)
+		{
+			HordeGroupBox.Enabled = HordeRadioBtn.Checked;
+			PerforceGroupBox.Enabled = PerforceRadioBtn.Checked;
+			DisableGroupBox.Enabled = DisableRadioBtn.Checked;
 		}
 	}
 }
diff --git a/Engine/Source/Programs/UnrealGameSync/UnrealGameSync/Launcher.cs b/Engine/Source/Programs/UnrealGameSync/UnrealGameSync/Launcher.cs
--- a/Engine/Source/Programs/UnrealGameSync/UnrealGameSync/Launcher.cs
+++ b/Engine/Source/Programs/UnrealGameSync/UnrealGameSync/Launcher.cs
@@ -24,14 +24,21 @@ namespace UnrealGameSync
 {
 	using JsonObject = System.Text.Json.Nodes.JsonObject;
 
+	enum LauncherResult
+	{
+		Continue, // Continue to run this instance
+		Exit, // Other instance has been spawned
+	}
+
 	static class Launcher
 	{
-		public static bool SyncAndRunLatest(Mutex instanceMutex, string[] args)
+		// Returns true if the application should keep running, false if it should quit.
+		public static LauncherResult SyncAndRunLatest(Mutex instanceMutex, string[] args)
 		{
 			// Don't do this if we're already running as a spawned instance
 			if (args.Any(x => x.StartsWith("-updatespawn=", StringComparison.OrdinalIgnoreCase)))
 			{
-				return false;
+				return LauncherResult.Continue;
 			}
 
 			// Figure out if we should sync the unstable build by default
@@ -53,7 +60,11 @@ namespace UnrealGameSync
 			{
 				// Show the settings window immediately
 				using UpdateSettingsWindow updateWindow = new UpdateSettingsWindow(null, null, launcherSettings, SyncAndRunWrapper);
-				return updateWindow.ShowDialog() != DialogResult.Cancel;
+				return updateWindow.ShowModal();
+			}
+			else if (launcherSettings.UpdateSource == LauncherUpdateSource.None)
+			{
+				return LauncherResult.Continue;
 			}
 			else
 			{
@@ -77,11 +88,11 @@ namespace UnrealGameSync
 				}
 				else if (task.Succeeded)
 				{
-					return true;
+					return LauncherResult.Exit;
 				}
 
 				using UpdateSettingsWindow updateWindow = new UpdateSettingsWindow("Unable to update UnrealGameSync from Perforce. Verify that your connection settings are correct.", logger.Render(Environment.NewLine), launcherSettings, SyncAndRunWrapper);
-				return updateWindow.DialogResult != DialogResult.Cancel;
+				return updateWindow.ShowModal();
 			}
 		}
 
@@ -101,6 +112,11 @@ namespace UnrealGameSync
 		{
 			try
 			{
+				if (launcherSettings.UpdateSource == LauncherUpdateSource.None)
+				{
+					return;
+				}
+
 				// Create the target folder
 				string applicationFolder = Path.Combine(Environment.GetFolderPath(Environment.SpecialFolder.LocalApplicationData), "UnrealGameSync", "Latest");
 				if (!SafeCreateDirectory(applicationFolder))
diff --git a/Engine/Source/Programs/UnrealGameSync/UnrealGameSync/Program.cs b/Engine/Source/Programs/UnrealGameSync/UnrealGameSync/Program.cs
--- a/Engine/Source/Programs/UnrealGameSync/UnrealGameSync/Program.cs
+++ b/Engine/Source/Programs/UnrealGameSync/UnrealGameSync/Program.cs
@@ -85,6 +85,8 @@ namespace UnrealGameSync
 			}
 		}
 
+		static WindowsFormsSynchronizationContext? _winformsSyncContext;
+
 		static void RealMain(string[] args)
 		{
 			bool firstInstance;
@@ -94,18 +96,29 @@ namespace UnrealGameSync
 				Application.SetCompatibleTextRenderingDefault(false);
 				Application.SetHighDpiMode(HighDpiMode.PerMonitorV2);
 
-				SynchronizationContext.SetSynchronizationContext(new WindowsFormsSynchronizationContext());
+				_winformsSyncContext = new WindowsFormsSynchronizationContext();
+				SynchronizationContext.SetSynchronizationContext(_winformsSyncContext);
 
 				using (EventWaitHandle activateEvent = new EventWaitHandle(false, EventResetMode.AutoReset, "ActivateUnrealGameSync"))
 				{
 					bool runUpdateCheck = ShouldRunAutoUpdate(args);
 
+					if (SynchronizationContext.Current is not WindowsFormsSynchronizationContext)
+					{
+						throw new NotImplementedException();
+					}
+
 					// Check for a newer version of the application
-					if (runUpdateCheck && Launcher.SyncAndRunLatest(instanceMutex, args))
+					if (runUpdateCheck && Launcher.SyncAndRunLatest(instanceMutex, args) != LauncherResult.Continue)
 					{
 						return;
 					}
 
+					if (SynchronizationContext.Current is not WindowsFormsSynchronizationContext)
+					{
+						throw new NotImplementedException();
+					}
+
 					// Handle any url passed in, possibly exiting
 					if (UriHandler.ProcessCommandLine(args, firstInstance, activateEvent))
 					{
diff --git a/Engine/Source/Programs/UnrealGameSync/UnrealGameSyncShared/LauncherSettings.cs b/Engine/Source/Programs/UnrealGameSync/UnrealGameSyncShared/LauncherSettings.cs
--- a/Engine/Source/Programs/UnrealGameSync/UnrealGameSyncShared/LauncherSettings.cs
+++ b/Engine/Source/Programs/UnrealGameSync/UnrealGameSyncShared/LauncherSettings.cs
@@ -14,6 +14,7 @@ namespace UnrealGameSync
 		Unknown = 0,
 		Perforce = 1,
 		Horde = 2,
+		None = 3,
 	}
 
 	public class LauncherSettings
