UGS: Disable parallel syncing to eliminate p4.exe dependency. UGS still manages threaded syncs explicitly.

diff --git a/Engine/Source/Programs/UnrealGameSync/UnrealGameSyncShared/WorkspaceUpdate.cs b/Engine/Source/Programs/UnrealGameSync/UnrealGameSyncShared/WorkspaceUpdate.cs
--- a/Engine/Source/Programs/UnrealGameSync/UnrealGameSyncShared/WorkspaceUpdate.cs
+++ b/Engine/Source/Programs/UnrealGameSync/UnrealGameSyncShared/WorkspaceUpdate.cs
@@ -789,7 +789,7 @@ namespace UnrealGameSync
 							string syncFilter = Context.Options.HasFlag(WorkspaceUpdateOptions.Sync) ? $"{Context.ChangeNumber}" : $"={Context.ChangeNumber}";
 
 							List<SyncFile> syncFiles = new List<SyncFile>();
-							await foreach (PerforceResponse<SyncRecord> response in perforce.TrySyncAsync(SyncOptions.PreviewOnly, -1, $"{syncPath}@{syncFilter}", cancellationToken))
+							await foreach (PerforceResponse<SyncRecord> response in perforce.TrySyncAsync(SyncOptions.PreviewOnly, -1, 0, -1, -1, -1, -1, $"{syncPath}@{syncFilter}", cancellationToken))
 							{
 								if (!response.Succeeded)
 								{
@@ -1757,8 +1757,8 @@ namespace UnrealGameSync
 
 		static async Task<(WorkspaceUpdateResult, string, string?)> StaticSyncFileRevisions(IPerforceConnection perforce, WorkspaceUpdateContext context, List<string> syncCommands, Action<SyncRecord> syncOutput, CancellationToken cancellationToken)
 		{
-			// Sync them all
-			List<PerforceResponse<SyncRecord>> responses = await perforce.TrySyncAsync(SyncOptions.None, -1, syncCommands, cancellationToken).ToListAsync(cancellationToken);
+			// Sync them all. Explicitly disable parallel syncing here to avoid shelling out to p4.exe.
+			List<PerforceResponse<SyncRecord>> responses = await perforce.TrySyncAsync(SyncOptions.None, -1, 0, -1, -1, -1, -1, syncCommands, cancellationToken).ToListAsync(cancellationToken);
 
 			List<string> tamperedFiles = new List<string>();
 			foreach (PerforceResponse<SyncRecord> response in responses)
@@ -1811,7 +1811,7 @@ namespace UnrealGameSync
 					bool shouldClobber = (context.Options & WorkspaceUpdateOptions.Clobber) != 0 || context.ClobberFiles[tamperedFile];
 					if (shouldClobber)
 					{
-						List<PerforceResponse<SyncRecord>> response = await perforce.TrySyncAsync(SyncOptions.Force, -1, tamperedFile, cancellationToken).ToListAsync(cancellationToken);
+						List<PerforceResponse<SyncRecord>> response = await perforce.TrySyncAsync(SyncOptions.Force, -1, 0, -1, -1, -1, -1, tamperedFile, cancellationToken).ToListAsync(cancellationToken);
 						if (!response.Succeeded())
 						{
 							return (WorkspaceUpdateResult.FailedToSync, $"Couldn't sync {tamperedFile}.", response.ToString());
