UGS: Fix cross-threading issue caused by Winforms installing and uninstalling a synchronization context when we show the settings window.

Also prevent P4 connection being created if auto-updates are disabled.

diff --git a/Engine/Source/Programs/UnrealGameSync/UnrealGameSync/Forms/UpdateSettingsWindow.cs b/Engine/Source/Programs/UnrealGameSync/UnrealGameSync/Forms/UpdateSettingsWindow.cs
--- a/Engine/Source/Programs/UnrealGameSync/UnrealGameSync/Forms/UpdateSettingsWindow.cs
+++ b/Engine/Source/Programs/UnrealGameSync/UnrealGameSync/Forms/UpdateSettingsWindow.cs
@@ -95,7 +95,7 @@ namespace UnrealGameSync
 
 		private void ConnectBtn_Click(object sender, EventArgs e)
 		{
-			_settings.UpdateSource = HordeRadioBtn.Checked ? LauncherUpdateSource.Horde : PerforceRadioBtn.Checked? LauncherUpdateSource.Perforce : LauncherUpdateSource.None;
+			_settings.UpdateSource = HordeRadioBtn.Checked ? LauncherUpdateSource.Horde : PerforceRadioBtn.Checked ? LauncherUpdateSource.Perforce : LauncherUpdateSource.None;
 			_settings.HordeServer = HordeServerTextBox.Text.Trim();
 			if (String.Equals(_settings.HordeServer, DeploymentSettings.Instance.HordeUrl, StringComparison.OrdinalIgnoreCase))
 			{
@@ -111,12 +111,12 @@ namespace UnrealGameSync
 			CaptureLogger logger = new CaptureLogger();
 
 			// Create the task for connecting to this server
-			ModalTask? task;
+			ModalTask? task = null;
 			if (_settings.UpdateSource == LauncherUpdateSource.Horde)
 			{
 				task = ModalTask.Execute(this, "Updating", "Checking for updates, please wait...", c => _syncAndRun(null, _settings, logger, c));
 			}
-			else
+			else if (_settings.UpdateSource == LauncherUpdateSource.Perforce)
 			{
 				PerforceSettings perforceSettings = new PerforceSettings(PerforceSettings.Default);
 				if (!String.IsNullOrEmpty(_settings.PerforceServerAndPort))
@@ -132,14 +132,16 @@ namespace UnrealGameSync
 				task = PerforceModalTask.Execute(this, "Updating", "Checking for updates, please wait...", perforceSettings, (p, c) => _syncAndRun(p, _settings, logger, c), logger);
 			}
 
+			if (task == null || task.Succeeded)
+			{
+				_settings.Save();
+				DialogResult = DialogResult.OK;
+				Close();
+				return;
+			}
+
 			if (task != null)
 			{
-				if (task.Succeeded)
-				{
-					_settings.Save();
-					DialogResult = DialogResult.OK;
-					Close();
-				}
 				PromptLabel.Text = task.Error;
 			}
 
diff --git a/Engine/Source/Programs/UnrealGameSync/UnrealGameSync/Program.cs b/Engine/Source/Programs/UnrealGameSync/UnrealGameSync/Program.cs
--- a/Engine/Source/Programs/UnrealGameSync/UnrealGameSync/Program.cs
+++ b/Engine/Source/Programs/UnrealGameSync/UnrealGameSync/Program.cs
@@ -85,8 +85,6 @@ namespace UnrealGameSync
 			}
 		}
 
-		static WindowsFormsSynchronizationContext? _winformsSyncContext;
-
 		static void RealMain(string[] args)
 		{
 			bool firstInstance;
@@ -96,29 +94,21 @@ namespace UnrealGameSync
 				Application.SetCompatibleTextRenderingDefault(false);
 				Application.SetHighDpiMode(HighDpiMode.PerMonitorV2);
 
-				_winformsSyncContext = new WindowsFormsSynchronizationContext();
-				SynchronizationContext.SetSynchronizationContext(_winformsSyncContext);
+				// Don't auto install (or - more importantly- auto *un-install*) the winforms sync context. We want to be able to access it from the
+				// constructor of our ApplicationContext, which will be after the temporary install/uninstall prompted by spawning the settings dialog.
+				WindowsFormsSynchronizationContext.AutoInstall = false;
+				SynchronizationContext.SetSynchronizationContext(new WindowsFormsSynchronizationContext());
 
 				using (EventWaitHandle activateEvent = new EventWaitHandle(false, EventResetMode.AutoReset, "ActivateUnrealGameSync"))
 				{
 					bool runUpdateCheck = ShouldRunAutoUpdate(args);
 
-					if (SynchronizationContext.Current is not WindowsFormsSynchronizationContext)
-					{
-						throw new NotImplementedException();
-					}
-
 					// Check for a newer version of the application
 					if (runUpdateCheck && Launcher.SyncAndRunLatest(instanceMutex, args) != LauncherResult.Continue)
 					{
 						return;
 					}
 
-					if (SynchronizationContext.Current is not WindowsFormsSynchronizationContext)
-					{
-						throw new NotImplementedException();
-					}
-
 					// Handle any url passed in, possibly exiting
 					if (UriHandler.ProcessCommandLine(args, firstInstance, activateEvent))
 					{
