UGS: Fix compile error

diff --git a/Engine/Source/Programs/UnrealGameSync/UnrealGameSync/UriHandlers/HordeHandler.cs b/Engine/Source/Programs/UnrealGameSync/UnrealGameSync/UriHandlers/HordeHandler.cs
--- a/Engine/Source/Programs/UnrealGameSync/UnrealGameSync/UriHandlers/HordeHandler.cs
+++ b/Engine/Source/Programs/UnrealGameSync/UnrealGameSync/UriHandlers/HordeHandler.cs
@@ -10,6 +10,7 @@ using EpicGames.Horde;
 using EpicGames.Horde.Storage;
 using EpicGames.Horde.Storage.Clients;
 using EpicGames.Horde.Storage.Nodes;
+using EpicGames.Horde.Storage.Bundles;
 using Microsoft.Extensions.DependencyInjection;
 using Microsoft.Extensions.Logging;
 using UnrealGameSync.Forms;
@@ -18,7 +19,7 @@ namespace UnrealGameSync.UriHandlers
 {
 	class HordeHandler
 	{
-		class CopyProgressAdapter : IProgress<ICopyStats>
+		class CopyProgressAdapter : IProgress<IExtractStats>
 		{
 			readonly IProgress<string> _inner;
 
@@ -27,7 +28,7 @@ namespace UnrealGameSync.UriHandlers
 				_inner = inner;
 			}
 
-			public void Report(ICopyStats value)
+			public void Report(IExtractStats value)
 			{
 				_inner.Report($"Copied {value.Count} files ({value.Size / (1024.0 * 1024.0):n1}mb, {value.Rate / (1024.0 * 1024.0):n1}mb/s)...");
 			}
