UGS: Fix hang on shutdown due to async dispose calls capturing the Winforms context.

diff --git a/Engine/Source/Programs/Shared/EpicGames.Core/BackgroundTask.cs b/Engine/Source/Programs/Shared/EpicGames.Core/BackgroundTask.cs
--- a/Engine/Source/Programs/Shared/EpicGames.Core/BackgroundTask.cs
+++ b/Engine/Source/Programs/Shared/EpicGames.Core/BackgroundTask.cs
@@ -32,7 +32,7 @@ namespace EpicGames.Core
 		/// <inheritdoc/>
 		public async ValueTask DisposeAsync()
 		{
-			await StopAsync();
+			await StopAsync().ConfigureAwait(false);
 			_cancellationTokenSource.Dispose();
 		}
 
@@ -83,7 +83,7 @@ namespace EpicGames.Core
 				try
 				{
 					_cancellationTokenSource.Cancel();
-					await _task.WaitAsync(cancellationToken);
+					await _task.WaitAsync(cancellationToken).ConfigureAwait(false);
 				}
 				catch (OperationCanceledException)
 				{
diff --git a/Engine/Source/Programs/UnrealGameSync/UnrealGameSync/Program.cs b/Engine/Source/Programs/UnrealGameSync/UnrealGameSync/Program.cs
--- a/Engine/Source/Programs/UnrealGameSync/UnrealGameSync/Program.cs
+++ b/Engine/Source/Programs/UnrealGameSync/UnrealGameSync/Program.cs
@@ -279,7 +279,9 @@ namespace UnrealGameSync
 
 		static void AsyncDispose(IAsyncDisposable disposable)
 		{
-			disposable.DisposeAsync().AsTask().ConfigureAwait(false).GetAwaiter().GetResult();
+			// Force the dispose to run on a task without a synchronization context, so we don't have to worry about waiting for it on the Winforms thread.
+			Task task = Task.Run(async () => await disposable.DisposeAsync());
+			task.GetAwaiter().GetResult();
 		}
 
 		public static string GetCurrentExecutable()
