UGS: Fix warnings building in Analyze configuration.

#rnx

diff --git a/Engine/Source/Programs/UnrealGameSync/UnrealGameSync/Controls/WorkspaceControl.cs b/Engine/Source/Programs/UnrealGameSync/UnrealGameSync/Controls/WorkspaceControl.cs
--- a/Engine/Source/Programs/UnrealGameSync/UnrealGameSync/Controls/WorkspaceControl.cs
+++ b/Engine/Source/Programs/UnrealGameSync/UnrealGameSync/Controls/WorkspaceControl.cs
@@ -3,6 +3,7 @@
 using System;
 using System.Collections.Generic;
 using System.Diagnostics;
+using System.Diagnostics.CodeAnalysis;
 using System.Drawing;
 using System.Drawing.Drawing2D;
 using System.IO;
@@ -11,18 +12,16 @@ using System.Runtime.InteropServices;
 using System.Text;
 using System.Text.Json;
 using System.Text.RegularExpressions;
+using System.Threading;
+using System.Threading.Tasks;
 using System.Windows.Forms;
 using System.Windows.Forms.VisualStyles;
-using System.Threading;
 using EpicGames.Core;
+using EpicGames.Horde;
 using EpicGames.Perforce;
-using System.Threading.Tasks;
-using Microsoft.Extensions.Logging;
 using Microsoft.Extensions.DependencyInjection;
-using System.Diagnostics.CodeAnalysis;
-using EpicGames.OIDC;
+using Microsoft.Extensions.Logging;
 using Microsoft.Win32;
-using EpicGames.Horde;
 
 namespace UnrealGameSync
 {
diff --git a/Engine/Source/Programs/UnrealGameSync/UnrealGameSync/Forms/MainWindow.cs b/Engine/Source/Programs/UnrealGameSync/UnrealGameSync/Forms/MainWindow.cs
--- a/Engine/Source/Programs/UnrealGameSync/UnrealGameSync/Forms/MainWindow.cs
+++ b/Engine/Source/Programs/UnrealGameSync/UnrealGameSync/Forms/MainWindow.cs
@@ -3,24 +3,23 @@
 using System;
 using System.Collections.Generic;
 using System.Data;
+using System.Diagnostics;
+using System.Diagnostics.CodeAnalysis;
 using System.Drawing;
 using System.IO;
 using System.Linq;
-using System.Runtime.InteropServices;
-using System.Windows.Forms;
-using System.Threading;
 using System.Reflection;
+using System.Runtime.InteropServices;
 using System.Text;
-using System.Diagnostics;
-using UnrealGameSync.Forms;
 using System.Text.RegularExpressions;
+using System.Threading;
+using System.Threading.Tasks;
+using System.Windows.Forms;
 using EpicGames.Core;
 using EpicGames.Perforce;
-using System.Threading.Tasks;
-using Microsoft.Extensions.Logging;
 using Microsoft.Extensions.DependencyInjection;
-using System.Diagnostics.CodeAnalysis;
-using EpicGames.OIDC;
+using Microsoft.Extensions.Logging;
+using UnrealGameSync.Forms;
 
 namespace UnrealGameSync
 {
@@ -156,14 +155,14 @@ namespace UnrealGameSync
 			SetupDefaultControl();
 
 			int selectTabIdx = -1;
-			foreach((UserSelectedProjectSettings project, ModalTask<OpenProjectInfo> startupTask) in startupTasks)
+			foreach ((UserSelectedProjectSettings project, ModalTask<OpenProjectInfo> startupTask) in startupTasks)
 			{
 				int tabIdx = -1;
 				if (startupTask.Succeeded)
 				{
 					tabIdx = TryOpenProject(startupTask.Result, -1, OpenProjectOptions.Quiet);
 				}
-				else if(startupTask.Failed)
+				else if (startupTask.Failed)
 				{
 					_logger.LogError("StartupProject Error: {Message}", startupTask.Error);
 					CreateErrorPanel(-1, project, startupTask.Error);
@@ -175,18 +174,18 @@ namespace UnrealGameSync
 				}
 			}
 
-			if(selectTabIdx != -1)
+			if (selectTabIdx != -1)
 			{
 				TabControl.SelectTab(selectTabIdx);
 			}
-			else if(TabControl.GetTabCount() > 0)
+			else if (TabControl.GetTabCount() > 0)
 			{
 				TabControl.SelectTab(0);
 			}
 
 			StartScheduleTimer();
 
-			if(_preview)
+			if (_preview)
 			{
 				Text += $" {Program.GetVersionString()} (UNSTABLE)";
 			}
@@ -205,7 +204,7 @@ namespace UnrealGameSync
 
 			_allowCreatingHandle = true;
 
-			foreach(WorkspaceIssueMonitor workspaceIssueMonitor in _workspaceIssueMonitors)
+			foreach (WorkspaceIssueMonitor workspaceIssueMonitor in _workspaceIssueMonitors)
 			{
 				workspaceIssueMonitor._issueMonitor.Start();
 			}
@@ -257,7 +256,7 @@ namespace UnrealGameSync
 					request.SetOutput(new AutomationRequestOutput(AutomationRequestResult.Invalid));
 				}
 			}
-			catch(Exception ex)
+			catch (Exception ex)
 			{
 				_logger.LogError(ex, "Exception running automation request");
 				request.SetOutput(new AutomationRequestOutput(AutomationRequestResult.Invalid));
@@ -284,7 +283,7 @@ namespace UnrealGameSync
 				return new AutomationRequestOutput(AutomationRequestResult.Error);
 			}
 
-			workspace.AddStartupCallback((control, cancel) => StartExecCommandAfterStartup(control, cancel, buildInfo.Sync? changelist : -1, buildInfo.ExecCommand));
+			workspace.AddStartupCallback((control, cancel) => StartExecCommandAfterStartup(control, cancel, buildInfo.Sync ? changelist : -1, buildInfo.ExecCommand));
 			return new AutomationRequestOutput(AutomationRequestResult.Ok);
 		}
 
@@ -292,7 +291,7 @@ namespace UnrealGameSync
 		{
 			if (!cancel)
 			{
-				if(changelist == -1)
+				if (changelist == -1)
 				{
 					StartExecCommandAfterSync(workspace, WorkspaceUpdateResult.Success, command);
 				}
@@ -321,18 +320,18 @@ namespace UnrealGameSync
 			string projectPath = reader.ReadString();
 
 			AutomatedSyncWindow.WorkspaceInfo? workspaceInfo;
-			if(!AutomatedSyncWindow.ShowModal(this, _defaultPerforceSettings, streamName, projectPath, out workspaceInfo, _serviceProvider))
+			if (!AutomatedSyncWindow.ShowModal(this, _defaultPerforceSettings, streamName, projectPath, out workspaceInfo, _serviceProvider))
 			{
 				return new AutomationRequestOutput(AutomationRequestResult.Canceled);
 			}
 
 			WorkspaceControl? workspace;
-			if(!OpenWorkspaceForAutomation(workspaceInfo, streamName, projectPath, out workspace))
+			if (!OpenWorkspaceForAutomation(workspaceInfo, streamName, projectPath, out workspace))
 			{
 				return new AutomationRequestOutput(AutomationRequestResult.Error);
 			}
 
-			if(!forceSync && workspace.CanLaunchEditor())
+			if (!forceSync && workspace.CanLaunchEditor())
 			{
 				return new AutomationRequestOutput(AutomationRequestResult.Ok, Encoding.UTF8.GetBytes(workspace.SelectedFileName.FullName));
 			}
@@ -399,7 +398,7 @@ namespace UnrealGameSync
 
 		private static void StartAutomatedSyncAfterStartup(WorkspaceControl workspace, bool cancel, AutomationRequest request)
 		{
-			if(cancel)
+			if (cancel)
 			{
 				request.SetOutput(new AutomationRequestOutput(AutomationRequestResult.Canceled));
 			}
@@ -411,11 +410,11 @@ namespace UnrealGameSync
 
 		static void CompleteAutomatedSync(WorkspaceUpdateResult result, FileReference selectedFileName, AutomationRequest request)
 		{
-			if(result == WorkspaceUpdateResult.Success)
+			if (result == WorkspaceUpdateResult.Success)
 			{
 				request.SetOutput(new AutomationRequestOutput(AutomationRequestResult.Ok, Encoding.UTF8.GetBytes(selectedFileName.FullName)));
 			}
-			else if(result == WorkspaceUpdateResult.Canceled)
+			else if (result == WorkspaceUpdateResult.Canceled)
 			{
 				request.SetOutput(new AutomationRequestOutput(AutomationRequestResult.Canceled));
 			}
@@ -431,16 +430,16 @@ namespace UnrealGameSync
 			string streamName = reader.ReadString();
 			string projectPath = reader.ReadString();
 
-			for(int existingTabIdx = 0; existingTabIdx < TabControl.GetTabCount(); existingTabIdx++)
+			for (int existingTabIdx = 0; existingTabIdx < TabControl.GetTabCount(); existingTabIdx++)
 			{
 				WorkspaceControl? existingWorkspace = TabControl.GetTabData(existingTabIdx) as WorkspaceControl;
-				if(existingWorkspace != null && String.Equals(existingWorkspace.StreamName, streamName, StringComparison.OrdinalIgnoreCase) && existingWorkspace.SelectedProject != null)
+				if (existingWorkspace != null && String.Equals(existingWorkspace.StreamName, streamName, StringComparison.OrdinalIgnoreCase) && existingWorkspace.SelectedProject != null)
 				{
 					string? clientPath = existingWorkspace.SelectedProject.ClientPath;
-					if(clientPath != null && clientPath.StartsWith("//", StringComparison.Ordinal))
+					if (clientPath != null && clientPath.StartsWith("//", StringComparison.Ordinal))
 					{
 						int slashIdx = clientPath.IndexOf('/', 2);
-						if(slashIdx != -1)
+						if (slashIdx != -1)
 						{
 							string existingProjectPath = clientPath.Substring(slashIdx);
 							if (String.Equals(existingProjectPath, projectPath, StringComparison.OrdinalIgnoreCase))
@@ -472,7 +471,7 @@ namespace UnrealGameSync
 						try
 						{
 							Task<IssueData> Func(CancellationToken cancellationToken) => RestApi.GetAsync<IssueData>($"{issueMonitor.ApiUrl}/api/issues/{issueId}", cancellationToken);
-							
+
 							ModalTask<IssueData>? issueTask = ModalTask.Execute(this, "Finding Issue", "Querying issue data, please wait...", Func);
 							if (issueTask == null)
 							{
@@ -510,7 +509,7 @@ namespace UnrealGameSync
 
 		void TabControl_OnButtonClick(int buttonIdx, Point location, MouseButtons buttons)
 		{
-			if(buttonIdx == 0)
+			if (buttonIdx == 0)
 			{
 				EditSelectedProject(TabControl.GetSelectedTabIndex());
 			}
@@ -518,21 +517,21 @@ namespace UnrealGameSync
 
 		void TabControl_OnTabClicked(object? tabData, Point location, MouseButtons buttons)
 		{
-			if(buttons == System.Windows.Forms.MouseButtons.Right)
+			if (buttons == System.Windows.Forms.MouseButtons.Right)
 			{
 				Activate();
 
 				int insertIdx = 0;
 
-				while(TabMenu_RecentProjects.DropDownItems[insertIdx] != TabMenu_Recent_Separator)
+				while (TabMenu_RecentProjects.DropDownItems[insertIdx] != TabMenu_Recent_Separator)
 				{
 					TabMenu_RecentProjects.DropDownItems.RemoveAt(insertIdx);
 				}
 
 				_tabMenuTabIdx = -1;
-				for(int idx = 0; idx < TabControl.GetTabCount(); idx++)
+				for (int idx = 0; idx < TabControl.GetTabCount(); idx++)
 				{
-					if(TabControl.GetTabData(idx) == tabData)
+					if (TabControl.GetTabData(idx) == tabData)
 					{
 						_tabMenuTabIdx = idx;
 						break;
@@ -541,7 +540,7 @@ namespace UnrealGameSync
 
 				TabMenu_OpenProject.Enabled = _tabMenuTabIdx != -1;
 
-				foreach(UserSelectedProjectSettings recentProject in _settings.RecentProjects)
+				foreach (UserSelectedProjectSettings recentProject in _settings.RecentProjects)
 				{
 #pragma warning disable CA2000 // warning CA2000: Call System.IDisposable.Dispose on object created by 'new ToolStripMenuItem(recentProject.ToString(), null, new EventHandler((o, e) => TabMenu_OpenRecentProject_Click(recentProject, _tabMenuTabIdx)))' before all references to it are out of scope
 					ToolStripMenuItem item = new ToolStripMenuItem(recentProject.ToString(), null, new EventHandler((o, e) => TabMenu_OpenRecentProject_Click(recentProject, _tabMenuTabIdx)));
@@ -570,7 +569,7 @@ namespace UnrealGameSync
 		void TabControl_OnTabClosed(object data)
 		{
 			IMainWindowTabPanel tabPanel = (IMainWindowTabPanel)data;
-			if(_currentTabPanel == tabPanel)
+			if (_currentTabPanel == tabPanel)
 			{
 				_currentTabPanel = null;
 			}
@@ -596,7 +595,7 @@ namespace UnrealGameSync
 				components.Dispose();
 			}
 
-			for(int idx = 0; idx < TabControl.GetTabCount(); idx++)
+			for (int idx = 0; idx < TabControl.GetTabCount(); idx++)
 			{
 				((IMainWindowTabPanel)TabControl.GetTabData(idx)).Dispose();
 			}
@@ -610,7 +609,7 @@ namespace UnrealGameSync
 			_defaultIssueMonitors.Clear();
 			Debug.Assert(_workspaceIssueMonitors.Count == 0);
 
-			if(_automationServer != null)
+			if (_automationServer != null)
 			{
 				_automationServer.Dispose();
 				_automationServer = null!;
@@ -640,7 +639,7 @@ namespace UnrealGameSync
 
 		private void MainWindow_FormClosing(object sender, FormClosingEventArgs eventArgs)
 		{
-			if(!_allowClose && _settings.KeepInTray)
+			if (!_allowClose && _settings.KeepInTray)
 			{
 				Hide();
 				eventArgs.Cancel = true;
@@ -665,7 +664,7 @@ namespace UnrealGameSync
 
 			_settings.WindowVisible = Visible;
 			_settings.WindowState = WindowState.ToString();
-			if(WindowState == FormWindowState.Normal)
+			if (WindowState == FormWindowState.Normal)
 			{
 				_settings.WindowBounds = new Rectangle(Location, Size);
 			}
@@ -709,7 +708,7 @@ namespace UnrealGameSync
 			string summaryText = String.Format("Unable to open '{0}'.", project.ToString());
 
 			int newContentWidth = Math.Max(TextRenderer.MeasureText(summaryText, errorPanel.Font).Width, 400);
-			if(!String.IsNullOrEmpty(message))
+			if (!String.IsNullOrEmpty(message))
 			{
 				newContentWidth = Math.Max(newContentWidth, TextRenderer.MeasureText(message, errorPanel.Font).Width);
 			}
@@ -722,11 +721,11 @@ namespace UnrealGameSync
 			summaryLine.AddText(summaryText);
 			lines.Add(summaryLine);
 
-			if(!String.IsNullOrEmpty(message))
+			if (!String.IsNullOrEmpty(message))
 			{
-				lines.Add(new StatusLine(){ LineHeight = 0.5f });
+				lines.Add(new StatusLine() { LineHeight = 0.5f });
 
-				foreach(string messageLine in message.Split('\n'))
+				foreach (string messageLine in message.Split('\n'))
 				{
 					StatusLine errorLine = new StatusLine();
 					errorLine.AddText(messageLine);
@@ -735,35 +734,35 @@ namespace UnrealGameSync
 				}
 			}
 
-			lines.Add(new StatusLine(){ LineHeight = 0.5f });
+			lines.Add(new StatusLine() { LineHeight = 0.5f });
 
 			StatusLine actionLine = new StatusLine();
-			actionLine.AddLink("Retry", FontStyle.Bold | FontStyle.Underline, () => 
-			{ 
-				BeginInvoke(new MethodInvoker(() => TryOpenProject(project, TabControl.FindTabIndex(errorPanel)))); 
+			actionLine.AddLink("Retry", FontStyle.Bold | FontStyle.Underline, () =>
+			{
+				BeginInvoke(new MethodInvoker(() => TryOpenProject(project, TabControl.FindTabIndex(errorPanel))));
 			});
 			actionLine.AddText(" | ");
-			actionLine.AddLink("Settings", FontStyle.Bold | FontStyle.Underline, () => 
-			{ 
-				BeginInvoke(new MethodInvoker(() => EditSelectedProject(errorPanel))); 
+			actionLine.AddLink("Settings", FontStyle.Bold | FontStyle.Underline, () =>
+			{
+				BeginInvoke(new MethodInvoker(() => EditSelectedProject(errorPanel)));
 			});
 			actionLine.AddText(" | ");
-			actionLine.AddLink("Close", FontStyle.Bold | FontStyle.Underline, () => 
-			{ 
-				BeginInvoke(new MethodInvoker(() => TabControl.RemoveTab(TabControl.FindTabIndex(errorPanel)))); 
+			actionLine.AddLink("Close", FontStyle.Bold | FontStyle.Underline, () =>
+			{
+				BeginInvoke(new MethodInvoker(() => TabControl.RemoveTab(TabControl.FindTabIndex(errorPanel))));
 			});
 			lines.Add(actionLine);
 
 			errorPanel.Set(lines, null, null, null);
 
 			string newProjectName = "Unknown";
-			if(project.Type == UserSelectedProjectType.Client && project.ClientPath != null)
+			if (project.Type == UserSelectedProjectType.Client && project.ClientPath != null)
 			{
 				newProjectName = project.ClientPath.Substring(project.ClientPath.LastIndexOf('/') + 1);
 			}
-			if(project.Type == UserSelectedProjectType.Local && project.LocalPath != null)
+			if (project.Type == UserSelectedProjectType.Local && project.LocalPath != null)
 			{
-				newProjectName = project.LocalPath.Substring(project.LocalPath.LastIndexOfAny(new char[]{ '/', '\\' }) + 1);
+				newProjectName = project.LocalPath.Substring(project.LocalPath.LastIndexOfAny(new char[] { '/', '\\' }) + 1);
 			}
 
 			string newTabName = String.Format("Error: {0}", newProjectName);
@@ -813,7 +812,7 @@ namespace UnrealGameSync
 			for (int tabIdx = 0; tabIdx < TabControl.GetTabCount(); tabIdx++)
 			{
 				IMainWindowTabPanel tabPanel = (IMainWindowTabPanel)TabControl.GetTabData(tabIdx)!;
-				if(tabPanel.IsBusy())
+				if (tabPanel.IsBusy())
 				{
 					return false;
 				}
@@ -863,7 +862,7 @@ namespace UnrealGameSync
 			StopScheduleTimer();
 
 			List<UserSelectedProjectSettings> openProjects = new List<UserSelectedProjectSettings>();
-			for(int tabIdx = 0; tabIdx < TabControl.GetTabCount(); tabIdx++)
+			for (int tabIdx = 0; tabIdx < TabControl.GetTabCount(); tabIdx++)
 			{
 				IMainWindowTabPanel tabPanel = (IMainWindowTabPanel)TabControl.GetTabData(tabIdx);
 				openProjects.Add(tabPanel.SelectedProject);
@@ -880,7 +879,7 @@ namespace UnrealGameSync
 			}
 
 			using ScheduleWindow schedule = new ScheduleWindow(_settings.ScheduleEnabled, _settings.ScheduleTime, _settings.ScheduleAnyOpenProject, _settings.ScheduleProjects, openProjects, projectToLatestChangeTypes);
-			if(schedule.ShowDialog() == System.Windows.Forms.DialogResult.OK)
+			if (schedule.ShowDialog() == System.Windows.Forms.DialogResult.OK)
 			{
 				schedule.CopySettings(_settings);
 				_settings.Save(_logger);
@@ -893,7 +892,7 @@ namespace UnrealGameSync
 		{
 			StopScheduleTimer();
 
-			if(_settings.ScheduleEnabled)
+			if (_settings.ScheduleEnabled)
 			{
 				DateTime currentTime = DateTime.Now;
 				Random rnd = new Random();
@@ -911,12 +910,12 @@ namespace UnrealGameSync
 				}
 
 				TimeSpan intervalToFirstTick = nextScheduleTime - currentTime;
-				_scheduleTimer = new System.Threading.Timer(x => _mainThreadSynchronizationContext.Post((o) => 
-				{ 
-					if(!IsDisposed)
-					{ 
-						ScheduleTimerElapsed(); 
-					} 
+				_scheduleTimer = new System.Threading.Timer(x => _mainThreadSynchronizationContext.Post((o) =>
+				{
+					if (!IsDisposed)
+					{
+						ScheduleTimerElapsed();
+					}
 				}, null), null, intervalToFirstTick, TimeSpan.FromDays(1));
 
 				_logger.LogInformation("Schedule: Started ScheduleTimer for {Time} ({Time} remaining)", nextScheduleTime, intervalToFirstTick);
@@ -1003,7 +1002,7 @@ namespace UnrealGameSync
 
 		private void StopScheduleTimer()
 		{
-			if(_scheduleTimer != null)
+			if (_scheduleTimer != null)
 			{
 				_scheduleTimer.Dispose();
 				_scheduleTimer = null;
@@ -1025,7 +1024,7 @@ namespace UnrealGameSync
 			}
 
 			// If we did open something, leave it for a while to populate with data before trying to start the sync.
-			if(TabControl.GetTabCount() > numInitialTabs)
+			if (TabControl.GetTabCount() > numInitialTabs)
 			{
 				StartScheduleSettledTimer();
 			}
@@ -1038,19 +1037,19 @@ namespace UnrealGameSync
 		private void StartScheduleSettledTimer()
 		{
 			StopScheduleSettledTimer();
-			_scheduleSettledTimer = new System.Threading.Timer(x => _mainThreadSynchronizationContext.Post((o) => 
-			{ 
-				if(!IsDisposed)
-				{ 
-					ScheduleSettledTimerElapsed(); 
-				} 
+			_scheduleSettledTimer = new System.Threading.Timer(x => _mainThreadSynchronizationContext.Post((o) =>
+			{
+				if (!IsDisposed)
+				{
+					ScheduleSettledTimerElapsed();
+				}
 			}, null), null, TimeSpan.FromSeconds(20.0), TimeSpan.FromMilliseconds(-1.0));
 			_logger.LogInformation("Schedule: Started ScheduleSettledTimer");
 		}
 
 		private void StopScheduleSettledTimer()
 		{
-			if(_scheduleSettledTimer != null)
+			if (_scheduleSettledTimer != null)
 			{
 				_scheduleSettledTimer.Dispose();
 				_scheduleSettledTimer = null;
@@ -1062,13 +1061,13 @@ namespace UnrealGameSync
 		private void ScheduleSettledTimerElapsed()
 		{
 			_logger.LogInformation("Schedule: Starting Sync");
-			for(int idx = 0; idx < TabControl.GetTabCount(); idx++)
+			for (int idx = 0; idx < TabControl.GetTabCount(); idx++)
 			{
 				WorkspaceControl? workspace = TabControl.GetTabData(idx) as WorkspaceControl;
-				if(workspace != null)
+				if (workspace != null)
 				{
 					_logger.LogInformation("Schedule: Considering {File}", workspace.SelectedFileName);
-					if(_settings.ScheduleAnyOpenProject || _settings.ScheduleProjects.Any(x => x.LocalPath != null && x.LocalPath.Equals(workspace.SelectedProject.LocalPath, StringComparison.OrdinalIgnoreCase)))
+					if (_settings.ScheduleAnyOpenProject || _settings.ScheduleProjects.Any(x => x.LocalPath != null && x.LocalPath.Equals(workspace.SelectedProject.LocalPath, StringComparison.OrdinalIgnoreCase)))
 					{
 						_logger.LogInformation("Schedule: Starting Sync");
 						workspace.ScheduleTimerElapsed();
@@ -1079,20 +1078,20 @@ namespace UnrealGameSync
 
 		void TabControl_OnTabChanged(object? newTabData)
 		{
-			if(IsHandleCreated)
+			if (IsHandleCreated)
 			{
 				SendMessage(Handle, WmSetredraw, 0, 0);
 			}
 
 			SuspendLayout();
 
-			if(_currentTabPanel != null)
+			if (_currentTabPanel != null)
 			{
 				_currentTabPanel.Deactivate();
 				_currentTabPanel.Hide();
 			}
 
-			if(newTabData == null)
+			if (newTabData == null)
 			{
 				_currentTabPanel = null;
 				_settings.LastProject = null;
@@ -1107,7 +1106,7 @@ namespace UnrealGameSync
 
 			_settings.Save(_logger);
 
-			if(_currentTabPanel != null)
+			if (_currentTabPanel != null)
 			{
 				_currentTabPanel.Activate();
 				_currentTabPanel.Show();
@@ -1115,7 +1114,7 @@ namespace UnrealGameSync
 
 			ResumeLayout();
 
-			if(IsHandleCreated)
+			if (IsHandleCreated)
 			{
 				SendMessage(Handle, WmSetredraw, 1, 0);
 			}
@@ -1126,9 +1125,9 @@ namespace UnrealGameSync
 		public void RequestProjectChange(WorkspaceControl workspace, UserSelectedProjectSettings project, bool modal)
 		{
 			int tabIdx = TabControl.FindTabIndex(workspace);
-			if(tabIdx != -1 && !workspace.IsBusy() && CanFocus)
+			if (tabIdx != -1 && !workspace.IsBusy() && CanFocus)
 			{
-				if(modal)
+				if (modal)
 				{
 					TryOpenProject(project, tabIdx);
 				}
@@ -1142,10 +1141,10 @@ namespace UnrealGameSync
 		public void OpenNewProject()
 		{
 			OpenProjectInfo? openProjectInfo = OpenProjectWindow.ShowModal(this, null, _settings, _defaultPerforceSettings, _serviceProvider, _logger);
-			if(openProjectInfo != null)
+			if (openProjectInfo != null)
 			{
 				int newTabIdx = TryOpenProject(openProjectInfo, -1, OpenProjectOptions.None);
-				if(newTabIdx != -1)
+				if (newTabIdx != -1)
 				{
 					TabControl.SelectTab(newTabIdx);
 					SaveTabSettings();
@@ -1171,11 +1170,11 @@ namespace UnrealGameSync
 		public void EditSelectedProject(int tabIdx)
 		{
 			object tabData = TabControl.GetTabData(tabIdx);
-			if(tabData is WorkspaceControl workspace)
+			if (tabData is WorkspaceControl workspace)
 			{
 				EditSelectedProject(tabIdx, workspace.SelectedProject);
 			}
-			else if(tabData is ErrorPanel error)
+			else if (tabData is ErrorPanel error)
 			{
 				EditSelectedProject(tabIdx, error.SelectedProject);
 			}
@@ -1184,7 +1183,7 @@ namespace UnrealGameSync
 		public void EditSelectedProject(WorkspaceControl workspace)
 		{
 			int tabIdx = TabControl.FindTabIndex(workspace);
-			if(tabIdx != -1)
+			if (tabIdx != -1)
 			{
 				EditSelectedProject(tabIdx, workspace.SelectedProject);
 			}
@@ -1193,7 +1192,7 @@ namespace UnrealGameSync
 		public void EditSelectedProject(ErrorPanel panel)
 		{
 			int tabIdx = TabControl.FindTabIndex(panel);
-			if(tabIdx != -1)
+			if (tabIdx != -1)
 			{
 				EditSelectedProject(tabIdx, panel.SelectedProject);
 			}
@@ -1202,10 +1201,10 @@ namespace UnrealGameSync
 		public void EditSelectedProject(int tabIdx, UserSelectedProjectSettings selectedProject)
 		{
 			OpenProjectInfo? openProjectInfo = OpenProjectWindow.ShowModal(this, selectedProject, _settings, _defaultPerforceSettings, _serviceProvider, _logger);
-			if(openProjectInfo != null)
+			if (openProjectInfo != null)
 			{
 				int newTabIdx = TryOpenProject(openProjectInfo, tabIdx, OpenProjectOptions.None);
-				if(newTabIdx != -1)
+				if (newTabIdx != -1)
 				{
 					TabControl.SelectTab(newTabIdx);
 					SaveTabSettings();
@@ -1219,7 +1218,7 @@ namespace UnrealGameSync
 			ILogger<OpenProjectInfo> projectLogger = _serviceProvider.GetRequiredService<ILogger<OpenProjectInfo>>();
 
 			ModalTaskFlags taskFlags = ModalTaskFlags.None;
-			if((options & OpenProjectOptions.Quiet) != 0)
+			if ((options & OpenProjectOptions.Quiet) != 0)
 			{
 				taskFlags |= ModalTaskFlags.Quiet;
 			}
@@ -1244,25 +1243,25 @@ namespace UnrealGameSync
 			_logger.LogInformation("Trying to open project {Project}", openProjectInfo.ProjectInfo.ClientFileName);
 
 			// Check that none of the other tabs already have it open
-			for(int tabIdx = 0; tabIdx < TabControl.GetTabCount(); tabIdx++)
+			for (int tabIdx = 0; tabIdx < TabControl.GetTabCount(); tabIdx++)
 			{
-				if(replaceTabIdx != tabIdx)
+				if (replaceTabIdx != tabIdx)
 				{
 					WorkspaceControl? workspace = TabControl.GetTabData(tabIdx) as WorkspaceControl;
-					if(workspace != null)
+					if (workspace != null)
 					{
-						if(workspace.SelectedFileName == openProjectInfo.ProjectInfo.LocalFileName)
+						if (workspace.SelectedFileName == openProjectInfo.ProjectInfo.LocalFileName)
 						{
 							_logger.LogInformation("  Already open in tab {TabIdx}", tabIdx);
-							if((options & OpenProjectOptions.Quiet) == 0)
+							if ((options & OpenProjectOptions.Quiet) == 0)
 							{
 								TabControl.SelectTab(tabIdx);
 							}
 							return tabIdx;
 						}
-						else if(openProjectInfo.ProjectInfo.LocalFileName.IsUnderDirectory(workspace.BranchDirectoryName))
+						else if (openProjectInfo.ProjectInfo.LocalFileName.IsUnderDirectory(workspace.BranchDirectoryName))
 						{
-							if((options & OpenProjectOptions.Quiet) == 0 && MessageBox.Show($"{workspace.SelectedFileName.GetFileNameWithoutExtension()} is already open under {workspace.BranchDirectoryName}.\n\nWould you like to close it?", "Branch already open", MessageBoxButtons.YesNo) == System.Windows.Forms.DialogResult.Yes)
+							if ((options & OpenProjectOptions.Quiet) == 0 && MessageBox.Show($"{workspace.SelectedFileName.GetFileNameWithoutExtension()} is already open under {workspace.BranchDirectoryName}.\n\nWould you like to close it?", "Branch already open", MessageBoxButtons.YesNo) == System.Windows.Forms.DialogResult.Yes)
 							{
 								_logger.LogInformation("  Another project already open in this workspace, tab {TabIdx}. Replacing.", tabIdx);
 								TabControl.RemoveTab(tabIdx);
@@ -1281,10 +1280,10 @@ namespace UnrealGameSync
 			DefaultControl.Hide();
 
 			// Remove the current tab. We need to ensure the workspace has been shut down before creating a new one with the same log files, etc...
-			if(replaceTabIdx != -1)
+			if (replaceTabIdx != -1)
 			{
 				WorkspaceControl? oldWorkspace = TabControl.GetTabData(replaceTabIdx) as WorkspaceControl;
-				if(oldWorkspace != null)
+				if (oldWorkspace != null)
 				{
 					oldWorkspace.Hide();
 					TabControl.SetTabData(replaceTabIdx, new ErrorPanel(openProjectInfo.SelectedProject));
@@ -1294,14 +1293,14 @@ namespace UnrealGameSync
 
 			// Now that we have the project settings, we can construct the tab
 			WorkspaceControl newWorkspace = new WorkspaceControl(this, _dataFolder, _apiUrl, openProjectInfo, _serviceProvider, _settings);
-			
+
 			newWorkspace.Parent = TabPanel;
 			newWorkspace.Dock = DockStyle.Fill;
 			newWorkspace.Hide();
 
 			// Add the tab
 			string newTabName = GetTabName(newWorkspace);
-			if(replaceTabIdx == -1)
+			if (replaceTabIdx == -1)
 			{
 				int newTabIdx = TabControl.InsertTab(-1, newTabName, newWorkspace, newWorkspace.TintColor);
 				_logger.LogInformation("  Inserted tab {TabIdx}", newTabIdx);
@@ -1318,27 +1317,27 @@ namespace UnrealGameSync
 
 		public void StreamChanged(WorkspaceControl workspace)
 		{
-			_mainThreadSynchronizationContext.Post((o) => 
-			{ 
-				if(!IsDisposed) 
-				{ 
-					StreamChangedCallback(workspace); 
-				} 
+			_mainThreadSynchronizationContext.Post((o) =>
+			{
+				if (!IsDisposed)
+				{
+					StreamChangedCallback(workspace);
+				}
 			}, null);
 		}
 
 		public void StreamChangedCallback(WorkspaceControl workspace)
 		{
-			if(_changingWorkspacesRefCount == 0)
+			if (_changingWorkspacesRefCount == 0)
 			{
 				_changingWorkspacesRefCount++;
 
-				for(int idx = 0; idx < TabControl.GetTabCount(); idx++)
+				for (int idx = 0; idx < TabControl.GetTabCount(); idx++)
 				{
-					if(TabControl.GetTabData(idx) == workspace)
+					if (TabControl.GetTabData(idx) == workspace)
 					{
 						UserSelectedProjectSettings project = workspace.SelectedProject;
-						if(TryOpenProject(project, idx) == -1)
+						if (TryOpenProject(project, idx) == -1)
 						{
 							TabControl.RemoveTab(idx);
 						}
@@ -1353,7 +1352,7 @@ namespace UnrealGameSync
 		void SaveTabSettings()
 		{
 			_settings.OpenProjects.Clear();
-			for(int tabIdx = 0; tabIdx < TabControl.GetTabCount(); tabIdx++)
+			for (int tabIdx = 0; tabIdx < TabControl.GetTabCount(); tabIdx++)
 			{
 				IMainWindowTabPanel tabPanel = (IMainWindowTabPanel)TabControl.GetTabData(tabIdx)!;
 				_settings.OpenProjects.Add(tabPanel.SelectedProject);
@@ -1363,7 +1362,7 @@ namespace UnrealGameSync
 
 		void TabControl_OnNewTabClick(Point location, MouseButtons buttons)
 		{
-			if(buttons == MouseButtons.Left)
+			if (buttons == MouseButtons.Left)
 			{
 				OpenNewProject();
 			}
@@ -1402,15 +1401,15 @@ namespace UnrealGameSync
 
 		public void SetTabNames(TabLabels newTabNames)
 		{
-			if(_settings.TabLabels != newTabNames)
+			if (_settings.TabLabels != newTabNames)
 			{
 				_settings.TabLabels = newTabNames;
 				_settings.Save(_logger);
 
-				for(int idx = 0; idx < TabControl.GetTabCount(); idx++)
+				for (int idx = 0; idx < TabControl.GetTabCount(); idx++)
 				{
 					WorkspaceControl? workspace = TabControl.GetTabData(idx) as WorkspaceControl;
-					if(workspace != null)
+					if (workspace != null)
 					{
 						TabControl.SetTabName(idx, GetTabName(workspace));
 					}
@@ -1471,22 +1470,22 @@ namespace UnrealGameSync
 			TaskbarState state = TaskbarState.NoProgress;
 			float progress = -1.0f;
 
-			for(int idx = 0; idx < TabControl.GetTabCount(); idx++)
+			for (int idx = 0; idx < TabControl.GetTabCount(); idx++)
 			{
 				IMainWindowTabPanel tabPanel = (IMainWindowTabPanel)TabControl.GetTabData(idx)!;
 
 				Tuple<TaskbarState, float> desiredTaskbarState = tabPanel.DesiredTaskbarState;
-				if(desiredTaskbarState.Item1 == TaskbarState.Error)
+				if (desiredTaskbarState.Item1 == TaskbarState.Error)
 				{
 					state = TaskbarState.Error;
 					TabControl.SetHighlight(idx, Tuple.Create(Color.FromArgb(204, 64, 64), 1.0f));
 				}
-				else if(desiredTaskbarState.Item1 == TaskbarState.Paused && state != TaskbarState.Error)
+				else if (desiredTaskbarState.Item1 == TaskbarState.Paused && state != TaskbarState.Error)
 				{
 					state = TaskbarState.Paused;
 					TabControl.SetHighlight(idx, Tuple.Create(Color.FromArgb(255, 242, 0), 1.0f));
 				}
-				else if(desiredTaskbarState.Item1 == TaskbarState.Normal && state != TaskbarState.Error && state != TaskbarState.Paused)
+				else if (desiredTaskbarState.Item1 == TaskbarState.Normal && state != TaskbarState.Error && state != TaskbarState.Paused)
 				{
 					state = TaskbarState.Normal;
 					progress = Math.Max(progress, desiredTaskbarState.Item2);
@@ -1498,9 +1497,9 @@ namespace UnrealGameSync
 				}
 			}
 
-			if(IsHandleCreated)
+			if (IsHandleCreated)
 			{
-				if(state == TaskbarState.Normal)
+				if (state == TaskbarState.Normal)
 				{
 					Taskbar.SetState(Handle, TaskbarState.Normal);
 					Taskbar.SetProgress(Handle, (ulong)(progress * 1000.0f), 1000);
@@ -1528,7 +1527,7 @@ namespace UnrealGameSync
 			{
 				_updateMonitor.TriggerUpdate(UpdateType.UserInitiated, false);
 			}
-			else if(result == ApplicationSettingsWindow.Result.RestartAndConfigureUpdate)
+			else if (result == ApplicationSettingsWindow.Result.RestartAndConfigureUpdate)
 			{
 				_updateMonitor.TriggerUpdate(UpdateType.UserInitiated, true);
 			}
@@ -1542,14 +1541,14 @@ namespace UnrealGameSync
 
 		private void MainWindow_Load(object sender, EventArgs e)
 		{
-			if(_settings.WindowBounds != null)
+			if (_settings.WindowBounds != null)
 			{
 				Rectangle windowBounds = _settings.WindowBounds.Value;
-				if(windowBounds.Width > MinimumSize.Width && windowBounds.Height > MinimumSize.Height)
+				if (windowBounds.Width > MinimumSize.Width && windowBounds.Height > MinimumSize.Height)
 				{
 					foreach (Screen screen in Screen.AllScreens)
 					{
-						if(windowBounds.IntersectsWith(screen.Bounds))
+						if (windowBounds.IntersectsWith(screen.Bounds))
 						{
 							Location = _settings.WindowBounds.Value.Location;
 							Size = _settings.WindowBounds.Value.Size;
@@ -1584,59 +1583,59 @@ namespace UnrealGameSync
 			}
 
 			HashSet<IssueData> allIssues = new HashSet<IssueData>();
-			foreach(IssueMonitor issueMonitor in _workspaceIssueMonitors.Select(x => x._issueMonitor))
+			foreach (IssueMonitor issueMonitor in _workspaceIssueMonitors.Select(x => x._issueMonitor))
 			{
 				List<IssueData> issues = issueMonitor.GetIssues();
-				foreach(IssueData issue in issues)
+				foreach (IssueData issue in issues)
 				{
 					IssueAlertReason reason = 0;
-					if(issue.FixChange == 0 && !issue.ResolvedAt.HasValue)
+					if (issue.FixChange == 0 && !issue.ResolvedAt.HasValue)
 					{
-						if(issue.Owner == null)
+						if (issue.Owner == null)
 						{
-							if(issue.Notify)
+							if (issue.Notify)
 							{
 								reason |= IssueAlertReason.Normal;
 							}
-							if(ShowNotificationForIssue(issue) && _settings.NotifyUnassignedMinutes >= 0 && issue.RetrievedAt - issue.CreatedAt >= TimeSpan.FromMinutes(_settings.NotifyUnassignedMinutes))
+							if (ShowNotificationForIssue(issue) && _settings.NotifyUnassignedMinutes >= 0 && issue.RetrievedAt - issue.CreatedAt >= TimeSpan.FromMinutes(_settings.NotifyUnassignedMinutes))
 							{
 								reason |= IssueAlertReason.UnassignedTimer;
 							}
 						}
-						else if(!issue.AcknowledgedAt.HasValue)
+						else if (!issue.AcknowledgedAt.HasValue)
 						{
-							if(String.Equals(issue.Owner, issueMonitor.UserName, StringComparison.OrdinalIgnoreCase))
+							if (String.Equals(issue.Owner, issueMonitor.UserName, StringComparison.OrdinalIgnoreCase))
 							{
 								reason |= IssueAlertReason.Owner;
 							}
-							else if(ShowNotificationForIssue(issue) && _settings.NotifyUnacknowledgedMinutes >= 0 && issue.RetrievedAt - issue.CreatedAt >= TimeSpan.FromMinutes(_settings.NotifyUnacknowledgedMinutes))
+							else if (ShowNotificationForIssue(issue) && _settings.NotifyUnacknowledgedMinutes >= 0 && issue.RetrievedAt - issue.CreatedAt >= TimeSpan.FromMinutes(_settings.NotifyUnacknowledgedMinutes))
 							{
 								reason |= IssueAlertReason.UnacknowledgedTimer;
 							}
 						}
-						if(ShowNotificationForIssue(issue) && _settings.NotifyUnresolvedMinutes >= 0 && issue.RetrievedAt - issue.CreatedAt >= TimeSpan.FromMinutes(_settings.NotifyUnresolvedMinutes))
+						if (ShowNotificationForIssue(issue) && _settings.NotifyUnresolvedMinutes >= 0 && issue.RetrievedAt - issue.CreatedAt >= TimeSpan.FromMinutes(_settings.NotifyUnresolvedMinutes))
 						{
 							reason |= IssueAlertReason.UnresolvedTimer;
 						}
 
 						IssueAlertReason prevReason;
-						if(issueMonitor.IssueIdToAlertReason.TryGetValue(issue.Id, out prevReason))
+						if (issueMonitor.IssueIdToAlertReason.TryGetValue(issue.Id, out prevReason))
 						{
 							reason &= ~prevReason;
 						}
 					}
 
 					IssueAlertWindow? alertWindow = _alertWindows.FirstOrDefault(x => x.IssueMonitor == issueMonitor && x.Issue.Id == issue.Id);
-					if(alertWindow == null)
+					if (alertWindow == null)
 					{
-						if(reason != 0)
+						if (reason != 0)
 						{
 							ShowAlertWindow(issueMonitor, issue, reason);
 						}
 					}
 					else
 					{
-						if(reason != 0)
+						if (reason != 0)
 						{
 							alertWindow.SetIssue(issue, reason);
 						}
@@ -1650,10 +1649,10 @@ namespace UnrealGameSync
 			}
 
 			// Close any alert windows which don't have an active issues
-			for(int idx = 0; idx < _alertWindows.Count; idx++)
+			for (int idx = 0; idx < _alertWindows.Count; idx++)
 			{
 				IssueAlertWindow alertWindow = _alertWindows[idx];
-				if(!allIssues.Contains(alertWindow.Issue))
+				if (!allIssues.Contains(alertWindow.Issue))
 				{
 					alertWindow.IssueMonitor.IssueIdToAlertReason.Remove(alertWindow.Issue.Id);
 					CloseAlertWindow(alertWindow);
@@ -1719,10 +1718,10 @@ namespace UnrealGameSync
 
 		void ShowIssueDetails(IssueAlertWindow alert)
 		{
-			for(int idx = 0; idx < TabControl.GetTabCount(); idx++)
+			for (int idx = 0; idx < TabControl.GetTabCount(); idx++)
 			{
 				WorkspaceControl? workspace = TabControl.GetTabData(idx) as WorkspaceControl;
-				if(workspace != null && workspace.GetIssueMonitor() == alert.IssueMonitor)
+				if (workspace != null && workspace.GetIssueMonitor() == alert.IssueMonitor)
 				{
 					workspace.ShowIssueDetails(alert.Issue);
 					break;
@@ -1737,12 +1736,12 @@ namespace UnrealGameSync
 
 			_alertWindows.Remove(alert);
 
-			for(int idx = 0; idx < _alertWindows.Count; idx++)
+			for (int idx = 0; idx < _alertWindows.Count; idx++)
 			{
 				SetAlertWindowPosition(idx);
 			}
 
-			if(_alertWindows.Count == 0)
+			if (_alertWindows.Count == 0)
 			{
 				UpdateAlertPositionsTimer.Enabled = false;
 			}
@@ -1804,7 +1803,7 @@ namespace UnrealGameSync
 		public void ReleaseIssueMonitor(IssueMonitor issueMonitor)
 		{
 			int index = _workspaceIssueMonitors.FindIndex(x => x._issueMonitor == issueMonitor);
-			if(index != -1)
+			if (index != -1)
 			{
 				WorkspaceIssueMonitor workspaceIssueMonitor = _workspaceIssueMonitors[index];
 				workspaceIssueMonitor._refCount--;
diff --git a/Engine/Source/Programs/UnrealGameSync/UnrealGameSync/Forms/OpenProjectWindow.cs b/Engine/Source/Programs/UnrealGameSync/UnrealGameSync/Forms/OpenProjectWindow.cs
--- a/Engine/Source/Programs/UnrealGameSync/UnrealGameSync/Forms/OpenProjectWindow.cs
+++ b/Engine/Source/Programs/UnrealGameSync/UnrealGameSync/Forms/OpenProjectWindow.cs
@@ -1,9 +1,5 @@
 // Copyright Epic Games, Inc. All Rights Reserved.
 
-using EpicGames.OIDC;
-using EpicGames.Perforce;
-using Microsoft.Extensions.DependencyInjection;
-using Microsoft.Extensions.Logging;
 using System;
 using System.Diagnostics.CodeAnalysis;
 using System.Drawing;
@@ -12,6 +8,9 @@ using System.Text;
 using System.Threading;
 using System.Threading.Tasks;
 using System.Windows.Forms;
+using EpicGames.Perforce;
+using Microsoft.Extensions.DependencyInjection;
+using Microsoft.Extensions.Logging;
 
 #nullable enable
 
@@ -176,15 +175,15 @@ namespace UnrealGameSync
 
 		private void UpdateOkButton()
 		{
-			OkBtn.Enabled = WorkspaceRadioBtn.Checked? TryGetClientPath(out _) : TryGetLocalPath(out _);
+			OkBtn.Enabled = WorkspaceRadioBtn.Checked ? TryGetClientPath(out _) : TryGetLocalPath(out _);
 		}
 
 		private void WorkspaceNewBtn_Click(object sender, EventArgs e)
 		{
 			WorkspaceRadioBtn.Checked = true;
-			
+
 			string? workspaceName;
-			if(NewWorkspaceWindow.ShowModal(this, Perforce, null, WorkspaceNameTextBox.Text, _serviceProvider, out workspaceName))
+			if (NewWorkspaceWindow.ShowModal(this, Perforce, null, WorkspaceNameTextBox.Text, _serviceProvider, out workspaceName))
 			{
 				WorkspaceNameTextBox.Text = workspaceName;
 				UpdateOkButton();
@@ -196,7 +195,7 @@ namespace UnrealGameSync
 			WorkspaceRadioBtn.Checked = true;
 
 			string? workspaceName = WorkspaceNameTextBox.Text;
-			if(SelectWorkspaceWindow.ShowModal(this, Perforce, workspaceName, _serviceProvider, out workspaceName))
+			if (SelectWorkspaceWindow.ShowModal(this, Perforce, workspaceName, _serviceProvider, out workspaceName))
 			{
 				WorkspaceNameTextBox.Text = workspaceName;
 			}
@@ -207,10 +206,10 @@ namespace UnrealGameSync
 			WorkspaceRadioBtn.Checked = true;
 
 			string? workspaceName;
-			if(TryGetWorkspaceName(out workspaceName))
+			if (TryGetWorkspaceName(out workspaceName))
 			{
 				string? workspacePath = WorkspacePathTextBox.Text.Trim();
-				if(SelectProjectFromWorkspaceWindow.ShowModal(this, Perforce, workspaceName, workspacePath, _serviceProvider, out workspacePath))
+				if (SelectProjectFromWorkspaceWindow.ShowModal(this, Perforce, workspaceName, workspacePath, _serviceProvider, out workspacePath))
 				{
 					WorkspacePathTextBox.Text = workspacePath;
 					UpdateOkButton();
@@ -221,7 +220,7 @@ namespace UnrealGameSync
 		private bool TryGetWorkspaceName([NotNullWhen(true)] out string? workspaceName)
 		{
 			string text = WorkspaceNameTextBox.Text.Trim();
-			if(text.Length == 0)
+			if (text.Length == 0)
 			{
 				workspaceName = null;
 				return false;
@@ -234,14 +233,14 @@ namespace UnrealGameSync
 		private bool TryGetClientPath([NotNullWhen(true)] out string? clientPath)
 		{
 			string? workspaceName;
-			if(!TryGetWorkspaceName(out workspaceName))
+			if (!TryGetWorkspaceName(out workspaceName))
 			{
 				clientPath = null;
 				return false;
 			}
 
 			string workspacePath = WorkspacePathTextBox.Text.Trim();
-			if(workspacePath.Length == 0 || workspacePath[0] != '/')
+			if (workspacePath.Length == 0 || workspacePath[0] != '/')
 			{
 				clientPath = null;
 				return false;
@@ -254,7 +253,7 @@ namespace UnrealGameSync
 		private bool TryGetLocalPath(out string? localPath)
 		{
 			string localFile = LocalFileTextBox.Text.Trim();
-			if(localFile.Length == 0)
+			if (localFile.Length == 0)
 			{
 				localPath = null;
 				return false;
@@ -266,10 +265,10 @@ namespace UnrealGameSync
 
 		private bool TryGetSelectedProject([NotNullWhen(true)] out UserSelectedProjectSettings? project)
 		{
-			if(WorkspaceRadioBtn.Checked)
+			if (WorkspaceRadioBtn.Checked)
 			{
 				string? clientPath;
-				if(TryGetClientPath(out clientPath))
+				if (TryGetClientPath(out clientPath))
 				{
 					project = new UserSelectedProjectSettings(_serverAndPortOverride, _userNameOverride, UserSelectedProjectType.Client, clientPath, null);
 					return true;
@@ -278,7 +277,7 @@ namespace UnrealGameSync
 			else
 			{
 				string? localPath;
-				if(TryGetLocalPath(out localPath))
+				if (TryGetLocalPath(out localPath))
 				{
 					project = new UserSelectedProjectSettings(_serverAndPortOverride, _userNameOverride, UserSelectedProjectType.Local, null, localPath);
 					return true;
@@ -292,7 +291,7 @@ namespace UnrealGameSync
 		private void OkBtn_Click(object sender, EventArgs e)
 		{
 			UserSelectedProjectSettings? selectedProject;
-			if(TryGetSelectedProject(out selectedProject))
+			if (TryGetSelectedProject(out selectedProject))
 			{
 				ILogger<OpenProjectInfo> logger = _serviceProvider.GetRequiredService<ILogger<OpenProjectInfo>>();
 
@@ -324,7 +323,7 @@ namespace UnrealGameSync
 
 		private void ChangeLink_LinkClicked(object sender, LinkLabelLinkClickedEventArgs e)
 		{
-			if(ConnectWindow.ShowModal(this, _defaultSettings, ref _serverAndPortOverride, ref _userNameOverride, _serviceProvider))
+			if (ConnectWindow.ShowModal(this, _defaultSettings, ref _serverAndPortOverride, ref _userNameOverride, _serviceProvider))
 			{
 				UpdateServerLabel();
 			}
@@ -335,10 +334,10 @@ namespace UnrealGameSync
 			LocalFileRadioBtn.Checked = true;
 
 			using OpenFileDialog dialog = new OpenFileDialog();
-			dialog.Filter = "Project files (*.uproject)|*.uproject|Project directory lists (*.uprojectdirs)|*.uprojectdirs|All supported files (*.uproject;*.uprojectdirs)|*.uproject;*.uprojectdirs|All files (*.*)|*.*" ;
+			dialog.Filter = "Project files (*.uproject)|*.uproject|Project directory lists (*.uprojectdirs)|*.uprojectdirs|All supported files (*.uproject;*.uprojectdirs)|*.uproject;*.uprojectdirs|All files (*.*)|*.*";
 			dialog.FilterIndex = _settings.FilterIndex;
-			
-			if(!String.IsNullOrEmpty(LocalFileTextBox.Text))
+
+			if (!String.IsNullOrEmpty(LocalFileTextBox.Text))
 			{
 				try
 				{
diff --git a/Engine/Source/Programs/UnrealGameSync/UnrealGameSync/Forms/UpdateSettingsWindow.cs b/Engine/Source/Programs/UnrealGameSync/UnrealGameSync/Forms/UpdateSettingsWindow.cs
--- a/Engine/Source/Programs/UnrealGameSync/UnrealGameSync/Forms/UpdateSettingsWindow.cs
+++ b/Engine/Source/Programs/UnrealGameSync/UnrealGameSync/Forms/UpdateSettingsWindow.cs
@@ -140,10 +140,7 @@ namespace UnrealGameSync
 				return;
 			}
 
-			if (task != null)
-			{
-				PromptLabel.Text = task.Error;
-			}
+			PromptLabel.Text = task.Error;
 
 			_logText = logger.Render(Environment.NewLine);
 			ViewLogBtn.Visible = true;
diff --git a/Engine/Source/Programs/UnrealGameSync/UnrealGameSync/Launcher.cs b/Engine/Source/Programs/UnrealGameSync/UnrealGameSync/Launcher.cs
--- a/Engine/Source/Programs/UnrealGameSync/UnrealGameSync/Launcher.cs
+++ b/Engine/Source/Programs/UnrealGameSync/UnrealGameSync/Launcher.cs
@@ -1,11 +1,5 @@
 // Copyright Epic Games, Inc. All Rights Reserved.
 
-using EpicGames.Core;
-using EpicGames.Horde;
-using EpicGames.Horde.Tools;
-using EpicGames.Perforce;
-using Microsoft.Extensions.DependencyInjection;
-using Microsoft.Extensions.Logging;
 using System;
 using System.Collections.Generic;
 using System.Diagnostics;
@@ -13,12 +7,16 @@ using System.Diagnostics.CodeAnalysis;
 using System.IO;
 using System.IO.Compression;
 using System.Linq;
-using System.Reflection;
 using System.Text.Json;
-using System.Text.Json.Nodes;
 using System.Threading;
 using System.Threading.Tasks;
 using System.Windows.Forms;
+using EpicGames.Core;
+using EpicGames.Horde;
+using EpicGames.Horde.Tools;
+using EpicGames.Perforce;
+using Microsoft.Extensions.DependencyInjection;
+using Microsoft.Extensions.Logging;
 
 namespace UnrealGameSync
 {
@@ -349,7 +347,7 @@ namespace UnrealGameSync
 				text = File.ReadAllText(fileName);
 				return true;
 			}
-			catch(Exception)
+			catch (Exception)
 			{
 				text = null;
 				return false;
@@ -363,7 +361,7 @@ namespace UnrealGameSync
 				File.WriteAllText(fileName, text);
 				return true;
 			}
-			catch(Exception)
+			catch (Exception)
 			{
 				return false;
 			}
@@ -376,7 +374,7 @@ namespace UnrealGameSync
 				Directory.CreateDirectory(directoryName);
 				return true;
 			}
-			catch(Exception)
+			catch (Exception)
 			{
 				return false;
 			}
@@ -389,7 +387,7 @@ namespace UnrealGameSync
 				Directory.Delete(directoryName, true);
 				return true;
 			}
-			catch(Exception)
+			catch (Exception)
 			{
 				return false;
 			}
@@ -400,18 +398,18 @@ namespace UnrealGameSync
 			try
 			{
 				DirectoryInfo directory = new DirectoryInfo(directoryName);
-				foreach(FileInfo childFile in directory.EnumerateFiles("*", SearchOption.AllDirectories))
+				foreach (FileInfo childFile in directory.EnumerateFiles("*", SearchOption.AllDirectories))
 				{
 					childFile.Attributes &= ~FileAttributes.ReadOnly;
 					childFile.Delete();
 				}
-				foreach(DirectoryInfo childDirectory in directory.EnumerateDirectories())
+				foreach (DirectoryInfo childDirectory in directory.EnumerateDirectories())
 				{
 					SafeDeleteDirectory(childDirectory.FullName);
 				}
 				return true;
 			}
-			catch(Exception)
+			catch (Exception)
 			{
 				return false;
 			}
diff --git a/Engine/Source/Programs/UnrealGameSync/UnrealGameSync/PerforceMonitor.cs b/Engine/Source/Programs/UnrealGameSync/UnrealGameSync/PerforceMonitor.cs
--- a/Engine/Source/Programs/UnrealGameSync/UnrealGameSync/PerforceMonitor.cs
+++ b/Engine/Source/Programs/UnrealGameSync/UnrealGameSync/PerforceMonitor.cs
@@ -1,11 +1,5 @@
 // Copyright Epic Games, Inc. All Rights Reserved.
 
-using EpicGames.Core;
-using EpicGames.OIDC;
-using EpicGames.Perforce;
-using EpicGames.Horde;
-using Microsoft.Extensions.DependencyInjection;
-using Microsoft.Extensions.Logging;
 using System;
 using System.Collections.Generic;
 using System.Diagnostics;
@@ -13,6 +7,11 @@ using System.Diagnostics.CodeAnalysis;
 using System.Linq;
 using System.Threading;
 using System.Threading.Tasks;
+using EpicGames.Core;
+using EpicGames.Horde;
+using EpicGames.Perforce;
+using Microsoft.Extensions.DependencyInjection;
+using Microsoft.Extensions.Logging;
 
 namespace UnrealGameSync
 {
@@ -44,7 +43,7 @@ namespace UnrealGameSync
 #pragma warning restore CA2213
 		int _pendingMaxChangesValue;
 		SortedSet<ChangesRecord> _changes = new SortedSet<ChangesRecord>(new PerforceChangeSorter());
-		readonly SortedDictionary<int, PerforceChangeDetails> _changeDetails = new SortedDictionary<int,PerforceChangeDetails>();
+		readonly SortedDictionary<int, PerforceChangeDetails> _changeDetails = new SortedDictionary<int, PerforceChangeDetails>();
 		readonly SortedSet<int> _promotedChangeNumbers = new SortedSet<int>();
 		List<BaseArchiveInfo> _archives = new List<BaseArchiveInfo>();
 		readonly AsyncEvent _refreshEvent = new AsyncEvent();
@@ -134,16 +133,16 @@ namespace UnrealGameSync
 		public int PendingMaxChanges
 		{
 			get => _pendingMaxChangesValue;
-			set 
-			{ 
-				lock(_lockObject)
-				{ 
-					if(value != _pendingMaxChangesValue)
-					{ 
-						_pendingMaxChangesValue = value; 
-						_refreshEvent.Set(); 
-					} 
-				} 
+			set
+			{
+				lock (_lockObject)
+				{
+					if (value != _pendingMaxChangesValue)
+					{
+						_pendingMaxChangesValue = value;
+						_refreshEvent.Set();
+					}
+				}
 			}
 		}
 
@@ -193,7 +192,7 @@ namespace UnrealGameSync
 				await UpdateArchivesAsync(perforce, cancellationToken);
 			}
 
-			while(!cancellationToken.IsCancellationRequested)
+			while (!cancellationToken.IsCancellationRequested)
 			{
 				Stopwatch timer = Stopwatch.StartNew();
 				Task nextRefreshTask = _refreshEvent.Task;
@@ -236,14 +235,14 @@ namespace UnrealGameSync
 					}
 				}
 
-                // Wait for another request, or scan for new builds after a timeout
+				// Wait for another request, or scan for new builds after a timeout
 				// Add random deviation to refresh event, to try and combat many UGS clients getting in sync and DDoSing Perforce
-                TimeSpan baseDelay = TimeSpan.FromMinutes(IsActive ? 5 : 30);
+				TimeSpan baseDelay = TimeSpan.FromMinutes(IsActive ? 5 : 30);
 
-                Random random = new Random();
-                TimeSpan randomDeviation = TimeSpan.FromSeconds(random.Next(0, 60));
+				Random random = new Random();
+				TimeSpan randomDeviation = TimeSpan.FromSeconds(random.Next(0, 60));
 
-                Task delayTask = Task.Delay(baseDelay + randomDeviation, cancellationToken);
+				Task delayTask = Task.Delay(baseDelay + randomDeviation, cancellationToken);
 				await Task.WhenAny(nextRefreshTask, delayTask);
 			}
 		}
@@ -298,7 +297,7 @@ namespace UnrealGameSync
 			}
 
 			// Read any new changes
-			int? minChangeNumber = (maxChanges > CurrentMaxChanges || newestChangeNumber == -1)? null : (newestChangeNumber + 1);
+			int? minChangeNumber = (maxChanges > CurrentMaxChanges || newestChangeNumber == -1) ? null : (newestChangeNumber + 1);
 			List<ChangesRecord> newChanges = await Utility.EnumerateChanges(perforce, depotPaths, minChangeNumber, null, maxChanges, cancellationToken).ToListAsync(cancellationToken);
 
 			// Remove anything we already have
@@ -339,21 +338,21 @@ namespace UnrealGameSync
 			const string roboMergePrefix = "#ROBOMERGE-AUTHOR:";
 			foreach (ChangesRecord change in newChanges)
 			{
-				if(change.Description != null && change.Description.StartsWith(roboMergePrefix, StringComparison.Ordinal))
+				if (change.Description != null && change.Description.StartsWith(roboMergePrefix, StringComparison.Ordinal))
 				{
 					int startIdx = roboMergePrefix.Length;
-					while(startIdx < change.Description.Length && change.Description[startIdx] == ' ')
+					while (startIdx < change.Description.Length && change.Description[startIdx] == ' ')
 					{
 						startIdx++;
 					}
 
 					int endIdx = startIdx;
-					while(endIdx < change.Description.Length && !Char.IsWhiteSpace(change.Description[endIdx]))
+					while (endIdx < change.Description.Length && !Char.IsWhiteSpace(change.Description[endIdx]))
 					{
 						endIdx++;
 					}
 
-					if(endIdx > startIdx)
+					if (endIdx > startIdx)
 					{
 						change.User = change.Description.Substring(startIdx, endIdx - startIdx);
 						change.Description = "ROBOMERGE: " + change.Description.Substring(endIdx).TrimStart();
@@ -362,21 +361,21 @@ namespace UnrealGameSync
 			}
 
 			// Process the new changes received
-			if(newChanges.Count > 0 || maxChanges < CurrentMaxChanges)
+			if (newChanges.Count > 0 || maxChanges < CurrentMaxChanges)
 			{
 				// Insert them into the builds list
-				lock(_lockObject)
+				lock (_lockObject)
 				{
 					_changes.UnionWith(newChanges);
-					if(_changes.Count > maxChanges)
+					if (_changes.Count > maxChanges)
 					{
 						// Remove changes to shrink it to the max requested size, being careful to avoid removing changes that would affect our ability to correctly
 						// show the availability for content changes using zipped binaries.
 						SortedSet<ChangesRecord> trimmedChanges = new SortedSet<ChangesRecord>(new PerforceChangeSorter());
-						foreach(ChangesRecord change in _changes)
+						foreach (ChangesRecord change in _changes)
 						{
 							trimmedChanges.Add(change);
-							if(trimmedChanges.Count >= maxChanges && _archives.Any(x => x.ChangeNumberToArchiveKey.Count == 0 || x.ChangeNumberToArchiveKey.ContainsKey(change.Number) || x.ChangeNumberToArchiveKey.First().Key > change.Number))
+							if (trimmedChanges.Count >= maxChanges && _archives.Any(x => x.ChangeNumberToArchiveKey.Count == 0 || x.ChangeNumberToArchiveKey.ContainsKey(change.Number) || x.ChangeNumberToArchiveKey.First().Key > change.Number))
 							{
 								break;
 							}
@@ -388,9 +387,9 @@ namespace UnrealGameSync
 
 				// Find the last submitted change by the current user
 				int newLastChangeByCurrentUser = -1;
-				foreach(ChangesRecord change in _changes)
+				foreach (ChangesRecord change in _changes)
 				{
-					if(String.Equals(change.User, perforce.Settings.UserName, StringComparison.OrdinalIgnoreCase))
+					if (String.Equals(change.User, perforce.Settings.UserName, StringComparison.OrdinalIgnoreCase))
 					{
 						newLastChangeByCurrentUser = Math.Max(newLastChangeByCurrentUser, change.Number);
 					}
@@ -415,11 +414,11 @@ namespace UnrealGameSync
 
 			// Find the changes we need to query
 			List<int> queryChangeNumbers = new List<int>();
-			lock(_lockObject)
+			lock (_lockObject)
 			{
-				foreach(ChangesRecord change in _changes)
+				foreach (ChangesRecord change in _changes)
 				{
-					if(!_changeDetails.ContainsKey(change.Number))
+					if (!_changeDetails.ContainsKey(change.Number))
 					{
 						queryChangeNumbers.Add(change.Number);
 					}
@@ -431,7 +430,7 @@ namespace UnrealGameSync
 			using (CancellationTokenSource cancellationSource = new CancellationTokenSource())
 			{
 				Task notifyTask = Task.CompletedTask;
-				await foreach(PerforceChangeDetails details in Utility.EnumerateChangeDetails(perforce, queryChangeNumbers.ToAsyncEnumerable(), codeRules, cancellationToken))
+				await foreach (PerforceChangeDetails details in Utility.EnumerateChangeDetails(perforce, queryChangeNumbers.ToAsyncEnumerable(), codeRules, cancellationToken))
 				{
 					cancellationToken.ThrowIfCancellationRequested();
 
@@ -469,12 +468,12 @@ namespace UnrealGameSync
 
 			// Find the last submitted code change by the current user
 			int newLastCodeChangeByCurrentUser = -1;
-			foreach(ChangesRecord change in _changes)
+			foreach (ChangesRecord change in _changes)
 			{
-				if(String.Equals(change.User, perforce.Settings.UserName, StringComparison.OrdinalIgnoreCase))
+				if (String.Equals(change.User, perforce.Settings.UserName, StringComparison.OrdinalIgnoreCase))
 				{
 					PerforceChangeDetails? otherDetails;
-					if(_changeDetails.TryGetValue(change.Number, out otherDetails) && otherDetails.ContainsCode)
+					if (_changeDetails.TryGetValue(change.Number, out otherDetails) && otherDetails.ContainsCode)
 					{
 						newLastCodeChangeByCurrentUser = Math.Max(newLastCodeChangeByCurrentUser, change.Number);
 					}
@@ -485,7 +484,7 @@ namespace UnrealGameSync
 			// Notify the main window that we've got an update
 			_synchronizationContext.Post(_ => OnUpdateMetadata?.Invoke(), null);
 
-			if(_localConfigFiles.Any(x => FileReference.GetLastWriteTimeUtc(x.Key) != x.Value))
+			if (_localConfigFiles.Any(x => FileReference.GetLastWriteTimeUtc(x.Key) != x.Value))
 			{
 				await UpdateProjectConfigFileAsync(perforce, cancellationToken);
 				_synchronizationContext.Post(_ => OnUpdateMetadata?.Invoke(), null);
@@ -521,7 +520,7 @@ namespace UnrealGameSync
 
 		public List<ChangesRecord> GetChanges()
 		{
-			lock(_lockObject)
+			lock (_lockObject)
 			{
 				return new List<ChangesRecord>(_changes);
 			}
@@ -529,7 +528,7 @@ namespace UnrealGameSync
 
 		public bool TryGetChangeDetails(int number, [NotNullWhen(true)] out PerforceChangeDetails? details)
 		{
-			lock(_lockObject)
+			lock (_lockObject)
 			{
 				return _changeDetails.TryGetValue(number, out details);
 			}
@@ -537,7 +536,7 @@ namespace UnrealGameSync
 
 		public HashSet<int> GetPromotedChangeNumbers()
 		{
-			lock(_lockObject)
+			lock (_lockObject)
 			{
 				return new HashSet<int>(_promotedChangeNumbers);
 			}
diff --git a/Engine/Source/Programs/UnrealGameSync/UnrealGameSync/Program.cs b/Engine/Source/Programs/UnrealGameSync/UnrealGameSync/Program.cs
--- a/Engine/Source/Programs/UnrealGameSync/UnrealGameSync/Program.cs
+++ b/Engine/Source/Programs/UnrealGameSync/UnrealGameSync/Program.cs
@@ -97,7 +97,9 @@ namespace UnrealGameSync
 				// Don't auto install (or - more importantly- auto *un-install*) the winforms sync context. We want to be able to access it from the
 				// constructor of our ApplicationContext, which will be after the temporary install/uninstall prompted by spawning the settings dialog.
 				WindowsFormsSynchronizationContext.AutoInstall = false;
-				SynchronizationContext.SetSynchronizationContext(new WindowsFormsSynchronizationContext());
+
+				using WindowsFormsSynchronizationContext synchronizationContext = new WindowsFormsSynchronizationContext();
+				SynchronizationContext.SetSynchronizationContext(synchronizationContext);
 
 				using (EventWaitHandle activateEvent = new EventWaitHandle(false, EventResetMode.AutoReset, "ActivateUnrealGameSync"))
 				{
@@ -229,12 +231,12 @@ namespace UnrealGameSync
 
 					using (ITelemetrySink telemetrySink = CreateTelemetrySink(launcherSettings.PerforceUserName, sessionId, telemetryLogger))
 					{
-						ITelemetrySink? prevTelemetrySink = Telemetry.ActiveSink;
+						ITelemetrySink? prevTelemetrySink = UgsTelemetry.ActiveSink;
 						try
 						{
-							Telemetry.ActiveSink = telemetrySink;
+							UgsTelemetry.ActiveSink = telemetrySink;
 
-							Telemetry.SendEvent("Startup", new { User = Environment.UserName, Machine = System.Net.Dns.GetHostName() });
+							UgsTelemetry.SendEvent("Startup", new { User = Environment.UserName, Machine = System.Net.Dns.GetHostName() });
 
 							AppDomain.CurrentDomain.UnhandledException += CurrentDomain_UnhandledException;
 
@@ -261,12 +263,12 @@ namespace UnrealGameSync
 						}
 						catch (Exception ex)
 						{
-							Telemetry.SendEvent("Crash", new { Exception = ex.ToString() });
+							UgsTelemetry.SendEvent("Crash", new { Exception = ex.ToString() });
 							throw;
 						}
 						finally
 						{
-							Telemetry.ActiveSink = prevTelemetrySink;
+							UgsTelemetry.ActiveSink = prevTelemetrySink;
 						}
 					}
 				}
@@ -337,7 +339,7 @@ namespace UnrealGameSync
 			Exception? ex = args.ExceptionObject as Exception;
 			if(ex != null)
 			{
-				Telemetry.SendEvent("Crash", new {Exception = ex.ToString()});
+				UgsTelemetry.SendEvent("Crash", new {Exception = ex.ToString()});
 			}
 		}
 
diff --git a/Engine/Source/Programs/UnrealGameSync/UnrealGameSync/ProgramApplicationContext.cs b/Engine/Source/Programs/UnrealGameSync/UnrealGameSync/ProgramApplicationContext.cs
--- a/Engine/Source/Programs/UnrealGameSync/UnrealGameSync/ProgramApplicationContext.cs
+++ b/Engine/Source/Programs/UnrealGameSync/UnrealGameSync/ProgramApplicationContext.cs
@@ -1,9 +1,5 @@
 ﻿// Copyright Epic Games, Inc. All Rights Reserved.
 
-using EpicGames.Core;
-using EpicGames.Perforce;
-using Microsoft.Extensions.DependencyInjection;
-using Microsoft.Extensions.Logging;
 using System;
 using System.Collections.Generic;
 using System.ComponentModel;
@@ -14,8 +10,11 @@ using System.Reflection;
 using System.Threading;
 using System.Threading.Tasks;
 using System.Windows.Forms;
-using EpicGames.OIDC;
+using EpicGames.Core;
 using EpicGames.Horde;
+using EpicGames.Perforce;
+using Microsoft.Extensions.DependencyInjection;
+using Microsoft.Extensions.Logging;
 
 #nullable enable
 
@@ -83,7 +82,7 @@ namespace UnrealGameSync
 
 			// Make sure a synchronization context is set. We spawn a bunch of threads (eg. UpdateMonitor) at startup, and need to make sure we can post messages 
 			// back to the main thread at any time.
-			if(SynchronizationContext.Current == null)
+			if (SynchronizationContext.Current == null)
 			{
 				_synchronizationContext = new WindowsFormsSynchronizationContext();
 				SynchronizationContext.SetSynchronizationContext(_synchronizationContext);
@@ -94,20 +93,20 @@ namespace UnrealGameSync
 
 			// Read the user's settings
 			_settings = UserSettings.Create(dataFolder, serviceProvider.GetRequiredService<ILogger<UserSettings>>());
-			if(!String.IsNullOrEmpty(projectFileName))
+			if (!String.IsNullOrEmpty(projectFileName))
 			{
 				string fullProjectFileName = Path.GetFullPath(projectFileName);
-				if(!_settings.OpenProjects.Any(x => x.LocalPath != null && String.Equals(x.LocalPath, fullProjectFileName, StringComparison.OrdinalIgnoreCase)))
+				if (!_settings.OpenProjects.Any(x => x.LocalPath != null && String.Equals(x.LocalPath, fullProjectFileName, StringComparison.OrdinalIgnoreCase)))
 				{
 					_settings.OpenProjects.Add(new UserSelectedProjectSettings(null, null, UserSelectedProjectType.Local, null, fullProjectFileName));
 				}
 			}
 
 			// Update the settings to the latest version
-			if(_settings.Version < UserSettingsVersion.Latest)
+			if (_settings.Version < UserSettingsVersion.Latest)
 			{
 				// Clear out the server settings for anything using the default server
-				if(_settings.Version < UserSettingsVersion.DefaultServerSettings)
+				if (_settings.Version < UserSettingsVersion.DefaultServerSettings)
 				{
 					_logger.LogInformation("Clearing project settings for default server");
 					for (int idx = 0; idx < _settings.OpenProjects.Count; idx++)
@@ -207,9 +206,9 @@ namespace UnrealGameSync
 			_startupWindow = new ModalTaskWindow("Opening Projects", "Opening projects, please wait...", FormStartPosition.CenterScreen, _startupTask, _startupCancellationSource);
 			_components.Add(_startupWindow);
 
-			if(restoreState)
+			if (restoreState)
 			{
-				if(_settings.WindowVisible)
+				if (_settings.WindowVisible)
 				{
 					_startupWindow.Show();
 				}
@@ -303,7 +302,7 @@ namespace UnrealGameSync
 			if (visible)
 			{
 				_mainWindowInstance.Show();
-				if(!_restoreState)
+				if (!_restoreState)
 				{
 					_mainWindowInstance.Activate();
 				}
@@ -318,7 +317,7 @@ namespace UnrealGameSync
 
 		private void OnActivationListenerCallback()
 		{
-			if(_mainWindowInstance != null && !_mainWindowInstance.IsDisposed)
+			if (_mainWindowInstance != null && !_mainWindowInstance.IsDisposed)
 			{
 				_mainWindowInstance.ShowAndActivate();
 			}
@@ -331,9 +330,9 @@ namespace UnrealGameSync
 
 		private void OnUpdateAvailable(UpdateType type)
 		{
-			if(_mainWindowInstance != null && !_isClosing)
+			if (_mainWindowInstance != null && !_isClosing)
 			{
-				if(type == UpdateType.UserInitiated || _mainWindowInstance.CanPerformUpdate())
+				if (type == UpdateType.UserInitiated || _mainWindowInstance.CanPerformUpdate())
 				{
 					_isClosing = true;
 					_mainWindowInstance.ForceClose();
@@ -343,7 +342,7 @@ namespace UnrealGameSync
 		}
 
 		private void OnUpdateAvailableCallback(UpdateType type)
-		{ 
+		{
 			_mainThreadSynchronizationContext.Post((o) => OnUpdateAvailable(type), null);
 		}
 
@@ -351,7 +350,7 @@ namespace UnrealGameSync
 		{
 			base.Dispose(disposing);
 
-			if(_activationListener != null)
+			if (_activationListener != null)
 			{
 				_activationListener.OnActivate -= OnActivationListenerAsyncCallback;
 				_activationListener.Stop();
@@ -371,13 +370,13 @@ namespace UnrealGameSync
 				_notifyIcon = null!;
 			}
 
-			if(_mainWindowInstance != null)
+			if (_mainWindowInstance != null)
 			{
 				_mainWindowInstance.ForceClose();
 				_mainWindowInstance = null!;
 			}
 
-			if(_startupWindow != null)
+			if (_startupWindow != null)
 			{
 				_startupWindow.Close();
 				_startupWindow = null;
@@ -443,7 +442,7 @@ namespace UnrealGameSync
 				_startupWindow = null;
 			}
 
-			if(_mainWindowInstance != null)
+			if (_mainWindowInstance != null)
 			{
 				_mainWindowInstance.ForceClose();
 				_mainWindowInstance = null;
@@ -456,7 +455,7 @@ namespace UnrealGameSync
 		{
 			base.ExitThreadCore();
 
-			if(_notifyIcon != null)
+			if (_notifyIcon != null)
 			{
 				_notifyIcon.Visible = false;
 			}
diff --git a/Engine/Source/Programs/UnrealGameSync/UnrealGameSyncAutomation/Program.cs b/Engine/Source/Programs/UnrealGameSync/UnrealGameSyncAutomation/Program.cs
--- a/Engine/Source/Programs/UnrealGameSync/UnrealGameSyncAutomation/Program.cs
+++ b/Engine/Source/Programs/UnrealGameSync/UnrealGameSyncAutomation/Program.cs
@@ -1,11 +1,11 @@
 // Copyright Epic Games, Inc. All Rights Reserved.
 
-using Microsoft.Win32;
 using System;
 using System.IO;
 using System.Net;
 using System.Net.Sockets;
 using System.Text;
+using Microsoft.Win32;
 
 namespace UnrealGameSyncAutomation
 {
@@ -32,14 +32,14 @@ namespace UnrealGameSyncAutomation
 		{
 			// Create the request data
 			MemoryStream InputDataStream = new MemoryStream();
-			using(BinaryWriter Writer = new BinaryWriter(InputDataStream))
+			using (BinaryWriter Writer = new BinaryWriter(InputDataStream))
 			{
 				Writer.Write("//UE4/Main");
 				Writer.Write("/Samples/Games/ShooterGame/ShooterGame.uproject");
 			}
 
 			Tuple<AutomationRequestResult, byte[]> Response = SendRequest(AutomationRequestType.OpenProject, InputDataStream.ToArray());
-		
+
 			string ResponseString = Encoding.UTF8.GetString(Response.Item2);
 			Console.WriteLine("{0}: {1}", Response.Item1, ResponseString);
 		}
@@ -47,20 +47,20 @@ namespace UnrealGameSyncAutomation
 		static Tuple<AutomationRequestResult, byte[]> SendRequest(AutomationRequestType Request, byte[] RequestData)
 		{
 			int PortNumber = (int)Registry.GetValue("HKEY_CURRENT_USER\\Software\\Epic Games\\UnrealGameSync", "AutomationPort", null);
-			using(TcpClient Client = new TcpClient())
+			using (TcpClient Client = new TcpClient())
 			{
 				Client.Connect(new IPEndPoint(IPAddress.Loopback, PortNumber));
 
-				using(NetworkStream Stream = Client.GetStream())
+				using (NetworkStream Stream = Client.GetStream())
 				{
 					// Send the request
-					BinaryWriter Writer = new BinaryWriter(Stream);
+					using BinaryWriter Writer = new BinaryWriter(Stream);
 					Writer.Write((int)Request);
 					Writer.Write(RequestData.Length);
 					Writer.Write(RequestData);
 
 					// Read the response
-					BinaryReader Reader = new BinaryReader(Stream);
+					using BinaryReader Reader = new BinaryReader(Stream);
 					AutomationRequestResult Result = (AutomationRequestResult)Reader.ReadInt32();
 					int ResponseLength = Reader.ReadInt32();
 					byte[] ResponseData = Reader.ReadBytes(ResponseLength);
diff --git a/Engine/Source/Programs/UnrealGameSync/UnrealGameSyncAutomation/Properties/AssemblyInfo.cs b/Engine/Source/Programs/UnrealGameSync/UnrealGameSyncAutomation/Properties/AssemblyInfo.cs
--- a/Engine/Source/Programs/UnrealGameSync/UnrealGameSyncAutomation/Properties/AssemblyInfo.cs
+++ b/Engine/Source/Programs/UnrealGameSync/UnrealGameSyncAutomation/Properties/AssemblyInfo.cs
@@ -2,6 +2,7 @@
 
 using System.Reflection;
 using System.Runtime.InteropServices;
+using System.Runtime.Versioning;
 
 // General Information about an assembly is controlled through the following
 // set of attributes. Change these attribute values to modify the information
@@ -14,6 +15,7 @@ using System.Runtime.InteropServices;
 [assembly: AssemblyCopyright("Copyright Epic Games, Inc. All Rights Reserved.")]
 [assembly: AssemblyTrademark("")]
 [assembly: AssemblyCulture("")]
+[assembly: SupportedOSPlatform("windows")]
 
 // Setting ComVisible to false makes the types in this assembly not visible
 // to COM components.  If you need to access a type in this assembly from
diff --git a/Engine/Source/Programs/UnrealGameSync/UnrealGameSyncShared/ConfigUtils.cs b/Engine/Source/Programs/UnrealGameSync/UnrealGameSyncShared/ConfigUtils.cs
--- a/Engine/Source/Programs/UnrealGameSync/UnrealGameSyncShared/ConfigUtils.cs
+++ b/Engine/Source/Programs/UnrealGameSync/UnrealGameSyncShared/ConfigUtils.cs
@@ -1,19 +1,16 @@
 // Copyright Epic Games, Inc. All Rights Reserved.
 
-using EpicGames.Core;
-using EpicGames.OIDC;
-using EpicGames.Perforce;
-using Microsoft.Extensions.Configuration;
-using Microsoft.Extensions.Logging;
 using System;
 using System.Collections.Generic;
 using System.Diagnostics.CodeAnalysis;
 using System.IO;
 using System.Linq;
 using System.Runtime.InteropServices;
-using System.Text;
 using System.Threading;
 using System.Threading.Tasks;
+using EpicGames.Core;
+using EpicGames.Perforce;
+using Microsoft.Extensions.Logging;
 
 namespace UnrealGameSync
 {
@@ -352,11 +349,11 @@ namespace UnrealGameSync
 			bool useCrashReportClientEditor = latestProjectConfigFile.GetValue("Options.UseCrashReportClientEditor", false);
 
 			string hostPlatform;
-			if(RuntimeInformation.IsOSPlatform(OSPlatform.OSX))
+			if (RuntimeInformation.IsOSPlatform(OSPlatform.OSX))
 			{
 				hostPlatform = "Mac";
 			}
-			else if(RuntimeInformation.IsOSPlatform(OSPlatform.Linux))
+			else if (RuntimeInformation.IsOSPlatform(OSPlatform.Linux))
 			{
 				hostPlatform = "Linux";
 			}
diff --git a/Engine/Source/Programs/UnrealGameSync/UnrealGameSyncShared/OpenProjectInfo.cs b/Engine/Source/Programs/UnrealGameSync/UnrealGameSyncShared/OpenProjectInfo.cs
--- a/Engine/Source/Programs/UnrealGameSync/UnrealGameSyncShared/OpenProjectInfo.cs
+++ b/Engine/Source/Programs/UnrealGameSync/UnrealGameSyncShared/OpenProjectInfo.cs
@@ -1,15 +1,14 @@
 // Copyright Epic Games, Inc. All Rights Reserved.
 
-using EpicGames.Core;
-using EpicGames.OIDC;
-using EpicGames.Perforce;
-using Microsoft.Extensions.Logging;
 using System;
 using System.Collections.Generic;
 using System.IO;
 using System.Linq;
 using System.Threading;
 using System.Threading.Tasks;
+using EpicGames.Core;
+using EpicGames.Perforce;
+using Microsoft.Extensions.Logging;
 
 namespace UnrealGameSync
 {
@@ -204,7 +203,7 @@ namespace UnrealGameSync
 						}
 					}
 				}
-				if(branchClientPath == null || branchDirectoryName == null)
+				if (branchClientPath == null || branchDirectoryName == null)
 				{
 					throw new UserErrorException($"Could not find engine in Perforce relative to project path ({newSelectedClientFileName})");
 				}
@@ -221,24 +220,35 @@ namespace UnrealGameSync
 				ProjectInfo projectInfo = await ProjectInfo.CreateAsync(perforceClient, userWorkspaceSettings, cancellationToken);
 
 				// Update the cached workspace state
-				WorkspaceStateWrapper workspaceStateWrapper = userSettings.FindOrAddWorkspaceState(projectInfo, userWorkspaceSettings);
+				WorkspaceStateWrapper? workspaceStateWrapper = null;
+				try
+				{
+#pragma warning disable CA2000
+					workspaceStateWrapper = userSettings.FindOrAddWorkspaceState(projectInfo, userWorkspaceSettings);
+#pragma warning restore CA2000
+
+					// Read the initial config file
+					List<KeyValuePair<FileReference, DateTime>> localConfigFiles = new List<KeyValuePair<FileReference, DateTime>>();
+					ConfigFile latestProjectConfigFile = await ConfigUtils.ReadProjectConfigFileAsync(perforceClient, projectInfo, localConfigFiles, logger, cancellationToken);
 
-				// Read the initial config file
-				List<KeyValuePair<FileReference, DateTime>> localConfigFiles = new List<KeyValuePair<FileReference, DateTime>>();
-				ConfigFile latestProjectConfigFile = await ConfigUtils.ReadProjectConfigFileAsync(perforceClient, projectInfo, localConfigFiles, logger, cancellationToken);
+					// Get the local config file and stream filter
+					ConfigFile workspaceProjectConfigFile = await WorkspaceUpdate.ReadProjectConfigFile(branchDirectoryName, newSelectedFileName, logger);
+					IReadOnlyList<string>? workspaceProjectStreamFilter = await WorkspaceUpdate.ReadProjectStreamFilter(perforceClient, workspaceProjectConfigFile, cancellationToken);
 
-				// Get the local config file and stream filter
-				ConfigFile workspaceProjectConfigFile = await WorkspaceUpdate.ReadProjectConfigFile(branchDirectoryName, newSelectedFileName, logger);
-				IReadOnlyList<string>? workspaceProjectStreamFilter = await WorkspaceUpdate.ReadProjectStreamFilter(perforceClient, workspaceProjectConfigFile, cancellationToken);
+					OpenProjectInfo workspaceSettings = new OpenProjectInfo(selectedProject, perforceSettings, projectInfo, userWorkspaceSettings, workspaceStateWrapper, latestProjectConfigFile, workspaceProjectConfigFile, workspaceProjectStreamFilter, localConfigFiles, generateP4Config);
 
-				OpenProjectInfo workspaceSettings = new OpenProjectInfo(selectedProject, perforceSettings, projectInfo, userWorkspaceSettings, workspaceStateWrapper, latestProjectConfigFile, workspaceProjectConfigFile, workspaceProjectStreamFilter, localConfigFiles, generateP4Config);
+					if (generateP4Config)
+					{
+						GenerateP4ConfigFile(perforceSettings, projectInfo, logger);
+					}
 
-				if (generateP4Config)
+					return workspaceSettings;
+				}
+				catch
 				{
-					GenerateP4ConfigFile(perforceSettings, projectInfo, logger);
+					workspaceStateWrapper?.Dispose();
+					throw;
 				}
-
-				return workspaceSettings;
 			}
 			finally
 			{
@@ -249,15 +259,15 @@ namespace UnrealGameSync
 		static async Task<List<IPerforceSettings>> FilterClients(List<ClientsRecord> clients, FileReference newSelectedFileName, IPerforceSettings defaultPerforceSettings, string? hostName, ILogger logger, CancellationToken cancellationToken)
 		{
 			List<IPerforceSettings> candidateClients = new List<IPerforceSettings>();
-			foreach(ClientsRecord client in clients)
+			foreach (ClientsRecord client in clients)
 			{
 				// Make sure the client is well formed
-				if(!String.IsNullOrEmpty(client.Name) && (!String.IsNullOrEmpty(client.Host) || !String.IsNullOrEmpty(client.Owner)) && !String.IsNullOrEmpty(client.Root))
+				if (!String.IsNullOrEmpty(client.Name) && (!String.IsNullOrEmpty(client.Host) || !String.IsNullOrEmpty(client.Owner)) && !String.IsNullOrEmpty(client.Root))
 				{
 					// Require either a username or host name match
-					if((String.IsNullOrEmpty(client.Host) || String.Equals(client.Host, hostName, StringComparison.OrdinalIgnoreCase)) && (String.IsNullOrEmpty(client.Owner) || String.Equals(client.Owner, defaultPerforceSettings.UserName, StringComparison.OrdinalIgnoreCase)))
+					if ((String.IsNullOrEmpty(client.Host) || String.Equals(client.Host, hostName, StringComparison.OrdinalIgnoreCase)) && (String.IsNullOrEmpty(client.Owner) || String.Equals(client.Owner, defaultPerforceSettings.UserName, StringComparison.OrdinalIgnoreCase)))
 					{
-						if(!Utility.SafeIsFileUnderDirectory(newSelectedFileName.FullName, client.Root))
+						if (!Utility.SafeIsFileUnderDirectory(newSelectedFileName.FullName, client.Root))
 						{
 							logger.LogInformation("Rejecting {ClientName} due to root mismatch ({RootPath})", client.Name, client.Root);
 							continue;
@@ -267,7 +277,7 @@ namespace UnrealGameSync
 						using IPerforceConnection candidateClient = await PerforceConnection.CreateAsync(candidateSettings, logger);
 
 						List<PerforceResponse<WhereRecord>> whereRecords = await candidateClient.TryWhereAsync(newSelectedFileName.FullName, cancellationToken).Where(x => x.Failed || !x.Data.Unmap).ToListAsync(cancellationToken);
-						if(!whereRecords.Succeeded() || whereRecords.Count != 1)
+						if (!whereRecords.Succeeded() || whereRecords.Count != 1)
 						{
 							logger.LogInformation("Rejecting {ClientName} due to file not existing in workspace", client.Name);
 							continue;
@@ -281,7 +291,7 @@ namespace UnrealGameSync
 						}
 
 						records.RemoveAll(x => !x.Data.IsMapped);
-						if(records.Count == 0)
+						if (records.Count == 0)
 						{
 							logger.LogInformation("Rejecting {ClientName} due to {NumRecords} matching records", client.Name, records.Count);
 							continue;
diff --git a/Engine/Source/Programs/UnrealGameSync/UnrealGameSyncShared/Telemetry.cs b/Engine/Source/Programs/UnrealGameSync/UnrealGameSyncShared/Telemetry.cs
--- a/Engine/Source/Programs/UnrealGameSync/UnrealGameSyncShared/Telemetry.cs
+++ b/Engine/Source/Programs/UnrealGameSync/UnrealGameSyncShared/Telemetry.cs
@@ -154,7 +154,7 @@ namespace UnrealGameSync
 	/// <summary>
 	/// Global telemetry static class
 	/// </summary>
-	public static class Telemetry
+	public static class UgsTelemetry
 	{
 		/// <summary>
 		/// The current telemetry provider
diff --git a/Engine/Source/Programs/UnrealGameSync/UnrealGameSyncShared/TelemetryStopwatch.cs b/Engine/Source/Programs/UnrealGameSync/UnrealGameSyncShared/TelemetryStopwatch.cs
--- a/Engine/Source/Programs/UnrealGameSync/UnrealGameSyncShared/TelemetryStopwatch.cs
+++ b/Engine/Source/Programs/UnrealGameSync/UnrealGameSyncShared/TelemetryStopwatch.cs
@@ -49,7 +49,7 @@ namespace UnrealGameSync
 			{
 				Stop("Aborted");
 			}
-			Telemetry.SendEvent(_eventName, _eventData);
+			UgsTelemetry.SendEvent(_eventName, _eventData);
 		}
 
 		public TimeSpan Elapsed => _timer.Elapsed;
