UGS: Run code cleanup on solution.

#rnx

diff --git a/Engine/Source/Programs/UnrealGameSync/UnrealGameSync/Automation/VisualStudioAutomation.cs b/Engine/Source/Programs/UnrealGameSync/UnrealGameSync/Automation/VisualStudioAutomation.cs
--- a/Engine/Source/Programs/UnrealGameSync/UnrealGameSync/Automation/VisualStudioAutomation.cs
+++ b/Engine/Source/Programs/UnrealGameSync/UnrealGameSync/Automation/VisualStudioAutomation.cs
@@ -2,14 +2,14 @@
 
 using System;
 using System.Collections.Generic;
+using System.Diagnostics;
 using System.IO;
 using System.Linq;
 using System.Runtime.InteropServices;
 using System.Runtime.InteropServices.ComTypes;
 using System.Windows.Forms;
-using Microsoft.VisualStudio.Setup.Configuration;
 using EnvDTE;
-using System.Diagnostics;
+using Microsoft.VisualStudio.Setup.Configuration;
 
 namespace UnrealGameSync
 {
diff --git a/Engine/Source/Programs/UnrealGameSync/UnrealGameSync/Controls/LogControl.cs b/Engine/Source/Programs/UnrealGameSync/UnrealGameSync/Controls/LogControl.cs
--- a/Engine/Source/Programs/UnrealGameSync/UnrealGameSync/Controls/LogControl.cs
+++ b/Engine/Source/Programs/UnrealGameSync/UnrealGameSync/Controls/LogControl.cs
@@ -1,16 +1,16 @@
 // Copyright Epic Games, Inc. All Rights Reserved.
 
 using System;
+using System.Collections.Concurrent;
 using System.Collections.Generic;
-using System.Drawing;
 using System.Data;
+using System.Drawing;
+using System.IO;
 using System.Linq;
-using System.Text;
-using System.Windows.Forms;
 using System.Runtime.InteropServices;
-using System.Collections.Concurrent;
-using System.IO;
+using System.Text;
 using System.Text.RegularExpressions;
+using System.Windows.Forms;
 using EpicGames.Core;
 using Microsoft.Extensions.Logging;
 
@@ -30,7 +30,7 @@ namespace UnrealGameSync
 		}
 
 		[StructLayout(LayoutKind.Sequential)]
-		class Scrollinfo 
+		class Scrollinfo
 		{
 			public int cbSize;
 			public ScrollInfoMask fMask;
@@ -131,7 +131,7 @@ namespace UnrealGameSync
 		{
 			BackColor = Color.White;
 			ForeColor = Color.FromArgb(255, 32, 32, 32);
-			
+
 			DoubleBuffered = true;
 
 			ContextMenuStrip = new ContextMenuStrip();
@@ -144,7 +144,7 @@ namespace UnrealGameSync
 
 		protected override void OnCreateControl()
 		{
- 			base.OnCreateControl();
+			base.OnCreateControl();
 
 			EnableScrollBar(Handle, ScrollBarType.SbBoth, ScrollBarArrows.EsbEnableBoth);
 
@@ -169,7 +169,7 @@ namespace UnrealGameSync
 		{
 			base.Dispose(disposing);
 
-			if(_updateTimer != null)
+			if (_updateTimer != null)
 			{
 				_updateTimer.Dispose();
 				_updateTimer = null;
@@ -192,7 +192,7 @@ namespace UnrealGameSync
 			{
 				_logFileStream = FileReference.Open(newLogFileName, FileMode.OpenOrCreate, FileAccess.ReadWrite, FileShare.Read);
 			}
-			catch(Exception)
+			catch (Exception)
 			{
 				return false;
 			}
@@ -201,7 +201,7 @@ namespace UnrealGameSync
 
 			using StreamReader reader = new StreamReader(_logFileStream, leaveOpen: true);
 			string text = reader.ReadToEnd().TrimEnd('\r', '\n');
-			if(text.Length > 0)
+			if (text.Length > 0)
 			{
 				AddLinesInternal(text.Split('\n').Select(x => x + "\n").ToList());
 			}
@@ -212,7 +212,7 @@ namespace UnrealGameSync
 
 		public void CloseFile()
 		{
-			if(_logFileStream != null)
+			if (_logFileStream != null)
 			{
 				_logFileStream.Dispose();
 				_logFileStream = null;
@@ -229,7 +229,7 @@ namespace UnrealGameSync
 			_isSelecting = false;
 			_selection = null;
 
-			if(IsHandleCreated)
+			if (IsHandleCreated)
 			{
 				Scrollinfo scrollInfo = new Scrollinfo();
 				scrollInfo.cbSize = Marshal.SizeOf(scrollInfo);
@@ -239,7 +239,7 @@ namespace UnrealGameSync
 			}
 
 			_queuedLines = new ConcurrentQueue<string>();
-			if(_logFileStream != null)
+			if (_logFileStream != null)
 			{
 				try
 				{
@@ -262,28 +262,28 @@ namespace UnrealGameSync
 
 		private void Tick()
 		{
-			if(!Focused || !Capture)
+			if (!Focused || !Capture)
 			{
 				List<string> newLines = new List<string>();
-				for(;;)
+				for (; ; )
 				{
 					string? nextLine;
-					if(!_queuedLines.TryDequeue(out nextLine))
+					if (!_queuedLines.TryDequeue(out nextLine))
 					{
 						break;
 					}
 					newLines.Add(nextLine.TrimEnd('\r', '\n') + "\n");
 				}
-				if(newLines.Count > 0)
+				if (newLines.Count > 0)
 				{
 					StringBuilder textToAppendBuilder = new StringBuilder();
-					foreach(string newLine in newLines)
+					foreach (string newLine in newLines)
 					{
 						textToAppendBuilder.AppendLine(newLine.TrimEnd('\n'));
 					}
 
 					string textToAppend = textToAppendBuilder.ToString();
-					if(_logFileStream != null)
+					if (_logFileStream != null)
 					{
 						byte[] data = Encoding.UTF8.GetBytes(textToAppend);
 						try
@@ -291,7 +291,7 @@ namespace UnrealGameSync
 							_logFileStream.Write(data, 0, data.Length);
 							_logFileStream.Flush();
 						}
-						catch(Exception ex)
+						catch (Exception ex)
 						{
 							newLines.Add(String.Format("Failed to write to log file ({0}): {1}\n", _logFileStream.Name, ex.ToString()));
 						}
@@ -306,13 +306,13 @@ namespace UnrealGameSync
 		{
 			_lines.AddRange(newLines);
 
-			if(IsHandleCreated)
+			if (IsHandleCreated)
 			{
 				// Figure out if we're tracking the last line
 				bool isTrackingLastLine = IsTrackingLastLine();
 
 				int newMaxLineLength = Math.Max(_maxLineLength, newLines.Max(x => x.Length));
-				if(newMaxLineLength > _maxLineLength)
+				if (newMaxLineLength > _maxLineLength)
 				{
 					_maxLineLength = newMaxLineLength;
 
@@ -330,7 +330,7 @@ namespace UnrealGameSync
 				verticalScroll.fMask = ScrollInfoMask.SifPos | ScrollInfoMask.SifRange | ScrollInfoMask.SifPage | ScrollInfoMask.SifTrackpos;
 				GetScrollInfo(Handle, ScrollBarType.SbVert, verticalScroll);
 
-				if(_trackingScroll)
+				if (_trackingScroll)
 				{
 					UpdateVerticalScrollPosition(verticalScroll.nTrackPos, ref verticalScroll);
 				}
@@ -338,7 +338,7 @@ namespace UnrealGameSync
 				{
 					verticalScroll.fMask = ScrollInfoMask.SifRange | ScrollInfoMask.SifPage;
 					verticalScroll.nMin = 0;
-					if(isTrackingLastLine)
+					if (isTrackingLastLine)
 					{
 						verticalScroll.fMask |= ScrollInfoMask.SifPos;
 						verticalScroll.nPos = Math.Max(_lines.Count - 1 - _scrollLinesPerPage + 1, 0);
@@ -365,19 +365,19 @@ namespace UnrealGameSync
 
 		private void CopySelection()
 		{
-			if(_lines.Count > 0)
+			if (_lines.Count > 0)
 			{
 				StringBuilder selectedText = new StringBuilder();
-				if(_selection == null || _selection.IsEmpty())
+				if (_selection == null || _selection.IsEmpty())
 				{
-					foreach(string line in _lines)
+					foreach (string line in _lines)
 					{
 						selectedText.Append(line);
 					}
 				}
 				else
 				{
-					for(int lineIdx = Math.Min(_selection.Start.LineIdx, _selection.End.LineIdx); lineIdx <= Math.Max(_selection.Start.LineIdx, _selection.End.LineIdx); lineIdx++)
+					for (int lineIdx = Math.Min(_selection.Start.LineIdx, _selection.End.LineIdx); lineIdx <= Math.Max(_selection.Start.LineIdx, _selection.End.LineIdx); lineIdx++)
 					{
 						int minIdx, maxIdx;
 						ClipSelectionToLine(lineIdx, out minIdx, out maxIdx);
@@ -400,7 +400,7 @@ namespace UnrealGameSync
 
 		public void SelectAll()
 		{
-			if(_lines.Count > 0)
+			if (_lines.Count > 0)
 			{
 				_selection = new TextSelection(new TextLocation(0, 0), new TextLocation(_lines.Count - 1, _lines[_lines.Count - 1].Length));
 				Invalidate();
@@ -409,18 +409,18 @@ namespace UnrealGameSync
 
 		TextLocation PointToTextLocation(Point clientPoint)
 		{
-			if(_lines.Count == 0)
+			if (_lines.Count == 0)
 			{
 				return new TextLocation(0, 0);
 			}
 			else
 			{
 				int unclippedLineIdx = (clientPoint.Y / _fontSize.Height) + _scrollLine;
-				if(unclippedLineIdx < 0)
+				if (unclippedLineIdx < 0)
 				{
 					return new TextLocation(0, 0);
 				}
-				else if(unclippedLineIdx >= _lines.Count)
+				else if (unclippedLineIdx >= _lines.Count)
 				{
 					return new TextLocation(_lines.Count - 1, _lines[_lines.Count - 1].Length);
 				}
@@ -441,7 +441,7 @@ namespace UnrealGameSync
 			newScrollLine = Math.Min(newScrollLine, _lines.Count - _scrollLinesPerPage);
 			newScrollLine = Math.Max(newScrollLine, 0);
 
-			if(_scrollLine != newScrollLine)
+			if (_scrollLine != newScrollLine)
 			{
 				_scrollLine = newScrollLine;
 				Invalidate();
@@ -489,7 +489,7 @@ namespace UnrealGameSync
 
 		private void UpdateFontMetrics()
 		{
-			using(Graphics graphics = CreateGraphics())
+			using (Graphics graphics = CreateGraphics())
 			{
 				_fontSize = TextRenderer.MeasureText(graphics, "A", Font, new Size(Int32.MaxValue, Int32.MaxValue), TextFormatFlags.NoPadding);
 			}
@@ -499,7 +499,7 @@ namespace UnrealGameSync
 		{
 			base.OnMouseDown(e);
 
-			if(e.Button == MouseButtons.Left)
+			if (e.Button == MouseButtons.Left)
 			{
 				_isSelecting = true;
 
@@ -525,16 +525,16 @@ namespace UnrealGameSync
 		{
 			base.OnMouseMove(e);
 
-			if(_isSelecting && _selection != null)
+			if (_isSelecting && _selection != null)
 			{
 				TextLocation newSelectionEnd = PointToTextLocation(e.Location);
-				if(newSelectionEnd.LineIdx != _selection.End.LineIdx || newSelectionEnd.ColumnIdx != _selection.End.ColumnIdx)
+				if (newSelectionEnd.LineIdx != _selection.End.LineIdx || newSelectionEnd.ColumnIdx != _selection.End.ColumnIdx)
 				{
 					_selection.End = newSelectionEnd;
 					Invalidate();
 				}
 
-				_autoScrollRate = ((e.Y < 0)? e.Y : Math.Max(e.Y - ClientSize.Height, 0)) / _fontSize.Height;
+				_autoScrollRate = ((e.Y < 0) ? e.Y : Math.Max(e.Y - ClientSize.Height, 0)) / _fontSize.Height;
 			}
 		}
 
@@ -545,7 +545,7 @@ namespace UnrealGameSync
 			EndSelection();
 			Capture = false;
 
-			if(e.Button == MouseButtons.Right)
+			if (e.Button == MouseButtons.Right)
 			{
 				ContextMenuStrip.Show(this, e.Location);
 			}
@@ -562,7 +562,7 @@ namespace UnrealGameSync
 
 		protected override bool IsInputKey(Keys key)
 		{
-			switch(key & ~Keys.Modifiers)
+			switch (key & ~Keys.Modifiers)
 			{
 				case Keys.Up:
 				case Keys.Down:
@@ -574,7 +574,7 @@ namespace UnrealGameSync
 
 		protected override void OnKeyDown(KeyEventArgs e)
 		{
-			switch(e.KeyCode & ~Keys.Modifiers)
+			switch (e.KeyCode & ~Keys.Modifiers)
 			{
 				case Keys.Up:
 					ScrollWindow(-1);
@@ -601,13 +601,13 @@ namespace UnrealGameSync
 					Invalidate();
 					break;
 				case Keys.A:
-					if(e.Control)
+					if (e.Control)
 					{
 						SelectAll();
 					}
 					break;
 				case Keys.C:
-					if(e.Control)
+					if (e.Control)
 					{
 						CopySelection();
 					}
@@ -617,9 +617,9 @@ namespace UnrealGameSync
 
 		void EndSelection()
 		{
-			if(_isSelecting)
+			if (_isSelecting)
 			{
-				if(_selection != null && _selection.IsEmpty())
+				if (_selection != null && _selection.IsEmpty())
 				{
 					_selection = null;
 				}
@@ -632,7 +632,7 @@ namespace UnrealGameSync
 		{
 			base.OnScroll(se);
 
-			if(se.ScrollOrientation == ScrollOrientation.HorizontalScroll)
+			if (se.ScrollOrientation == ScrollOrientation.HorizontalScroll)
 			{
 				// Get the current scroll position
 				Scrollinfo scrollInfo = new Scrollinfo();
@@ -642,7 +642,7 @@ namespace UnrealGameSync
 
 				// Get the new scroll position
 				int targetScrollPos = scrollInfo.nPos;
-				switch(se.Type)
+				switch (se.Type)
 				{
 					case ScrollEventType.SmallDecrement:
 						targetScrollPos = scrollInfo.nPos - 1;
@@ -666,7 +666,7 @@ namespace UnrealGameSync
 				_scrollColumn = Math.Max(0, Math.Min(targetScrollPos, scrollInfo.nMax - scrollInfo.nPage + 1));
 
 				// Update the scroll bar if we're not tracking
-				if(se.Type != ScrollEventType.ThumbTrack)
+				if (se.Type != ScrollEventType.ThumbTrack)
 				{
 					scrollInfo.fMask = ScrollInfoMask.SifPos;
 					scrollInfo.nPos = _scrollColumn;
@@ -675,7 +675,7 @@ namespace UnrealGameSync
 
 				Invalidate();
 			}
-			else if(se.ScrollOrientation == ScrollOrientation.VerticalScroll)
+			else if (se.ScrollOrientation == ScrollOrientation.VerticalScroll)
 			{
 				// Get the current scroll position
 				Scrollinfo scrollInfo = new Scrollinfo();
@@ -685,7 +685,7 @@ namespace UnrealGameSync
 
 				// Get the new scroll position
 				int targetScrollPos = scrollInfo.nPos;
-				switch(se.Type)
+				switch (se.Type)
 				{
 					case ScrollEventType.SmallDecrement:
 						targetScrollPos = scrollInfo.nPos - 1;
@@ -711,7 +711,7 @@ namespace UnrealGameSync
 				UpdateVerticalScrollPosition(targetScrollPos, ref scrollInfo);
 
 				// Update the new range as well
-				if(se.Type == ScrollEventType.ThumbTrack)
+				if (se.Type == ScrollEventType.ThumbTrack)
 				{
 					_trackingScroll = true;
 				}
@@ -719,7 +719,7 @@ namespace UnrealGameSync
 				{
 					// If we've just finished tracking, allow updating the range
 					scrollInfo.nPos = _scrollLine;
-					if(_trackingScroll)
+					if (_trackingScroll)
 					{
 						scrollInfo.fMask |= ScrollInfoMask.SifRange;
 						scrollInfo.nMax = Math.Max(_lines.Count, 1) - 1;
@@ -746,18 +746,18 @@ namespace UnrealGameSync
 			base.OnPaint(e);
 
 			int textX = -_scrollColumn * _fontSize.Width;
-			for(int idx = _scrollLine; idx < _scrollLine + _scrollLinesPerPage + 1 && idx < _lines.Count; idx++)
+			for (int idx = _scrollLine; idx < _scrollLine + _scrollLinesPerPage + 1 && idx < _lines.Count; idx++)
 			{
 				int selectMinIdx;
 				int selectMaxIdx;
 				ClipSelectionToLine(idx, out selectMinIdx, out selectMaxIdx);
 
 				Color textColor;
-				if(Regex.IsMatch(_lines[idx], "(?<!\\w)error[: ]", RegexOptions.IgnoreCase))
+				if (Regex.IsMatch(_lines[idx], "(?<!\\w)error[: ]", RegexOptions.IgnoreCase))
 				{
 					textColor = Color.FromArgb(189, 54, 47);
 				}
-				else if(Regex.IsMatch(_lines[idx], "(?<!\\w)warning[: ]", RegexOptions.IgnoreCase))
+				else if (Regex.IsMatch(_lines[idx], "(?<!\\w)warning[: ]", RegexOptions.IgnoreCase))
 				{
 					textColor = Color.FromArgb(128, 128, 0);
 				}
@@ -767,16 +767,16 @@ namespace UnrealGameSync
 				}
 
 				int textY = (idx - _scrollLine) * _fontSize.Height;
-				if(selectMinIdx > 0)
+				if (selectMinIdx > 0)
 				{
 					TextRenderer.DrawText(e.Graphics, _lines[idx].Substring(0, selectMinIdx), Font, new Point(textX, textY), textColor, TextFormatFlags.NoPadding);
 				}
-				if(selectMinIdx < selectMaxIdx)
+				if (selectMinIdx < selectMaxIdx)
 				{
 					e.Graphics.FillRectangle(SystemBrushes.Highlight, textX + (selectMinIdx * _fontSize.Width), textY, (selectMaxIdx - selectMinIdx) * _fontSize.Width, _fontSize.Height);
 					TextRenderer.DrawText(e.Graphics, _lines[idx].Substring(selectMinIdx, selectMaxIdx - selectMinIdx), Font, new Point(textX + (selectMinIdx * _fontSize.Width), textY), SystemColors.HighlightText, TextFormatFlags.NoPadding);
 				}
-				if(selectMaxIdx < _lines[idx].Length)
+				if (selectMaxIdx < _lines[idx].Length)
 				{
 					TextRenderer.DrawText(e.Graphics, _lines[idx].Substring(selectMaxIdx), Font, new Point(textX + (selectMaxIdx * _fontSize.Width), textY), textColor, TextFormatFlags.NoPadding);
 				}
@@ -799,7 +799,7 @@ namespace UnrealGameSync
 
 		private void SelectionScrollTimer_TimerElapsed(object? sender, EventArgs args)
 		{
-			if(_autoScrollRate != 0 && _selection != null)
+			if (_autoScrollRate != 0 && _selection != null)
 			{
 				ScrollWindow(_autoScrollRate);
 
@@ -808,16 +808,16 @@ namespace UnrealGameSync
 				verticalScroll.fMask = ScrollInfoMask.SifAll;
 				GetScrollInfo(Handle, ScrollBarType.SbVert, verticalScroll);
 
-				if(_autoScrollRate < 0)
+				if (_autoScrollRate < 0)
 				{
-					if(_selection.End.LineIdx > verticalScroll.nPos)
+					if (_selection.End.LineIdx > verticalScroll.nPos)
 					{
 						_selection.End = new TextLocation(verticalScroll.nPos, 0);
 					}
 				}
 				else
 				{
-					if(_selection.End.LineIdx < verticalScroll.nPos + verticalScroll.nPage)
+					if (_selection.End.LineIdx < verticalScroll.nPos + verticalScroll.nPage)
 					{
 						int lineIdx = Math.Min(verticalScroll.nPos + verticalScroll.nPage, _lines.Count - 1);
 
@@ -834,29 +834,29 @@ namespace UnrealGameSync
 
 		protected void ClipSelectionToLine(int lineIdx, out int selectMinIdx, out int selectMaxIdx)
 		{
-			if(_selection == null)
+			if (_selection == null)
 			{
 				selectMinIdx = 0;
 				selectMaxIdx = 0;
 			}
-			else if(_selection.Start.LineIdx < _selection.End.LineIdx)
+			else if (_selection.Start.LineIdx < _selection.End.LineIdx)
 			{
-				if(lineIdx < _selection.Start.LineIdx)
+				if (lineIdx < _selection.Start.LineIdx)
 				{
 					selectMinIdx = _lines[lineIdx].Length;
 					selectMaxIdx = _lines[lineIdx].Length;
 				}
-				else if(lineIdx == _selection.Start.LineIdx)
+				else if (lineIdx == _selection.Start.LineIdx)
 				{
 					selectMinIdx = _selection.Start.ColumnIdx;
 					selectMaxIdx = _lines[lineIdx].Length;
 				}
-				else if(lineIdx < _selection.End.LineIdx)
+				else if (lineIdx < _selection.End.LineIdx)
 				{
 					selectMinIdx = 0;
 					selectMaxIdx = _lines[lineIdx].Length;
 				}
-				else if(lineIdx == _selection.End.LineIdx)
+				else if (lineIdx == _selection.End.LineIdx)
 				{
 					selectMinIdx = 0;
 					selectMaxIdx = _selection.End.ColumnIdx;
@@ -867,24 +867,24 @@ namespace UnrealGameSync
 					selectMaxIdx = 0;
 				}
 			}
-			else if(_selection.Start.LineIdx > _selection.End.LineIdx)
+			else if (_selection.Start.LineIdx > _selection.End.LineIdx)
 			{
-				if(lineIdx < _selection.End.LineIdx)
+				if (lineIdx < _selection.End.LineIdx)
 				{
 					selectMinIdx = _lines[lineIdx].Length;
 					selectMaxIdx = _lines[lineIdx].Length;
 				}
-				else if(lineIdx == _selection.End.LineIdx)
+				else if (lineIdx == _selection.End.LineIdx)
 				{
 					selectMinIdx = _selection.End.ColumnIdx;
 					selectMaxIdx = _lines[lineIdx].Length;
 				}
-				else if(lineIdx < _selection.Start.LineIdx)
+				else if (lineIdx < _selection.Start.LineIdx)
 				{
 					selectMinIdx = 0;
 					selectMaxIdx = _lines[lineIdx].Length;
 				}
-				else if(lineIdx == _selection.Start.LineIdx)
+				else if (lineIdx == _selection.Start.LineIdx)
 				{
 					selectMinIdx = 0;
 					selectMaxIdx = _selection.Start.ColumnIdx;
@@ -897,7 +897,7 @@ namespace UnrealGameSync
 			}
 			else
 			{
-				if(lineIdx == _selection.Start.LineIdx)
+				if (lineIdx == _selection.Start.LineIdx)
 				{
 					selectMinIdx = Math.Min(_selection.Start.ColumnIdx, _selection.End.ColumnIdx);
 					selectMaxIdx = Math.Max(_selection.Start.ColumnIdx, _selection.End.ColumnIdx);
diff --git a/Engine/Source/Programs/UnrealGameSync/UnrealGameSync/EventMonitor.cs b/Engine/Source/Programs/UnrealGameSync/UnrealGameSync/EventMonitor.cs
--- a/Engine/Source/Programs/UnrealGameSync/UnrealGameSync/EventMonitor.cs
+++ b/Engine/Source/Programs/UnrealGameSync/UnrealGameSync/EventMonitor.cs
@@ -4,15 +4,15 @@ using System;
 using System.Collections.Concurrent;
 using System.Collections.Generic;
 using System.Diagnostics;
+using System.Diagnostics.CodeAnalysis;
 using System.IO;
 using System.Linq;
 using System.Text.Json;
 using System.Threading;
-using EpicGames.Core;
 using System.Threading.Tasks;
-using Microsoft.Extensions.Logging;
+using EpicGames.Core;
 using Microsoft.Extensions.DependencyInjection;
-using System.Diagnostics.CodeAnalysis;
+using Microsoft.Extensions.Logging;
 
 namespace UnrealGameSync
 {
@@ -26,7 +26,7 @@ namespace UnrealGameSync
 		Good,
 		Bad,
 		Unknown,
-		
+
 		// Starred builds
 		Starred,
 		Unstarred,
@@ -106,7 +106,7 @@ namespace UnrealGameSync
 				}
 
 				int idx = BuildType.IndexOf(':', StringComparison.Ordinal);
-				if(idx == -1)
+				if (idx == -1)
 				{
 					return BuildType;
 				}
@@ -127,7 +127,7 @@ namespace UnrealGameSync
 				}
 
 				int idx = BuildType.IndexOf(':', StringComparison.Ordinal);
-				if(idx == -1)
+				if (idx == -1)
 				{
 					return BuildType;
 				}
@@ -280,7 +280,7 @@ namespace UnrealGameSync
 				}
 			}
 
-			if(_apiUrl == null)
+			if (_apiUrl == null)
 			{
 				LastStatusMessage = "Database functionality disabled due to empty ApiUrl.";
 			}
@@ -299,7 +299,7 @@ namespace UnrealGameSync
 		{
 			OnUpdatesReady = null;
 
-			if(_workerTask != null)
+			if (_workerTask != null)
 			{
 				_cancellationSource.Cancel();
 				_asyncDisposer.Add(_workerTask.ContinueWith(_ => _cancellationSource.Dispose(), CancellationToken.None, TaskContinuationOptions.None, TaskScheduler.Default));
@@ -339,15 +339,15 @@ namespace UnrealGameSync
 
 			// Clear out the list of active users for each review we have
 			_userNameToLastSyncEvent.Clear();
-			foreach(EventSummary summary in _changeNumberToSummary.Values)
+			foreach (EventSummary summary in _changeNumberToSummary.Values)
 			{
 				summary.CurrentUsers.Clear();
 			}
 
 			// Add all the user reviews back in again
-			foreach(EventSummary summary in _changeNumberToSummary.Values)
+			foreach (EventSummary summary in _changeNumberToSummary.Values)
 			{
-				foreach(EventData syncEvent in summary.SyncEvents)
+				foreach (EventData syncEvent in summary.SyncEvents)
 				{
 					ApplyFilteredUpdate(syncEvent);
 				}
@@ -366,7 +366,7 @@ namespace UnrealGameSync
 		protected EventSummary FindOrAddSummary(int changeNumber)
 		{
 			EventSummary? summary;
-			if(!_changeNumberToSummary.TryGetValue(changeNumber, out summary))
+			if (!_changeNumberToSummary.TryGetValue(changeNumber, out summary))
 			{
 				summary = new EventSummary();
 				summary.ChangeNumber = changeNumber;
@@ -390,19 +390,19 @@ namespace UnrealGameSync
 			}
 
 			EventData? evt;
-			while(_incomingEvents.TryDequeue(out evt))
+			while (_incomingEvents.TryDequeue(out evt))
 			{
 				ApplyEventUpdate(evt);
 			}
 
 			BadgeData? badge;
-			while(_incomingBadges.TryDequeue(out badge))
+			while (_incomingBadges.TryDequeue(out badge))
 			{
 				ApplyBadgeUpdate(badge);
 			}
 
 			CommentData? comment;
-			while(_incomingComments.TryDequeue(out comment))
+			while (_incomingComments.TryDequeue(out comment))
 			{
 				ApplyCommentUpdate(comment);
 			}
@@ -411,41 +411,41 @@ namespace UnrealGameSync
 		void ApplyEventUpdate(EventData evt)
 		{
 			EventSummary summary = FindOrAddSummary(evt.Change);
-			if(evt.Type == EventType.Starred || evt.Type == EventType.Unstarred)
+			if (evt.Type == EventType.Starred || evt.Type == EventType.Unstarred)
 			{
 				// If it's a star or un-star review, process that separately
-				if(summary.LastStarReview == null || evt.Id > summary.LastStarReview.Id)
+				if (summary.LastStarReview == null || evt.Id > summary.LastStarReview.Id)
 				{
 					summary.LastStarReview = evt;
 				}
 			}
-			else if(evt.Type == EventType.Investigating || evt.Type == EventType.Resolved)
+			else if (evt.Type == EventType.Investigating || evt.Type == EventType.Resolved)
 			{
 				// Insert it sorted in the investigation list
 				int insertIdx = 0;
-				while(insertIdx < _investigationEvents.Count && _investigationEvents[insertIdx].Id < evt.Id)
+				while (insertIdx < _investigationEvents.Count && _investigationEvents[insertIdx].Id < evt.Id)
 				{
 					insertIdx++;
 				}
-				if(insertIdx == _investigationEvents.Count || _investigationEvents[insertIdx].Id != evt.Id)
+				if (insertIdx == _investigationEvents.Count || _investigationEvents[insertIdx].Id != evt.Id)
 				{
 					_investigationEvents.Insert(insertIdx, evt);
 				}
 				_activeInvestigations = null;
 			}
-			else if(evt.Type == EventType.Syncing)
+			else if (evt.Type == EventType.Syncing)
 			{
 				summary.SyncEvents.RemoveAll(x => String.Equals(x.UserName, evt.UserName, StringComparison.OrdinalIgnoreCase));
 				summary.SyncEvents.Add(evt);
 				ApplyFilteredUpdate(evt);
 			}
-			else if(IsReview(evt.Type))
+			else if (IsReview(evt.Type))
 			{
 				// Try to find an existing review by this user. If we already have a newer review, ignore this one. Otherwise remove it.
 				EventData? existingReview = summary.Reviews.Find(x => String.Equals(x.UserName, evt.UserName, StringComparison.OrdinalIgnoreCase));
-				if(existingReview != null)
+				if (existingReview != null)
 				{
-					if(existingReview.Id <= evt.Id)
+					if (existingReview.Id <= evt.Id)
 					{
 						summary.Reviews.Remove(existingReview);
 					}
@@ -470,9 +470,9 @@ namespace UnrealGameSync
 			EventSummary summary = FindOrAddSummary(badge.ChangeNumber);
 
 			BadgeData? existingBadge = summary.Badges.Find(x => x.ChangeNumber == badge.ChangeNumber && x.BuildType == badge.BuildType);
-			if(existingBadge != null)
+			if (existingBadge != null)
 			{
-				if(existingBadge.Id <= badge.Id)
+				if (existingBadge.Id <= badge.Id)
 				{
 					summary.Badges.Remove(existingBadge);
 				}
@@ -486,7 +486,7 @@ namespace UnrealGameSync
 			summary.Verdict = GetVerdict(summary.Reviews, summary.Badges);
 
 			BadgeData? latestBadge;
-			if(!_badgeNameToLatestData.TryGetValue(badge.BadgeName, out latestBadge) || badge.ChangeNumber > latestBadge.ChangeNumber || (badge.ChangeNumber == latestBadge.ChangeNumber && badge.Id > latestBadge.Id))
+			if (!_badgeNameToLatestData.TryGetValue(badge.BadgeName, out latestBadge) || badge.ChangeNumber > latestBadge.ChangeNumber || (badge.ChangeNumber == latestBadge.ChangeNumber && badge.Id > latestBadge.Id))
 			{
 				_badgeNameToLatestData[badge.BadgeName] = badge;
 			}
@@ -495,7 +495,7 @@ namespace UnrealGameSync
 		void ApplyCommentUpdate(CommentData comment)
 		{
 			EventSummary summary = FindOrAddSummary(comment.ChangeNumber);
-			if(String.Equals(comment.UserName, _currentUserName, StringComparison.OrdinalIgnoreCase) && summary.Comments.Count > 0 && summary.Comments.Last().Id == Int64.MaxValue)
+			if (String.Equals(comment.UserName, _currentUserName, StringComparison.OrdinalIgnoreCase) && summary.Comments.Count > 0 && summary.Comments.Last().Id == Int64.MaxValue)
 			{
 				// This comment was added by PostComment(), to mask the latency of a round trip to the server. Remove it now we have the sorted comment.
 				summary.Comments.RemoveAt(summary.Comments.Count - 1);
@@ -507,9 +507,9 @@ namespace UnrealGameSync
 		{
 			int insertIdx = items.Count;
 
-			for(; insertIdx > 0 && idSelector(items[insertIdx - 1]) >= idSelector(newItem); insertIdx--)
+			for (; insertIdx > 0 && idSelector(items[insertIdx - 1]) >= idSelector(newItem); insertIdx--)
 			{
-				if(String.Equals(userSelector(items[insertIdx - 1]), userSelector(newItem), StringComparison.OrdinalIgnoreCase))
+				if (String.Equals(userSelector(items[insertIdx - 1]), userSelector(newItem), StringComparison.OrdinalIgnoreCase))
 				{
 					return false;
 				}
@@ -517,9 +517,9 @@ namespace UnrealGameSync
 
 			items.Insert(insertIdx, newItem);
 
-			for(; insertIdx > 0; insertIdx--)
+			for (; insertIdx > 0; insertIdx--)
 			{
-				if(String.Equals(userSelector(items[insertIdx - 1]), userSelector(newItem), StringComparison.OrdinalIgnoreCase))
+				if (String.Equals(userSelector(items[insertIdx - 1]), userSelector(newItem), StringComparison.OrdinalIgnoreCase))
 				{
 					items.RemoveAt(insertIdx - 1);
 				}
@@ -647,21 +647,21 @@ namespace UnrealGameSync
 		{
 			int numPositiveReviews = events.Count(x => x.Type == EventType.Good);
 			int numNegativeReviews = events.Count(x => x.Type == EventType.Bad);
-			if(numPositiveReviews > 0 || numNegativeReviews > 0)
+			if (numPositiveReviews > 0 || numNegativeReviews > 0)
 			{
 				return GetVerdict(numPositiveReviews, numNegativeReviews);
 			}
 
 			int numCompiles = events.Count(x => x.Type == EventType.Compiles);
 			int numFailedCompiles = events.Count(x => x.Type == EventType.DoesNotCompile);
-			if(numCompiles > 0 || numFailedCompiles > 0)
+			if (numCompiles > 0 || numFailedCompiles > 0)
 			{
 				return GetVerdict(numCompiles, numFailedCompiles);
 			}
 
 			int numBadges = badges.Count(x => x.BuildType == "Editor" && x.IsSuccess);
 			int numFailedBadges = badges.Count(x => x.BuildType == "Editor" && x.IsFailure);
-			if(numBadges > 0 || numFailedBadges > 0)
+			if (numBadges > 0 || numFailedBadges > 0)
 			{
 				return GetVerdict(numBadges, numFailedBadges);
 			}
@@ -671,11 +671,11 @@ namespace UnrealGameSync
 
 		static ReviewVerdict GetVerdict(int numPositive, int numNegative)
 		{
-			if(numPositive > (int)(numNegative * 1.5))
+			if (numPositive > (int)(numNegative * 1.5))
 			{
 				return ReviewVerdict.Good;
 			}
-			else if(numPositive >= numNegative)
+			else if (numPositive >= numNegative)
 			{
 				return ReviewVerdict.Mixed;
 			}
@@ -687,13 +687,13 @@ namespace UnrealGameSync
 
 		void ApplyFilteredUpdate(EventData evt)
 		{
-			if(evt.Type == EventType.Syncing && _filterChangeNumbers.Contains(evt.Change) && !String.IsNullOrEmpty(evt.UserName))
+			if (evt.Type == EventType.Syncing && _filterChangeNumbers.Contains(evt.Change) && !String.IsNullOrEmpty(evt.UserName))
 			{
 				// Update the active users list for this change
 				EventData? lastSync;
-				if(_userNameToLastSyncEvent.TryGetValue(evt.UserName, out lastSync))
+				if (_userNameToLastSyncEvent.TryGetValue(evt.UserName, out lastSync))
 				{
-					if(evt.Id > lastSync.Id)
+					if (evt.Id > lastSync.Id)
 					{
 						_changeNumberToSummary[lastSync.Change].CurrentUsers.RemoveAll(x => String.Equals(x, evt.UserName, StringComparison.OrdinalIgnoreCase));
 						FindOrAddSummary(evt.Change).CurrentUsers.Add(evt.UserName);
@@ -723,20 +723,20 @@ namespace UnrealGameSync
 				if (_apiUrl != null)
 				{
 					// Post all the reviews to the database. We don't send them out of order, so keep the review outside the queue until the next update if it fails
-					while(evt != null || _outgoingEvents.TryDequeue(out evt))
+					while (evt != null || _outgoingEvents.TryDequeue(out evt))
 					{
 						await SendEventToBackendAsync(evt, cancellationToken);
 						evt = null;
 					}
 
 					// Post all the comments to the database.
-					while(comment != null || _outgoingComments.TryDequeue(out comment))
+					while (comment != null || _outgoingComments.TryDequeue(out comment))
 					{
 						await SendCommentToBackendAsync(comment, cancellationToken);
 						comment = null;
 					}
 
-					if(timer.Elapsed > TimeSpan.FromSeconds(requestThrottle))
+					if (timer.Elapsed > TimeSpan.FromSeconds(requestThrottle))
 					{
 						updateThrottledRequests = true;
 						timer.Restart();
@@ -746,7 +746,7 @@ namespace UnrealGameSync
 					await ReadEventsFromBackendAsync(updateThrottledRequests, cancellationToken);
 
 					// Send a notification that we're ready to update
-					if(!_incomingMetadata.IsEmpty || !_incomingEvents.IsEmpty || !_incomingBadges.IsEmpty || !_incomingComments.IsEmpty)
+					if (!_incomingMetadata.IsEmpty || !_incomingEvents.IsEmpty || !_incomingBadges.IsEmpty || !_incomingComments.IsEmpty)
 					{
 						_synchronizationContext.Post(_ => OnUpdatesReady?.Invoke(), null);
 					}
@@ -793,7 +793,7 @@ namespace UnrealGameSync
 				}
 				return true;
 			}
-			catch(Exception ex)
+			catch (Exception ex)
 			{
 				_logger.LogError(ex, "Failed with exception.");
 				return false;
@@ -817,7 +817,7 @@ namespace UnrealGameSync
 				_logger.LogInformation("Done in {Time}ms.", timer.ElapsedMilliseconds);
 				return true;
 			}
-			catch(Exception ex)
+			catch (Exception ex)
 			{
 				_logger.LogError(ex, "Failed with exception.");
 				return false;
@@ -886,7 +886,7 @@ namespace UnrealGameSync
 				if (_apiVersion == 0)
 				{
 					LatestData initialIds = await RestApi.GetAsync<LatestData>($"{_apiUrl}/api/latest?project={_project}", cancellationToken);
-					_apiVersion = (initialIds.Version == 0)? 1 : initialIds.Version;
+					_apiVersion = (initialIds.Version == 0) ? 1 : initialIds.Version;
 					_latestIds.LastBuildId = initialIds.LastBuildId;
 					_latestIds.LastCommentId = initialIds.LastCommentId;
 					_latestIds.LastEventId = initialIds.LastEventId;
@@ -973,7 +973,7 @@ namespace UnrealGameSync
 				_logger.LogInformation("Done in {Time}ms.", timer.ElapsedMilliseconds);
 				return true;
 			}
-			catch(Exception ex)
+			catch (Exception ex)
 			{
 				_logger.LogError(ex, "Failed with exception.");
 				LastStatusMessage = String.Format("Last update failed: ({0})", ex.ToString());
@@ -983,7 +983,7 @@ namespace UnrealGameSync
 
 		public void PostEvent(int changeNumber, EventType type)
 		{
-			if(_apiUrl != null)
+			if (_apiUrl != null)
 			{
 				EventData evt = new EventData();
 				evt.Id = ++_latestIds.LastEventId;
@@ -1001,7 +1001,7 @@ namespace UnrealGameSync
 
 		public void PostComment(int changeNumber, string text)
 		{
-			if(_apiUrl != null)
+			if (_apiUrl != null)
 			{
 				CommentData comment = new CommentData();
 				comment.Id = Int64.MaxValue;
@@ -1020,14 +1020,14 @@ namespace UnrealGameSync
 		public bool GetCommentByCurrentUser(int changeNumber, [NotNullWhen(true)] out string? commentText)
 		{
 			EventSummary? summary = GetSummaryForChange(changeNumber);
-			if(summary == null)
+			if (summary == null)
 			{
 				commentText = null;
 				return false;
 			}
 
 			CommentData? comment = summary.Comments.Find(x => String.Equals(x.UserName, _currentUserName, StringComparison.OrdinalIgnoreCase));
-			if(comment == null || String.IsNullOrWhiteSpace(comment.Text))
+			if (comment == null || String.IsNullOrWhiteSpace(comment.Text))
 			{
 				commentText = null;
 				return false;
@@ -1040,13 +1040,13 @@ namespace UnrealGameSync
 		public EventData? GetReviewByCurrentUser(int changeNumber)
 		{
 			EventSummary? summary = GetSummaryForChange(changeNumber);
-			if(summary == null)
+			if (summary == null)
 			{
 				return null;
 			}
 
 			EventData? evt = summary.Reviews.FirstOrDefault(x => String.Equals(x.UserName, _currentUserName, StringComparison.OrdinalIgnoreCase));
-			if(evt == null || evt.Type == EventType.Unknown)
+			if (evt == null || evt.Type == EventType.Unknown)
 			{
 				return null;
 			}
@@ -1099,18 +1099,18 @@ namespace UnrealGameSync
 
 		protected void UpdateActiveInvestigations()
 		{
-			if(_activeInvestigations == null)
+			if (_activeInvestigations == null)
 			{
 				// Insert investigation events into the active list, sorted by change number.
 				_activeInvestigations = new List<EventData>();
-				foreach(EventData investigationEvent in _investigationEvents)
+				foreach (EventData investigationEvent in _investigationEvents)
 				{
-					if(_filterChangeNumbers.Contains(investigationEvent.Change))
+					if (_filterChangeNumbers.Contains(investigationEvent.Change))
 					{
-						if(investigationEvent.Type == EventType.Investigating)
+						if (investigationEvent.Type == EventType.Investigating)
 						{
 							int insertIdx = 0;
-							while(insertIdx < _activeInvestigations.Count && _activeInvestigations[insertIdx].Change > investigationEvent.Change)
+							while (insertIdx < _activeInvestigations.Count && _activeInvestigations[insertIdx].Change > investigationEvent.Change)
 							{
 								insertIdx++;
 							}
@@ -1124,11 +1124,11 @@ namespace UnrealGameSync
 				}
 
 				// Remove any duplicate users
-				for(int idx = 0; idx < _activeInvestigations.Count; idx++)
+				for (int idx = 0; idx < _activeInvestigations.Count; idx++)
 				{
-					for(int otherIdx = 0; otherIdx < idx; otherIdx++)
+					for (int otherIdx = 0; otherIdx < idx; otherIdx++)
 					{
-						if(String.Equals(_activeInvestigations[idx].UserName, _activeInvestigations[otherIdx].UserName, StringComparison.OrdinalIgnoreCase))
+						if (String.Equals(_activeInvestigations[idx].UserName, _activeInvestigations[otherIdx].UserName, StringComparison.OrdinalIgnoreCase))
 						{
 							_activeInvestigations.RemoveAt(idx--);
 							break;
diff --git a/Engine/Source/Programs/UnrealGameSync/UnrealGameSync/Forms/DiagnosticsWindow.cs b/Engine/Source/Programs/UnrealGameSync/UnrealGameSync/Forms/DiagnosticsWindow.cs
--- a/Engine/Source/Programs/UnrealGameSync/UnrealGameSync/Forms/DiagnosticsWindow.cs
+++ b/Engine/Source/Programs/UnrealGameSync/UnrealGameSync/Forms/DiagnosticsWindow.cs
@@ -1,13 +1,13 @@
 // Copyright Epic Games, Inc. All Rights Reserved.
 
-using EpicGames.Core;
 using System;
 using System.Collections.Generic;
 using System.Diagnostics;
-using System.Linq;
-using System.Windows.Forms;
 using System.IO;
 using System.IO.Compression;
+using System.Linq;
+using System.Windows.Forms;
+using EpicGames.Core;
 
 namespace UnrealGameSync
 {
@@ -45,14 +45,14 @@ namespace UnrealGameSync
 			dialog.Filter = "Zip Files (*.zip)|*.zip|AllFiles (*.*)|*.*";
 			dialog.InitialDirectory = Environment.GetFolderPath(Environment.SpecialFolder.Desktop);
 			dialog.FileName = Path.Combine(dialog.InitialDirectory, "UGS-Diagnostics.zip");
-			if(dialog.ShowDialog() == DialogResult.OK)
+			if (dialog.ShowDialog() == DialogResult.OK)
 			{
 				FileReference diagnosticsFileName = FileReference.Combine(_appDataFolder, "Diagnostics.txt");
 				try
 				{
 					FileReference.WriteAllLines(diagnosticsFileName, DiagnosticsTextBox.Lines);
 				}
-				catch(Exception ex)
+				catch (Exception ex)
 				{
 					MessageBox.Show(String.Format("Couldn't write to '{0}'\n\n{1}", diagnosticsFileName, ex.ToString()));
 					return;
@@ -68,7 +68,7 @@ namespace UnrealGameSync
 
 						foreach (FileReference extraFile in _extraFiles)
 						{
-							if(FileReference.Exists(extraFile))
+							if (FileReference.Exists(extraFile))
 							{
 								using (FileStream inputStream = FileReference.Open(extraFile, FileMode.Open, FileAccess.Read, FileShare.ReadWrite))
 								{
@@ -82,7 +82,7 @@ namespace UnrealGameSync
 						}
 					}
 				}
-				catch(Exception ex)
+				catch (Exception ex)
 				{
 					MessageBox.Show(String.Format("Couldn't save '{0}'\n\n{1}", zipFileName, ex.ToString()));
 					return;
diff --git a/Engine/Source/Programs/UnrealGameSync/UnrealGameSync/UriHandlers/UriHandler.cs b/Engine/Source/Programs/UnrealGameSync/UnrealGameSync/UriHandlers/UriHandler.cs
--- a/Engine/Source/Programs/UnrealGameSync/UnrealGameSync/UriHandlers/UriHandler.cs
+++ b/Engine/Source/Programs/UnrealGameSync/UnrealGameSync/UriHandlers/UriHandler.cs
@@ -2,16 +2,16 @@
 
 using System;
 using System.Collections.Generic;
-using Microsoft.Win32;
+using System.Collections.Specialized;
+using System.Diagnostics;
+using System.IO;
+using System.Linq;
 using System.Reflection;
 using System.Threading;
-using System.Collections.Specialized;
 using System.Web;
-using System.Linq;
-using System.Diagnostics;
 using System.Windows.Forms;
-using System.IO;
 using Microsoft.Extensions.Logging;
+using Microsoft.Win32;
 
 namespace UnrealGameSync
 {
@@ -32,7 +32,7 @@ namespace UnrealGameSync
 	static class UriHandler
 	{
 		public static UriResult HandleUri(string uriIn)
-		{			
+		{
 			try
 			{
 				Uri uri = new Uri(uriIn);
@@ -48,7 +48,7 @@ namespace UnrealGameSync
 				List<object> parameters = new List<object>();
 
 				foreach (ParameterInfo param in info.GetParameters())
-				{					
+				{
 					string? value = query.Get(param.Name);
 
 					if (value == null)
@@ -101,7 +101,7 @@ namespace UnrealGameSync
 					}
 				}
 
-				return (UriResult)info.Invoke(null, parameters.ToArray())!;					
+				return (UriResult)info.Invoke(null, parameters.ToArray())!;
 			}
 			catch (Exception ex)
 			{
@@ -211,7 +211,7 @@ namespace UnrealGameSync
 		{
 			foreach (Type type in Assembly.GetExecutingAssembly().GetTypes())
 			{
-				List<MethodInfo> handlerMethods = new List<MethodInfo> (type.GetMethods().Where(methodInfo => methodInfo.GetCustomAttribute<UriHandlerAttribute>() != null));
+				List<MethodInfo> handlerMethods = new List<MethodInfo>(type.GetMethods().Where(methodInfo => methodInfo.GetCustomAttribute<UriHandlerAttribute>() != null));
 				foreach (MethodInfo methodInfo in handlerMethods)
 				{
 					if (!methodInfo.IsStatic)
@@ -260,7 +260,7 @@ namespace UnrealGameSync
 
 		public UriHandlerAttribute(bool terminate = false)
 		{
-			Terminate = terminate;			
+			Terminate = terminate;
 		}
 	}
 
@@ -350,7 +350,7 @@ namespace UnrealGameSync
 					}
 				}
 
-				return hasAll? ProtocolHandlerState.Installed : hasAny? ProtocolHandlerState.Unknown : ProtocolHandlerState.NotInstalled;
+				return hasAll ? ProtocolHandlerState.Installed : hasAny ? ProtocolHandlerState.Unknown : ProtocolHandlerState.NotInstalled;
 			}
 			catch
 			{
diff --git a/Engine/Source/Programs/UnrealGameSync/UnrealGameSyncCmd/Program.cs b/Engine/Source/Programs/UnrealGameSync/UnrealGameSyncCmd/Program.cs
--- a/Engine/Source/Programs/UnrealGameSync/UnrealGameSyncCmd/Program.cs
+++ b/Engine/Source/Programs/UnrealGameSync/UnrealGameSyncCmd/Program.cs
@@ -1,30 +1,30 @@
 // Copyright Epic Games, Inc. All Rights Reserved.
 
-using EpicGames.Core;
-using EpicGames.Perforce;
-using Microsoft.Extensions.Logging;
-using Serilog;
 using System;
 using System.Collections.Generic;
 using System.Diagnostics;
+using System.Globalization;
 using System.IO;
+using System.IO.Compression;
 using System.Linq;
 using System.Net;
+using System.Net.Http;
+using System.Net.Http.Json;
 using System.Reflection;
 using System.Runtime.InteropServices;
+using System.Text;
 using System.Text.RegularExpressions;
 using System.Threading;
 using System.Threading.Tasks;
+using EpicGames.Core;
+using EpicGames.Horde;
 using EpicGames.OIDC;
+using EpicGames.Perforce;
 using Microsoft.Extensions.Configuration;
-using UnrealGameSync;
-using System.Text;
-using System.Globalization;
-using System.Net.Http;
-using System.Net.Http.Json;
-using System.IO.Compression;
 using Microsoft.Extensions.DependencyInjection;
-using EpicGames.Horde;
+using Microsoft.Extensions.Logging;
+using Serilog;
+using UnrealGameSync;
 
 namespace UnrealGameSyncCmd
 {
@@ -120,8 +120,8 @@ namespace UnrealGameSyncCmd
 				"ugs version",
 				"Prints the current application version"
 			),
-			new CommandInfo("install", typeof(InstallCommand), null, 
-				null, 
+			new CommandInfo("install", typeof(InstallCommand), null,
+				null,
 				null
 			),
 			new CommandInfo("upgrade", typeof(UpgradeCommand), typeof(UpgradeCommandOptions),
@@ -831,7 +831,7 @@ namespace UnrealGameSyncCmd
 					}
 				}
 
-				WorkspaceUpdateOptions options = syncOptions.SingleChange? WorkspaceUpdateOptions.SyncSingleChange : WorkspaceUpdateOptions.Sync;
+				WorkspaceUpdateOptions options = syncOptions.SingleChange ? WorkspaceUpdateOptions.SyncSingleChange : WorkspaceUpdateOptions.Sync;
 				if (syncOptions.Clean)
 				{
 					options |= WorkspaceUpdateOptions.Clean;
@@ -858,7 +858,7 @@ namespace UnrealGameSyncCmd
 				string[] syncFilter = ReadSyncFilter(settings, context.UserSettings, projectConfig);
 
 				using WorkspaceLock? workspaceLock = CreateWorkspaceLock(settings.RootDir);
-				if(workspaceLock != null && !await workspaceLock.TryAcquireAsync())
+				if (workspaceLock != null && !await workspaceLock.TryAcquireAsync())
 				{
 					logger.LogError("Another process is already syncing this workspace.");
 					return;
@@ -1045,7 +1045,7 @@ namespace UnrealGameSyncCmd
 				using IPerforceConnection perforceClient = await ConnectAsync(settings, context.LoggerFactory);
 
 				List<ChangesRecord> changes = await perforceClient.GetChangesAsync(EpicGames.Perforce.ChangesOptions.None, options.Count, ChangeStatus.Submitted, $"//{settings.ClientName}/...");
-				foreach(IEnumerable<ChangesRecord> changesBatch in changes.Batch(10))
+				foreach (IEnumerable<ChangesRecord> changesBatch in changes.Batch(10))
 				{
 					List<DescribeRecord> describeRecords = await perforceClient.DescribeAsync(changesBatch.Select(x => x.Number).ToArray());
 
@@ -1217,7 +1217,7 @@ namespace UnrealGameSyncCmd
 						enabled = syncCategory.Enable;
 					}
 
-					logger.LogInformation("  {Id,30} {Enabled,3} {Scope,-9} {Name}", syncCategory.UniqueId, enabled? "Yes" : "No", scope, syncCategory.Name);
+					logger.LogInformation("  {Id,30} {Enabled,3} {Scope,-9} {Name}", syncCategory.UniqueId, enabled ? "Yes" : "No", scope, syncCategory.Name);
 				}
 
 				if (globalFilter.View.Count > 0)
@@ -1370,7 +1370,7 @@ namespace UnrealGameSyncCmd
 
 				if (result != WorkspaceUpdateResult.Success)
 				{
-					throw new UserErrorException("{Message}", message); 
+					throw new UserErrorException("{Message}", message);
 				}
 			}
 		}
@@ -1427,13 +1427,13 @@ namespace UnrealGameSyncCmd
 				context.Arguments.CheckAllArgumentsUsed();
 
 				UserWorkspaceSettings settings = ReadRequiredUserWorkspaceSettings();
-				
+
 				// Find the valid config file paths
 				DirectoryInfo engineDir = DirectoryReference.Combine(settings.RootDir, "Engine").ToDirectoryInfo();
 				DirectoryInfo gameDir = new DirectoryInfo(settings.ProjectPath);
 				using ITokenStore tokenStore = TokenStoreFactory.CreateTokenStore();
 				IConfiguration providerConfiguration = ProviderConfigurationFactory.ReadConfiguration(engineDir, gameDir);
-				OidcTokenManager oidcTokenManager = OidcTokenManager.CreateTokenManager(providerConfiguration, tokenStore, new List<string>() {providerIdentifier});
+				OidcTokenManager oidcTokenManager = OidcTokenManager.CreateTokenManager(providerConfiguration, tokenStore, new List<string>() { providerIdentifier });
 				OidcTokenInfo result = await oidcTokenManager.Login(providerIdentifier);
 
 				logger.LogInformation("Logged in to provider {ProviderIdentifier}", providerIdentifier);
@@ -1572,7 +1572,7 @@ namespace UnrealGameSyncCmd
 			public override async Task ExecuteAsync(CommandContext context)
 			{
 				ILogger logger = context.Logger;
-				
+
 				UpgradeCommandOptions options = new UpgradeCommandOptions();
 				context.Arguments.ApplyTo(options);
 				string? targetDirStr = context.Arguments.GetStringOrDefault("-TargetDir=", null);
@@ -1602,7 +1602,7 @@ namespace UnrealGameSyncCmd
 
 					DirectoryReference currentDir = new FileReference(Assembly.GetExecutingAssembly().Location).Directory;
 
-					DirectoryReference targetDir = (targetDirStr == null)? currentDir : DirectoryReference.Combine(currentDir, targetDirStr);
+					DirectoryReference targetDir = (targetDirStr == null) ? currentDir : DirectoryReference.Combine(currentDir, targetDirStr);
 					DirectoryReference.CreateDirectory(targetDir);
 
 					FileReference tempFile = FileReference.Combine(targetDir, "update.zip");
diff --git a/Engine/Source/Programs/UnrealGameSync/UnrealGameSyncShared/DeploymentSettings.cs b/Engine/Source/Programs/UnrealGameSync/UnrealGameSyncShared/DeploymentSettings.cs
--- a/Engine/Source/Programs/UnrealGameSync/UnrealGameSyncShared/DeploymentSettings.cs
+++ b/Engine/Source/Programs/UnrealGameSync/UnrealGameSyncShared/DeploymentSettings.cs
@@ -3,18 +3,12 @@
 #pragma warning disable IDE0005
 #pragma warning disable CA1802 // warning CA1802: Field 'EnableAlerts' is declared as 'readonly' but is initialized with a constant value. Mark this field as 'const' instead.
 
-using EpicGames.Core;
-using EpicGames.Horde;
-using EpicGames.Horde.Tools;
-using Microsoft.Extensions.Logging;
-using Microsoft.Extensions.Logging.Abstractions;
 using System.Collections.Generic;
-using System.Diagnostics.CodeAnalysis;
-using System.IO;
 using System.Reflection;
 using System.Text.Json;
 using System.Text.Json.Serialization;
-using System.Xml;
+using EpicGames.Core;
+using EpicGames.Horde.Tools;
 
 namespace UnrealGameSync
 {
diff --git a/Engine/Source/Programs/UnrealGameSync/UnrealGameSyncShared/WorkspaceUpdate.cs b/Engine/Source/Programs/UnrealGameSync/UnrealGameSyncShared/WorkspaceUpdate.cs
--- a/Engine/Source/Programs/UnrealGameSync/UnrealGameSyncShared/WorkspaceUpdate.cs
+++ b/Engine/Source/Programs/UnrealGameSync/UnrealGameSyncShared/WorkspaceUpdate.cs
@@ -1,13 +1,11 @@
 // Copyright Epic Games, Inc. All Rights Reserved.
 
-using EpicGames.Core;
-using EpicGames.Perforce;
-using Microsoft.Extensions.Logging;
 using System;
 using System.Collections.Generic;
 using System.Diagnostics;
 using System.IO;
 using System.Linq;
+using System.Runtime.InteropServices;
 using System.Security.Cryptography;
 using System.Text;
 using System.Text.Encodings.Web;
@@ -15,7 +13,9 @@ using System.Text.Json;
 using System.Text.RegularExpressions;
 using System.Threading;
 using System.Threading.Tasks;
-using System.Runtime.InteropServices;
+using EpicGames.Core;
+using EpicGames.Perforce;
+using Microsoft.Extensions.Logging;
 
 namespace UnrealGameSync
 {
@@ -69,7 +69,7 @@ namespace UnrealGameSync
 
 		public int? MaxCommandsPerBatch { get; set; }
 		public int? MaxSizePerBatch { get; set; }
-		
+
 		public int? NumSyncErrorRetries { get; set; }
 
 		public PerforceSyncOptions Clone()
@@ -176,7 +176,7 @@ namespace UnrealGameSync
 
 		public static void ApplyDelta(Dictionary<Guid, bool> categories, Dictionary<Guid, bool> delta)
 		{
-			foreach(KeyValuePair<Guid, bool> pair in delta)
+			foreach (KeyValuePair<Guid, bool> pair in delta)
 			{
 				categories[pair.Key] = pair.Value;
 			}
@@ -207,7 +207,7 @@ namespace UnrealGameSync
 		public string ClientFileName => $"//{ClientName}{BranchPath}{ProjectPath}";
 		public string TelemetryProjectIdentifier => PerforceUtils.GetClientOrDepotDirectoryName(ProjectIdentifier);
 		public DirectoryReference EngineDir => DirectoryReference.Combine(LocalRootPath, "Engine");
-		public DirectoryReference? ProjectDir => ProjectPath.EndsWith(".uproject", StringComparison.OrdinalIgnoreCase)? LocalFileName.Directory : null;
+		public DirectoryReference? ProjectDir => ProjectPath.EndsWith(".uproject", StringComparison.OrdinalIgnoreCase) ? LocalFileName.Directory : null;
 		public DirectoryReference DataFolder => GetDataFolder(LocalRootPath);
 		public DirectoryReference CacheFolder => GetCacheFolder(LocalRootPath);
 
@@ -565,11 +565,11 @@ namespace UnrealGameSync
 			public void Dispose() => Release();
 		}
 
-		public static string ShellScriptExt { get; }= RuntimeInformation.IsOSPlatform(OSPlatform.Windows)? "bat" : "sh";
+		public static string ShellScriptExt { get; } = RuntimeInformation.IsOSPlatform(OSPlatform.Windows) ? "bat" : "sh";
 
 		public static Task<int> ExecuteShellCommandAsync(string commandLine, Action<string> processOutput, CancellationToken cancellationToken)
 		{
-			if(RuntimeInformation.IsOSPlatform(OSPlatform.Windows))
+			if (RuntimeInformation.IsOSPlatform(OSPlatform.Windows))
 			{
 				string cmdExe = Path.Combine(Environment.GetFolderPath(Environment.SpecialFolder.System), "cmd.exe");
 				return Utility.ExecuteProcessAsync(cmdExe, null, $"/C \"{commandLine}\"", processOutput, cancellationToken);
@@ -974,7 +974,7 @@ namespace UnrealGameSync
 									}
 								}
 							}
-							catch(EpicGames.Perforce.PerforceException)
+							catch (EpicGames.Perforce.PerforceException)
 							{
 								logger.LogInformation("Falling back to the slow way of finding last code change for CL {Number}...", Context.ChangeNumber);
 
@@ -1335,11 +1335,11 @@ namespace UnrealGameSync
 						logger.LogInformation("{Status}", step.StatusText);
 
 						DirectoryReference batchFilesDir = DirectoryReference.Combine(project.LocalRootPath, "Engine", "Build", "BatchFiles");
-						if(RuntimeInformation.IsOSPlatform(OSPlatform.OSX))
+						if (RuntimeInformation.IsOSPlatform(OSPlatform.OSX))
 						{
 							batchFilesDir = DirectoryReference.Combine(batchFilesDir, "Mac");
 						}
-						else if(RuntimeInformation.IsOSPlatform(OSPlatform.Linux))
+						else if (RuntimeInformation.IsOSPlatform(OSPlatform.Linux))
 						{
 							batchFilesDir = DirectoryReference.Combine(batchFilesDir, "Linux");
 						}
@@ -1727,7 +1727,7 @@ namespace UnrealGameSync
 				string statusMessage = "";
 
 				int maxRetries = context.PerforceSyncOptions?.NumSyncErrorRetries ?? PerforceSyncOptions.DefaultNumSyncErrorRetries;
-				for (int attempt = 0; ;attempt++)
+				for (int attempt = 0; ; attempt++)
 				{
 					// Sync the files
 					string? errorMessage;
