UGS: Increase the default number of threads to 4, now that we disable parallel syncing on individual UGS sync threads.

#rnx

diff --git a/Engine/Source/Programs/UnrealGameSync/UnrealGameSyncShared/WorkspaceUpdate.cs b/Engine/Source/Programs/UnrealGameSync/UnrealGameSyncShared/WorkspaceUpdate.cs
--- a/Engine/Source/Programs/UnrealGameSync/UnrealGameSyncShared/WorkspaceUpdate.cs
+++ b/Engine/Source/Programs/UnrealGameSync/UnrealGameSyncShared/WorkspaceUpdate.cs
@@ -57,7 +57,7 @@ namespace UnrealGameSync
 	public class PerforceSyncOptions
 	{
 		public const int DefaultNumRetries = 0;
-		public const int DefaultNumThreads = 2;
+		public const int DefaultNumThreads = 4;
 		public const int DefaultTcpBufferSize = 0;
 		public const int DefaultFileBufferSize = 0;
 
