UGS: Show a dialog box showing how to launch the legacy settings window if using an older version of UnrealGameSyncLauncher.

diff --git a/Engine/Source/Programs/UnrealGameSync/UnrealGameSync/Forms/ApplicationSettingsWindow.cs b/Engine/Source/Programs/UnrealGameSync/UnrealGameSync/Forms/ApplicationSettingsWindow.cs
--- a/Engine/Source/Programs/UnrealGameSync/UnrealGameSync/Forms/ApplicationSettingsWindow.cs
+++ b/Engine/Source/Programs/UnrealGameSync/UnrealGameSync/Forms/ApplicationSettingsWindow.cs
@@ -2,6 +2,7 @@
 
 using System;
 using System.Collections.Generic;
+using System.IO;
 using System.Linq;
 using System.Threading;
 using System.Threading.Tasks;
@@ -18,6 +19,7 @@ namespace UnrealGameSync
 		{
 			Cancel,
 			Ok,
+			Quit,
 			Restart,
 			RestartAndConfigureUpdate,
 		}
@@ -204,6 +206,15 @@ namespace UnrealGameSync
 
 		private void UpdateSettingsBtn_Click(object sender, EventArgs e)
 		{
+			if (Path.GetFileName(_originalExecutableFileName).Contains("Launcher", StringComparison.OrdinalIgnoreCase))
+			{
+				if (MessageBox.Show("To configure update settings, quit UnrealGameSync and relaunch from the start menu while holding down the shift key.\n\nWould you like to quit now?", "Restart Required", MessageBoxButtons.YesNo) == DialogResult.Yes)
+				{
+					ApplySettings(Result.Quit);
+					return;
+				}
+			}
+
 			if (MessageBox.Show("UnrealGameSync must be restarted to configure update settings.\n\nWould you like to restart now?", "Restart Required", MessageBoxButtons.OKCancel) != DialogResult.OK)
 			{
 				return;
diff --git a/Engine/Source/Programs/UnrealGameSync/UnrealGameSync/Forms/MainWindow.cs b/Engine/Source/Programs/UnrealGameSync/UnrealGameSync/Forms/MainWindow.cs
--- a/Engine/Source/Programs/UnrealGameSync/UnrealGameSync/Forms/MainWindow.cs
+++ b/Engine/Source/Programs/UnrealGameSync/UnrealGameSync/Forms/MainWindow.cs
@@ -1531,6 +1531,10 @@ namespace UnrealGameSync
 			{
 				_updateMonitor.TriggerUpdate(UpdateType.UserInitiated, true);
 			}
+			else if (result == ApplicationSettingsWindow.Result.Quit)
+			{
+				ForceClose();
+			}
 
 			for (int idx = 0; idx < TabControl.GetTabCount(); idx++)
 			{
