UGS: Fix legacy update prompt falling back to new update prompt if user cancels.

#rnx

diff --git a/Engine/Source/Programs/UnrealGameSync/UnrealGameSync/Forms/ApplicationSettingsWindow.cs b/Engine/Source/Programs/UnrealGameSync/UnrealGameSync/Forms/ApplicationSettingsWindow.cs
--- a/Engine/Source/Programs/UnrealGameSync/UnrealGameSync/Forms/ApplicationSettingsWindow.cs
+++ b/Engine/Source/Programs/UnrealGameSync/UnrealGameSync/Forms/ApplicationSettingsWindow.cs
@@ -211,16 +211,15 @@ namespace UnrealGameSync
 				if (MessageBox.Show("To configure update settings, quit UnrealGameSync and relaunch from the start menu while holding down the shift key.\n\nWould you like to quit now?", "Restart Required", MessageBoxButtons.YesNo) == DialogResult.Yes)
 				{
 					ApplySettings(Result.Quit);
-					return;
 				}
 			}
-
-			if (MessageBox.Show("UnrealGameSync must be restarted to configure update settings.\n\nWould you like to restart now?", "Restart Required", MessageBoxButtons.OKCancel) != DialogResult.OK)
+			else
 			{
-				return;
+				if (MessageBox.Show("UnrealGameSync must be restarted to configure update settings.\n\nWould you like to restart now?", "Restart Required", MessageBoxButtons.OKCancel) == DialogResult.OK)
+				{
+					ApplySettings(Result.RestartAndConfigureUpdate);
+				}
 			}
-
-			ApplySettings(Result.RestartAndConfigureUpdate);
 		}
 
 		private void OkBtn_Click(object sender, EventArgs e)
