UGS: Build fix.

#rnx

diff --git a/Engine/Source/Programs/UnrealGameSync/UnrealGameSyncShared/ArchiveInfo.cs b/Engine/Source/Programs/UnrealGameSync/UnrealGameSyncShared/ArchiveInfo.cs
--- a/Engine/Source/Programs/UnrealGameSync/UnrealGameSyncShared/ArchiveInfo.cs
+++ b/Engine/Source/Programs/UnrealGameSync/UnrealGameSyncShared/ArchiveInfo.cs
@@ -288,7 +288,7 @@ namespace UnrealGameSync
 				StreamId streamId = new StreamId(StreamName);
 				int? minChange = null;
 				int? maxChange = null;
-				List<GetArtifactResponse> artifactResponse = await hordeHttpClient.FindArtifactsByTypeAsync(artifactType, streamId, minChange, maxChange, cancellationToken);
+				List<GetArtifactResponse> artifactResponse = await hordeHttpClient.FindArtifactsByTypeAsync(artifactType, streamId, minChange, maxChange, null, cancellationToken);
 
 				foreach (GetArtifactResponse response in artifactResponse)
 				{
