Horde: Fix uncaught exception querying Horde for artifacts

diff --git a/Engine/Source/Programs/UnrealGameSync/UnrealGameSyncShared/ArchiveInfo.cs b/Engine/Source/Programs/UnrealGameSync/UnrealGameSyncShared/ArchiveInfo.cs
--- a/Engine/Source/Programs/UnrealGameSync/UnrealGameSyncShared/ArchiveInfo.cs
+++ b/Engine/Source/Programs/UnrealGameSync/UnrealGameSyncShared/ArchiveInfo.cs
@@ -333,7 +333,13 @@ namespace UnrealGameSync
 				// Make sure the zipped binaries path exists
 				foreach (BaseArchiveInfo newArchive in newArchives)
 				{
-					await newArchive.FindArtifacts(perforce, cancellationToken);
+					try
+					{
+						await newArchive.FindArtifacts(perforce, cancellationToken);
+					}
+					catch
+					{
+					}
 				}
 			}
 
