UGS: Exclude deleted tools from ToolUpdateMonitor.

diff --git a/Engine/Source/Programs/UnrealGameSync/UnrealGameSync/ToolUpdateMonitor.cs b/Engine/Source/Programs/UnrealGameSync/UnrealGameSync/ToolUpdateMonitor.cs
--- a/Engine/Source/Programs/UnrealGameSync/UnrealGameSync/ToolUpdateMonitor.cs
+++ b/Engine/Source/Programs/UnrealGameSync/UnrealGameSync/ToolUpdateMonitor.cs
@@ -312,6 +312,8 @@ namespace UnrealGameSync
 				Dictionary<string, ToolInfo> newPerforceTools = new Dictionary<string, ToolInfo>(StringComparer.Ordinal);
 
 				List<FStatRecord> fileRecords = await perforce.FStatAsync($"{DeploymentSettings.Instance.ToolsDepotPath}/...", cancellationToken).ToListAsync(cancellationToken);
+				fileRecords.RemoveAll(x => x.Action == FileAction.Delete || x.Action == FileAction.MoveDelete);
+
 				foreach (FStatRecord fileRecord in fileRecords)
 				{
 					if (fileRecord.DepotFile != null && fileRecord.DepotFile.EndsWith(".ini", StringComparison.OrdinalIgnoreCase))
