Horde: Prevent UGS trying to auto-update itself when launched from the AppData\Local folder. Can happen when users pin the application to the start menu.

diff --git a/Engine/Source/Programs/UnrealGameSync/UnrealGameSync/Launcher.cs b/Engine/Source/Programs/UnrealGameSync/UnrealGameSync/Launcher.cs
--- a/Engine/Source/Programs/UnrealGameSync/UnrealGameSync/Launcher.cs
+++ b/Engine/Source/Programs/UnrealGameSync/UnrealGameSync/Launcher.cs
@@ -7,6 +7,7 @@ using System.Diagnostics.CodeAnalysis;
 using System.IO;
 using System.IO.Compression;
 using System.Linq;
+using System.Reflection;
 using System.Text.Json;
 using System.Threading;
 using System.Threading.Tasks;
@@ -39,6 +40,13 @@ namespace UnrealGameSync
 				return LauncherResult.Continue;
 			}
 
+			// Also check we're not running under the 
+			DirectoryReference applicationFolder = new DirectoryReference(GetSyncFolder());
+			if (new FileReference(Assembly.GetExecutingAssembly().Location).IsUnderDirectory(applicationFolder))
+			{
+				return LauncherResult.Continue;
+			}
+
 			// Figure out if we should sync the unstable build by default
 			bool preview = args.Contains("-unstable", StringComparer.InvariantCultureIgnoreCase) || args.Contains("-preview", StringComparer.InvariantCultureIgnoreCase);
 			bool openSettings = args.Contains("-settings", StringComparer.OrdinalIgnoreCase);
@@ -106,6 +114,11 @@ namespace UnrealGameSync
 			public void Dispose() { }
 		}
 
+		static string GetSyncFolder()
+		{
+			return Path.Combine(Environment.GetFolderPath(Environment.SpecialFolder.LocalApplicationData), "UnrealGameSync", "Latest");
+		}
+
 		public static async Task SyncAndRun(IPerforceConnection? perforce, LauncherSettings launcherSettings, string[] args, Mutex instanceMutex, ILogger logger, CancellationToken cancellationToken)
 		{
 			try
@@ -116,7 +129,7 @@ namespace UnrealGameSync
 				}
 
 				// Create the target folder
-				string applicationFolder = Path.Combine(Environment.GetFolderPath(Environment.SpecialFolder.LocalApplicationData), "UnrealGameSync", "Latest");
+				string applicationFolder = GetSyncFolder();
 				if (!SafeCreateDirectory(applicationFolder))
 				{
 					throw new UserErrorException($"Couldn't create directory: {applicationFolder}");
