UGS: Update Horde builds at the same time as P4 builds.

#rnx

diff --git a/Engine/Source/Programs/UnrealGameSync/BuildUGS.xml b/Engine/Source/Programs/UnrealGameSync/BuildUGS.xml
--- a/Engine/Source/Programs/UnrealGameSync/BuildUGS.xml
+++ b/Engine/Source/Programs/UnrealGameSync/BuildUGS.xml
@@ -53,6 +53,11 @@
 
 			<Copy From="$(UgsInstallerDir)/bin/Release/UnrealGameSync.msi" To="$(UgsInstallerPublishDir)/UnrealGameSync.msi"/>
 		</Node>
+		<Node Name="Publish UnrealGameSync Installer Horde" Requires="Build UnrealGameSync Installer">
+			<Error Message="Missing -HordeSettings=... argument" If="'$(HordeSettings)' == ''"/>
+			<Zip FromDir="$(UgsInstallerPublishDir)" Files="$(UgsInstallerPublishDir)/UnrealGameSync.msi" ZipFile="$(UgsZipPublishDir)/UgsInstaller.zip"/>
+			<DeployTool Id="ugs" Version="$(InformationalVersion)" Settings="$(HordeSettings)" File="$(UgsZipPublishDir)/UgsInstaller.zip"/>
+		</Node>
 
 		<Property Name="InstallerDependency" Value=""/>
 		<Property Name="InstallerDependency" Value="Build UnrealGameSync Installer" If="$(IncludeInstaller)"/>
