UGS: Add an IArchive type to encapsulate a downloadable archive, and renamed IArchiveInfo to IArchiveChannel for clarity.

Also allow enumerating Horde artifacts by searching for a project identifier registered as a key with the artifact.

diff --git a/Engine/Source/Programs/UnrealGameSync/UnrealGameSync/Controls/WorkspaceControl.cs b/Engine/Source/Programs/UnrealGameSync/UnrealGameSync/Controls/WorkspaceControl.cs
--- a/Engine/Source/Programs/UnrealGameSync/UnrealGameSync/Controls/WorkspaceControl.cs
+++ b/Engine/Source/Programs/UnrealGameSync/UnrealGameSync/Controls/WorkspaceControl.cs
@@ -197,7 +197,6 @@ namespace UnrealGameSync
 		readonly SynchronizationContext _mainThreadSynchronizationContext;
 		bool _isDisposing;
 
-		//		JupiterMonitor _jupiterMonitor;
 		PerforceMonitor _perforceMonitor;
 		Workspace _workspace;
 #pragma warning disable CA2213 //warning CA2213: 'WorkspaceControl' contains field '_issueMonitor' that is of IDisposable type 'IssueMonitor', but it is never disposed. Change the Dispose method on 'WorkspaceControl' to call Close or Dispose on this field.
@@ -210,7 +209,7 @@ namespace UnrealGameSync
 		HashSet<int> _promotedChangeNumbers = new HashSet<int>();
 		List<int> _listIndexToChangeIndex = new List<int>();
 		List<int> _sortedChangeNumbers = new List<int>();
-		readonly Dictionary<string, Dictionary<int, string?>> _archiveToChangeNumberToArchiveKey = new Dictionary<string, Dictionary<int, string?>>();
+		readonly Dictionary<string, Dictionary<int, IArchive?>> _archiveToChangeNumberToArchiveKey = new Dictionary<string, Dictionary<int, IArchive?>>();
 		readonly Dictionary<int, ChangeLayoutInfo> _changeNumberToLayoutInfo = new Dictionary<int, ChangeLayoutInfo>();
 		readonly List<ContextMenuStrip> _customStatusPanelMenus = new List<ContextMenuStrip>();
 		readonly List<(string, Action<Point, Rectangle>)> _customStatusPanelLinks = new List<(string, Action<Point, Rectangle>)>();
@@ -365,9 +364,6 @@ namespace UnrealGameSync
 			_eventMonitor = new EventMonitor(_apiUrl, PerforceUtils.GetClientOrDepotDirectoryName(SelectedProjectIdentifier), openProjectInfo.PerforceSettings.UserName, _serviceProvider);
 			_eventMonitor.OnUpdatesReady += UpdateReviewsCallback;
 
-			//			ILogger<JupiterMonitor> jupiterLogger = _serviceProvider.GetRequiredService<ILogger<JupiterMonitor>>();
-			//			_jupiterMonitor = JupiterMonitor.CreateFromConfigFile(inOidcTokenManager, jupiterLogger, openProjectInfo.LatestProjectConfigFile, SelectedProjectIdentifier);
-
 			UpdateColumnSettings(true);
 			Font = new System.Drawing.Font("Segoe UI", 8.25F, System.Drawing.FontStyle.Regular, System.Drawing.GraphicsUnit.Point);
 
@@ -741,11 +737,6 @@ namespace UnrealGameSync
 				_badgeFont.Dispose();
 				_badgeFont = null!;
 			}
-			//			if (_jupiterMonitor != null)
-			//			{
-			//				_jupiterMonitor.Dispose();
-			//				_jupiterMonitor = null!;
-			//			}
 
 			base.Dispose(disposing);
 		}
@@ -909,23 +900,23 @@ namespace UnrealGameSync
 			WorkspaceUpdateContext context = new WorkspaceUpdateContext(changeNumber, options, GetEditorBuildConfig(), combinedSyncFilter, _projectSettings.BuildSteps, null);
 			if (options.HasFlag(WorkspaceUpdateOptions.SyncArchives))
 			{
-				IReadOnlyList<IArchiveInfo> archives = GetArchives();
-				foreach (IArchiveInfo archive in archives)
+				IReadOnlyList<IArchiveChannel> archives = GetArchiveChannels();
+				foreach (IArchiveChannel archive in archives)
 				{
 					context.ArchiveTypeToArchive[archive.Type] = null;
 				}
 
-				List<IArchiveInfo> selectedArchives = GetSelectedArchives(archives);
-				foreach (IArchiveInfo archive in selectedArchives)
+				List<IArchiveChannel> selectedArchiveChannels = GetSelectedArchiveChannels(archives);
+				foreach (IArchiveChannel archiveChannel in selectedArchiveChannels)
 				{
-					string? archivePath = GetArchiveKeyForChangeNumber(archive, changeNumber);
+					IArchive? archivePath = GetArchiveForChangeNumber(archiveChannel, changeNumber);
 					if (archivePath == null)
 					{
-						MessageBox.Show(String.Format("There are no compiled {0} binaries for this change. To sync it, you must disable syncing of precompiled editor binaries.", archive.Name));
+						MessageBox.Show(String.Format("There are no compiled {0} binaries for this change. To sync it, you must disable syncing of precompiled editor binaries.", archiveChannel.Name));
 						return;
 					}
 
-					if (archive.Type == IArchiveInfo.EditorArchiveType)
+					if (archiveChannel.Type == IArchiveChannel.EditorArchiveType)
 					{
 						context.Options &= ~(WorkspaceUpdateOptions.Build | WorkspaceUpdateOptions.GenerateProjectFiles | WorkspaceUpdateOptions.OpenSolutionAfterSync);
 					}
@@ -936,7 +927,7 @@ namespace UnrealGameSync
 						context.SyncFilter.AddRange(zippedBinariesSyncFilter);
 					}
 
-					context.ArchiveTypeToArchive[archive.Type] = new Tuple<IArchiveInfo, string>(archive, archivePath);
+					context.ArchiveTypeToArchive[archiveChannel.Type] = archivePath;
 				}
 			}
 			StartWorkspaceUpdate(context, callback);
@@ -965,8 +956,8 @@ namespace UnrealGameSync
 			{
 				if (!context.Options.HasFlag(WorkspaceUpdateOptions.ContentOnly) && (context.CustomBuildSteps == null || context.CustomBuildSteps.Count == 0))
 				{
-					bool usingPrecompiledEditor = context.ArchiveTypeToArchive.TryGetValue(IArchiveInfo.EditorArchiveType, out Tuple<IArchiveInfo, string>? archiveInfo) && archiveInfo != null;
-					bool usingLastSyncedEditorArchive = archiveInfo != null && archiveInfo.Item2 == _workspace.State.LastSyncEditorArchive;
+					bool usingPrecompiledEditor = context.ArchiveTypeToArchive.TryGetValue(IArchiveChannel.EditorArchiveType, out IArchive? archive) && archive != null;
+					bool usingLastSyncedEditorArchive = archive != null && archive.Key == _workspace.State.LastSyncEditorArchive;
 					if (!usingPrecompiledEditor || !usingLastSyncedEditorArchive)
 					{
 						FileReference targetFile = ConfigUtils.GetEditorTargetFile(_workspace.Project, _workspace.ProjectConfigFile);
@@ -2029,40 +2020,58 @@ namespace UnrealGameSync
 			}
 		}
 
-		private string? GetArchiveKeyForChangeNumber(IArchiveInfo archive, int changeNumber)
+		private IArchive? GetArchiveForChangeNumber(IArchiveChannel archiveChannel, int changeNumber)
 		{
-			Dictionary<int, string?>? changeNumberToArchivePath;
-			if (!_archiveToChangeNumberToArchiveKey.TryGetValue(archive.Name, out changeNumberToArchivePath))
+			Dictionary<int, IArchive?>? changeNumberToArchive;
+			if (!_archiveToChangeNumberToArchiveKey.TryGetValue(archiveChannel.Name, out changeNumberToArchive))
 			{
-				changeNumberToArchivePath = new Dictionary<int, string?>();
-				_archiveToChangeNumberToArchiveKey[archive.Name] = changeNumberToArchivePath;
+				changeNumberToArchive = new Dictionary<int, IArchive?>();
+				_archiveToChangeNumberToArchiveKey[archiveChannel.Name] = changeNumberToArchive;
 			}
-			return GetArchiveKeyForChangeNumber(archive, changeNumber, changeNumber, changeNumberToArchivePath);
+			return GetArchiveForChangeNumber(archiveChannel, changeNumber, changeNumber, changeNumberToArchive);
 		}
 
-		private string? GetArchiveKeyForChangeNumber(IArchiveInfo archive, int changeNumber, int maxChangeNumber, Dictionary<int, string?> changeNumberToArchivePath)
+		private IArchive? GetArchiveForChangeNumber(IArchiveChannel archiveChannel, int changeNumber, int maxChangeNumber, Dictionary<int, IArchive?> changeNumberToArchive)
 		{
-			string? archivePath;
-			if (!changeNumberToArchivePath.TryGetValue(changeNumber, out archivePath))
+			IArchive? archive;
+			if (!changeNumberToArchive.TryGetValue(changeNumber, out archive))
 			{
-				PerforceChangeDetails? details;
-				if (_perforceMonitor.TryGetChangeDetails(changeNumber, out details))
+				archive = GetArchiveForChangeNumberUncached(archiveChannel, changeNumber, maxChangeNumber, changeNumberToArchive);
+				changeNumberToArchive.Add(changeNumber, archive);
+			}
+			return archive;
+		}
+
+		private IArchive? GetArchiveForChangeNumberUncached(IArchiveChannel archiveChannel, int changeNumber, int maxChangeNumber, Dictionary<int, IArchive?> changeNumberToArchive)
+		{
+			int idx = _sortedChangeNumbers.BinarySearch(changeNumber);
+			if (idx < 0)
+			{
+				return null;
+			}
+
+			PerforceChangeDetails? details;
+			while (idx > 0 && _perforceMonitor.TryGetChangeDetails(_sortedChangeNumbers[idx], out details) && !details.ContainsCode)
+			{
+				idx--;
+			}
+
+			for (; ; )
+			{
+				IArchive? archive = archiveChannel.TryGetArchiveForChangeNumber(_sortedChangeNumbers[idx], maxChangeNumber);
+				if (archive != null)
 				{
-					// Try to get the archive for this CL
-					if (!archive.TryGetArchiveKeyForChangeNumber(changeNumber, maxChangeNumber, out archivePath) && !details.ContainsCode)
-					{
-						// Otherwise if it's a content-only change, find the previous build any use the archive path from that
-						int index = _sortedChangeNumbers.BinarySearch(changeNumber);
-						if (index > 0)
-						{
-							archivePath = GetArchiveKeyForChangeNumber(archive, _sortedChangeNumbers[index - 1], maxChangeNumber, changeNumberToArchivePath);
-						}
-					}
+					return archive;
 				}
-				changeNumberToArchivePath.Add(changeNumber, archivePath);
-			}
 
-			return archivePath;
+				idx++;
+				if (idx >= _sortedChangeNumbers.Count 
+					|| !_perforceMonitor.TryGetChangeDetails(_sortedChangeNumbers[idx], out details) 
+					|| details.ContainsCode)
+				{
+					return null;
+				}
+			}
 		}
 
 		private static Color Blend(Color first, Color second, float t)
@@ -2077,8 +2086,8 @@ namespace UnrealGameSync
 				return false;
 			}
 
-			List<IArchiveInfo> selectedArchives = GetSelectedArchives(GetArchives());
-			return selectedArchives.Count == 0 || selectedArchives.All(x => GetArchiveKeyForChangeNumber(x, changeNumber) != null);
+			List<IArchiveChannel> selectedArchives = GetSelectedArchiveChannels(GetArchiveChannels());
+			return selectedArchives.Count == 0 || selectedArchives.All(x => GetArchiveForChangeNumber(x, changeNumber) != null);
 		}
 
 		/// <summary>
@@ -4554,10 +4563,10 @@ namespace UnrealGameSync
 			return null;
 		}
 
-		private List<IArchiveInfo> GetSelectedArchives(IReadOnlyList<IArchiveInfo> archives)
+		private List<IArchiveChannel> GetSelectedArchiveChannels(IReadOnlyList<IArchiveChannel> archives)
 		{
-			Dictionary<string, KeyValuePair<IArchiveInfo, int>> archiveTypeToSelection = new Dictionary<string, KeyValuePair<IArchiveInfo, int>>();
-			foreach (IArchiveInfo archive in archives)
+			Dictionary<string, KeyValuePair<IArchiveChannel, int>> archiveTypeToSelection = new Dictionary<string, KeyValuePair<IArchiveChannel, int>>();
+			foreach (IArchiveChannel archive in archives)
 			{
 				ArchiveSettings? archiveSettings = _settings.Archives.FirstOrDefault(x => x.Type == archive.Type);
 				if (archiveSettings != null && archiveSettings.Enabled)
@@ -4568,10 +4577,10 @@ namespace UnrealGameSync
 						preference = archiveSettings.Order.Count;
 					}
 
-					KeyValuePair<IArchiveInfo, int> existingItem;
+					KeyValuePair<IArchiveChannel, int> existingItem;
 					if (!archiveTypeToSelection.TryGetValue(archive.Type, out existingItem) || existingItem.Value > preference)
 					{
-						archiveTypeToSelection[archive.Type] = new KeyValuePair<IArchiveInfo, int>(archive, preference);
+						archiveTypeToSelection[archive.Type] = new KeyValuePair<IArchiveChannel, int>(archive, preference);
 					}
 				}
 			}
@@ -4589,7 +4598,7 @@ namespace UnrealGameSync
 			UpdateSelectedArchives();
 		}
 
-		private void SetSelectedArchive(IArchiveInfo archive, bool selected)
+		private void SetSelectedArchive(IArchiveChannel archive, bool selected)
 		{
 			ArchiveSettings? archiveSettings = _settings.Archives.FirstOrDefault(x => x.Type == archive.Type);
 			if (archiveSettings == null)
@@ -4627,8 +4636,8 @@ namespace UnrealGameSync
 
 			OptionsContextMenu_SyncPrecompiledBinaries.DropDownItems.Clear();
 
-			IReadOnlyList<IArchiveInfo> archives = GetArchives();
-			if (archives == null || archives.Count == 0)
+			IReadOnlyList<IArchiveChannel> channels = GetArchiveChannels();
+			if (channels == null || channels.Count == 0)
 			{
 				OptionsContextMenu_SyncPrecompiledBinaries.Enabled = false;
 				OptionsContextMenu_SyncPrecompiledBinaries.ToolTipText = String.Format("Precompiled binaries are not available for {0}", SelectedProjectIdentifier);
@@ -4636,13 +4645,13 @@ namespace UnrealGameSync
 			}
 			else
 			{
-				List<IArchiveInfo> selectedArchives = GetSelectedArchives(archives);
+				List<IArchiveChannel> selectedArchives = GetSelectedArchiveChannels(channels);
 
 				OptionsContextMenu_SyncPrecompiledBinaries.Enabled = true;
 				OptionsContextMenu_SyncPrecompiledBinaries.ToolTipText = null;
 				OptionsContextMenu_SyncPrecompiledBinaries.Checked = selectedArchives.Count > 0;
 
-				if (archives.Count > 1 || archives[0].Type != IArchiveInfo.EditorArchiveType)
+				if (channels.Count > 1 || channels[0].Type != IArchiveChannel.EditorArchiveType)
 				{
 					ToolStripMenuItem disableItem = new ToolStripMenuItem("Disable (compile locally)");
 					disableItem.Checked = (selectedArchives.Count == 0);
@@ -4652,16 +4661,13 @@ namespace UnrealGameSync
 					ToolStripSeparator separator = new ToolStripSeparator();
 					OptionsContextMenu_SyncPrecompiledBinaries.DropDownItems.Add(separator);
 
-					foreach (IArchiveInfo archive in archives)
+					foreach (IArchiveChannel channel in channels)
 					{
-						ToolStripMenuItem item = new ToolStripMenuItem(archive.Name);
-						item.Enabled = archive.Exists();
-						if (!item.Enabled)
-						{
-							item.ToolTipText = String.Format("No valid archives found at {0}", archive.BasePath);
-						}
-						item.Checked = selectedArchives.Contains(archive);
-						item.Click += (sender, args) => SetSelectedArchive(archive, !item.Checked);
+						ToolStripMenuItem item = new ToolStripMenuItem(channel.Name);
+						item.Enabled = channel.HasAny();
+						item.ToolTipText = channel.ToolTip;
+						item.Checked = selectedArchives.Contains(channel);
+						item.Click += (sender, args) => SetSelectedArchive(channel, !item.Checked);
 						OptionsContextMenu_SyncPrecompiledBinaries.DropDownItems.Add(item);
 					}
 				}
@@ -4683,21 +4689,14 @@ namespace UnrealGameSync
 			OptionsContextMenu.Show(OptionsButton, new Point(OptionsButton.Width - OptionsContextMenu.Size.Width, OptionsButton.Height));
 		}
 
-		private IReadOnlyList<IArchiveInfo> GetArchives()
+		private IReadOnlyList<IArchiveChannel> GetArchiveChannels()
 		{
-			//			IReadOnlyList<IArchiveInfo>? availableArchives = _jupiterMonitor?.AvailableArchives;
-			//			if (availableArchives != null && availableArchives.Count != 0)
-			//			{
-			//				return availableArchives;
-			//			}
-
-			// if jupiter had no archives we fallback to the perforce monitor
 			if (_perforceMonitor != null)
 			{
-				return _perforceMonitor.AvailableArchives;
+				return _perforceMonitor.AvailableArchiveChannels;
 			}
 
-			return new List<IArchiveInfo>();
+			return new List<IArchiveChannel>();
 		}
 
 		private void BuildAfterSyncCheckBox_CheckedChanged(object sender, EventArgs e)
@@ -5686,7 +5685,7 @@ namespace UnrealGameSync
 			}
 		}
 
-		private bool ShouldSyncPrecompiledEditor => _settings.Archives.Any(x => x.Enabled && x.Type == IArchiveInfo.EditorArchiveType) && GetArchives().Any(x => x.Type == "Editor");
+		private bool ShouldSyncPrecompiledEditor => _settings.Archives.Any(x => x.Enabled && x.Type == IArchiveChannel.EditorArchiveType) && GetArchiveChannels().Any(x => x.Type == "Editor");
 
 		public BuildConfig GetEditorBuildConfig()
 		{
@@ -5723,7 +5722,7 @@ namespace UnrealGameSync
 		{
 			if (OptionsContextMenu_SyncPrecompiledBinaries.DropDownItems.Count == 0)
 			{
-				IArchiveInfo? editorArchive = GetArchives().FirstOrDefault(x => x.Type == IArchiveInfo.EditorArchiveType);
+				IArchiveChannel? editorArchive = GetArchiveChannels().FirstOrDefault(x => x.Type == IArchiveChannel.EditorArchiveType);
 				if (editorArchive != null)
 				{
 					SetSelectedArchive(editorArchive, !OptionsContextMenu_SyncPrecompiledBinaries.Checked);
diff --git a/Engine/Source/Programs/UnrealGameSync/UnrealGameSync/PerforceMonitor.cs b/Engine/Source/Programs/UnrealGameSync/UnrealGameSync/PerforceMonitor.cs
--- a/Engine/Source/Programs/UnrealGameSync/UnrealGameSync/PerforceMonitor.cs
+++ b/Engine/Source/Programs/UnrealGameSync/UnrealGameSync/PerforceMonitor.cs
@@ -17,7 +17,7 @@ namespace UnrealGameSync
 {
 	interface IArchiveInfoSource
 	{
-		IReadOnlyList<IArchiveInfo> AvailableArchives { get; }
+		IReadOnlyList<IArchiveChannel> AvailableArchiveChannels { get; }
 	}
 
 	class PerforceMonitor : IDisposable, IArchiveInfoSource
@@ -45,7 +45,7 @@ namespace UnrealGameSync
 		SortedSet<ChangesRecord> _changes = new SortedSet<ChangesRecord>(new PerforceChangeSorter());
 		readonly SortedDictionary<int, PerforceChangeDetails> _changeDetails = new SortedDictionary<int, PerforceChangeDetails>();
 		readonly SortedSet<int> _promotedChangeNumbers = new SortedSet<int>();
-		List<BaseArchiveInfo> _archives = new List<BaseArchiveInfo>();
+		List<BaseArchiveChannel> _archives = new List<BaseArchiveChannel>();
 		readonly AsyncEvent _refreshEvent = new AsyncEvent();
 		readonly ILogger _logger;
 		readonly bool _isEnterpriseProject;
@@ -89,7 +89,7 @@ namespace UnrealGameSync
 			_cancellationSource = new CancellationTokenSource();
 			_hordeClient = serviceProvider.GetRequiredService<IHordeClient>();
 
-			AvailableArchives = (new List<IArchiveInfo>()).AsReadOnly();
+			AvailableArchiveChannels = (new List<IArchiveChannel>()).AsReadOnly();
 		}
 
 		public void Start()
@@ -316,9 +316,9 @@ namespace UnrealGameSync
 			if (perforceConfigSection != null && perforceConfigSection.GetValue("FindAllChangesForPCBs", false))
 			{
 				int minZippedChangeNumber = -1;
-				foreach (BaseArchiveInfo archive in _archives)
+				foreach (BaseArchiveChannel archive in _archives)
 				{
-					foreach (int changeNumber in archive.ChangeNumberToArchiveKey.Keys)
+					foreach (int changeNumber in archive.ChangeNumberToArchive.Keys)
 					{
 						if (changeNumber > minZippedChangeNumber && changeNumber <= oldestChangeNumber)
 						{
@@ -375,7 +375,7 @@ namespace UnrealGameSync
 						foreach (ChangesRecord change in _changes)
 						{
 							trimmedChanges.Add(change);
-							if (trimmedChanges.Count >= maxChanges && _archives.Any(x => x.ChangeNumberToArchiveKey.Count == 0 || x.ChangeNumberToArchiveKey.ContainsKey(change.Number) || x.ChangeNumberToArchiveKey.First().Key > change.Number))
+							if (trimmedChanges.Count >= maxChanges && _archives.Any(x => x.ChangeNumberToArchive.Count == 0 || x.ChangeNumberToArchive.ContainsKey(change.Number) || x.ChangeNumberToArchive.First().Key > change.Number))
 							{
 								break;
 							}
@@ -495,13 +495,13 @@ namespace UnrealGameSync
 
 		async Task<bool> UpdateArchivesAsync(IPerforceConnection perforce, CancellationToken cancellationToken)
 		{
-			List<BaseArchiveInfo> newArchives = await BaseArchive.EnumerateAsync(perforce, _hordeClient, LatestProjectConfigFile, _selectedProjectIdentifier, cancellationToken);
+			List<BaseArchiveChannel> newArchives = await BaseArchive.EnumerateChannelsAsync(perforce, _hordeClient, LatestProjectConfigFile, _selectedProjectIdentifier, cancellationToken);
 
 			// Check if the information has changed
 			if (!Enumerable.SequenceEqual(_archives, newArchives))
 			{
 				_archives = newArchives;
-				AvailableArchives = _archives.Select(x => (IArchiveInfo)x).ToList();
+				AvailableArchiveChannels = _archives.Select(x => (IArchiveChannel)x).ToList();
 
 				if (_changes.Count > 0)
 				{
@@ -565,7 +565,7 @@ namespace UnrealGameSync
 			private set;
 		}
 
-		public IReadOnlyList<IArchiveInfo> AvailableArchives
+		public IReadOnlyList<IArchiveChannel> AvailableArchiveChannels
 		{
 			get;
 			private set;
diff --git a/Engine/Source/Programs/UnrealGameSync/UnrealGameSyncCmd/Program.cs b/Engine/Source/Programs/UnrealGameSync/UnrealGameSyncCmd/Program.cs
--- a/Engine/Source/Programs/UnrealGameSync/UnrealGameSyncCmd/Program.cs
+++ b/Engine/Source/Programs/UnrealGameSync/UnrealGameSyncCmd/Program.cs
@@ -869,9 +869,9 @@ namespace UnrealGameSyncCmd
 
 				if (syncOptions.Binaries)
 				{
-					List<BaseArchiveInfo> archives = await BaseArchive.EnumerateAsync(perforceClient, context._hordeClient, projectConfig, state.Current.ProjectIdentifier, CancellationToken.None);
+					List<BaseArchiveChannel> archives = await BaseArchive.EnumerateChannelsAsync(perforceClient, context._hordeClient, projectConfig, state.Current.ProjectIdentifier, CancellationToken.None);
 
-					BaseArchiveInfo? editorArchiveInfo = archives.FirstOrDefault(x => x.Name == IArchiveInfo.EditorArchiveType);
+					BaseArchiveChannel? editorArchiveInfo = archives.FirstOrDefault(x => x.Name == IArchiveChannel.EditorArchiveType);
 					if (editorArchiveInfo == null)
 					{
 						throw new UserErrorException("No editor archives found for project");
@@ -907,7 +907,7 @@ namespace UnrealGameSyncCmd
 					}
 
 					updateContext.Options |= WorkspaceUpdateOptions.SyncArchives;
-					updateContext.ArchiveTypeToArchive[IArchiveInfo.EditorArchiveType] = Tuple.Create<IArchiveInfo, string>(editorArchiveInfo, revision.Value);
+					updateContext.ArchiveTypeToArchive[IArchiveChannel.EditorArchiveType] = Tuple.Create<IArchiveChannel, string>(editorArchiveInfo, revision.Value);
 				}
 
 				WorkspaceUpdate update = new WorkspaceUpdate(updateContext);
diff --git a/Engine/Source/Programs/UnrealGameSync/UnrealGameSyncShared/ArchiveInfo.cs b/Engine/Source/Programs/UnrealGameSync/UnrealGameSyncShared/ArchiveInfo.cs
--- a/Engine/Source/Programs/UnrealGameSync/UnrealGameSyncShared/ArchiveInfo.cs
+++ b/Engine/Source/Programs/UnrealGameSync/UnrealGameSyncShared/ArchiveInfo.cs
@@ -16,121 +16,66 @@ using Microsoft.Extensions.Logging;
 
 namespace UnrealGameSync
 {
-	public interface IArchiveInfo
+	public interface IArchive
+	{
+		public string Key { get; }
+
+		Task<bool> DownloadAsync(IPerforceConnection perforce, DirectoryReference localRootPath, FileReference manifestFileName, ILogger logger, ProgressValue progress, CancellationToken cancellationToken);
+	}
+
+	public interface IArchiveChannel
 	{
 		public const string EditorArchiveType = "Editor";
 
+		// Name to display in the UI
 		string Name { get; }
+
+		// Type key; only one item of each type may be enabled
 		string Type { get; }
-		string BasePath { get; }
-		string? Target { get; }
 
-		bool Exists();
-		bool TryGetArchiveKeyForChangeNumber(int changeNumber, int maxChangeNumber, [NotNullWhen(true)] out string? archiveKey);
-		Task<bool> DownloadArchive(IPerforceConnection perforce, string archiveKey, DirectoryReference localRootPath, FileReference manifestFileName, ILogger logger, ProgressValue progress, CancellationToken cancellationToken);
+		// Tooltip when hovering over item in UI
+		string ToolTip { get; }
+
+		bool HasAny();
+		IArchive? TryGetArchiveForChangeNumber(int changeNumber, int maxChangeNumber);
 	}
 
-	public abstract class BaseArchiveInfo : IArchiveInfo
+	public abstract class BaseArchiveChannel : IArchiveChannel
 	{
 		public string Name { get; }
 		public string Type { get; }
-		public string BasePath { get; set; }
-		public string? Target { get; }
-
-		public abstract Task<bool> DownloadArchive(IPerforceConnection perforce, string archiveKey, DirectoryReference localRootPath, FileReference manifestFileName, ILogger logger, ProgressValue progress, CancellationToken cancellationToken);
 
-		// Abstract BaseArchiveInfo Helper Functions
-		public abstract Task FindArtifacts(IPerforceConnection perforce, CancellationToken cancellationToken);
+		public virtual string ToolTip { get; } = "";
 
 		// TODO: executable/configuration?
-		public SortedList<int, string> ChangeNumberToArchiveKey { get; } = new SortedList<int, string>();
+		public SortedList<int, IArchive> ChangeNumberToArchive { get; } = new SortedList<int, IArchive>();
 
-		protected BaseArchiveInfo(string name, string type, string basePath, string? target)
+		protected BaseArchiveChannel(string name, string type)
 		{
 			Name = name;
 			Type = type;
-			BasePath = basePath;
-			Target = target;
 		}
 
-		public override bool Equals(object? other)
+		public bool HasAny()
 		{
-			BaseArchiveInfo? otherArchive = other as BaseArchiveInfo;
-			return otherArchive != null && Name == otherArchive.Name && Type == otherArchive.Type && BasePath == otherArchive.BasePath && Target == otherArchive.Target && Enumerable.SequenceEqual(ChangeNumberToArchiveKey, otherArchive.ChangeNumberToArchiveKey);
+			return ChangeNumberToArchive.Count > 0;
 		}
 
-		public override int GetHashCode()
+		public IArchive? TryGetArchiveForChangeNumber(int changeNumber, int maxChangeNumber)
 		{
-			throw new NotSupportedException();
-		}
-
-		public bool Exists()
-		{
-			return ChangeNumberToArchiveKey.Count > 0;
-		}
-
-		public static bool TryParseConfigEntry(IHordeClient hordeClient, string text, [NotNullWhen(true)] out BaseArchiveInfo? info)
-		{
-			ConfigObject obj = new ConfigObject(text);
-
-			string? name = obj.GetValue("Name", null);
-			if (name == null)
-			{
-				info = null;
-				return false;
-			}
-
-			// Where to find archives, you'll have either Perforce (DepotPath) or Horde (ArchiveType)
-			string? depotPath = obj.GetValue("DepotPath", null);
-			string? archiveType = obj.GetValue("ArchiveType", null);
-
-			// We only want one of the other, not both or none
-			if (((depotPath == null) && (archiveType == null)) || ((depotPath != null) && (archiveType != null)))
-			{
-				info = null;
-				return false;
-			}
-
-			string? target = obj.GetValue("Target", null);
-
-			string type = obj.GetValue("Type", null) ?? name;
-
-			string? streamName = obj.GetValue("StreamName", null);
-
-			if (depotPath != null)
-			{
-				info = new PerforceArchiveInfo(name, type, depotPath, target);
-			}
-			else if ((archiveType != null) && (streamName != null))
-			{
-				info = new HordeArchiveInfo(hordeClient, name, type, archiveType, target, streamName);
-			}
-			else
-			{
-				info = null;
-				return false;
-			}
-			return true;
-		}
-
-		public bool TryGetArchiveKeyForChangeNumber(int changeNumber, int maxChangeNumber, [NotNullWhen(true)] out string? archiveKey)
-		{
-			int idx = ChangeNumberToArchiveKey.Keys.AsReadOnlyList().BinarySearch(changeNumber);
+			int idx = ChangeNumberToArchive.Keys.AsReadOnlyList().BinarySearch(changeNumber);
 			if (idx >= 0)
 			{
-				archiveKey = ChangeNumberToArchiveKey.Values[idx];
-				return true;
+				return ChangeNumberToArchive.Values[idx];
 			}
 
 			int nextIdx = ~idx;
-			if (nextIdx < ChangeNumberToArchiveKey.Count && ChangeNumberToArchiveKey.Keys[nextIdx] <= maxChangeNumber)
+			if (nextIdx < ChangeNumberToArchive.Count && ChangeNumberToArchive.Keys[nextIdx] <= maxChangeNumber)
 			{
-				archiveKey = ChangeNumberToArchiveKey.Values[nextIdx];
-				return true;
+				return ChangeNumberToArchive.Values[nextIdx];
 			}
 
-			archiveKey = null;
-			return false;
+			return null;
 		}
 
 		public override string ToString()
@@ -139,23 +84,27 @@ namespace UnrealGameSync
 		}
 	}
 
-	public class PerforceArchiveInfo : BaseArchiveInfo
+	public class PerforceArchiveChannel : BaseArchiveChannel
 	{
-		public string DepotPath
-		{
-			get => BasePath;
-			set => BasePath = value;
-		}
+		public string DepotPath { get; set; }
+		public string? Target { get; }
+
+		public override string ToolTip
+			=> HasAny() ? "" : $"No valid archives found at {DepotPath}";
 
-		public PerforceArchiveInfo(string name, string type, string depotPath, string? target)
-			: base(name, type, depotPath, target)
+		public PerforceArchiveChannel(string name, string type, string depotPath, string? target)
+			: base(name, type)
 		{
+			Target = target;
+			DepotPath = depotPath;
+			Target = target;
 		}
 
 		public override bool Equals(object? other)
 		{
-			PerforceArchiveInfo? otherArchive = other as PerforceArchiveInfo;
-			return otherArchive != null && Name == otherArchive.Name && Type == otherArchive.Type && DepotPath == otherArchive.DepotPath && Target == otherArchive.Target && Enumerable.SequenceEqual(ChangeNumberToArchiveKey, otherArchive.ChangeNumberToArchiveKey);
+			PerforceArchiveChannel? otherArchive = other as PerforceArchiveChannel;
+			return otherArchive != null && Name == otherArchive.Name && Type == otherArchive.Type && DepotPath == otherArchive.DepotPath && Target == otherArchive.Target 
+				&& Enumerable.SequenceEqual(ChangeNumberToArchive.Select(x => (x.Key, x.Value)), otherArchive.ChangeNumberToArchive.Select(x => (x.Key, x.Value)));
 		}
 
 		public override int GetHashCode()
@@ -163,32 +112,70 @@ namespace UnrealGameSync
 			throw new NotSupportedException();
 		}
 
-		public override async Task<bool> DownloadArchive(IPerforceConnection perforce, string archiveKey, DirectoryReference localRootPath, FileReference manifestFileName, ILogger logger, ProgressValue progress, CancellationToken cancellationToken)
+		class PerforceArchive : IArchive
 		{
-			DirectoryReference configDir = UserSettings.GetConfigDir(localRootPath);
-			UserSettings.CreateConfigDir(configDir);
+			public string Key { get; }
 
-			FileReference tempZipFileName = FileReference.Combine(configDir, "archive.zip");
-			try
+			public PerforceArchive(string key)
+				=> Key = key;
+
+			public async Task<bool> DownloadAsync(IPerforceConnection perforce, DirectoryReference localRootPath, FileReference manifestFileName, ILogger logger, ProgressValue progress, CancellationToken cancellationToken)
 			{
-				PrintRecord record = await perforce.PrintAsync(tempZipFileName.FullName, archiveKey, cancellationToken);
+				DirectoryReference configDir = UserSettings.GetConfigDir(localRootPath);
+				UserSettings.CreateConfigDir(configDir);
 
-				if (tempZipFileName.ToFileInfo().Length == 0)
+				FileReference tempZipFileName = FileReference.Combine(configDir, "archive.zip");
+				try
 				{
-					return false;
+					PrintRecord record = await perforce.PrintAsync(tempZipFileName.FullName, Key, cancellationToken);
+
+					if (tempZipFileName.ToFileInfo().Length == 0)
+					{
+						return false;
+					}
+					ArchiveUtils.ExtractFiles(tempZipFileName, localRootPath, manifestFileName, progress, logger);
 				}
-				ArchiveUtils.ExtractFiles(tempZipFileName, localRootPath, manifestFileName, progress, logger);
+				finally
+				{
+					FileReference.SetAttributes(tempZipFileName, FileAttributes.Normal);
+					FileReference.Delete(tempZipFileName);
+				}
+
+				return true;
+			}
+		}
+
+		public static bool TryParseConfigEntryAsync(string text, [NotNullWhen(true)] out PerforceArchiveChannel? channel)
+		{
+			ConfigObject obj = new ConfigObject(text);
+
+			string? name = obj.GetValue("Name", null);
+			if (name == null)
+			{
+				channel = null;
+				return false;
 			}
-			finally
+
+			// Where to find archives, you'll have either Perforce (DepotPath) or Horde (ArchiveType)
+			string? depotPath = obj.GetValue("DepotPath", null);
+			if (depotPath == null)
 			{
-				FileReference.SetAttributes(tempZipFileName, FileAttributes.Normal);
-				FileReference.Delete(tempZipFileName);
+				channel = null;
+				return false;
 			}
 
+			string? target = obj.GetValue("Target", null);
+
+			string type = obj.GetValue("Type", null) ?? name;
+
+			string? streamName = obj.GetValue("StreamName", null);
+
+			// Build a new list of zipped binaries
+			channel = new PerforceArchiveChannel(name, type, depotPath, target);
 			return true;
 		}
 
-		public override async Task FindArtifacts(IPerforceConnection perforce, CancellationToken cancellationToken)
+		public async Task FindArtifactsAsync(IPerforceConnection perforce, CancellationToken cancellationToken)
 		{
 			PerforceResponseList<FileLogRecord> response = await perforce.TryFileLogAsync(128, FileLogOptions.FullDescriptions, DepotPath, cancellationToken);
 			if (response.Succeeded)
@@ -204,9 +191,10 @@ namespace UnrealGameSync
 							if (tokens[0].StartsWith("[CL", StringComparison.Ordinal) && tokens[1].EndsWith("]", StringComparison.Ordinal))
 							{
 								int originalChangeNumber;
-								if (Int32.TryParse(tokens[1].Substring(0, tokens[1].Length - 1), out originalChangeNumber) && !ChangeNumberToArchiveKey.ContainsKey(originalChangeNumber))
+								if (Int32.TryParse(tokens[1].Substring(0, tokens[1].Length - 1), out originalChangeNumber) && !ChangeNumberToArchive.ContainsKey(originalChangeNumber))
 								{
-									ChangeNumberToArchiveKey[originalChangeNumber] = $"{DepotPath}#{revision.RevisionNumber}";
+									PerforceArchive archive = new PerforceArchive($"{DepotPath}#{revision.RevisionNumber}");
+									ChangeNumberToArchive[originalChangeNumber] = archive;
 								}
 							}
 						}
@@ -214,30 +202,56 @@ namespace UnrealGameSync
 				}
 			}
 		}
-	}
 
-	public class HordeArchiveInfo : BaseArchiveInfo
-	{
-		public string ArchiveType
+		public static async Task<List<PerforceArchiveChannel>> GetChannelsAsync(IPerforceConnection perforce, ConfigFile latestProjectConfigFile, string projectIdentifier, CancellationToken cancellationToken)
 		{
-			get => BasePath;
-			set => BasePath = value;
-		}
+			List<PerforceArchiveChannel> channels = new List<PerforceArchiveChannel>();
+
+			// Find all the zipped binaries under this stream
+			ConfigSection? projectConfigSection = latestProjectConfigFile.FindSection(projectIdentifier);
+			if (projectConfigSection != null)
+			{
+				// Legacy
+				string? legacyEditorArchivePath = projectConfigSection.GetValue("ZippedBinariesPath", null);
+				if (legacyEditorArchivePath != null)
+				{
+					// Only Perforce uses the legacy method
+					PerforceArchiveChannel legacyChannel = new PerforceArchiveChannel("Editor", "Editor", legacyEditorArchivePath, null);
+					await legacyChannel.FindArtifactsAsync(perforce, cancellationToken);
+					channels.Add(legacyChannel);
+				}
+
+				// New style
+				foreach (string archiveValue in projectConfigSection.GetValues("Archives", Array.Empty<string>()))
+				{
+					PerforceArchiveChannel? channel;
+					if (PerforceArchiveChannel.TryParseConfigEntryAsync(archiveValue, out channel))
+					{
+						await channel.FindArtifactsAsync(perforce, cancellationToken);
+						channels.Add(channel!);
+					}
+				}
+			}
 
-		public string StreamName { get; set; }
+			return channels;
+		}
+	}
 
+	public class HordeArchiveChannel : BaseArchiveChannel
+	{
 		private readonly IHordeClient _hordeClient;
 
-		public HordeArchiveInfo(IHordeClient hordeClient, string name, string type, string archiveType, string? target, string streamName)
-			: base(name, type, archiveType, target)
+		public HordeArchiveChannel(IHordeClient hordeClient, string name, string type)
+			: base(name, type)
 		{
-			StreamName = streamName;
 			_hordeClient = hordeClient;
 		}
+
 		public override bool Equals(object? other)
 		{
-			HordeArchiveInfo? otherArchive = other as HordeArchiveInfo;
-			return otherArchive != null && Name == otherArchive.Name && Type == otherArchive.Type && ArchiveType == otherArchive.ArchiveType && Target == otherArchive.Target && StreamName == otherArchive.StreamName && Enumerable.SequenceEqual(ChangeNumberToArchiveKey, otherArchive.ChangeNumberToArchiveKey);
+			HordeArchiveChannel? otherArchive = other as HordeArchiveChannel;
+			return otherArchive != null && Name == otherArchive.Name && Type == otherArchive.Type
+				&& Enumerable.SequenceEqual(ChangeNumberToArchive.Select(x => (x.Key, x.Value)), otherArchive.ChangeNumberToArchive.Select(x => (x.Key, x.Value)));
 		}
 
 		public override int GetHashCode()
@@ -245,104 +259,103 @@ namespace UnrealGameSync
 			throw new NotSupportedException();
 		}
 
-		public override async Task<bool> DownloadArchive(IPerforceConnection perforce, string archiveKey, DirectoryReference localRootPath, FileReference manifestFileName, ILogger logger, ProgressValue progress, CancellationToken cancellationToken)
+		class HordeArchive : IArchive
 		{
-			DirectoryReference configDir = UserSettings.GetConfigDir(localRootPath);
-			UserSettings.CreateConfigDir(configDir);
+			readonly IHordeClient _hordeClient;
+			readonly ArtifactId _artifactId;
 
-			FileReference tempZipFileName = FileReference.Combine(configDir, "archive.zip");
-			try
+			public string Key { get; }
+
+			public HordeArchive(IHordeClient hordeClient, ArtifactId artifactId)
 			{
-				HordeHttpClient hordeHttpClient = _hordeClient.CreateHttpClient();
+				_hordeClient = hordeClient;
+				_artifactId = artifactId;
+				Key = artifactId.ToString();
+			}
 
-				ArtifactId artifactId = ArtifactId.Parse(archiveKey);
+			public async Task<bool> DownloadAsync(IPerforceConnection perforce, DirectoryReference localRootPath, FileReference manifestFileName, ILogger logger, ProgressValue progress, CancellationToken cancellationToken)
+			{
+				DirectoryReference configDir = UserSettings.GetConfigDir(localRootPath);
+				UserSettings.CreateConfigDir(configDir);
 
-				using (FileStream stream = FileReference.Open(tempZipFileName, FileMode.Create, FileAccess.Write, FileShare.None))
+				FileReference tempZipFileName = FileReference.Combine(configDir, "archive.zip");
+				try
 				{
-					await using Stream sourceStream = await hordeHttpClient.GetArtifactZipAsync(artifactId, cancellationToken);
-					await sourceStream.CopyToAsync(stream, cancellationToken);
-				}
+					HordeHttpClient hordeHttpClient = _hordeClient.CreateHttpClient();
+
+					using (FileStream stream = FileReference.Open(tempZipFileName, FileMode.Create, FileAccess.Write, FileShare.None))
+					{
+						await using Stream sourceStream = await hordeHttpClient.GetArtifactZipAsync(_artifactId, cancellationToken);
+						await sourceStream.CopyToAsync(stream, cancellationToken);
+					}
 
-				if (tempZipFileName.ToFileInfo().Length == 0)
+					if (tempZipFileName.ToFileInfo().Length == 0)
+					{
+						return false;
+					}
+					ArchiveUtils.ExtractFiles(tempZipFileName, localRootPath, manifestFileName, progress, logger);
+				}
+				finally
 				{
-					return false;
+					FileReference.SetAttributes(tempZipFileName, FileAttributes.Normal);
+					FileReference.Delete(tempZipFileName);
 				}
-				ArchiveUtils.ExtractFiles(tempZipFileName, localRootPath, manifestFileName, progress, logger);
-			}
-			finally
-			{
-				FileReference.SetAttributes(tempZipFileName, FileAttributes.Normal);
-				FileReference.Delete(tempZipFileName);
-			}
 
-			return true;
+				return true;
+			}
 		}
 
-		public override async Task FindArtifacts(IPerforceConnection perforce, CancellationToken cancellationToken)
+		public static async Task<List<HordeArchiveChannel>> GetChannelsAsync(IHordeClient hordeClient, string projectIdentifier, CancellationToken cancellationToken)
 		{
+			List<HordeArchiveChannel> channels = new List<HordeArchiveChannel>();
 			try
 			{
-				HordeHttpClient hordeHttpClient = _hordeClient.CreateHttpClient();
+				HordeHttpClient hordeHttpClient = hordeClient.CreateHttpClient();
 
-				ArtifactType artifactType = new ArtifactType(ArchiveType);
-				StreamId streamId = new StreamId(StreamName);
-				int? minChange = null;
-				int? maxChange = null;
-				List<GetArtifactResponse> artifactResponse = await hordeHttpClient.FindArtifactsByTypeAsync(artifactType, streamId, minChange, maxChange, null, cancellationToken);
+				ArtifactType artifactType = new ArtifactType("ugs-pcb");
+				string[] artifactKeys = new[] { $"ugs-project={projectIdentifier}" };
 
-				foreach (GetArtifactResponse response in artifactResponse)
+				List<GetArtifactResponse> artifactResponses = await hordeHttpClient.FindArtifactsByTypeAsync(artifactType, keys: artifactKeys, cancellationToken: cancellationToken);
+				foreach (IGrouping<ArtifactName, GetArtifactResponse> group in artifactResponses.GroupBy(x => x.Name))
 				{
-					ChangeNumberToArchiveKey[response.Change] = response.Id.ToString();
+					GetArtifactResponse? first = group.FirstOrDefault();
+					if (first != null)
+					{
+						string name = first.Description ?? first.Name.ToString();
+
+						const string ArchiveTypePrefix = "ArchiveType=";
+						string? archiveType = first.Metadata?.FirstOrDefault(x => x.StartsWith(ArchiveTypePrefix));
+
+						string type = IArchiveChannel.EditorArchiveType;
+						if (archiveType != null)
+						{
+							type = archiveType.Substring(ArchiveTypePrefix.Length);
+						}
+
+						HordeArchiveChannel channel = new HordeArchiveChannel(hordeClient, name, type);
+						foreach (GetArtifactResponse response in group)
+						{
+							HordeArchive archive = new HordeArchive(hordeClient, response.Id);
+							channel.ChangeNumberToArchive[response.Change] = archive;
+						}
+						channels.Add(channel);
+					}
 				}
 			}
 			catch (Exception)
 			{
-				return;
 			}
+			return channels;
 		}
 	}
 
 	public static class BaseArchive
 	{
-		public static async Task<List<BaseArchiveInfo>> EnumerateAsync(IPerforceConnection perforce, IHordeClient hordeClient, ConfigFile latestProjectConfigFile, string projectIdentifier, CancellationToken cancellationToken)
+		public static async Task<List<BaseArchiveChannel>> EnumerateChannelsAsync(IPerforceConnection perforce, IHordeClient hordeClient, ConfigFile latestProjectConfigFile, string projectIdentifier, CancellationToken cancellationToken)
 		{
-			List<BaseArchiveInfo> newArchives = new List<BaseArchiveInfo>();
-
-			// Find all the zipped binaries under this stream
-			ConfigSection? projectConfigSection = latestProjectConfigFile.FindSection(projectIdentifier);
-			if (projectConfigSection != null)
-			{
-				// Legacy
-				string? legacyEditorArchivePath = projectConfigSection.GetValue("ZippedBinariesPath", null);
-				if (legacyEditorArchivePath != null)
-				{
-					// Only Perforce uses the legacy method
-					newArchives.Add(new PerforceArchiveInfo("Editor", "Editor", legacyEditorArchivePath, null));
-				}
-
-				// New style
-				foreach (string archiveValue in projectConfigSection.GetValues("Archives", Array.Empty<string>()))
-				{
-					BaseArchiveInfo? archive;
-					if (BaseArchiveInfo.TryParseConfigEntry(hordeClient, archiveValue, out archive))
-					{
-						newArchives.Add(archive!);
-					}
-				}
-
-				// Make sure the zipped binaries path exists
-				foreach (BaseArchiveInfo newArchive in newArchives)
-				{
-					try
-					{
-						await newArchive.FindArtifacts(perforce, cancellationToken);
-					}
-					catch
-					{
-					}
-				}
-			}
-
+			List<BaseArchiveChannel> newArchives = new List<BaseArchiveChannel>();
+			newArchives.AddRange(await PerforceArchiveChannel.GetChannelsAsync(perforce, latestProjectConfigFile, projectIdentifier, cancellationToken));
+			newArchives.AddRange(await HordeArchiveChannel.GetChannelsAsync(hordeClient, projectIdentifier, cancellationToken));
 			return newArchives;
 		}
 	}
diff --git a/Engine/Source/Programs/UnrealGameSync/UnrealGameSyncShared/WorkspaceUpdate.cs b/Engine/Source/Programs/UnrealGameSync/UnrealGameSyncShared/WorkspaceUpdate.cs
--- a/Engine/Source/Programs/UnrealGameSync/UnrealGameSyncShared/WorkspaceUpdate.cs
+++ b/Engine/Source/Programs/UnrealGameSync/UnrealGameSyncShared/WorkspaceUpdate.cs
@@ -86,7 +86,7 @@ namespace UnrealGameSync
 		public WorkspaceUpdateOptions Options { get; set; }
 		public BuildConfig EditorConfig { get; set; }
 		public List<string> SyncFilter { get; } = new List<string>();
-		public Dictionary<string, Tuple<IArchiveInfo, string>?> ArchiveTypeToArchive { get; } = new Dictionary<string, Tuple<IArchiveInfo, string>?>();
+		public Dictionary<string, IArchive?> ArchiveTypeToArchive { get; } = new Dictionary<string, IArchive?>();
 		public Dictionary<string, bool> DeleteFiles { get; } = new Dictionary<string, bool>();
 		public Dictionary<string, bool> ClobberFiles { get; } = new Dictionary<string, bool>();
 		public List<ConfigObject> UserBuildStepObjects { get; } = new List<ConfigObject>();
@@ -1162,15 +1162,13 @@ namespace UnrealGameSync
 					DirectoryReference.CreateDirectory(manifestDirectoryName);
 
 					// Sync and extract (or just remove) the given archives
-					foreach (KeyValuePair<string, Tuple<IArchiveInfo, string>?> archiveTypeAndArchive in Context.ArchiveTypeToArchive)
+					foreach ((string archiveType, IArchive? archive) in Context.ArchiveTypeToArchive)
 					{
-						string archiveType = archiveTypeAndArchive.Key;
-
 						// Remove any existing binaries
 						FileReference manifestFileName = FileReference.Combine(manifestDirectoryName, String.Format("{0}.zipmanifest", archiveType));
 						if (FileReference.Exists(manifestFileName))
 						{
-							bool isNotLastSyncedEditorArchive = archiveType != IArchiveInfo.EditorArchiveType || (archiveTypeAndArchive.Value != null && archiveTypeAndArchive.Value.Item2 != state.LastSyncEditorArchive);
+							bool isNotLastSyncedEditorArchive = archiveType != IArchiveChannel.EditorArchiveType || (archive != null && archive.Key != state.LastSyncEditorArchive);
 							if (isNotLastSyncedEditorArchive)
 							{
 								logger.LogInformation("Removing {ArchiveType} binaries...", archiveType);
@@ -1182,21 +1180,19 @@ namespace UnrealGameSync
 						}
 
 						// If we have a new depot path, sync it down and extract it
-						if (archiveTypeAndArchive.Value != null)
+						if (archive != null)
 						{
-							IArchiveInfo archiveInfo = archiveTypeAndArchive.Value.Item1;
-							string archiveKey = archiveTypeAndArchive.Value.Item2;
-
-							if (archiveType != IArchiveInfo.EditorArchiveType || archiveKey != state.LastSyncEditorArchive)
+							string archiveKey = archive.Key;
+							if (archiveType != IArchiveChannel.EditorArchiveType || archiveKey != state.LastSyncEditorArchive)
 							{
 								logger.LogInformation("Syncing {ArchiveType} binaries...", archiveType);
 								Progress.Set(String.Format("Syncing {0} binaries...", archiveType), 0.0f);
-								if (!await archiveInfo.DownloadArchive(perforce, archiveKey, project.LocalRootPath, manifestFileName, logger, Progress, CancellationToken.None))
+								if (!await archive.DownloadAsync(perforce, project.LocalRootPath, manifestFileName, logger, Progress, CancellationToken.None))
 								{
 									return (WorkspaceUpdateResult.FailedToSync, $"Couldn't read {archiveKey}");
 								}
 								// Update last synced editor archive
-								if (archiveType == IArchiveInfo.EditorArchiveType)
+								if (archiveType == IArchiveChannel.EditorArchiveType)
 								{
 									state = stateMgr.Modify(x =>
 									{
@@ -1206,7 +1202,7 @@ namespace UnrealGameSync
 							}
 							else
 							{
-								logger.LogInformation("Skipping {ArchiveType} binaries download, already downloaded", IArchiveInfo.EditorArchiveType);
+								logger.LogInformation("Skipping {ArchiveType} binaries download, already downloaded", IArchiveChannel.EditorArchiveType);
 							}
 						}
 					}
@@ -1278,7 +1274,7 @@ namespace UnrealGameSync
 				FileReference editorReceiptFile = ConfigUtils.GetReceiptFile(project, Context.ProjectConfigFile, editorTargetFile, Context.EditorConfig.ToString());
 
 				// Get the build steps
-				bool usingPrecompiledEditor = Context.ArchiveTypeToArchive.TryGetValue(IArchiveInfo.EditorArchiveType, out Tuple<IArchiveInfo, string>? archiveInfo) && archiveInfo != null;
+				bool usingPrecompiledEditor = Context.ArchiveTypeToArchive.TryGetValue(IArchiveChannel.EditorArchiveType, out IArchive? archive) && archive != null;
 				Dictionary<Guid, ConfigObject> buildStepObjects = ConfigUtils.GetDefaultBuildStepObjects(project, editorTargetName, Context.EditorConfig, Context.ProjectConfigFile, usingPrecompiledEditor);
 				BuildStep.MergeBuildStepObjects(buildStepObjects, Context.ProjectConfigFile.GetValues("Build.Step", Array.Empty<string>()).Select(x => new ConfigObject(x)));
 				BuildStep.MergeBuildStepObjects(buildStepObjects, Context.UserBuildStepObjects);
