UGS: Fix compile error for command line tool.

#rnx

diff --git a/Engine/Source/Programs/UnrealGameSync/UnrealGameSyncCmd/Program.cs b/Engine/Source/Programs/UnrealGameSync/UnrealGameSyncCmd/Program.cs
--- a/Engine/Source/Programs/UnrealGameSync/UnrealGameSyncCmd/Program.cs
+++ b/Engine/Source/Programs/UnrealGameSync/UnrealGameSyncCmd/Program.cs
@@ -877,7 +877,7 @@ namespace UnrealGameSyncCmd
 						throw new UserErrorException("No editor archives found for project");
 					}
 
-					KeyValuePair<int, string> revision = editorArchiveInfo.ChangeNumberToArchiveKey.LastOrDefault(x => x.Key <= change);
+					KeyValuePair<int, IArchive> revision = editorArchiveInfo.ChangeNumberToArchive.LastOrDefault(x => x.Key <= change);
 					if (revision.Key == 0)
 					{
 						throw new UserErrorException($"No editor archives found for CL {change}");
@@ -907,7 +907,7 @@ namespace UnrealGameSyncCmd
 					}
 
 					updateContext.Options |= WorkspaceUpdateOptions.SyncArchives;
-					updateContext.ArchiveTypeToArchive[IArchiveChannel.EditorArchiveType] = Tuple.Create<IArchiveChannel, string>(editorArchiveInfo, revision.Value);
+					updateContext.ArchiveTypeToArchive[IArchiveChannel.EditorArchiveType] = revision.Value;
 				}
 
 				WorkspaceUpdate update = new WorkspaceUpdate(updateContext);
