UGS: Add server and username to sync status message to help debugging.

diff --git a/Engine/Source/Programs/UnrealGameSync/UnrealGameSyncShared/WorkspaceUpdate.cs b/Engine/Source/Programs/UnrealGameSync/UnrealGameSyncShared/WorkspaceUpdate.cs
--- a/Engine/Source/Programs/UnrealGameSync/UnrealGameSyncShared/WorkspaceUpdate.cs
+++ b/Engine/Source/Programs/UnrealGameSync/UnrealGameSyncShared/WorkspaceUpdate.cs
@@ -594,7 +594,7 @@ namespace UnrealGameSync
 			{
 				using (TelemetryStopwatch syncTelemetryStopwatch = new TelemetryStopwatch("Workspace_Sync", project.TelemetryProjectIdentifier))
 				{
-					logger.LogInformation("Syncing to {Change}...", Context.ChangeNumber);
+					logger.LogInformation("Syncing to {Change} on {ServerAndPort} as {UserName}...", Context.ChangeNumber, perforceSettings.ServerAndPort, perforceSettings.UserName);
 
 					// Make sure we're logged in
 					PerforceResponse<LoginRecord> loginResponse = await perforce.TryGetLoginStateAsync(cancellationToken);
