UGS: Register UGS as a handler for .uartifact files.

diff --git a/Engine/Source/Programs/Shared/EpicGames.Horde/Artifacts/ArtifactDescriptor.cs b/Engine/Source/Programs/Shared/EpicGames.Horde/Artifacts/ArtifactDescriptor.cs
--- a/Engine/Source/Programs/Shared/EpicGames.Horde/Artifacts/ArtifactDescriptor.cs
+++ b/Engine/Source/Programs/Shared/EpicGames.Horde/Artifacts/ArtifactDescriptor.cs
@@ -7,6 +7,8 @@ using System.Threading;
 using System.Threading.Tasks;
 using EpicGames.Core;
 
+#pragma warning disable CA2227
+
 namespace EpicGames.Horde.Storage
 {
 	/// <summary>
@@ -27,12 +29,21 @@ namespace EpicGames.Horde.Storage
 		/// <summary>
 		/// Filter for the files selected for download
 		/// </summary>
-		public List<string>? Filter { get; }
+		public List<string>? Filter { get; set; }
+
+		/// <summary>
+		/// Default constructor
+		/// </summary>
+		public ArtifactDescriptor()
+		{
+			BaseUrl = new Uri("http://horde");
+			RefName = new RefName("default");
+		}
 
 		/// <summary>
 		/// Constructor
 		/// </summary>
-		public ArtifactDescriptor(Uri baseUrl, RefName refName, IReadOnlyCollection<string>? filter)
+		public ArtifactDescriptor(Uri baseUrl, RefName refName, IReadOnlyCollection<string>? filter = null)
 		{
 			BaseUrl = baseUrl;
 			RefName = refName;
diff --git a/Engine/Source/Programs/Shared/EpicGames.Horde/Storage/Nodes/DirectoryNode.Extract.cs b/Engine/Source/Programs/Shared/EpicGames.Horde/Storage/Nodes/DirectoryNode.Extract.cs
--- a/Engine/Source/Programs/Shared/EpicGames.Horde/Storage/Nodes/DirectoryNode.Extract.cs
+++ b/Engine/Source/Programs/Shared/EpicGames.Horde/Storage/Nodes/DirectoryNode.Extract.cs
@@ -46,15 +46,18 @@ namespace EpicGames.Horde.Storage.Nodes
 		readonly object _lockObject = new object();
 		readonly Stopwatch _timer = Stopwatch.StartNew();
 		readonly IProgress<IExtractStats>? _progress;
+		readonly TimeSpan _frequency;
 		long _lastTotalSize;
+		readonly Queue<double> _rateSamples = new Queue<double>();
 
 		public int Count { get; set; }
 		public long Size { get; set; }
 		public double Rate { get; set; }
 
-		public ExtractStats(IProgress<IExtractStats>? progress)
+		public ExtractStats(IProgress<IExtractStats>? progress, TimeSpan frequency)
 		{
 			_progress = progress;
+			_frequency = frequency;
 		}
 
 		public void Update(int count, long size)
@@ -65,7 +68,7 @@ namespace EpicGames.Horde.Storage.Nodes
 				{
 					Count += count;
 					Size += size;
-					if (_timer.Elapsed > TimeSpan.FromSeconds(5.0))
+					if (_timer.Elapsed > _frequency)
 					{
 						FlushInternal();
 					}
@@ -86,7 +89,13 @@ namespace EpicGames.Horde.Storage.Nodes
 
 		void FlushInternal()
 		{
-			Rate = (Size - _lastTotalSize) / _timer.Elapsed.TotalSeconds;
+			_rateSamples.Enqueue((Size - _lastTotalSize) / _timer.Elapsed.TotalSeconds);
+			while (_rateSamples.Count > Math.Max(1, 10.0 / _frequency.TotalSeconds))
+			{
+				_rateSamples.Dequeue();
+			}
+
+			Rate = _rateSamples.Average();
 			_lastTotalSize = Size;
 
 			_progress!.Report(this);
@@ -272,12 +281,26 @@ namespace EpicGames.Horde.Storage.Nodes
 		/// <param name="progress">Sink for progress updates</param>
 		/// <param name="logger">Logger for output</param>
 		/// <param name="cancellationToken">Cancellation token for the operation</param>
-		public static async Task CopyToDirectoryAsync(this DirectoryNode directoryNode, DirectoryInfo directoryInfo, IProgress<IExtractStats>? progress, ILogger logger, CancellationToken cancellationToken)
+		public static Task CopyToDirectoryAsync(this DirectoryNode directoryNode, DirectoryInfo directoryInfo, IProgress<IExtractStats>? progress, ILogger logger, CancellationToken cancellationToken)
+		{
+			return CopyToDirectoryAsync(directoryNode, directoryInfo, progress, TimeSpan.FromSeconds(5.0), logger, cancellationToken);
+		}
+
+		/// <summary>
+		/// Utility function to allow extracting a packed directory to disk
+		/// </summary>
+		/// <param name="directoryNode">Directory to update</param>
+		/// <param name="directoryInfo">Direcotry to write to</param>
+		/// <param name="progress">Sink for progress updates</param>
+		/// <param name="frequency">Frequency for progress updates</param>
+		/// <param name="logger">Logger for output</param>
+		/// <param name="cancellationToken">Cancellation token for the operation</param>
+		public static async Task CopyToDirectoryAsync(this DirectoryNode directoryNode, DirectoryInfo directoryInfo, IProgress<IExtractStats>? progress, TimeSpan frequency, ILogger logger, CancellationToken cancellationToken)
 		{
 			int numTasks = Math.Min(1 + (int)(directoryNode.Length / (16 * 1024 * 1024)), 16);
 			logger.LogInformation("Splitting read into {NumThreads} threads", numTasks);
 
-			ExtractStats extractStats = new ExtractStats(progress);
+			ExtractStats extractStats = new ExtractStats(progress, frequency);
 			using (CancellationTokenSource cancellationSource = CancellationTokenSource.CreateLinkedTokenSource(cancellationToken))
 			{
 				// Helper method to run a background task and set a cancellation source on error
diff --git a/Engine/Source/Programs/UnrealGameSync/UnrealGameSync/ArtifactDownload.cs b/Engine/Source/Programs/UnrealGameSync/UnrealGameSync/ArtifactDownload.cs
new file mode 100644
--- /dev/null
+++ b/Engine/Source/Programs/UnrealGameSync/UnrealGameSync/ArtifactDownload.cs
@@ -0,0 +1,157 @@
+// Copyright Epic Games, Inc. All Rights Reserved.
+
+using System;
+using System.Diagnostics;
+using System.IO;
+using System.Linq;
+using System.Threading;
+using System.Threading.Tasks;
+using System.Windows.Forms;
+using EpicGames.Core;
+using EpicGames.Horde;
+using EpicGames.Horde.Artifacts;
+using EpicGames.Horde.Storage;
+using EpicGames.Horde.Storage.Bundles;
+using EpicGames.Horde.Storage.Clients;
+using EpicGames.Horde.Storage.Nodes;
+using Microsoft.Extensions.DependencyInjection;
+using Microsoft.Extensions.Logging;
+using Microsoft.Win32;
+using UnrealGameSync.Forms;
+
+namespace UnrealGameSync
+{
+	class ArtifactDownload
+	{
+		class CopyProgressAdapter : IProgress<IExtractStats>
+		{
+			readonly IProgress<string> _inner;
+
+			public CopyProgressAdapter(IProgress<string> inner)
+			{
+				_inner = inner;
+			}
+
+			public void Report(IExtractStats value)
+			{
+				_inner.Report($"Copied {value.Count} files ({value.Size / (1024.0 * 1024.0):n1}mb, {value.Rate / (1024.0 * 1024.0):n1}mb/s)...");
+			}
+		}
+
+		static async Task DoSyncAsync(Uri baseUri, RefName refName, DirectoryReference outputDir, IProgress<string> progress, CancellationToken cancellationToken)
+		{
+			DirectoryReference stateFolder = DirectoryReference.Combine(outputDir, ".ugs");
+			DirectoryReference.CreateDirectory(stateFolder);
+
+			using (ILoggerProvider loggerProvider = Logging.CreateLoggerProvider(FileReference.Combine(stateFolder, "sync.log")))
+			{
+				Uri serverUrl = new Uri(baseUri.GetLeftPart(UriPartial.Authority));
+
+				ServiceCollection services = new ServiceCollection();
+				services.AddLogging(builder => builder.AddProvider(loggerProvider));
+				services.AddHorde(options =>
+				{
+					options.ServerUrl = serverUrl;
+					options.AllowAuthPrompt = true;
+				});
+
+				await using ServiceProvider serviceProvider = services.BuildServiceProvider();
+
+				IHordeClient hordeClient = serviceProvider.GetRequiredService<IHordeClient>();
+				using IStorageClient storageClient = hordeClient.CreateStorageClient(baseUri.AbsolutePath);
+
+				progress.Report("Connecting to server...");
+
+				DirectoryNode? node = await storageClient.ReadRefTargetAsync<DirectoryNode>(refName, cancellationToken: cancellationToken);
+
+				progress.Report("Starting...");
+				await node.CopyToDirectoryAsync(outputDir.ToDirectoryInfo(), new CopyProgressAdapter(progress), TimeSpan.FromSeconds(0.2), serviceProvider.GetRequiredService<ILogger<ArtifactDownload>>(), cancellationToken);
+			}
+		}
+
+		public static void RegisterFileAssociations(string application)
+		{
+			string progId = "UnrealGameSync.Artifact";
+
+			using (RegistryKey baseKey = Registry.CurrentUser.CreateSubKey("SOFTWARE\\Classes"))
+			{
+				using (RegistryKey extensionKey = baseKey.CreateSubKey(".uartifact"))
+				{
+					extensionKey.SetValue("", progId);
+				}
+
+				using (RegistryKey progKey = baseKey.CreateSubKey(progId))
+				{
+					progKey.SetValue("", "Horde artifact descriptor");
+
+					using (RegistryKey defaultIconKey = progKey.CreateSubKey("DefaultIcon"))
+					{
+						defaultIconKey.SetValue("", $"\"{application}\",0");
+					}
+
+					using (RegistryKey shellKey = progKey.CreateSubKey("Shell"))
+					{
+						using (RegistryKey openKey = shellKey.CreateSubKey("open"))
+						{
+							using RegistryKey commandKey = openKey.CreateSubKey("command");
+							commandKey.SetValue("", $"\"{application}\" -artifact=\"%1\"");
+						}
+					}
+				}
+			}
+		}
+
+		public static bool ProcessCommandLine(string[] args)
+		{
+			const string ArtifactArg = "-artifact=";
+
+			string? artifactFile = args.FirstOrDefault(x => x.StartsWith(ArtifactArg, StringComparison.OrdinalIgnoreCase));
+			if (artifactFile == null)
+			{
+				return false;
+			}
+
+			try
+			{
+				Download(artifactFile.Substring(ArtifactArg.Length));
+			}
+			catch(Exception ex)
+			{
+				MessageBox.Show(ex.ToString());
+			}
+
+			return true;
+		}
+
+		static void Download(string artifactFile)
+		{
+			byte[] data = File.ReadAllBytes(artifactFile);
+
+			ArtifactDescriptor descriptor = ArtifactDescriptor.Deserialize(data);
+
+			string defaultOutputDir = Directory.GetCurrentDirectory();
+
+			DirectoryReference? outputDir;
+			if (!DownloadSettingsWindow.Show(descriptor.BaseUrl.ToString(), defaultOutputDir, out outputDir))
+			{
+				return;
+			}
+
+			try
+			{
+				DownloadProgressWindow.Execute((p, ctx) => DoSyncAsync(descriptor.BaseUrl, descriptor.RefName, outputDir, p, ctx), CancellationToken.None);
+			}
+			catch (Exception ex)
+			{
+				MessageBox.Show($"An error occurred while downloading the specified item.\n\n{ex}");
+				return;
+			}
+
+			ProcessStartInfo startInfo = new ProcessStartInfo();
+			startInfo.FileName = "explorer.exe";
+			startInfo.ArgumentList.Add(outputDir.ToString());
+			startInfo.UseShellExecute = true;
+			using Process? _ = Process.Start(startInfo);
+		}
+	}
+}
diff --git a/Engine/Source/Programs/UnrealGameSync/UnrealGameSync/Forms/DownloadSettingsWindow.cs b/Engine/Source/Programs/UnrealGameSync/UnrealGameSync/Forms/DownloadSettingsWindow.cs
--- a/Engine/Source/Programs/UnrealGameSync/UnrealGameSync/Forms/DownloadSettingsWindow.cs
+++ b/Engine/Source/Programs/UnrealGameSync/UnrealGameSync/Forms/DownloadSettingsWindow.cs
@@ -3,6 +3,7 @@
 using System;
 using System.Diagnostics.CodeAnalysis;
 using System.Windows.Forms;
+using EpicGames.Core;
 
 namespace UnrealGameSync.Forms
 {
@@ -14,7 +15,7 @@ namespace UnrealGameSync.Forms
 			Font = new System.Drawing.Font("Segoe UI", 8.25F, System.Drawing.FontStyle.Regular, System.Drawing.GraphicsUnit.Point, ((byte)(0)));
 		}
 
-		public static bool Show(string source, string defaultOutputDir, [NotNullWhen(true)] out string? outputDir)
+		public static bool Show(string source, string defaultOutputDir, [NotNullWhen(true)] out DirectoryReference? outputDir)
 		{
 			using DownloadSettingsWindow window = new DownloadSettingsWindow();
 			window.SourceText.Text = source;
@@ -26,7 +27,7 @@ namespace UnrealGameSync.Forms
 				return false;
 			}
 
-			outputDir = window.OutputFolderText.Text;
+			outputDir = new DirectoryReference(window.OutputFolderText.Text);
 			return true;
 		}
 
diff --git a/Engine/Source/Programs/UnrealGameSync/UnrealGameSync/Program.cs b/Engine/Source/Programs/UnrealGameSync/UnrealGameSync/Program.cs
--- a/Engine/Source/Programs/UnrealGameSync/UnrealGameSync/Program.cs
+++ b/Engine/Source/Programs/UnrealGameSync/UnrealGameSync/Program.cs
@@ -117,6 +117,12 @@ namespace UnrealGameSync
 						return;
 					}
 
+					// Handle any .uartifact downloads
+					if (ArtifactDownload.ProcessCommandLine(args))
+					{
+						return;
+					}
+
 					// Launch the application proper
 					if (firstInstance)
 					{
@@ -187,6 +193,8 @@ namespace UnrealGameSync
 				}
 			}
 
+			ArtifactDownload.RegisterFileAssociations(updateSpawn ?? GetCurrentExecutable());
+
 			DirectoryReference dataFolder = DirectoryReference.Combine(DirectoryReference.GetSpecialFolder(Environment.SpecialFolder.LocalApplicationData)!, "UnrealGameSync");
 			DirectoryReference.CreateDirectory(dataFolder);
 
diff --git a/Engine/Source/Programs/UnrealGameSync/UnrealGameSync/UriHandlers/HordeHandler.cs b/Engine/Source/Programs/UnrealGameSync/UnrealGameSync/UriHandlers/HordeHandler.cs
deleted file mode 100644
diff --git a/Engine/Source/Programs/UnrealGameSync/UnrealGameSyncShared/WorkspaceUpdate.cs b/Engine/Source/Programs/UnrealGameSync/UnrealGameSyncShared/WorkspaceUpdate.cs
--- a/Engine/Source/Programs/UnrealGameSync/UnrealGameSyncShared/WorkspaceUpdate.cs
+++ b/Engine/Source/Programs/UnrealGameSync/UnrealGameSyncShared/WorkspaceUpdate.cs
@@ -590,7 +590,7 @@ namespace UnrealGameSync
 			List<Tuple<string, TimeSpan>> times = new List<Tuple<string, TimeSpan>>();
 
 			int numFilesSynced = 0;
-			if (Context.Options.HasFlag(WorkspaceUpdateOptions.Sync) || Context.Options.HasFlag(WorkspaceUpdateOptions.SyncSingleChange))
+			if (Context.Options.HasFlag(WorkspaceUpdateOptions.Sync) && Context.Options.HasFlag(WorkspaceUpdateOptions.SyncSingleChange))
 			{
 				using (TelemetryStopwatch syncTelemetryStopwatch = new TelemetryStopwatch("Workspace_Sync", project.TelemetryProjectIdentifier))
 				{
