UGS: Fix some static analysis warnings.

diff --git a/Engine/Source/Programs/UnrealGameSync/UnrealGameSync/ArtifactDownload.cs b/Engine/Source/Programs/UnrealGameSync/UnrealGameSync/ArtifactDownload.cs
--- a/Engine/Source/Programs/UnrealGameSync/UnrealGameSync/ArtifactDownload.cs
+++ b/Engine/Source/Programs/UnrealGameSync/UnrealGameSync/ArtifactDownload.cs
@@ -9,10 +9,7 @@ using System.Threading.Tasks;
 using System.Windows.Forms;
 using EpicGames.Core;
 using EpicGames.Horde;
-using EpicGames.Horde.Artifacts;
 using EpicGames.Horde.Storage;
-using EpicGames.Horde.Storage.Bundles;
-using EpicGames.Horde.Storage.Clients;
 using EpicGames.Horde.Storage.Nodes;
 using Microsoft.Extensions.DependencyInjection;
 using Microsoft.Extensions.Logging;
@@ -115,7 +112,7 @@ namespace UnrealGameSync
 			{
 				Download(artifactFile.Substring(ArtifactArg.Length));
 			}
-			catch(Exception ex)
+			catch (Exception ex)
 			{
 				MessageBox.Show(ex.ToString());
 			}
diff --git a/Engine/Source/Programs/UnrealGameSync/UnrealGameSync/Controls/WorkspaceControl.cs b/Engine/Source/Programs/UnrealGameSync/UnrealGameSync/Controls/WorkspaceControl.cs
--- a/Engine/Source/Programs/UnrealGameSync/UnrealGameSync/Controls/WorkspaceControl.cs
+++ b/Engine/Source/Programs/UnrealGameSync/UnrealGameSync/Controls/WorkspaceControl.cs
@@ -2036,13 +2036,13 @@ namespace UnrealGameSync
 			IArchive? archive;
 			if (!changeNumberToArchive.TryGetValue(changeNumber, out archive))
 			{
-				archive = GetArchiveForChangeNumberUncached(archiveChannel, changeNumber, maxChangeNumber, changeNumberToArchive);
+				archive = GetArchiveForChangeNumberUncached(archiveChannel, changeNumber, maxChangeNumber);
 				changeNumberToArchive.Add(changeNumber, archive);
 			}
 			return archive;
 		}
 
-		private IArchive? GetArchiveForChangeNumberUncached(IArchiveChannel archiveChannel, int changeNumber, int maxChangeNumber, Dictionary<int, IArchive?> changeNumberToArchive)
+		private IArchive? GetArchiveForChangeNumberUncached(IArchiveChannel archiveChannel, int changeNumber, int maxChangeNumber)
 		{
 			int idx = _sortedChangeNumbers.BinarySearch(changeNumber);
 			if (idx < 0)
diff --git a/Engine/Source/Programs/UnrealGameSync/UnrealGameSyncShared/ArchiveInfo.cs b/Engine/Source/Programs/UnrealGameSync/UnrealGameSyncShared/ArchiveInfo.cs
--- a/Engine/Source/Programs/UnrealGameSync/UnrealGameSyncShared/ArchiveInfo.cs
+++ b/Engine/Source/Programs/UnrealGameSync/UnrealGameSyncShared/ArchiveInfo.cs
@@ -10,7 +10,6 @@ using System.Threading.Tasks;
 using EpicGames.Core;
 using EpicGames.Horde;
 using EpicGames.Horde.Artifacts;
-using EpicGames.Horde.Streams;
 using EpicGames.Perforce;
 using Microsoft.Extensions.Logging;
 
@@ -103,7 +102,7 @@ namespace UnrealGameSync
 		public override bool Equals(object? other)
 		{
 			PerforceArchiveChannel? otherArchive = other as PerforceArchiveChannel;
-			return otherArchive != null && Name == otherArchive.Name && Type == otherArchive.Type && DepotPath == otherArchive.DepotPath && Target == otherArchive.Target 
+			return otherArchive != null && Name == otherArchive.Name && Type == otherArchive.Type && DepotPath == otherArchive.DepotPath && Target == otherArchive.Target
 				&& Enumerable.SequenceEqual(ChangeNumberToArchive.Select(x => (x.Key, x.Value)), otherArchive.ChangeNumberToArchive.Select(x => (x.Key, x.Value)));
 		}
 
@@ -168,8 +167,6 @@ namespace UnrealGameSync
 
 			string type = obj.GetValue("Type", null) ?? name;
 
-			string? streamName = obj.GetValue("StreamName", null);
-
 			// Build a new list of zipped binaries
 			channel = new PerforceArchiveChannel(name, type, depotPath, target);
 			return true;
@@ -239,12 +236,9 @@ namespace UnrealGameSync
 
 	public class HordeArchiveChannel : BaseArchiveChannel
 	{
-		private readonly IHordeClient _hordeClient;
-
-		public HordeArchiveChannel(IHordeClient hordeClient, string name, string type)
+		public HordeArchiveChannel(string name, string type)
 			: base(name, type)
 		{
-			_hordeClient = hordeClient;
 		}
 
 		public override bool Equals(object? other)
@@ -324,7 +318,7 @@ namespace UnrealGameSync
 						string name = first.Description ?? first.Name.ToString();
 
 						const string ArchiveTypePrefix = "ArchiveType=";
-						string? archiveType = first.Metadata?.FirstOrDefault(x => x.StartsWith(ArchiveTypePrefix));
+						string? archiveType = first.Metadata?.FirstOrDefault(x => x.StartsWith(ArchiveTypePrefix, StringComparison.OrdinalIgnoreCase));
 
 						string type = IArchiveChannel.EditorArchiveType;
 						if (archiveType != null)
@@ -332,7 +326,7 @@ namespace UnrealGameSync
 							type = archiveType.Substring(ArchiveTypePrefix.Length);
 						}
 
-						HordeArchiveChannel channel = new HordeArchiveChannel(hordeClient, name, type);
+						HordeArchiveChannel channel = new HordeArchiveChannel(name, type);
 						foreach (GetArtifactResponse response in group)
 						{
 							HordeArchive archive = new HordeArchive(hordeClient, response.Id);
