UGS: Fix sync step not executing.

diff --git a/Engine/Source/Programs/UnrealGameSync/UnrealGameSyncShared/WorkspaceUpdate.cs b/Engine/Source/Programs/UnrealGameSync/UnrealGameSyncShared/WorkspaceUpdate.cs
--- a/Engine/Source/Programs/UnrealGameSync/UnrealGameSyncShared/WorkspaceUpdate.cs
+++ b/Engine/Source/Programs/UnrealGameSync/UnrealGameSyncShared/WorkspaceUpdate.cs
@@ -590,7 +590,7 @@ namespace UnrealGameSync
 			List<Tuple<string, TimeSpan>> times = new List<Tuple<string, TimeSpan>>();
 
 			int numFilesSynced = 0;
-			if (Context.Options.HasFlag(WorkspaceUpdateOptions.Sync) && Context.Options.HasFlag(WorkspaceUpdateOptions.SyncSingleChange))
+			if (Context.Options.HasFlag(WorkspaceUpdateOptions.Sync) || Context.Options.HasFlag(WorkspaceUpdateOptions.SyncSingleChange))
 			{
 				using (TelemetryStopwatch syncTelemetryStopwatch = new TelemetryStopwatch("Workspace_Sync", project.TelemetryProjectIdentifier))
 				{
