UGS: Fix compile error.

#rnx

diff --git a/Engine/Source/Programs/UnrealGameSync/UnrealGameSyncCmd/Program.cs b/Engine/Source/Programs/UnrealGameSync/UnrealGameSyncCmd/Program.cs
--- a/Engine/Source/Programs/UnrealGameSync/UnrealGameSyncCmd/Program.cs
+++ b/Engine/Source/Programs/UnrealGameSync/UnrealGameSyncCmd/Program.cs
@@ -1434,7 +1434,7 @@ namespace UnrealGameSyncCmd
 				using ITokenStore tokenStore = TokenStoreFactory.CreateTokenStore();
 				IConfiguration providerConfiguration = ProviderConfigurationFactory.ReadConfiguration(engineDir, gameDir);
 				OidcTokenManager oidcTokenManager = OidcTokenManager.CreateTokenManager(providerConfiguration, tokenStore, new List<string>() { providerIdentifier });
-				OidcTokenInfo result = await oidcTokenManager.Login(providerIdentifier);
+				OidcTokenInfo result = await oidcTokenManager.LoginAsync(providerIdentifier);
 
 				logger.LogInformation("Logged in to provider {ProviderIdentifier}", providerIdentifier);
 			}
