UGS: Fix cancellation exceptions propagating out from update monitor.

diff --git a/Engine/Source/Programs/UnrealGameSync/UnrealGameSync/UpdateMonitor.cs b/Engine/Source/Programs/UnrealGameSync/UnrealGameSync/UpdateMonitor.cs
--- a/Engine/Source/Programs/UnrealGameSync/UnrealGameSync/UpdateMonitor.cs
+++ b/Engine/Source/Programs/UnrealGameSync/UnrealGameSync/UpdateMonitor.cs
@@ -155,7 +155,14 @@ namespace UnrealGameSync
 		{
 			for (; ; )
 			{
-				await Task.Delay(TimeSpan.FromMinutes(5.0), cancellationToken);
+				try
+				{
+					await Task.Delay(TimeSpan.FromMinutes(5.0), cancellationToken);
+				}
+				catch (OperationCanceledException)
+				{
+					break;
+				}
 
 				IPerforceConnection? perforce = null;
 				try
@@ -172,6 +179,10 @@ namespace UnrealGameSync
 				{
 					_logger.LogInformation(ex, "Perforce exception while attempting to poll for updates.");
 				}
+				catch (OperationCanceledException) when (cancellationToken.IsCancellationRequested)
+				{
+					break;
+				}
 				catch (Exception ex)
 				{
 					_logger.LogWarning(ex, "Exception while attempting to poll for updates.");
