UGS: Fix null reference exception from Sentry.

diff --git a/Engine/Source/Programs/UnrealGameSync/UnrealGameSync/ProgramApplicationContext.cs b/Engine/Source/Programs/UnrealGameSync/UnrealGameSync/ProgramApplicationContext.cs
--- a/Engine/Source/Programs/UnrealGameSync/UnrealGameSync/ProgramApplicationContext.cs
+++ b/Engine/Source/Programs/UnrealGameSync/UnrealGameSync/ProgramApplicationContext.cs
@@ -278,6 +278,11 @@ namespace UnrealGameSync
 
 		private void OnStartupComplete(List<(UserSelectedProjectSettings, ModalTask<OpenProjectInfo>)> startupTasks)
 		{
+			if (_startupWindow == null)
+			{
+				return;
+			}
+
 			// Close the startup window
 			bool visible = _startupWindow!.Visible;
 			_startupWindow = null;
