[UGS]: Add a feature to Archives to optionally have a required set of badges (tests) that have to be successful in order for a change to be synced.

diff --git a/Engine/Source/Programs/UnrealGameSync/UnrealGameSync/Controls/WorkspaceControl.cs b/Engine/Source/Programs/UnrealGameSync/UnrealGameSync/Controls/WorkspaceControl.cs
--- a/Engine/Source/Programs/UnrealGameSync/UnrealGameSync/Controls/WorkspaceControl.cs
+++ b/Engine/Source/Programs/UnrealGameSync/UnrealGameSync/Controls/WorkspaceControl.cs
@@ -912,10 +912,20 @@ namespace UnrealGameSync
 					IArchive? archivePath = GetArchiveForChangeNumber(archiveChannel, changeNumber);
 					if (archivePath == null)
 					{
-						MessageBox.Show(String.Format("There are no compiled {0} binaries for this change. To sync it, you must disable syncing of precompiled editor binaries.", archiveChannel.Name));
+						MessageBox.Show(String.Format("There are no compiled \"{0}\" binaries for this change. To sync it, you must disable syncing of precompiled editor binaries.", archiveChannel.Name));
 						return;
 					}
 
+					PerforceArchiveChannel? perforceArchiveChannel = archiveChannel as PerforceArchiveChannel;
+					if ((perforceArchiveChannel != null) && (perforceArchiveChannel.RequiredBadges.Count > 0))
+					{
+						if (!CanSyncChange(changeNumber))
+						{
+							MessageBox.Show(String.Format("This change has not passed the required tests \"{0}\". To sync it, you must disable syncing of precompiled editor binaries.", String.Join(", ", perforceArchiveChannel.RequiredBadges)));
+							return;
+						}
+					}
+
 					if (archiveChannel.Type == IArchiveChannel.EditorArchiveType)
 					{
 						context.Options &= ~(WorkspaceUpdateOptions.Build | WorkspaceUpdateOptions.GenerateProjectFiles | WorkspaceUpdateOptions.OpenSolutionAfterSync);
@@ -2087,7 +2097,46 @@ namespace UnrealGameSync
 			}
 
 			List<IArchiveChannel> selectedArchives = GetSelectedArchiveChannels(GetArchiveChannels());
-			return selectedArchives.Count == 0 || selectedArchives.All(x => GetArchiveForChangeNumber(x, changeNumber) != null);
+			bool returnValue = selectedArchives.Count == 0 || selectedArchives.All(x => GetArchiveForChangeNumber(x, changeNumber) != null);
+
+			// If we can sync this change and we're using perforce archives, let's double check for any required badges
+			if (returnValue && (selectedArchives.Count == 1))
+			{
+				PerforceArchiveChannel? archiveChannel = selectedArchives[0] as PerforceArchiveChannel;
+
+				if ((archiveChannel != null) && (archiveChannel.RequiredBadges.Count > 0))
+				{
+					int currentChangeIdx = _sortedChangeNumbers.BinarySearch(changeNumber);
+					if (currentChangeIdx < 0)
+					{
+						return returnValue;
+					}
+
+					bool foundRequiredBadges = false;
+					bool badgesSuccessful = false;
+
+					for (int idx = currentChangeIdx; idx >= 0; idx--)
+					{
+						int localChangeNumber = _sortedChangeNumbers[idx];
+
+						EventSummary? summary = _eventMonitor.GetSummaryForChange(localChangeNumber);
+
+						if (summary != null)
+						{
+							foundRequiredBadges = DoRequiredBadgesExist(archiveChannel.RequiredBadges, summary.Badges, out badgesSuccessful);
+						}
+
+						if (foundRequiredBadges)
+						{
+							break;
+						}
+					}
+
+					return badgesSuccessful;
+				}
+			}
+
+			return returnValue;
 		}
 
 		/// <summary>
@@ -2097,6 +2146,15 @@ namespace UnrealGameSync
 		/// </summary>
 		private static bool DoRequiredBadgesExist(List<string> requiredBadgeList, List<BadgeData> inBadges)
 		{
+			bool badgesSuccessful;
+			bool badgesExist = DoRequiredBadgesExist(requiredBadgeList, inBadges, out badgesSuccessful);
+			return badgesExist && badgesSuccessful;
+		}
+
+		private static bool DoRequiredBadgesExist(List<string> requiredBadgeList, List<BadgeData> inBadges, out bool badgesSuccessful)
+		{
+			badgesSuccessful = true;
+
 			Dictionary<string, BadgeData> inBadgeDictionary = new Dictionary<string, BadgeData>();
 			foreach (BadgeData badge in inBadges)
 			{
@@ -2109,24 +2167,21 @@ namespace UnrealGameSync
 				BadgeData? badge;
 				if (inBadgeDictionary.TryGetValue(badgeName, out badge))
 				{
-					// If any required badge is not successful then the filter isn't matched.
-					if (!badge.IsSuccess)
-					{
-						return false;
-					}
+					badgesSuccessful = badgesSuccessful && badge.IsSuccess;
 				}
 				else
 				{
-					// required badge not found.
+					// required badge not found, early out
+					badgesSuccessful = false;
 					return false;
 				}
 			}
 
-			// All required badges existed and were successful 
+			// All required badges exist, if they were all successful, "badgesSuccessful" will be set to true
 			return true;
 		}
 
-		private bool CanSyncChangeType(LatestChangeType changeType, ChangesRecord change, EventSummary? summary)
+		private bool CanSyncChangeType(LatestChangeType changeType, int changeNumber, EventSummary? summary)
 		{
 			if (changeType.Good)
 			{
@@ -2138,7 +2193,7 @@ namespace UnrealGameSync
 
 			if (changeType.Starred)
 			{
-				if ((summary == null || summary.LastStarReview == null || summary.LastStarReview.Type != EventType.Starred) && !_promotedChangeNumbers.Contains(change.Number))
+				if ((summary == null || summary.LastStarReview == null || summary.LastStarReview.Type != EventType.Starred) && !_promotedChangeNumbers.Contains(changeNumber))
 				{
 					return false;
 				}
@@ -5343,7 +5398,7 @@ namespace UnrealGameSync
 				{
 					EventSummary? summary = _eventMonitor.GetSummaryForChange(change.Number);
 					if (CanSyncChange(change.Number)
-						&& CanSyncChangeType(changeType, change, summary))
+						&& CanSyncChangeType(changeType, change.Number, summary))
 					{
 						if (changeType.FindNewestGoodContent)
 						{
@@ -5387,6 +5442,11 @@ namespace UnrealGameSync
 					break;
 				}
 
+				if (!CanSyncChange(nextChangeNumber))
+				{
+					break;
+				}
+
 				index = nextIndex;
 			}
 
diff --git a/Engine/Source/Programs/UnrealGameSync/UnrealGameSyncShared/ArchiveInfo.cs b/Engine/Source/Programs/UnrealGameSync/UnrealGameSyncShared/ArchiveInfo.cs
--- a/Engine/Source/Programs/UnrealGameSync/UnrealGameSyncShared/ArchiveInfo.cs
+++ b/Engine/Source/Programs/UnrealGameSync/UnrealGameSyncShared/ArchiveInfo.cs
@@ -91,12 +91,21 @@ namespace UnrealGameSync
 		public override string ToolTip
 			=> HasAny() ? "" : $"No valid archives found at {DepotPath}";
 
-		public PerforceArchiveChannel(string name, string type, string depotPath, string? target)
+		// Get an optional list of badges that have to be green for this archive to be syncable
+
+		public List<string> RequiredBadges { get; init; } = new List<string>();
+
+		public PerforceArchiveChannel(string name, string type, string depotPath, string? target, List<string>? requiredBadges)
 			: base(name, type)
 		{
 			Target = target;
 			DepotPath = depotPath;
 			Target = target;
+
+			if (requiredBadges != null)
+			{
+				RequiredBadges = requiredBadges;
+			}
 		}
 
 		public override bool Equals(object? other)
@@ -167,8 +176,11 @@ namespace UnrealGameSync
 
 			string type = obj.GetValue("Type", null) ?? name;
 
+			List<string> requiredBadges = new List<string>();
+			requiredBadges.AddRange(obj.GetValue("RequiredBadges", "").Split(new char[] { ',' }, StringSplitOptions.RemoveEmptyEntries | StringSplitOptions.TrimEntries));
+
 			// Build a new list of zipped binaries
-			channel = new PerforceArchiveChannel(name, type, depotPath, target);
+			channel = new PerforceArchiveChannel(name, type, depotPath, target, requiredBadges);
 			return true;
 		}
 
@@ -213,7 +225,7 @@ namespace UnrealGameSync
 				if (legacyEditorArchivePath != null)
 				{
 					// Only Perforce uses the legacy method
-					PerforceArchiveChannel legacyChannel = new PerforceArchiveChannel("Editor", "Editor", legacyEditorArchivePath, null);
+					PerforceArchiveChannel legacyChannel = new PerforceArchiveChannel("Editor", "Editor", legacyEditorArchivePath, null, null);
 					await legacyChannel.FindArtifactsAsync(perforce, cancellationToken);
 					channels.Add(legacyChannel);
 				}
