UGS: Remove hard dependency on Horde (#11702)

diff --git a/Engine/Source/Programs/UnrealGameSync/UnrealGameSync/PerforceMonitor.cs b/Engine/Source/Programs/UnrealGameSync/UnrealGameSync/PerforceMonitor.cs
--- a/Engine/Source/Programs/UnrealGameSync/UnrealGameSync/PerforceMonitor.cs
+++ b/Engine/Source/Programs/UnrealGameSync/UnrealGameSync/PerforceMonitor.cs
@@ -52,7 +52,7 @@ namespace UnrealGameSync
 		readonly DirectoryReference _cacheFolder;
 		readonly List<KeyValuePair<FileReference, DateTime>> _localConfigFiles;
 		readonly IAsyncDisposer _asyncDisposeTasks;
-		readonly IHordeClient _hordeClient;
+		readonly IHordeClient? _hordeClient;
 
 		string[] prevCodeRules = Array.Empty<string>();
 
@@ -87,7 +87,7 @@ namespace UnrealGameSync
 			_asyncDisposeTasks = serviceProvider.GetRequiredService<IAsyncDisposer>();
 			_synchronizationContext = SynchronizationContext.Current!;
 			_cancellationSource = new CancellationTokenSource();
-			_hordeClient = serviceProvider.GetRequiredService<IHordeClient>();
+			_hordeClient = serviceProvider.GetService<IHordeClient>();
 
 			AvailableArchiveChannels = (new List<IArchiveChannel>()).AsReadOnly();
 		}
diff --git a/Engine/Source/Programs/UnrealGameSync/UnrealGameSyncShared/ArchiveInfo.cs b/Engine/Source/Programs/UnrealGameSync/UnrealGameSyncShared/ArchiveInfo.cs
--- a/Engine/Source/Programs/UnrealGameSync/UnrealGameSyncShared/ArchiveInfo.cs
+++ b/Engine/Source/Programs/UnrealGameSync/UnrealGameSyncShared/ArchiveInfo.cs
@@ -357,11 +357,14 @@ namespace UnrealGameSync
 
 	public static class BaseArchive
 	{
-		public static async Task<List<BaseArchiveChannel>> EnumerateChannelsAsync(IPerforceConnection perforce, IHordeClient hordeClient, ConfigFile latestProjectConfigFile, string projectIdentifier, CancellationToken cancellationToken)
+		public static async Task<List<BaseArchiveChannel>> EnumerateChannelsAsync(IPerforceConnection perforce, IHordeClient? hordeClient, ConfigFile latestProjectConfigFile, string projectIdentifier, CancellationToken cancellationToken)
 		{
 			List<BaseArchiveChannel> newArchives = new List<BaseArchiveChannel>();
 			newArchives.AddRange(await PerforceArchiveChannel.GetChannelsAsync(perforce, latestProjectConfigFile, projectIdentifier, cancellationToken));
-			newArchives.AddRange(await HordeArchiveChannel.GetChannelsAsync(hordeClient, projectIdentifier, cancellationToken));
+			if (hordeClient != null)
+			{
+				newArchives.AddRange(await HordeArchiveChannel.GetChannelsAsync(hordeClient, projectIdentifier, cancellationToken));
+			}
 			return newArchives;
 		}
 	}
