[UGS]: Move the "RequiredBadges" functionality from the Archive level to the Project level.

diff --git a/Engine/Source/Programs/UnrealGameSync/UnrealGameSync/Controls/WorkspaceControl.cs b/Engine/Source/Programs/UnrealGameSync/UnrealGameSync/Controls/WorkspaceControl.cs
--- a/Engine/Source/Programs/UnrealGameSync/UnrealGameSync/Controls/WorkspaceControl.cs
+++ b/Engine/Source/Programs/UnrealGameSync/UnrealGameSync/Controls/WorkspaceControl.cs
@@ -288,6 +288,16 @@ namespace UnrealGameSync
 			_projectSettings = settings.FindOrAddProjectSettings(openProjectInfo.ProjectInfo, openProjectInfo.WorkspaceSettings, _logger);
 			_hordeClient = serviceProvider.GetService<IHordeClient>();
 
+			ConfigFile projectConfigFile = openProjectInfo.LatestProjectConfigFile;
+			if (projectConfigFile != null)
+			{
+				string? requiredBadges;
+				if (TryGetProjectSetting(projectConfigFile, "RequiredBadges", out requiredBadges))
+				{
+					 _projectSettings.RequiredBadges.AddRange(requiredBadges.Split(new char[] { ',' }, StringSplitOptions.RemoveEmptyEntries | StringSplitOptions.TrimEntries));
+				}
+			}
+
 			DesiredTaskbarState = Tuple.Create(TaskbarState.NoProgress, 0.0f);
 
 			System.Reflection.PropertyInfo doubleBufferedProperty = typeof(Control).GetProperty("DoubleBuffered", System.Reflection.BindingFlags.NonPublic | System.Reflection.BindingFlags.Instance)!;
@@ -916,16 +926,6 @@ namespace UnrealGameSync
 						return;
 					}
 
-					PerforceArchiveChannel? perforceArchiveChannel = archiveChannel as PerforceArchiveChannel;
-					if ((perforceArchiveChannel != null) && (perforceArchiveChannel.RequiredBadges.Count > 0))
-					{
-						if (!CanSyncChange(changeNumber))
-						{
-							MessageBox.Show(String.Format("This change has not passed the required tests \"{0}\". To sync it, you must disable syncing of precompiled editor binaries.", String.Join(", ", perforceArchiveChannel.RequiredBadges)));
-							return;
-						}
-					}
-
 					if (archiveChannel.Type == IArchiveChannel.EditorArchiveType)
 					{
 						context.Options &= ~(WorkspaceUpdateOptions.Build | WorkspaceUpdateOptions.GenerateProjectFiles | WorkspaceUpdateOptions.OpenSolutionAfterSync);
@@ -939,6 +939,15 @@ namespace UnrealGameSync
 
 					context.ArchiveTypeToArchive[archiveChannel.Type] = archivePath;
 				}
+
+				if ((selectedArchiveChannels.Count > 0) && (_projectSettings.RequiredBadges.Count > 0))
+				{
+					if (!CanSyncChange(changeNumber))
+					{
+						MessageBox.Show(String.Format("This change has not passed the required tests \"{0}\". To sync it, you must disable syncing of precompiled editor binaries.", String.Join(", ", _projectSettings.RequiredBadges)));
+						return;
+					}
+				}
 			}
 			StartWorkspaceUpdate(context, callback);
 		}
@@ -2099,41 +2108,36 @@ namespace UnrealGameSync
 			List<IArchiveChannel> selectedArchives = GetSelectedArchiveChannels(GetArchiveChannels());
 			bool returnValue = selectedArchives.Count == 0 || selectedArchives.All(x => GetArchiveForChangeNumber(x, changeNumber) != null);
 
-			// If we can sync this change and we're using perforce archives, let's double check for any required badges
-			if (returnValue && (selectedArchives.Count == 1))
+			// If we can sync this change and we're using archives, let's double check for any required badges
+			if (returnValue && (selectedArchives.Count > 0) && (_projectSettings.RequiredBadges.Count > 0))
 			{
-				PerforceArchiveChannel? archiveChannel = selectedArchives[0] as PerforceArchiveChannel;
-
-				if ((archiveChannel != null) && (archiveChannel.RequiredBadges.Count > 0))
+				int currentChangeIdx = _sortedChangeNumbers.BinarySearch(changeNumber);
+				if (currentChangeIdx < 0)
 				{
-					int currentChangeIdx = _sortedChangeNumbers.BinarySearch(changeNumber);
-					if (currentChangeIdx < 0)
-					{
-						return returnValue;
-					}
-
-					bool foundRequiredBadges = false;
-					bool badgesSuccessful = false;
+					return returnValue;
+				}
 
-					for (int idx = currentChangeIdx; idx >= 0; idx--)
-					{
-						int localChangeNumber = _sortedChangeNumbers[idx];
+				bool foundRequiredBadges = false;
+				bool badgesSuccessful = false;
 
-						EventSummary? summary = _eventMonitor.GetSummaryForChange(localChangeNumber);
+				for (int idx = currentChangeIdx; idx >= 0; idx--)
+				{
+					int localChangeNumber = _sortedChangeNumbers[idx];
 
-						if (summary != null)
-						{
-							foundRequiredBadges = DoRequiredBadgesExist(archiveChannel.RequiredBadges, summary.Badges, out badgesSuccessful);
-						}
+					EventSummary? summary = _eventMonitor.GetSummaryForChange(localChangeNumber);
 
-						if (foundRequiredBadges)
-						{
-							break;
-						}
+					if (summary != null)
+					{
+						foundRequiredBadges = DoRequiredBadgesExist(_projectSettings.RequiredBadges, summary.Badges, out badgesSuccessful);
 					}
 
-					return badgesSuccessful;
+					if (foundRequiredBadges)
+					{
+						break;
+					}
 				}
+
+				return badgesSuccessful;
 			}
 
 			return returnValue;
diff --git a/Engine/Source/Programs/UnrealGameSync/UnrealGameSyncShared/ArchiveInfo.cs b/Engine/Source/Programs/UnrealGameSync/UnrealGameSyncShared/ArchiveInfo.cs
--- a/Engine/Source/Programs/UnrealGameSync/UnrealGameSyncShared/ArchiveInfo.cs
+++ b/Engine/Source/Programs/UnrealGameSync/UnrealGameSyncShared/ArchiveInfo.cs
@@ -91,21 +91,12 @@ namespace UnrealGameSync
 		public override string ToolTip
 			=> HasAny() ? "" : $"No valid archives found at {DepotPath}";
 
-		// Get an optional list of badges that have to be green for this archive to be syncable
-
-		public List<string> RequiredBadges { get; init; } = new List<string>();
-
-		public PerforceArchiveChannel(string name, string type, string depotPath, string? target, List<string>? requiredBadges)
+		public PerforceArchiveChannel(string name, string type, string depotPath, string? target)
 			: base(name, type)
 		{
 			Target = target;
 			DepotPath = depotPath;
 			Target = target;
-
-			if (requiredBadges != null)
-			{
-				RequiredBadges = requiredBadges;
-			}
 		}
 
 		public override bool Equals(object? other)
@@ -176,11 +167,8 @@ namespace UnrealGameSync
 
 			string type = obj.GetValue("Type", null) ?? name;
 
-			List<string> requiredBadges = new List<string>();
-			requiredBadges.AddRange(obj.GetValue("RequiredBadges", "").Split(new char[] { ',' }, StringSplitOptions.RemoveEmptyEntries | StringSplitOptions.TrimEntries));
-
 			// Build a new list of zipped binaries
-			channel = new PerforceArchiveChannel(name, type, depotPath, target, requiredBadges);
+			channel = new PerforceArchiveChannel(name, type, depotPath, target);
 			return true;
 		}
 
@@ -225,7 +213,7 @@ namespace UnrealGameSync
 				if (legacyEditorArchivePath != null)
 				{
 					// Only Perforce uses the legacy method
-					PerforceArchiveChannel legacyChannel = new PerforceArchiveChannel("Editor", "Editor", legacyEditorArchivePath, null, null);
+					PerforceArchiveChannel legacyChannel = new PerforceArchiveChannel("Editor", "Editor", legacyEditorArchivePath, null);
 					await legacyChannel.FindArtifactsAsync(perforce, cancellationToken);
 					channels.Add(legacyChannel);
 				}
diff --git a/Engine/Source/Programs/UnrealGameSync/UnrealGameSyncShared/UserSettings.cs b/Engine/Source/Programs/UnrealGameSync/UnrealGameSyncShared/UserSettings.cs
--- a/Engine/Source/Programs/UnrealGameSync/UnrealGameSyncShared/UserSettings.cs
+++ b/Engine/Source/Programs/UnrealGameSync/UnrealGameSyncShared/UserSettings.cs
@@ -419,6 +419,7 @@ namespace UnrealGameSync
 		public List<ConfigObject> BuildSteps { get; init; } = new List<ConfigObject>();
 		public FilterType FilterType { get; set; }
 		public HashSet<string> FilterBadges { get; init; } = new HashSet<string>(StringComparer.OrdinalIgnoreCase);
+		public List<string> RequiredBadges { get; init; } = new List<string>();
 
 		static readonly object _syncRoot = new object();
 
