[UGS]: Make it so that individual archives can choose to ignore the "RequiredBadges"

diff --git a/Engine/Source/Programs/UnrealGameSync/UnrealGameSync/Controls/WorkspaceControl.cs b/Engine/Source/Programs/UnrealGameSync/UnrealGameSync/Controls/WorkspaceControl.cs
--- a/Engine/Source/Programs/UnrealGameSync/UnrealGameSync/Controls/WorkspaceControl.cs
+++ b/Engine/Source/Programs/UnrealGameSync/UnrealGameSync/Controls/WorkspaceControl.cs
@@ -940,7 +940,7 @@ namespace UnrealGameSync
 					context.ArchiveTypeToArchive[archiveChannel.Type] = archivePath;
 				}
 
-				if ((selectedArchiveChannels.Count > 0) && (_projectSettings.RequiredBadges.Count > 0))
+				if ((selectedArchiveChannels.Count == 1) && (!selectedArchiveChannels[0].IgnoreRequiredBadges) && (_projectSettings.RequiredBadges.Count > 0))
 				{
 					if (!CanSyncChange(changeNumber))
 					{
@@ -2109,7 +2109,7 @@ namespace UnrealGameSync
 			bool returnValue = selectedArchives.Count == 0 || selectedArchives.All(x => GetArchiveForChangeNumber(x, changeNumber) != null);
 
 			// If we can sync this change and we're using archives, let's double check for any required badges
-			if (returnValue && (selectedArchives.Count > 0) && (_projectSettings.RequiredBadges.Count > 0))
+			if (returnValue && (selectedArchives.Count == 1) && (!selectedArchives[0].IgnoreRequiredBadges) && (_projectSettings.RequiredBadges.Count > 0))
 			{
 				int currentChangeIdx = _sortedChangeNumbers.BinarySearch(changeNumber);
 				if (currentChangeIdx < 0)
diff --git a/Engine/Source/Programs/UnrealGameSync/UnrealGameSyncShared/ArchiveInfo.cs b/Engine/Source/Programs/UnrealGameSync/UnrealGameSyncShared/ArchiveInfo.cs
--- a/Engine/Source/Programs/UnrealGameSync/UnrealGameSyncShared/ArchiveInfo.cs
+++ b/Engine/Source/Programs/UnrealGameSync/UnrealGameSyncShared/ArchiveInfo.cs
@@ -35,6 +35,9 @@ namespace UnrealGameSync
 		// Tooltip when hovering over item in UI
 		string ToolTip { get; }
 
+		// Does this archive channel ignore required badges?  Default is false
+		bool IgnoreRequiredBadges { get; }
+
 		bool HasAny();
 		IArchive? TryGetArchiveForChangeNumber(int changeNumber, int maxChangeNumber);
 	}
@@ -46,13 +49,16 @@ namespace UnrealGameSync
 
 		public virtual string ToolTip { get; } = "";
 
+		public bool IgnoreRequiredBadges { get; } = false;
+
 		// TODO: executable/configuration?
 		public SortedList<int, IArchive> ChangeNumberToArchive { get; } = new SortedList<int, IArchive>();
 
-		protected BaseArchiveChannel(string name, string type)
+		protected BaseArchiveChannel(string name, string type, bool bIgnoreRequiredBadges = false)
 		{
 			Name = name;
 			Type = type;
+			IgnoreRequiredBadges = bIgnoreRequiredBadges;
 		}
 
 		public bool HasAny()
@@ -91,8 +97,8 @@ namespace UnrealGameSync
 		public override string ToolTip
 			=> HasAny() ? "" : $"No valid archives found at {DepotPath}";
 
-		public PerforceArchiveChannel(string name, string type, string depotPath, string? target)
-			: base(name, type)
+		public PerforceArchiveChannel(string name, string type, string depotPath, string? target, bool bIgnoreRequiredBadges = false)
+			: base(name, type, bIgnoreRequiredBadges)
 		{
 			Target = target;
 			DepotPath = depotPath;
@@ -167,8 +173,10 @@ namespace UnrealGameSync
 
 			string type = obj.GetValue("Type", null) ?? name;
 
+			bool bIgnoreRequiredBadges = obj.GetValue("bIgnoreRequiredBadges", false);
+
 			// Build a new list of zipped binaries
-			channel = new PerforceArchiveChannel(name, type, depotPath, target);
+			channel = new PerforceArchiveChannel(name, type, depotPath, target, bIgnoreRequiredBadges);
 			return true;
 		}
 
@@ -236,8 +244,8 @@ namespace UnrealGameSync
 
 	public class HordeArchiveChannel : BaseArchiveChannel
 	{
-		public HordeArchiveChannel(string name, string type)
-			: base(name, type)
+		public HordeArchiveChannel(string name, string type, bool bIgnoreRequiredBadges = false)
+			: base(name, type, bIgnoreRequiredBadges)
 		{
 		}
 
