UGS: Fix garbled error message text.

diff --git a/Engine/Source/Programs/UnrealGameSync/UnrealGameSync/Controls/WorkspaceControl.cs b/Engine/Source/Programs/UnrealGameSync/UnrealGameSync/Controls/WorkspaceControl.cs
--- a/Engine/Source/Programs/UnrealGameSync/UnrealGameSync/Controls/WorkspaceControl.cs
+++ b/Engine/Source/Programs/UnrealGameSync/UnrealGameSync/Controls/WorkspaceControl.cs
@@ -4863,7 +4863,7 @@ namespace UnrealGameSync
 							break;
 						}
 					}
-					ShowErrorDialog(String.Format("Couldn't find any {0}changelist. Double-click on the change you want to sync manually.", syncTypeName));
+					ShowErrorDialog(String.Format("Couldn't find any change matching '{0}'. Double-click on the change you want to sync manually.", syncTypeName));
 				}
 				else if (changeNumber < _workspace.CurrentChangeNumber)
 				{
