[UGS]: Make the "RequiredBadges" more robust when dealing with multiple required badges

diff --git a/Engine/Source/Programs/UnrealGameSync/UnrealGameSync/Controls/WorkspaceControl.cs b/Engine/Source/Programs/UnrealGameSync/UnrealGameSync/Controls/WorkspaceControl.cs
--- a/Engine/Source/Programs/UnrealGameSync/UnrealGameSync/Controls/WorkspaceControl.cs
+++ b/Engine/Source/Programs/UnrealGameSync/UnrealGameSync/Controls/WorkspaceControl.cs
@@ -294,7 +294,8 @@ namespace UnrealGameSync
 				string? requiredBadges;
 				if (TryGetProjectSetting(projectConfigFile, "RequiredBadges", out requiredBadges))
 				{
-					 _projectSettings.RequiredBadges.AddRange(requiredBadges.Split(new char[] { ',' }, StringSplitOptions.RemoveEmptyEntries | StringSplitOptions.TrimEntries));
+					_projectSettings.RequiredBadges.Clear();
+					_projectSettings.RequiredBadges.AddRange(requiredBadges.Split(new char[] { ',' }, StringSplitOptions.RemoveEmptyEntries | StringSplitOptions.TrimEntries).Distinct().ToList());
 				}
 			}
 
@@ -2111,33 +2112,9 @@ namespace UnrealGameSync
 			// If we can sync this change and we're using archives, let's double check for any required badges
 			if (returnValue && (selectedArchives.Count == 1) && (!selectedArchives[0].IgnoreRequiredBadges) && (_projectSettings.RequiredBadges.Count > 0))
 			{
-				int currentChangeIdx = _sortedChangeNumbers.BinarySearch(changeNumber);
-				if (currentChangeIdx < 0)
-				{
-					return returnValue;
-				}
+				Dictionary<string, bool?> requiredBadgesStatus = GetRequiredBadgesStatusAtChange(changeNumber, _projectSettings.RequiredBadges);
 
-				bool foundRequiredBadges = false;
-				bool badgesSuccessful = false;
-
-				for (int idx = currentChangeIdx; idx >= 0; idx--)
-				{
-					int localChangeNumber = _sortedChangeNumbers[idx];
-
-					EventSummary? summary = _eventMonitor.GetSummaryForChange(localChangeNumber);
-
-					if (summary != null)
-					{
-						foundRequiredBadges = DoRequiredBadgesExist(_projectSettings.RequiredBadges, summary.Badges, out badgesSuccessful);
-					}
-
-					if (foundRequiredBadges)
-					{
-						break;
-					}
-				}
-
-				return badgesSuccessful;
+				return requiredBadgesStatus.Values.All(result => result.HasValue && result.Value);
 			}
 
 			return returnValue;
@@ -2185,6 +2162,61 @@ namespace UnrealGameSync
 			return true;
 		}
 
+		private Dictionary<string, bool?> GetRequiredBadgesStatusAtChange(int changeNumber, List<string> requiredBadgeList)
+		{
+			Dictionary<string, bool?> badgeResults = requiredBadgeList.ToDictionary(badge => badge, badge => (bool?)null);
+
+			int currentChangeIdx = _sortedChangeNumbers.BinarySearch(changeNumber);
+			if (currentChangeIdx < 0)
+			{
+				return badgeResults;
+			}
+
+			for (int idx = currentChangeIdx; idx >= 0; idx--)
+			{
+				int localChangeNumber = _sortedChangeNumbers[idx];
+
+				EventSummary? summary = _eventMonitor.GetSummaryForChange(localChangeNumber);
+
+				if (summary != null)
+				{
+					Dictionary<string, BadgeData> inBadgeDictionary = new Dictionary<string, BadgeData>();
+					foreach (BadgeData badge in summary.Badges)
+					{
+						inBadgeDictionary.Add(badge.BadgeName, badge);
+					}
+
+					foreach (string badgeName in requiredBadgeList)
+					{
+						BadgeData? badge;
+						if (inBadgeDictionary.TryGetValue(badgeName, out badge))
+						{
+							bool? result;
+							if (badgeResults.TryGetValue(badgeName, out result))
+							{
+								if (result != null)
+								{
+									badgeResults[badgeName] = result.Value && badge.IsSuccess;
+								}
+								else
+								{
+									badgeResults[badgeName] = badge.IsSuccess;
+								}
+							}
+						}
+					}
+				}
+
+				// If we have found all of the required badges (or ran out of changes to test) exit, the loop
+				if (badgeResults.Values.All(result => result.HasValue))
+				{
+					break;
+				}
+			}
+
+			return badgeResults;
+		}
+
 		private bool CanSyncChangeType(LatestChangeType changeType, int changeNumber, EventSummary? summary)
 		{
 			if (changeType.Good)
diff --git a/Engine/Source/Programs/UnrealGameSync/UnrealGameSyncShared/UserSettings.cs b/Engine/Source/Programs/UnrealGameSync/UnrealGameSyncShared/UserSettings.cs
--- a/Engine/Source/Programs/UnrealGameSync/UnrealGameSyncShared/UserSettings.cs
+++ b/Engine/Source/Programs/UnrealGameSync/UnrealGameSyncShared/UserSettings.cs
@@ -992,6 +992,7 @@ namespace UnrealGameSync
 			currentProject.BuildSteps.AddRange(projectSection.GetValues("BuildStep", Array.Empty<string>()).Select(x => new ConfigObject(x)));
 			currentProject.FilterType = projectSection.GetEnumValue("FilterType", FilterType.None);
 			currentProject.FilterBadges.UnionWith(projectSection.GetValues("FilterBadges", Array.Empty<string>()));
+			currentProject.RequiredBadges.AddRange(projectSection.GetValues("RequiredBadges", Array.Empty<string>())); 
 		}
 
 		public override bool Save(ILogger logger)
