UGS: Prevent tool update monitor from deleting tools if enumerating tools from a particular source fails with an error.

diff --git a/Engine/Source/Programs/UnrealGameSync/UnrealGameSync/Controls/WorkspaceControl.cs b/Engine/Source/Programs/UnrealGameSync/UnrealGameSync/Controls/WorkspaceControl.cs
--- a/Engine/Source/Programs/UnrealGameSync/UnrealGameSync/Controls/WorkspaceControl.cs
+++ b/Engine/Source/Programs/UnrealGameSync/UnrealGameSync/Controls/WorkspaceControl.cs
@@ -5846,6 +5846,7 @@ namespace UnrealGameSync
 			diagnosticsText.AppendFormat("Perforce monitor: {0}\n", (_perforceMonitor == null) ? "(inactive)" : _perforceMonitor.LastStatusMessage);
 			diagnosticsText.AppendFormat("Event monitor: {0}\n", (_eventMonitor == null) ? "(inactive)" : _eventMonitor.LastStatusMessage);
 			diagnosticsText.AppendFormat("Issue monitor: {0}\n", (_issueMonitor == null) ? "(inactive)" : _issueMonitor.LastStatusMessage);
+			diagnosticsText.AppendFormat("Tool update monitor: {0}\n", _owner.ToolUpdateMonitor.LastStatus.Item2);
 
 			using DiagnosticsWindow diagnostics = new DiagnosticsWindow(_appDataFolder, _workspaceDataFolder, diagnosticsText.ToString(), _settings.GetCachedFilePaths());
 			diagnostics.ShowDialog(this);
diff --git a/Engine/Source/Programs/UnrealGameSync/UnrealGameSync/Forms/ApplicationSettingsWindow.cs b/Engine/Source/Programs/UnrealGameSync/UnrealGameSync/Forms/ApplicationSettingsWindow.cs
--- a/Engine/Source/Programs/UnrealGameSync/UnrealGameSync/Forms/ApplicationSettingsWindow.cs
+++ b/Engine/Source/Programs/UnrealGameSync/UnrealGameSync/Forms/ApplicationSettingsWindow.cs
@@ -4,6 +4,7 @@ using System;
 using System.Collections.Generic;
 using System.IO;
 using System.Linq;
+using System.Text.RegularExpressions;
 using System.Threading;
 using System.Threading.Tasks;
 using System.Windows.Forms;
@@ -46,6 +47,7 @@ namespace UnrealGameSync
 		readonly ProtocolHandlerState _initialProtocolHandlerState;
 
 		readonly ToolUpdateMonitor _toolUpdateMonitor;
+		readonly Tuple<bool, string> _toolUpdateStatus;
 
 		Result _result = Result.Ok;
 
@@ -97,6 +99,7 @@ namespace UnrealGameSync
 			_originalExecutableFileName = originalExecutableFileName;
 			_settings = settings;
 			_toolUpdateMonitor = toolUpdateMonitor;
+			_toolUpdateStatus = toolUpdateMonitor.LastStatus;
 			_logger = logger;
 
 			LauncherSettings launcherSettings = new LauncherSettings();
@@ -148,34 +151,45 @@ namespace UnrealGameSync
 				EnableProtocolHandlerCheckBox.CheckState = CheckState.Indeterminate;
 			}
 
-			List<ToolItem> toolItems = new List<ToolItem>();
-			Dictionary<Guid, ToolItem> idToToolItem = new Dictionary<Guid, ToolItem>();
-			foreach (ToolInfo tool in toolUpdateMonitor.GetTools().OrderBy(x => x.Name))
+			if (_toolUpdateStatus.Item1)
 			{
-				ToolItem toolItem = new ToolItem(toolItems.Count, tool, settings.EnabledTools.Contains(tool.Id));
-				idToToolItem[tool.Id] = toolItem;
-				toolItems.Add(toolItem);
-			}
-
-			HashSet<ToolItem> dependsOnToolItems = new HashSet<ToolItem>();
-			foreach (ToolItem toolItem in toolItems)
-			{
-				dependsOnToolItems.Clear();
-				FindDependencies(toolItem, dependsOnToolItems, idToToolItem);
-				toolItem.RequiresTools.AddRange(dependsOnToolItems);
+				List<ToolItem> toolItems = new List<ToolItem>();
+				Dictionary<Guid, ToolItem> idToToolItem = new Dictionary<Guid, ToolItem>();
+				foreach (ToolInfo tool in toolUpdateMonitor.GetTools().OrderBy(x => x.Name))
+				{
+					ToolItem toolItem = new ToolItem(toolItems.Count, tool, settings.EnabledTools.Contains(tool.Id));
+					idToToolItem[tool.Id] = toolItem;
+					toolItems.Add(toolItem);
+				}
 
-				if (toolItem.Enabled)
+				HashSet<ToolItem> dependsOnToolItems = new HashSet<ToolItem>();
+				foreach (ToolItem toolItem in toolItems)
 				{
-					foreach (ToolItem requiredTool in toolItem.RequiresTools)
+					dependsOnToolItems.Clear();
+					FindDependencies(toolItem, dependsOnToolItems, idToToolItem);
+					toolItem.RequiresTools.AddRange(dependsOnToolItems);
+
+					if (toolItem.Enabled)
 					{
-						requiredTool.DependencyRefCount++;
+						foreach (ToolItem requiredTool in toolItem.RequiresTools)
+						{
+							requiredTool.DependencyRefCount++;
+						}
 					}
 				}
-			}
 
-			foreach (ToolItem toolItem in toolItems)
+				foreach (ToolItem toolItem in toolItems)
+				{
+					CustomToolsListBox.Items.Add(toolItem, toolItem.GetCheckState());
+				}
+
+				CustomToolsListBox.Enabled = true;
+			}
+			else
 			{
-				CustomToolsListBox.Items.Add(toolItem, toolItem.GetCheckState());
+				string message = Regex.Replace(_toolUpdateStatus.Item2, "\n.*$", "...");
+				CustomToolsListBox.Items.Add(message);
+				CustomToolsListBox.Enabled = _toolUpdateStatus.Item1;
 			}
 		}
 
@@ -297,20 +311,23 @@ namespace UnrealGameSync
 				_settings.Save(_logger);
 			}
 
-			List<Guid> newEnabledTools = new List<Guid>();
-			foreach (ToolItem? item in CustomToolsListBox.Items)
+			if (_toolUpdateStatus.Item1)
 			{
-				if (item != null && item.Enabled)
+				List<Guid> newEnabledTools = new List<Guid>();
+				foreach (ToolItem? item in CustomToolsListBox.Items)
 				{
-					newEnabledTools.Add(item.Definition.Id);
+					if (item != null && item.Enabled)
+					{
+						newEnabledTools.Add(item.Definition.Id);
+					}
+				}
+				if (!newEnabledTools.SequenceEqual(_settings.EnabledTools))
+				{
+					_settings.EnabledTools.Clear();
+					_settings.EnabledTools.UnionWith(newEnabledTools);
+					_settings.Save(_logger);
+					_toolUpdateMonitor.UpdateNow();
 				}
-			}
-			if (!newEnabledTools.SequenceEqual(_settings.EnabledTools))
-			{
-				_settings.EnabledTools.Clear();
-				_settings.EnabledTools.UnionWith(newEnabledTools);
-				_settings.Save(_logger);
-				_toolUpdateMonitor.UpdateNow();
 			}
 
 			if (EnableProtocolHandlerCheckBox.CheckState == CheckState.Checked)
