UGS: Missing file.

#rnx

diff --git a/Engine/Source/Programs/UnrealGameSync/UnrealGameSync/ToolUpdateMonitor.cs b/Engine/Source/Programs/UnrealGameSync/UnrealGameSync/ToolUpdateMonitor.cs
--- a/Engine/Source/Programs/UnrealGameSync/UnrealGameSync/ToolUpdateMonitor.cs
+++ b/Engine/Source/Programs/UnrealGameSync/UnrealGameSync/ToolUpdateMonitor.cs
@@ -5,6 +5,7 @@ using System.Collections.Generic;
 using System.Diagnostics;
 using System.IO;
 using System.Linq;
+using System.Security.Permissions;
 using System.Text;
 using System.Text.Json;
 using System.Threading;
@@ -95,6 +96,8 @@ namespace UnrealGameSync
 
 		public Action? OnChange;
 
+		public Tuple<bool, string> LastStatus { get; private set; } = Tuple.Create(false, "Starting...");
+
 		public ToolUpdateMonitor(IPerforceSettings perforceSettings, DirectoryReference dataDir, UserSettings settings, IServiceProvider serviceProvider)
 		{
 			_cancellationSource = new CancellationTokenSource();
@@ -135,6 +138,7 @@ namespace UnrealGameSync
 
 		public void Dispose()
 		{
+			LastStatus = Tuple.Create(false, "Stopped");
 			OnChange = null;
 
 			if (_workerTask != null)
@@ -198,7 +202,8 @@ namespace UnrealGameSync
 				}
 				catch (Exception ex)
 				{
-					_logger.LogError(ex, "Exception while checking for tool updates");
+					LastStatus = Tuple.Create(false, $"Exception while checking for tool updates: {ex.Message}");
+					_logger.LogError(ex, "Exception while checking for tool updates: {Message}", ex.Message);
 				}
 
 				Task delayTask = Task.Delay(TimeSpan.FromMinutes(60.0), cancellationToken);
@@ -211,6 +216,7 @@ namespace UnrealGameSync
 			IPerforceConnection? perforce = null;
 			try
 			{
+				Stopwatch timer = Stopwatch.StartNew();
 				// Update all the available tools
 				List<ToolInfo> tools = new List<ToolInfo>();
 				if (!String.IsNullOrEmpty(DeploymentSettings.Instance.ToolsDepotPath))
@@ -222,7 +228,9 @@ namespace UnrealGameSync
 					}
 					catch (Exception ex) when (ex is not OperationCanceledException)
 					{
+						LastStatus = Tuple.Create(false, $"Error while polling Perforce for available tools: {ex.Message}");
 						_logger.LogWarning(ex, "Error while polling Perforce for available tools: {Message}", ex.Message);
+						return;
 					}
 				}
 				using (HordeHttpClient? hordeHttpClient = _serviceProvider.GetService<HordeHttpClient>())
@@ -235,7 +243,9 @@ namespace UnrealGameSync
 						}
 						catch (Exception ex) when (ex is not OperationCanceledException)
 						{
+							LastStatus = Tuple.Create(false, $"Error while polling Horde for available tools: {ex.Message}");
 							_logger.LogWarning(ex, "Error while polling Horde for available tools: {Message}", ex.Message);
+							return;
 						}
 					}
 				}
@@ -283,6 +293,8 @@ namespace UnrealGameSync
 				{
 					_synchronizationContext.Post(_ => OnChange?.Invoke(), null);
 				}
+
+				LastStatus = Tuple.Create(true, $"Last update took {timer.ElapsedMilliseconds}ms (completed at {DateTime.Now.ToShortTimeString()})");
 			}
 			finally
 			{
