[UGS]: Make "RequiredBadges" more strict by also considering badges after the changelist

diff --git a/Engine/Source/Programs/UnrealGameSync/UnrealGameSync/Controls/WorkspaceControl.cs b/Engine/Source/Programs/UnrealGameSync/UnrealGameSync/Controls/WorkspaceControl.cs
--- a/Engine/Source/Programs/UnrealGameSync/UnrealGameSync/Controls/WorkspaceControl.cs
+++ b/Engine/Source/Programs/UnrealGameSync/UnrealGameSync/Controls/WorkspaceControl.cs
@@ -2112,7 +2112,7 @@ namespace UnrealGameSync
 			// If we can sync this change and we're using archives, let's double check for any required badges
 			if (returnValue && (selectedArchives.Count == 1) && (!selectedArchives[0].IgnoreRequiredBadges) && (_projectSettings.RequiredBadges.Count > 0))
 			{
-				Dictionary<string, bool?> requiredBadgesStatus = GetRequiredBadgesStatusAtChange(changeNumber, _projectSettings.RequiredBadges);
+				Dictionary<string, bool?> requiredBadgesStatus = GetRequiredBadgesStatusAtChange(changeNumber, _projectSettings.RequiredBadges, /* includeFutureChanges= */ true);
 
 				return requiredBadgesStatus.Values.All(result => result.HasValue && result.Value);
 			}
@@ -2162,8 +2162,40 @@ namespace UnrealGameSync
 			return true;
 		}
 
-		private Dictionary<string, bool?> GetRequiredBadgesStatusAtChange(int changeNumber, List<string> requiredBadgeList)
+		private void GetRequiredBadgesStatusOnlyAtChangeCumulative(int changeNumber, List<string> requiredBadgeList, Dictionary<string, bool?> badgeResults)
 		{
+			EventSummary? summary = _eventMonitor.GetSummaryForChange(changeNumber);
+
+			if (summary != null)
+			{
+				Dictionary<string, BadgeData> inBadgeDictionary = new Dictionary<string, BadgeData>();
+				foreach (BadgeData badge in summary.Badges)
+				{
+					inBadgeDictionary.Add(badge.BadgeName, badge);
+				}
+
+				foreach (string badgeName in requiredBadgeList)
+				{
+					BadgeData? badge;
+					if (inBadgeDictionary.TryGetValue(badgeName, out badge))
+					{
+						bool? result;
+						if (badgeResults.TryGetValue(badgeName, out result))
+						{
+							if (result == null)
+							{
+								badgeResults[badgeName] = badge.IsSuccess;
+							}
+						}
+					}
+				}
+			}
+		}
+
+		private Dictionary<string, bool?> GetRequiredBadgesStatusAtChange(int changeNumber, List<string> requiredBadgeList, bool includeFutureChanges = false)
+		{
+			// This is function is brute force and could probably use some caching
+
 			Dictionary<string, bool?> badgeResults = requiredBadgeList.ToDictionary(badge => badge, badge => (bool?)null);
 
 			int currentChangeIdx = _sortedChangeNumbers.BinarySearch(changeNumber);
@@ -2176,41 +2208,39 @@ namespace UnrealGameSync
 			{
 				int localChangeNumber = _sortedChangeNumbers[idx];
 
-				EventSummary? summary = _eventMonitor.GetSummaryForChange(localChangeNumber);
+				GetRequiredBadgesStatusOnlyAtChangeCumulative(localChangeNumber, requiredBadgeList, badgeResults);
 
-				if (summary != null)
+				// If we have found all of the required badges (or ran out of changes to test) exit, the loop
+				if (badgeResults.Values.All(result => result.HasValue))
 				{
-					Dictionary<string, BadgeData> inBadgeDictionary = new Dictionary<string, BadgeData>();
-					foreach (BadgeData badge in summary.Badges)
-					{
-						inBadgeDictionary.Add(badge.BadgeName, badge);
-					}
+					break;
+				}
+			}
+
+			if (includeFutureChanges)
+			{
+				Dictionary<string, bool?> futureBadgeResults = requiredBadgeList.ToDictionary(badge => badge, badge => (bool?)null);
 
-					foreach (string badgeName in requiredBadgeList)
+				for (int idx = currentChangeIdx; idx < _sortedChangeNumbers.Count; idx++)
+				{
+					int localChangeNumber = _sortedChangeNumbers[idx];
+
+					GetRequiredBadgesStatusOnlyAtChangeCumulative(localChangeNumber, requiredBadgeList, futureBadgeResults);
+
+					// If we have found all of the required badges (or ran out of changes to test) exit, the loop
+					if (futureBadgeResults.Values.All(result => result.HasValue))
 					{
-						BadgeData? badge;
-						if (inBadgeDictionary.TryGetValue(badgeName, out badge))
-						{
-							bool? result;
-							if (badgeResults.TryGetValue(badgeName, out result))
-							{
-								if (result != null)
-								{
-									badgeResults[badgeName] = result.Value && badge.IsSuccess;
-								}
-								else
-								{
-									badgeResults[badgeName] = badge.IsSuccess;
-								}
-							}
-						}
+						break;
 					}
 				}
 
-				// If we have found all of the required badges (or ran out of changes to test) exit, the loop
-				if (badgeResults.Values.All(result => result.HasValue))
+				// Combine results with &&, guaranteed to have the same keys
+				foreach (KeyValuePair<string, bool?> idx in badgeResults)
 				{
-					break;
+					bool? currentValue = badgeResults[idx.Key];
+					bool? futureValue = futureBadgeResults[idx.Key];
+
+					badgeResults[idx.Key] = (currentValue.HasValue && futureValue.HasValue) ? (currentValue.Value && futureValue.Value) : null;
 				}
 			}
 
