UGS: Use the native storage API to download UGS updates from Horde.

diff --git a/Engine/Source/Programs/UnrealGameSync/UnrealGameSync/Launcher.cs b/Engine/Source/Programs/UnrealGameSync/UnrealGameSync/Launcher.cs
--- a/Engine/Source/Programs/UnrealGameSync/UnrealGameSync/Launcher.cs
+++ b/Engine/Source/Programs/UnrealGameSync/UnrealGameSync/Launcher.cs
@@ -14,6 +14,8 @@ using System.Threading.Tasks;
 using System.Windows.Forms;
 using EpicGames.Core;
 using EpicGames.Horde;
+using EpicGames.Horde.Storage;
+using EpicGames.Horde.Storage.Nodes;
 using EpicGames.Horde.Tools;
 using EpicGames.Perforce;
 using Microsoft.Extensions.DependencyInjection;
@@ -160,10 +162,12 @@ namespace UnrealGameSync
 
 					ServiceCollection services = new ServiceCollection();
 					services.AddLogging(builder => builder.AddProvider(new LoggerProviderAdapter(logger)));
-					services.AddHorde(options => options.ServerUrl = hordeServerUrl);
+					services.AddHorde(options => { options.ServerUrl = hordeServerUrl; options.AllowAuthPrompt = true; });
 
 					await using ServiceProvider serviceProvider = services.BuildServiceProvider();
-					HordeHttpClient httpClient = serviceProvider.GetRequiredService<HordeHttpClient>();
+					IHordeClient hordeClient = serviceProvider.GetRequiredService<IHordeClient>();
+
+					using HordeHttpClient httpClient = hordeClient.CreateHttpClient();
 
 					ToolId toolId = DeploymentSettings.Instance.HordeToolId;
 
@@ -177,17 +181,12 @@ namespace UnrealGameSync
 						// Delete the output directory
 						await SafeDeleteDirectoryContentsWithRetryAsync(applicationFolder, cancellationToken);
 
-						// Download and extract the zip file
-						string zipFile = Path.Combine(applicationFolder, "update.zip");
-						using (Stream requestStream = await httpClient.GetToolDeploymentZipAsync(toolId, deployment.Id, cancellationToken))
-						{
-							using (Stream tempFileStream = File.Open(zipFile, FileMode.Create, FileAccess.Write, FileShare.None))
-							{
-								await requestStream.CopyToAsync(tempFileStream, cancellationToken);
-							}
-						}
-						ZipFile.ExtractToDirectory(zipFile, applicationFolder);
-						File.Delete(zipFile);
+						// Download and extract the data
+						IStorageClient storageClient = hordeClient.CreateStorageClient(toolId);
+						IBlobHandle handle = storageClient.CreateBlobHandle(deployment.Locator);
+
+						DirectoryNode rootNode = await handle.ReadBlobAsync<DirectoryNode>(cancellationToken: cancellationToken);
+						await rootNode.CopyToDirectoryAsync(new DirectoryInfo(applicationFolder), new ExtractStatsLogger(0, rootNode.Length, logger), logger, cancellationToken);
 
 						// Update the version
 						if (!TryWriteAllText(syncVersionFile, requiredSyncText))
