UGS: Fix command line crash on Linux due to IHordeClient not being registered.

diff --git a/Engine/Source/Programs/UnrealGameSync/UnrealGameSyncCmd/Program.cs b/Engine/Source/Programs/UnrealGameSync/UnrealGameSyncCmd/Program.cs
--- a/Engine/Source/Programs/UnrealGameSync/UnrealGameSyncCmd/Program.cs
+++ b/Engine/Source/Programs/UnrealGameSync/UnrealGameSyncCmd/Program.cs
@@ -137,15 +137,15 @@ namespace UnrealGameSyncCmd
 			public ILoggerFactory LoggerFactory { get; }
 			public GlobalSettingsFile UserSettings { get; }
 
-			public IHordeClient _hordeClient { get; }
+			public IHordeClient? HordeClient { get; }
 
-			public CommandContext(CommandLineArguments arguments, ILogger logger, ILoggerFactory loggerFactory, GlobalSettingsFile userSettings, IHordeClient hordeClient)
+			public CommandContext(CommandLineArguments arguments, ILogger logger, ILoggerFactory loggerFactory, GlobalSettingsFile userSettings, IHordeClient? hordeClient)
 			{
 				Arguments = arguments;
 				Logger = logger;
 				LoggerFactory = loggerFactory;
 				UserSettings = userSettings;
-				_hordeClient = hordeClient;
+				HordeClient = hordeClient;
 			}
 		}
 
@@ -242,7 +242,7 @@ namespace UnrealGameSyncCmd
 					services.AddHorde(options => options.ServerUrl = new Uri(launcherSettings.HordeServer));
 				}
 				ServiceProvider serviceProvider = services.BuildServiceProvider();
-				IHordeClient hordeClient = serviceProvider.GetRequiredService<IHordeClient>();
+				IHordeClient? hordeClient = serviceProvider.GetService<IHordeClient>();
 
 				GlobalSettingsFile settings;
 				if (RuntimeInformation.IsOSPlatform(OSPlatform.Windows))
@@ -869,7 +869,7 @@ namespace UnrealGameSyncCmd
 
 				if (syncOptions.Binaries)
 				{
-					List<BaseArchiveChannel> archives = await BaseArchive.EnumerateChannelsAsync(perforceClient, context._hordeClient, projectConfig, state.Current.ProjectIdentifier, CancellationToken.None);
+					List<BaseArchiveChannel> archives = await BaseArchive.EnumerateChannelsAsync(perforceClient, context.HordeClient, projectConfig, state.Current.ProjectIdentifier, CancellationToken.None);
 
 					BaseArchiveChannel? editorArchiveInfo = archives.FirstOrDefault(x => x.Name == IArchiveChannel.EditorArchiveType);
 					if (editorArchiveInfo == null)
