UGS: Fix some static analysis warnings.

#rnx

diff --git a/Engine/Source/Programs/UnrealGameSync/UnrealGameSync/Launcher.cs b/Engine/Source/Programs/UnrealGameSync/UnrealGameSync/Launcher.cs
--- a/Engine/Source/Programs/UnrealGameSync/UnrealGameSync/Launcher.cs
+++ b/Engine/Source/Programs/UnrealGameSync/UnrealGameSync/Launcher.cs
@@ -162,7 +162,11 @@ namespace UnrealGameSync
 
 					ServiceCollection services = new ServiceCollection();
 					services.AddLogging(builder => builder.AddProvider(new LoggerProviderAdapter(logger)));
-					services.AddHorde(options => { options.ServerUrl = hordeServerUrl; options.AllowAuthPrompt = true; });
+					services.AddHorde(options => 
+					{ 
+						options.ServerUrl = hordeServerUrl; 
+						options.AllowAuthPrompt = true; 
+					});
 
 					await using ServiceProvider serviceProvider = services.BuildServiceProvider();
 					IHordeClient hordeClient = serviceProvider.GetRequiredService<IHordeClient>();
@@ -182,7 +186,7 @@ namespace UnrealGameSync
 						await SafeDeleteDirectoryContentsWithRetryAsync(applicationFolder, cancellationToken);
 
 						// Download and extract the data
-						IStorageClient storageClient = hordeClient.CreateStorageClient(toolId);
+						using IStorageClient storageClient = hordeClient.CreateStorageClient(toolId);
 						IBlobHandle handle = storageClient.CreateBlobHandle(deployment.Locator);
 
 						DirectoryNode rootNode = await handle.ReadBlobAsync<DirectoryNode>(cancellationToken: cancellationToken);
diff --git a/Engine/Source/Programs/UnrealGameSync/UnrealGameSync/ToolUpdateMonitor.cs b/Engine/Source/Programs/UnrealGameSync/UnrealGameSync/ToolUpdateMonitor.cs
--- a/Engine/Source/Programs/UnrealGameSync/UnrealGameSync/ToolUpdateMonitor.cs
+++ b/Engine/Source/Programs/UnrealGameSync/UnrealGameSync/ToolUpdateMonitor.cs
@@ -5,7 +5,6 @@ using System.Collections.Generic;
 using System.Diagnostics;
 using System.IO;
 using System.Linq;
-using System.Security.Permissions;
 using System.Text;
 using System.Text.Json;
 using System.Threading;
