Fix UGS throwing "Error while polling Perforce for available tools: An item with the same key has already been added." when the querying tools if the working directory happens to be mapped to a workspace that includes //depot/InternalTools

diff --git a/Engine/Source/Programs/UnrealGameSync/UnrealGameSync/ToolUpdateMonitor.cs b/Engine/Source/Programs/UnrealGameSync/UnrealGameSync/ToolUpdateMonitor.cs
--- a/Engine/Source/Programs/UnrealGameSync/UnrealGameSync/ToolUpdateMonitor.cs
+++ b/Engine/Source/Programs/UnrealGameSync/UnrealGameSync/ToolUpdateMonitor.cs
@@ -322,24 +322,23 @@ namespace UnrealGameSync
 			{
 				Dictionary<string, ToolInfo> newPerforceTools = new Dictionary<string, ToolInfo>(StringComparer.Ordinal);
 
-				List<FStatRecord> fileRecords = await perforce.FStatAsync($"{DeploymentSettings.Instance.ToolsDepotPath}/...", cancellationToken).ToListAsync(cancellationToken);
-				fileRecords.RemoveAll(x => x.Action == FileAction.Delete || x.Action == FileAction.MoveDelete);
+				List<FilesRecord> fileRecords = await perforce.FilesAsync(FilesOptions.ExcludeDeleted, $"{DeploymentSettings.Instance.ToolsDepotPath}/...", cancellationToken);
 
-				foreach (FStatRecord fileRecord in fileRecords)
+				foreach (FilesRecord fileRecord in fileRecords)
 				{
 					if (fileRecord.DepotFile != null && fileRecord.DepotFile.EndsWith(".ini", StringComparison.OrdinalIgnoreCase))
 					{
 						string zipFile = fileRecord.DepotFile.Substring(0, fileRecord.DepotFile.Length - 4) + ".zip";
-						FStatRecord? zipRecord = fileRecords.FirstOrDefault(x => String.Equals(x.DepotFile, zipFile, StringComparison.OrdinalIgnoreCase));
+						FilesRecord? zipRecord = fileRecords.FirstOrDefault(x => String.Equals(x.DepotFile, zipFile, StringComparison.OrdinalIgnoreCase));
 
 						if (zipRecord != null)
 						{
-							string toolRevision = $"{zipFile}@{Math.Max(fileRecord.HeadChange, zipRecord.HeadChange)}";
+							string toolRevision = $"{zipFile}@{Math.Max(fileRecord.Change, zipRecord.Change)}";
 
 							ToolInfo? toolInfo;
 							if (!_perforceTools.TryGetValue(toolRevision, out toolInfo))
 							{
-								toolInfo = await ReadToolDefinitionAsync(perforce, $"{fileRecord.DepotFile}@{fileRecord.HeadChange}", toolRevision, cancellationToken);
+								toolInfo = await ReadToolDefinitionAsync(perforce, $"{fileRecord.DepotFile}@{fileRecord.Change}", toolRevision, cancellationToken);
 							}
 							if (toolInfo != null)
 							{
