[UGS]: Add ability to have per project (or branch) default editor arguments

diff --git a/Engine/Programs/UnrealGameSync/UnrealGameSync.ini b/Engine/Programs/UnrealGameSync/UnrealGameSync.ini
--- a/Engine/Programs/UnrealGameSync/UnrealGameSync.ini
+++ b/Engine/Programs/UnrealGameSync/UnrealGameSync.ini
@@ -43,6 +43,8 @@ BuildUnrealHeaderTool=0
 
 [Default]
 EditorTarget=UnrealEditor
++DefaultEditorArgument=(Name="-log", Enabled=false, Locked=false)
++DefaultEditorArgument=(Name="-fastload", Enabled=false, Locked=false)
 
 ; Set this to a network share which contains the SDK installers for your site
 SdkInstallerDir= 
diff --git a/Engine/Source/Programs/UnrealGameSync/UnrealGameSync/Controls/WorkspaceControl.cs b/Engine/Source/Programs/UnrealGameSync/UnrealGameSync/Controls/WorkspaceControl.cs
--- a/Engine/Source/Programs/UnrealGameSync/UnrealGameSync/Controls/WorkspaceControl.cs
+++ b/Engine/Source/Programs/UnrealGameSync/UnrealGameSync/Controls/WorkspaceControl.cs
@@ -3253,11 +3253,14 @@ namespace UnrealGameSync
 					{
 						launchArguments.AppendFormat("\"{0}\"", SelectedFileName);
 					}
-					foreach (Tuple<string, bool> editorArgument in _settings.EditorArguments)
+
+					List<LockableEditorArgument> tempEditorArguments = GetEditorArgumentsWithDefaults();
+
+					foreach (LockableEditorArgument editorArgument in tempEditorArguments)
 					{
-						if (editorArgument.Item2)
+						if (editorArgument.Enabled)
 						{
-							launchArguments.AppendFormat(" {0}", editorArgument.Item1);
+							launchArguments.AppendFormat(" {0}", editorArgument.Name);
 						}
 					}
 					if (editorBuildConfig == BuildConfig.Debug || editorBuildConfig == BuildConfig.DebugGame)
@@ -4797,7 +4800,7 @@ namespace UnrealGameSync
 			OptionsContextMenu_EditorBuildConfiguration.Enabled = !ShouldSyncPrecompiledEditor;
 			UpdateCheckedBuildConfig();
 			OptionsContextMenu_CustomizeBuildSteps.Enabled = (_workspace != null);
-			OptionsContextMenu_EditorArguments.Checked = _settings.EditorArguments.Any(x => x.Item2);
+			OptionsContextMenu_EditorArguments.Checked = GetEditorArgumentsWithDefaults().Any(x => x.Enabled);
 			OptionsContextMenu_ScheduledSync.Checked = _settings.ScheduleEnabled;
 			OptionsContextMenu_TimeZone_Local.Checked = _settings.ShowLocalTimes;
 			OptionsContextMenu_TimeZone_PerforceServer.Checked = !_settings.ShowLocalTimes;
@@ -5269,13 +5272,85 @@ namespace UnrealGameSync
 			ModifyEditorArguments();
 		}
 
+		private List<LockableEditorArgument> GetDefaultEditorArguments()
+		{
+			List<string> defaultEditorArgumentDefinitions = new List<string>();
+			ConfigUtils.GetProjectSettings(_perforceMonitor.LatestProjectConfigFile, SelectedProjectIdentifier, "DefaultEditorArgument", defaultEditorArgumentDefinitions);
+
+			List<LockableEditorArgument> defaultEditorArguments = new List<LockableEditorArgument>(); 
+			foreach (string editorArgumentDefinition in defaultEditorArgumentDefinitions.Distinct())
+			{
+				LockableEditorArgument? editorArgument;
+				if (LockableEditorArgument.TryParseConfigEntry(editorArgumentDefinition, out editorArgument))
+				{
+					defaultEditorArguments.Add(editorArgument);
+				}
+			}
+
+			return defaultEditorArguments;
+		}
+
+		private List<LockableEditorArgument> GetEditorArgumentsWithDefaults()
+		{
+			List<LockableEditorArgument> defaultEditorArguments = GetDefaultEditorArguments();
+
+			List<LockableEditorArgument> currentEditorArguments = new List<LockableEditorArgument>();
+
+			foreach (LockableEditorArgument editorArgument in _settings.EditorArguments)
+			{
+				currentEditorArguments.Add( new LockableEditorArgument(editorArgument.Name, editorArgument.Enabled) );
+			}
+			
+			foreach (LockableEditorArgument defaultEditorArgument in defaultEditorArguments)
+			{
+				// Check to see if the user already has this default argument in their list
+				int index = currentEditorArguments.FindIndex(x => x.Name == defaultEditorArgument.Name);
+
+				if (index != -1)
+				{
+					if (defaultEditorArgument.Locked)
+					{
+						currentEditorArguments[index] = new LockableEditorArgument(defaultEditorArgument.Name, defaultEditorArgument.Enabled, /* Locked = */ true);
+					}
+				}
+				else
+				{
+					currentEditorArguments.Add(new LockableEditorArgument(defaultEditorArgument));
+				}
+			}
+
+			return currentEditorArguments;
+		}
+
 		private bool ModifyEditorArguments()
 		{
-			using ArgumentsWindow arguments = new ArgumentsWindow(_settings.EditorArguments, _settings.EditorArgumentsPrompt);
+			List<LockableEditorArgument> tempEditorArguments = GetEditorArgumentsWithDefaults();
+			List<LockableEditorArgument> defaultEditorArguments = GetDefaultEditorArguments();
+
+			using ArgumentsWindow arguments = new ArgumentsWindow(tempEditorArguments, defaultEditorArguments, _settings.EditorArgumentsPrompt);
 			if (arguments.ShowDialog(this) == DialogResult.OK)
 			{
 				_settings.EditorArguments.Clear();
-				_settings.EditorArguments.AddRange(arguments.GetItems());
+
+				// If the user hasn't changed the value of a default editor argument,
+				// we don't need to save it in the global user editor arguments.
+				foreach (LockableEditorArgument editorArgument in arguments.GetItems())
+				{
+					int index = defaultEditorArguments.FindIndex(x => x.Name == editorArgument.Name);
+
+					if (index != -1)
+					{
+						LockableEditorArgument defaultEditorArgument = defaultEditorArguments[index];
+
+						if (editorArgument.Enabled == defaultEditorArgument.Enabled)
+						{
+							continue;
+						}
+					}
+
+					_settings.EditorArguments.Add(new LockableEditorArgument(editorArgument));
+				}
+
 				_settings.EditorArgumentsPrompt = arguments.PromptBeforeLaunch;
 				_settings.Save(_logger);
 				return true;
diff --git a/Engine/Source/Programs/UnrealGameSync/UnrealGameSync/Forms/ArgumentsWindow.cs b/Engine/Source/Programs/UnrealGameSync/UnrealGameSync/Forms/ArgumentsWindow.cs
--- a/Engine/Source/Programs/UnrealGameSync/UnrealGameSync/Forms/ArgumentsWindow.cs
+++ b/Engine/Source/Programs/UnrealGameSync/UnrealGameSync/Forms/ArgumentsWindow.cs
@@ -22,7 +22,10 @@ namespace UnrealGameSync
 		[DllImport("user32.dll")]
 		static extern bool SetWindowText(IntPtr hWnd, string text);
 
-		public ArgumentsWindow(List<Tuple<string, bool>> arguments, bool promptBeforeLaunch)
+		private readonly List<LockableEditorArgument> _editorArguments = new List<LockableEditorArgument>();
+		private readonly List<LockableEditorArgument> _defaultEditorArguments = new List<LockableEditorArgument>();
+
+		public ArgumentsWindow(List<LockableEditorArgument> arguments, List<LockableEditorArgument> defaultArguments, bool promptBeforeLaunch)
 		{
 			InitializeComponent();
 			Font = new System.Drawing.Font("Segoe UI", 8.25F, System.Drawing.FontStyle.Regular, System.Drawing.GraphicsUnit.Point, ((byte)(0)));
@@ -33,11 +36,19 @@ namespace UnrealGameSync
 
 			ArgumentsList.Items.Clear();
 
-			foreach (Tuple<string, bool> argument in arguments)
+			foreach (LockableEditorArgument argument in arguments)
 			{
-				ListViewItem item = new ListViewItem(argument.Item1);
-				item.Checked = argument.Item2;
+				ListViewItem item = new ListViewItem(argument.Name);
+				item.Checked = argument.Enabled;
 				ArgumentsList.Items.Add(item);
+
+				// keep a local copy of the editor arguments
+				_editorArguments.Add(new LockableEditorArgument(argument));
+			}
+
+			foreach (LockableEditorArgument defaultArgument in defaultArguments)
+			{
+				_defaultEditorArguments.Add(defaultArgument);
 			}
 
 			ListViewItem addAnotherItem = new ListViewItem("Click to add an item...", 0);
@@ -48,13 +59,13 @@ namespace UnrealGameSync
 
 		public bool PromptBeforeLaunch => PromptBeforeLaunchCheckBox.Checked;
 
-		public List<Tuple<string, bool>> GetItems()
+		public List<LockableEditorArgument> GetItems()
 		{
-			List<Tuple<string, bool>> items = new List<Tuple<string, bool>>();
+			List<LockableEditorArgument> items = new List<LockableEditorArgument>();
 			for (int idx = 0; idx < ArgumentsList.Items.Count - 1; idx++)
 			{
 				ListViewItem item = ArgumentsList.Items[idx];
-				items.Add(new Tuple<string, bool>(item.Text, item.Checked));
+				items.Add(new LockableEditorArgument(item.Text, item.Checked));
 			}
 			return items;
 		}
@@ -68,6 +79,14 @@ namespace UnrealGameSync
 
 		private void ArgumentsList_AfterLabelEdit(object sender, LabelEditEventArgs e)
 		{
+			// don't allow default editor arguments to have their text modified
+			bool isItemDefaultArgument = (GetDefaultArgumentForItem(e.Item) != null);
+			if (isItemDefaultArgument)
+			{
+				e.CancelEdit = true;
+				return;
+			}
+
 			if ((e.Label == null && ArgumentsList.Items[e.Item].Text.Length == 0) || (e.Label != null && e.Label.Trim().Length == 0))
 			{
 				e.CancelEdit = true;
@@ -77,6 +96,14 @@ namespace UnrealGameSync
 
 		private void ArgumentsList_BeforeLabelEdit(object sender, LabelEditEventArgs e)
 		{
+			// don't allow default editor arguments to have their text modified
+			bool isItemDefaultArgument = (GetDefaultArgumentForItem(e.Item) != null);
+			if (isItemDefaultArgument)
+			{
+				e.CancelEdit = true;
+				return;
+			}
+
 			if (e.Item == ArgumentsList.Items.Count - 1)
 			{
 				e.CancelEdit = true;
@@ -95,6 +122,13 @@ namespace UnrealGameSync
 			}
 			else
 			{
+				// don't allow default editor arguments to have their text modified
+				bool isItemDefaultArgument = (GetDefaultArgumentForItem(info.Item) != null);
+				if (isItemDefaultArgument)
+				{
+					return;
+				}
+
 				using (Graphics graphics = ArgumentsList.CreateGraphics())
 				{
 					int labelOffset = e.X - CheckBoxPadding - CheckBoxRenderer.GetGlyphSize(graphics, CheckBoxState.CheckedNormal).Width - CheckBoxPadding;
@@ -126,23 +160,27 @@ namespace UnrealGameSync
 
 		private void DrawItemLabel(Graphics graphics, Color normalColor, ListViewItem item)
 		{
+			string extraInfo = GetDefaultArgumentDebugText(item);
+
 			Rectangle labelRect = GetLabelRectangle(graphics, item);
 			if (item.Selected)
 			{
 				graphics.FillRectangle(SystemBrushes.Highlight, labelRect);
-				TextRenderer.DrawText(graphics, item.Text, ArgumentsList.Font, labelRect, SystemColors.HighlightText, SystemColors.Highlight, TextFormatFlags.Left | TextFormatFlags.VerticalCenter | TextFormatFlags.SingleLine);
+				TextRenderer.DrawText(graphics, item.Text + extraInfo, ArgumentsList.Font, labelRect, SystemColors.HighlightText, SystemColors.Highlight, TextFormatFlags.Left | TextFormatFlags.VerticalCenter | TextFormatFlags.SingleLine);
 			}
 			else
 			{
-				TextRenderer.DrawText(graphics, item.Text, ArgumentsList.Font, labelRect, normalColor, TextFormatFlags.Left | TextFormatFlags.VerticalCenter | TextFormatFlags.SingleLine);
+				TextRenderer.DrawText(graphics, item.Text + extraInfo, ArgumentsList.Font, labelRect, normalColor, TextFormatFlags.Left | TextFormatFlags.VerticalCenter | TextFormatFlags.SingleLine);
 			}
 		}
 
 		private Rectangle GetLabelRectangle(Graphics graphics, ListViewItem item)
 		{
+			string extraInfo = GetDefaultArgumentDebugText(item);
+
 			CheckBoxState state = item.Checked ? CheckBoxState.CheckedNormal : CheckBoxState.UncheckedNormal;
 			Size checkSize = CheckBoxRenderer.GetGlyphSize(graphics, state);
-			Size labelSize = TextRenderer.MeasureText(item.Text, ArgumentsList.Font);
+			Size labelSize = TextRenderer.MeasureText(item.Text + extraInfo, ArgumentsList.Font);
 
 			int labelIndent = CheckBoxPadding + checkSize.Width + CheckBoxPadding;
 			return new Rectangle(item.Bounds.Left + labelIndent, item.Bounds.Top, labelSize.Width, item.Bounds.Height);
@@ -153,6 +191,14 @@ namespace UnrealGameSync
 			if (e.KeyCode == Keys.Delete && ArgumentsList.SelectedIndices.Count == 1)
 			{
 				int index = ArgumentsList.SelectedIndices[0];
+
+				// don't allow default editor arguments to be deleted
+				bool isItemDefaultArgument = (GetDefaultArgumentForItem(index) != null);
+				if (isItemDefaultArgument)
+				{
+					return;
+				}
+
 				ArgumentsList.Items.RemoveAt(index);
 				if (index < ArgumentsList.Items.Count - 1)
 				{
@@ -166,6 +212,14 @@ namespace UnrealGameSync
 			if (ArgumentsList.SelectedItems.Count == 1 && !Char.IsControl(e.KeyChar))
 			{
 				ListViewItem item = ArgumentsList.SelectedItems[0];
+
+				// don't allow default editor arguments to have their text modified
+				bool isItemDefaultArgument = (GetDefaultArgumentForItem(item) != null);
+				if (isItemDefaultArgument)
+				{
+					return;
+				}
+
 				if (item.Index == ArgumentsList.Items.Count - 1)
 				{
 					item = new ListViewItem();
@@ -209,6 +263,19 @@ namespace UnrealGameSync
 			}
 		}
 
+		private void ResetDefaultsButton_Click(object sender, EventArgs e)
+		{
+			foreach (ListViewItem item in ArgumentsList.Items)
+			{
+				LockableEditorArgument? defaultArgument = GetDefaultArgumentForItem(item);
+
+				if (defaultArgument != null)
+				{
+					item.Checked = defaultArgument.Enabled;
+				}
+			}
+		}
+
 		private void ArgumentsList_SelectedIndexChanged(object sender, EventArgs e)
 		{
 			UpdateMoveButtons();
@@ -218,6 +285,62 @@ namespace UnrealGameSync
 		{
 			MoveUpButton.Enabled = (ArgumentsList.SelectedIndices.Count == 1 && ArgumentsList.SelectedIndices[0] > 0);
 			MoveDownButton.Enabled = (ArgumentsList.SelectedIndices.Count == 1 && ArgumentsList.SelectedIndices[0] < ArgumentsList.Items.Count - 2);
+			ResetDefaultsButton.Enabled = true;
+		}
+
+		private LockableEditorArgument? GetDefaultArgumentForItem(ListViewItem item)
+		{
+			int defaultArgumentindex = _defaultEditorArguments.FindIndex(x => x.Name == item.Text);
+
+			if (defaultArgumentindex != -1)
+			{
+				return _defaultEditorArguments[defaultArgumentindex];
+			}
+
+			return null;
+		}
+
+		private LockableEditorArgument? GetDefaultArgumentForItem(int index)
+		{
+			if (index < 0 || index >= ArgumentsList.Items.Count)
+			{
+				return null;
+			}
+
+			return GetDefaultArgumentForItem(ArgumentsList.Items[index]);
+		}
+
+		private string GetDefaultArgumentDebugText(ListViewItem item)
+		{
+			string extraInfo = "";
+
+			LockableEditorArgument? defaultArgument = GetDefaultArgumentForItem(item);
+
+			if (defaultArgument != null)
+			{
+				extraInfo = extraInfo + "    Default";
+
+				if (defaultArgument.Locked)
+				{
+					extraInfo = extraInfo + ": Locked";
+				}
+				else if (item.Checked != defaultArgument.Enabled)
+				{
+					extraInfo = extraInfo + ": Modifed";
+				}
+			}
+
+			return extraInfo;
+		}
+
+		private void ArgumentsList_ItemCheck(object? sender, ItemCheckEventArgs e)
+		{
+			LockableEditorArgument? defaultArgument = GetDefaultArgumentForItem(e.Index);
+
+			if ((defaultArgument != null) && defaultArgument.Locked)
+			{ 
+				e.NewValue = defaultArgument.Enabled ? CheckState.Checked : CheckState.Unchecked;
+			}
 		}
 	}
 }
diff --git a/Engine/Source/Programs/UnrealGameSync/UnrealGameSyncShared/UserSettings.cs b/Engine/Source/Programs/UnrealGameSync/UnrealGameSyncShared/UserSettings.cs
--- a/Engine/Source/Programs/UnrealGameSync/UnrealGameSyncShared/UserSettings.cs
+++ b/Engine/Source/Programs/UnrealGameSync/UnrealGameSyncShared/UserSettings.cs
@@ -50,6 +50,52 @@ namespace UnrealGameSync
 		Content
 	}
 
+	/// <summary>
+	/// Config specified class to add a lockable editor argument to a project
+	/// </summary>
+	public class LockableEditorArgument
+	{
+		public string Name { get; set; } = "";
+		public bool Enabled { get; set; } = false;
+		public bool Locked { get; set; } = false;
+
+		public LockableEditorArgument( string name, bool enabled, bool locked = false )
+		{
+			Name = name;
+			Enabled = enabled;
+			Locked = locked;
+		}
+
+		public LockableEditorArgument(LockableEditorArgument lockableEditorArgument)
+		{
+			Name = lockableEditorArgument.Name;
+			Enabled = lockableEditorArgument.Enabled;
+			Locked = lockableEditorArgument.Locked;
+		}
+
+		public static bool TryParseConfigEntry(string text, [NotNullWhen(true)] out LockableEditorArgument? editorArgument)
+		{
+			ConfigObject definitionObject = new ConfigObject(text);
+
+			string editorArgumentName = definitionObject.GetValue("Name", "");
+			if (editorArgumentName.Length > 0)
+			{
+				editorArgument = new LockableEditorArgument(
+					editorArgumentName,
+					definitionObject.GetValue("Enabled", false),
+					definitionObject.GetValue("Locked", false)
+				);
+				return true;
+			}
+			else
+			{
+				editorArgument = null;
+			}
+
+			return false;
+		}
+	}
+
 	/// <summary>
 	/// Config specified class to determine what is the Latest Change to Sync
 	/// Can be configured using badges, good and starred CLs
@@ -419,6 +465,7 @@ namespace UnrealGameSync
 		public List<ConfigObject> BuildSteps { get; init; } = new List<ConfigObject>();
 		public FilterType FilterType { get; set; }
 		public HashSet<string> FilterBadges { get; init; } = new HashSet<string>(StringComparer.OrdinalIgnoreCase);
+		[JsonIgnore]
 		public List<string> RequiredBadges { get; init; } = new List<string>();
 
 		static readonly object _syncRoot = new object();
@@ -524,7 +571,7 @@ namespace UnrealGameSync
 		public List<UserSelectedProjectSettings> ScheduleProjects { get; init; } = new List<UserSelectedProjectSettings>();
 
 		// Run configuration
-		public List<Tuple<string, bool>> EditorArguments { get; init; } = new List<Tuple<string, bool>>();
+		public List<LockableEditorArgument> EditorArguments { get; init; } = new List<LockableEditorArgument>();
 		public bool EditorArgumentsPrompt { get; set; }
 
 		// Notification settings
@@ -703,7 +750,7 @@ namespace UnrealGameSync
 			TabLabels = _configFile.GetEnumValue("General.TabLabels", TabLabels.Stream);
 
 			// Editor arguments
-			string[] arguments = _configFile.GetValues("General.EditorArguments", new string[] { "0:-log", "0:-fastload" });
+			string[] arguments = _configFile.GetValues("General.EditorArguments", Array.Empty<string>());
 			if (Version < UserSettingsVersion.XgeShaderCompilation)
 			{
 				arguments = Enumerable.Concat(arguments, new string[] { "0:-noxgeshadercompile" }).ToArray();
@@ -712,15 +759,15 @@ namespace UnrealGameSync
 			{
 				if (argument.StartsWith("0:", StringComparison.Ordinal))
 				{
-					EditorArguments.Add(new Tuple<string, bool>(argument.Substring(2), false));
+					EditorArguments.Add(new LockableEditorArgument(argument.Substring(2), /* Enabled = */ false));
 				}
 				else if (argument.StartsWith("1:", StringComparison.Ordinal))
 				{
-					EditorArguments.Add(new Tuple<string, bool>(argument.Substring(2), true));
+					EditorArguments.Add(new LockableEditorArgument(argument.Substring(2), /* Enabled = */ true));
 				}
 				else
 				{
-					EditorArguments.Add(new Tuple<string, bool>(argument, true));
+					EditorArguments.Add(new LockableEditorArgument(argument, /* Enabled = */ true));
 				}
 			}
 			EditorArgumentsPrompt = _configFile.GetValue("General.EditorArgumentsPrompt", false);
@@ -1035,9 +1082,9 @@ namespace UnrealGameSync
 
 			// Editor arguments
 			List<string> editorArgumentList = new List<string>();
-			foreach (Tuple<string, bool> editorArgument in EditorArguments)
+			foreach (LockableEditorArgument editorArgument in  EditorArguments)
 			{
-				editorArgumentList.Add(String.Format("{0}:{1}", editorArgument.Item2 ? 1 : 0, editorArgument.Item1));
+				editorArgumentList.Add(String.Format("{0}:{1}", editorArgument.Enabled ? 1 : 0, editorArgument.Name));
 			}
 			generalSection.SetValues("EditorArguments", editorArgumentList.ToArray());
 			generalSection.SetValue("EditorArgumentsPrompt", EditorArgumentsPrompt);
