[UGS]: Fix a bug where the tools being updated would not update the "More..." list

diff --git a/Engine/Source/Programs/UnrealGameSync/UnrealGameSync/ToolUpdateMonitor.cs b/Engine/Source/Programs/UnrealGameSync/UnrealGameSync/ToolUpdateMonitor.cs
--- a/Engine/Source/Programs/UnrealGameSync/UnrealGameSync/ToolUpdateMonitor.cs
+++ b/Engine/Source/Programs/UnrealGameSync/UnrealGameSync/ToolUpdateMonitor.cs
@@ -248,8 +248,17 @@ namespace UnrealGameSync
 						}
 					}
 				}
+
+				bool hasChanged = false;
+				int previousToolsCount = _tools.Count;
+
 				_tools = tools;
 
+				if (previousToolsCount != _tools.Count)
+				{
+					hasChanged = true;
+				}
+
 				// When upgrading from older UGS versions, read the legacy sync CL from plain-text config files
 				if (_readLegacyConfig)
 				{
@@ -262,7 +271,6 @@ namespace UnrealGameSync
 				FindEnabledTools(Settings.EnabledTools, tools, enabledToolIds);
 
 				// Install or update any new tools
-				bool hasChanged = false;
 				foreach (ToolInfo toolInfo in _tools)
 				{
 					if (enabledToolIds.Contains(toolInfo.Id))
