[UGS]: Fix a debug text typo

diff --git a/Engine/Source/Programs/UnrealGameSync/UnrealGameSync/Forms/ArgumentsWindow.cs b/Engine/Source/Programs/UnrealGameSync/UnrealGameSync/Forms/ArgumentsWindow.cs
--- a/Engine/Source/Programs/UnrealGameSync/UnrealGameSync/Forms/ArgumentsWindow.cs
+++ b/Engine/Source/Programs/UnrealGameSync/UnrealGameSync/Forms/ArgumentsWindow.cs
@@ -326,7 +326,7 @@ namespace UnrealGameSync
 				}
 				else if (item.Checked != defaultArgument.Enabled)
 				{
-					extraInfo = extraInfo + ": Modifed";
+					extraInfo = extraInfo + ": Modified";
 				}
 			}
 
