UGS: Ignore required badges by default for non-editor archive types.

diff --git a/Engine/Source/Programs/UnrealGameSync/UnrealGameSyncShared/ArchiveInfo.cs b/Engine/Source/Programs/UnrealGameSync/UnrealGameSyncShared/ArchiveInfo.cs
--- a/Engine/Source/Programs/UnrealGameSync/UnrealGameSyncShared/ArchiveInfo.cs
+++ b/Engine/Source/Programs/UnrealGameSync/UnrealGameSyncShared/ArchiveInfo.cs
@@ -334,7 +334,9 @@ namespace UnrealGameSync
 							type = archiveType.Substring(ArchiveTypePrefix.Length);
 						}
 
-						HordeArchiveChannel channel = new HordeArchiveChannel(name, type);
+						bool ignoreRequiredBadges = !String.Equals(type, IArchiveChannel.EditorArchiveType, StringComparison.OrdinalIgnoreCase);
+
+						HordeArchiveChannel channel = new HordeArchiveChannel(name, type, ignoreRequiredBadges);
 						foreach (GetArtifactResponse response in group)
 						{
 							HordeArchive archive = new HordeArchive(hordeClient, response.Id);
