Parallelize determination of files to remove logic.
Saves 17s on workspace with 2.7m files synced.
#rnx

diff --git a/Engine/Source/Programs/UnrealGameSync/UnrealGameSyncShared/WorkspaceUpdate.cs b/Engine/Source/Programs/UnrealGameSync/UnrealGameSyncShared/WorkspaceUpdate.cs
--- a/Engine/Source/Programs/UnrealGameSync/UnrealGameSyncShared/WorkspaceUpdate.cs
+++ b/Engine/Source/Programs/UnrealGameSync/UnrealGameSyncShared/WorkspaceUpdate.cs
@@ -1,6 +1,7 @@
 // Copyright Epic Games, Inc. All Rights Reserved.
 
 using System;
+using System.Collections.Concurrent;
 using System.Collections.Generic;
 using System.Diagnostics;
 using System.IO;
@@ -667,25 +668,26 @@ namespace UnrealGameSync
 
 							// Remove all the files that are not included by the filter
 							const int MaxLogFiles = 1000;
-							List<string> removeDepotPaths = new List<string>();
-							foreach (HaveRecord haveFile in haveFiles)
+							ConcurrentBag<string> removeDepotPathsBag = new ConcurrentBag<string>();
+							Parallel.ForEach(haveFiles, haveFile =>
 							{
 								try
 								{
 									FileReference fullPath = new FileReference(haveFile.Path);
 									if (MatchFilter(project, fullPath, syncPathsFilter) && !MatchFilter(project, fullPath, userFilter))
 									{
-										if (removeDepotPaths.Count <= MaxLogFiles)
-										{
-											logger.LogInformation("  {DepotFile}", haveFile.DepotFile);
-										}
-										removeDepotPaths.Add(haveFile.DepotFile);
+										removeDepotPathsBag.Add(haveFile.DepotFile);
 									}
 								}
 								catch (PathTooLongException)
 								{
 									// We don't actually care about this when looking for files to remove. Perforce may think that it's synced the path, and silently failed. Just ignore it.
 								}
+							});
+							List<string> removeDepotPaths = removeDepotPathsBag.ToList();
+							for (int i=0; i<Math.Min(removeDepotPaths.Count, MaxLogFiles); i++)
+							{
+								logger.LogInformation("  {DepotFile}", removeDepotPaths[i]);
 							}
 							if (removeDepotPaths.Count > MaxLogFiles)
 							{
