[UGS]: Make the "Reset To Defaults" button uncheck any non default editor arguments.  Make the argument order more consistent (defaults first).

diff --git a/Engine/Source/Programs/UnrealGameSync/UnrealGameSync/Forms/ArgumentsWindow.cs b/Engine/Source/Programs/UnrealGameSync/UnrealGameSync/Forms/ArgumentsWindow.cs
--- a/Engine/Source/Programs/UnrealGameSync/UnrealGameSync/Forms/ArgumentsWindow.cs
+++ b/Engine/Source/Programs/UnrealGameSync/UnrealGameSync/Forms/ArgumentsWindow.cs
@@ -22,7 +22,6 @@ namespace UnrealGameSync
 		[DllImport("user32.dll")]
 		static extern bool SetWindowText(IntPtr hWnd, string text);
 
-		private readonly List<LockableEditorArgument> _editorArguments = new List<LockableEditorArgument>();
 		private readonly List<LockableEditorArgument> _defaultEditorArguments = new List<LockableEditorArgument>();
 
 		public ArgumentsWindow(List<LockableEditorArgument> arguments, List<LockableEditorArgument> defaultArguments, bool promptBeforeLaunch)
@@ -36,19 +35,36 @@ namespace UnrealGameSync
 
 			ArgumentsList.Items.Clear();
 
-			foreach (LockableEditorArgument argument in arguments)
+			// Keep a local copy of the default arguments
+			foreach (LockableEditorArgument defaultArgument in defaultArguments)
+			{
+				_defaultEditorArguments.Add(defaultArgument);
+			}
+
+			// Add any arguments that are in the default list first
+			foreach (LockableEditorArgument defaultArgument in _defaultEditorArguments)
 			{
-				ListViewItem item = new ListViewItem(argument.Name);
-				item.Checked = argument.Enabled;
-				ArgumentsList.Items.Add(item);
+				int index = arguments.FindIndex(x => x.Name == defaultArgument.Name);
 
-				// keep a local copy of the editor arguments
-				_editorArguments.Add(new LockableEditorArgument(argument));
+				if (index != -1)
+				{
+					ListViewItem item = new ListViewItem(arguments[index].Name);
+					item.Checked = arguments[index].Enabled;
+					ArgumentsList.Items.Add(item);
+				}
 			}
 
-			foreach (LockableEditorArgument defaultArgument in defaultArguments)
+			// Add any user arguments that are not defaults second
+			foreach (LockableEditorArgument argument in arguments)
 			{
-				_defaultEditorArguments.Add(defaultArgument);
+				int index = _defaultEditorArguments.FindIndex(x => x.Name == argument.Name);
+
+				if (index == -1)
+				{
+					ListViewItem item = new ListViewItem(argument.Name);
+					item.Checked = argument.Enabled;
+					ArgumentsList.Items.Add(item);
+				}
 			}
 
 			ListViewItem addAnotherItem = new ListViewItem("Click to add an item...", 0);
@@ -263,7 +279,7 @@ namespace UnrealGameSync
 			}
 		}
 
-		private void ResetDefaultsButton_Click(object sender, EventArgs e)
+		private void ResetToDefaultsButton_Click(object sender, EventArgs e)
 		{
 			foreach (ListViewItem item in ArgumentsList.Items)
 			{
@@ -273,6 +289,10 @@ namespace UnrealGameSync
 				{
 					item.Checked = defaultArgument.Enabled;
 				}
+				else
+				{
+					item.Checked = false;
+				}
 			}
 		}
 
@@ -285,7 +305,7 @@ namespace UnrealGameSync
 		{
 			MoveUpButton.Enabled = (ArgumentsList.SelectedIndices.Count == 1 && ArgumentsList.SelectedIndices[0] > 0);
 			MoveDownButton.Enabled = (ArgumentsList.SelectedIndices.Count == 1 && ArgumentsList.SelectedIndices[0] < ArgumentsList.Items.Count - 2);
-			ResetDefaultsButton.Enabled = true;
+			ResetToDefaultsButton.Enabled = true;
 		}
 
 		private LockableEditorArgument? GetDefaultArgumentForItem(ListViewItem item)
