[UGS]: Add some extra telemetry information

diff --git a/Engine/Source/Programs/UnrealGameSync/UnrealGameSync/Program.cs b/Engine/Source/Programs/UnrealGameSync/UnrealGameSync/Program.cs
--- a/Engine/Source/Programs/UnrealGameSync/UnrealGameSync/Program.cs
+++ b/Engine/Source/Programs/UnrealGameSync/UnrealGameSync/Program.cs
@@ -275,7 +275,7 @@ namespace UnrealGameSync
 						}
 						catch (Exception ex)
 						{
-							UgsTelemetry.SendEvent("Crash", new { Exception = ex.ToString() });
+							UgsTelemetry.SendEvent("Crash", new { User = Environment.UserName, Machine = System.Net.Dns.GetHostName(), Exception = ex.ToString() });
 							throw;
 						}
 						finally
@@ -351,7 +351,7 @@ namespace UnrealGameSync
 			Exception? ex = args.ExceptionObject as Exception;
 			if (ex != null)
 			{
-				UgsTelemetry.SendEvent("Crash", new { Exception = ex.ToString() });
+				UgsTelemetry.SendEvent("Crash", new { User = Environment.UserName, Machine = System.Net.Dns.GetHostName(), Exception = ex.ToString() });
 			}
 		}
 
diff --git a/Engine/Source/Programs/UnrealGameSync/UnrealGameSyncShared/TelemetryStopwatch.cs b/Engine/Source/Programs/UnrealGameSync/UnrealGameSyncShared/TelemetryStopwatch.cs
--- a/Engine/Source/Programs/UnrealGameSync/UnrealGameSyncShared/TelemetryStopwatch.cs
+++ b/Engine/Source/Programs/UnrealGameSync/UnrealGameSyncShared/TelemetryStopwatch.cs
@@ -39,6 +39,7 @@ namespace UnrealGameSync
 
 				_eventData["Result"] = inResult;
 				_eventData["TimeSeconds"] = _timer.Elapsed.TotalSeconds;
+				_eventData["TimeMinutes"] = _timer.Elapsed.TotalSeconds / 60.0f;
 			}
 			return Elapsed;
 		}
diff --git a/Engine/Source/Programs/UnrealGameSync/UnrealGameSyncShared/WorkspaceUpdate.cs b/Engine/Source/Programs/UnrealGameSync/UnrealGameSyncShared/WorkspaceUpdate.cs
--- a/Engine/Source/Programs/UnrealGameSync/UnrealGameSyncShared/WorkspaceUpdate.cs
+++ b/Engine/Source/Programs/UnrealGameSync/UnrealGameSyncShared/WorkspaceUpdate.cs
@@ -597,6 +597,8 @@ namespace UnrealGameSync
 				{
 					logger.LogInformation("Syncing to {Change} on {ServerAndPort} as {UserName}...", Context.ChangeNumber, perforceSettings.ServerAndPort, perforceSettings.UserName);
 
+					syncTelemetryStopwatch.AddData(new { MachineName = System.Net.Dns.GetHostName(), DomainName = Environment.UserDomainName, ServerAndPort = perforce.Settings.ServerAndPort, ClientName = perforce.Settings.ClientName, UserName = perforce.Settings.UserName });
+
 					// Make sure we're logged in
 					PerforceResponse<LoginRecord> loginResponse = await perforce.TryGetLoginStateAsync(cancellationToken);
 					if (!loginResponse.Succeeded)
@@ -638,6 +640,8 @@ namespace UnrealGameSync
 						{
 							logger.LogInformation("Filter has changed ({PrevHash} -> {NextHash}); finding files in workspace that need to be removed.", (String.IsNullOrEmpty(state.CurrentSyncFilterHash)) ? "None" : state.CurrentSyncFilterHash, nextSyncFilterHash);
 
+							filterStopwatch.AddData(new { MachineName = System.Net.Dns.GetHostName(), DomainName = Environment.UserDomainName, ServerAndPort = perforce.Settings.ServerAndPort, ClientName = perforce.Settings.ClientName, UserName = perforce.Settings.UserName });
+
 							// Find all the files that are in this workspace
 							List<HaveRecord> haveFiles = Context.HaveFiles;
 							if (haveFiles.Count == 0)
@@ -900,7 +904,7 @@ namespace UnrealGameSync
 
 					using (TelemetryStopwatch transferStopwatch = new TelemetryStopwatch("Workspace_Sync_TransferFiles", project.TelemetryProjectIdentifier))
 					{
-						transferStopwatch.AddData(new { MachineName = System.Net.Dns.GetHostName(), DomainName = Environment.UserDomainName, ServerAndPort = perforce.Settings.ServerAndPort, UserName = perforce.Settings.UserName, IncludedFiles = syncTree.TotalIncludedFiles, ExcludedFiles = syncTree.TotalExcludedFiles, Size = syncTree.TotalSize, NumThreads = Context.PerforceSyncOptions?.NumThreads ?? PerforceSyncOptions.DefaultNumThreads });
+						transferStopwatch.AddData(new { MachineName = System.Net.Dns.GetHostName(), DomainName = Environment.UserDomainName, ServerAndPort = perforce.Settings.ServerAndPort, ClientName = perforce.Settings.ClientName, UserName = perforce.Settings.UserName, IncludedFiles = syncTree.TotalIncludedFiles, ExcludedFiles = syncTree.TotalExcludedFiles, Size = syncTree.TotalSize, NumThreads = Context.PerforceSyncOptions?.NumThreads ?? PerforceSyncOptions.DefaultNumThreads });
 
 						(WorkspaceUpdateResult, string) syncResult = await SyncFileRevisions(perforce, "Syncing files...", Context, batchBuilder.Batches, remainingDepotPaths, Progress, logger, cancellationToken);
 						if (syncResult.Item1 != WorkspaceUpdateResult.Success)
@@ -1151,6 +1155,8 @@ namespace UnrealGameSync
 			{
 				using (TelemetryStopwatch stopwatch = new TelemetryStopwatch("Workspace_SyncArchives", project.TelemetryProjectIdentifier))
 				{
+					stopwatch.AddData(new { MachineName = System.Net.Dns.GetHostName(), DomainName = Environment.UserDomainName, ServerAndPort = perforce.Settings.ServerAndPort, ClientName = perforce.Settings.ClientName, UserName = perforce.Settings.UserName });
+
 					// Create the directory for extracted archive manifests
 					DirectoryReference manifestDirectoryName;
 					if (project.LocalFileName.HasExtension(".uproject"))
@@ -1239,6 +1245,8 @@ namespace UnrealGameSync
 				{
 					Progress.Set("Generating project files...", 0.0f);
 
+					stopwatch.AddData(new { MachineName = System.Net.Dns.GetHostName(), DomainName = Environment.UserDomainName, ServerAndPort = perforce.Settings.ServerAndPort, ClientName = perforce.Settings.ClientName, UserName = perforce.Settings.UserName });
+
 					StringBuilder commandLine = new StringBuilder();
 					commandLine.AppendFormat("\"{0}\"", FileReference.Combine(project.LocalRootPath, $"GenerateProjectFiles.{ShellScriptExt}"));
 					if ((Context.Options & WorkspaceUpdateOptions.SyncAllProjects) == 0 && (Context.Options & WorkspaceUpdateOptions.IncludeAllProjectsInSolution) == 0)
@@ -1321,6 +1329,8 @@ namespace UnrealGameSync
 				{
 					Progress.Set("Starting build...", 0.0f);
 
+					stopwatch.AddData(new { MachineName = System.Net.Dns.GetHostName(), DomainName = Environment.UserDomainName, ServerAndPort = perforce.Settings.ServerAndPort, ClientName = perforce.Settings.ClientName, UserName = perforce.Settings.UserName });
+
 					// Execute all the steps
 					float maxProgressFraction = 0.0f;
 					foreach (BuildStep step in buildSteps)
@@ -1358,6 +1368,7 @@ namespace UnrealGameSync
 								case BuildStepType.Compile:
 									using (TelemetryStopwatch stepStopwatch = new TelemetryStopwatch("Workspace_Execute_Compile", project.TelemetryProjectIdentifier))
 									{
+										stepStopwatch.AddData(new { MachineName = System.Net.Dns.GetHostName(), DomainName = Environment.UserDomainName, ServerAndPort = perforce.Settings.ServerAndPort, ClientName = perforce.Settings.ClientName, UserName = perforce.Settings.UserName });
 										stepStopwatch.AddData(new { Target = step.Target });
 
 										FileReference buildBat = FileReference.Combine(batchFilesDir, $"Build.{ShellScriptExt}");
@@ -1395,6 +1406,7 @@ namespace UnrealGameSync
 								case BuildStepType.Cook:
 									using (TelemetryStopwatch stepStopwatch = new TelemetryStopwatch("Workspace_Execute_Cook", project.TelemetryProjectIdentifier))
 									{
+										stepStopwatch.AddData(new { MachineName = System.Net.Dns.GetHostName(), DomainName = Environment.UserDomainName, ServerAndPort = perforce.Settings.ServerAndPort, ClientName = perforce.Settings.ClientName, UserName = perforce.Settings.UserName });
 										stepStopwatch.AddData(new { Project = Path.GetFileNameWithoutExtension(step.FileName) });
 
 										FileReference localRunUat = FileReference.Combine(batchFilesDir, $"RunUAT.{ShellScriptExt}");
@@ -1414,6 +1426,7 @@ namespace UnrealGameSync
 								case BuildStepType.Other:
 									using (TelemetryStopwatch stepStopwatch = new TelemetryStopwatch("Workspace_Execute_Custom", project.TelemetryProjectIdentifier))
 									{
+										stepStopwatch.AddData(new { MachineName = System.Net.Dns.GetHostName(), DomainName = Environment.UserDomainName, ServerAndPort = perforce.Settings.ServerAndPort, ClientName = perforce.Settings.ClientName, UserName = perforce.Settings.UserName });
 										stepStopwatch.AddData(new { FileName = Path.GetFileNameWithoutExtension(step.FileName) });
 
 										FileReference toolFileName = FileReference.Combine(project.LocalRootPath, Utility.ExpandVariables(step.FileName ?? "unknown", variables));
