Support for multiple capture groups in regex badges

diff --git a/Engine/Source/Programs/UnrealGameSync/UnrealGameSync/Controls/WorkspaceControl.cs b/Engine/Source/Programs/UnrealGameSync/UnrealGameSync/Controls/WorkspaceControl.cs
--- a/Engine/Source/Programs/UnrealGameSync/UnrealGameSync/Controls/WorkspaceControl.cs
+++ b/Engine/Source/Programs/UnrealGameSync/UnrealGameSync/Controls/WorkspaceControl.cs
@@ -2687,6 +2687,8 @@ namespace UnrealGameSync
 						string hoverColor = badgeDefinitionObject.GetValue("HoverColor", "#b0b0b0");
 						string? url = badgeDefinitionObject.GetValue("Url", null);
 						string? arguments = badgeDefinitionObject.GetValue("Arguments", null);
+						string? multiCaptureGroup = badgeDefinitionObject.GetValue("MultiCaptureGroup", null);
+						string? subGroupSplit = badgeDefinitionObject.GetValue("SubGroupSplit", null);
 						bool singleMatch = badgeDefinitionObject.GetValue("SingleMatch", false);
 						if (!String.IsNullOrEmpty(name) && !String.IsNullOrEmpty(pattern))
 						{
@@ -2699,28 +2701,59 @@ namespace UnrealGameSync
 
 									string? uniqueId = String.IsNullOrEmpty(url) ? null : String.Format("Description:{0}:{1}", change.Number, badges.Count);
 
-									string? expandedUrl = ReplaceRegexMatches(url, matchResult);
-									string? expandedArguments = ReplaceRegexMatches(arguments, matchResult);
-
-									Action? clickHandler;
-									if (String.IsNullOrEmpty(expandedUrl))
-									{
-										clickHandler = null;
-									}
-									else if (String.IsNullOrEmpty(expandedArguments))
+									string? expandedMultiCaptureGroup = ReplaceRegexMatches(multiCaptureGroup, matchResult);
+									if (!String.IsNullOrWhiteSpace(expandedMultiCaptureGroup))
 									{
-										clickHandler = () => SafeProcessStart(expandedUrl);
+										string[] capturedGroups = expandedMultiCaptureGroup.Split(subGroupSplit);
+										for (int i = 0; i < capturedGroups.Length; ++i)
+										{
+											string? currentCapture = capturedGroups[i].Trim();
+
+											Action? clickHandler;
+											if (String.IsNullOrEmpty(url))
+											{
+												clickHandler = null;
+											}
+											else if (String.IsNullOrEmpty(arguments))
+											{
+												string? expandedUrl = url.Replace("$*", currentCapture, StringComparison.OrdinalIgnoreCase);
+												clickHandler = () => SafeProcessStart(expandedUrl);
+											}
+											else
+											{
+												string? expandedUrl = url.Replace("$*", currentCapture, StringComparison.OrdinalIgnoreCase);
+												string? expandedArguments = url.Replace("$*", currentCapture, StringComparison.OrdinalIgnoreCase);
+												clickHandler = () => SafeProcessStart(expandedUrl, expandedArguments);
+											}
+
+											badges.Add(new BadgeInfo(name.Replace("$*", currentCapture, StringComparison.OrdinalIgnoreCase), group, uniqueId, badgeColor, hoverBadgeColor, clickHandler, null));
+										}
 									}
 									else
 									{
-										clickHandler = () => SafeProcessStart(expandedUrl, expandedArguments);
-									}
+										string? expandedUrl = ReplaceRegexMatches(url, matchResult);
+										string? expandedArguments = ReplaceRegexMatches(arguments, matchResult);
+
+										Action? clickHandler;
+										if (String.IsNullOrEmpty(expandedUrl))
+										{
+											clickHandler = null;
+										}
+										else if (String.IsNullOrEmpty(expandedArguments))
+										{
+											clickHandler = () => SafeProcessStart(expandedUrl);
+										}
+										else
+										{
+											clickHandler = () => SafeProcessStart(expandedUrl, expandedArguments);
+										}
 
-									badges.Add(new BadgeInfo(ReplaceRegexMatches(name, matchResult), group, uniqueId, badgeColor, hoverBadgeColor, clickHandler, null));
+										badges.Add(new BadgeInfo(ReplaceRegexMatches(name, matchResult), group, uniqueId, badgeColor, hoverBadgeColor, clickHandler, null));
 
-									if (singleMatch)
-									{
-										break;
+										if (singleMatch)
+										{
+											break;
+										}
 									}
 								}
 							}
