[UGS]: Add a setting to wait an optional number of milliseconds between sync error retries

diff --git a/Engine/Source/Programs/UnrealGameSync/UnrealGameSync/Forms/PerforceSyncSettingsWindow.cs b/Engine/Source/Programs/UnrealGameSync/UnrealGameSync/Forms/PerforceSyncSettingsWindow.cs
--- a/Engine/Source/Programs/UnrealGameSync/UnrealGameSync/Forms/PerforceSyncSettingsWindow.cs
+++ b/Engine/Source/Programs/UnrealGameSync/UnrealGameSync/Forms/PerforceSyncSettingsWindow.cs
@@ -26,6 +26,7 @@ namespace UnrealGameSync
 			numericUpDownMaxCommandsPerBatch.Value = syncOptions.MaxCommandsPerBatch ?? PerforceSyncOptions.DefaultMaxCommandsPerBatch;
 			numericUpDownMaxSizePerBatch.Value = (syncOptions.MaxSizePerBatch ?? PerforceSyncOptions.DefaultMaxSizePerBatch) / 1024 / 1024;
 			numericUpDownRetriesOnSyncError.Value = syncOptions.NumSyncErrorRetries ?? PerforceSyncOptions.DefaultNumSyncErrorRetries;
+			numericUpDownSyncErrorRetryDelay.Value = syncOptions.SyncErrorRetryDelay ?? PerforceSyncOptions.DefaultSyncErrorRetryDelay;
 		}
 
 		private void OkButton_Click(object sender, EventArgs e)
@@ -33,6 +34,7 @@ namespace UnrealGameSync
 			_settings.SyncOptions.MaxCommandsPerBatch = GetValueIfNotDefault((int)numericUpDownMaxCommandsPerBatch.Value, PerforceSyncOptions.DefaultMaxCommandsPerBatch);
 			_settings.SyncOptions.MaxSizePerBatch = GetValueIfNotDefault((int)numericUpDownMaxSizePerBatch.Value, PerforceSyncOptions.DefaultMaxSizePerBatch) * 1024 * 1024;
 			_settings.SyncOptions.NumSyncErrorRetries = GetValueIfNotDefault((int)numericUpDownRetriesOnSyncError.Value, PerforceSyncOptions.DefaultNumSyncErrorRetries);
+			_settings.SyncOptions.SyncErrorRetryDelay = GetValueIfNotDefault((int)numericUpDownSyncErrorRetryDelay.Value, PerforceSyncOptions.DefaultSyncErrorRetryDelay);
 			_settings.Save(_logger);
 
 			DialogResult = System.Windows.Forms.DialogResult.OK;
@@ -55,6 +57,7 @@ namespace UnrealGameSync
 			numericUpDownMaxCommandsPerBatch.Value = PerforceSyncOptions.DefaultMaxCommandsPerBatch;
 			numericUpDownMaxSizePerBatch.Value = PerforceSyncOptions.DefaultMaxSizePerBatch / 1024 / 1024;
 			numericUpDownRetriesOnSyncError.Value = PerforceSyncOptions.DefaultNumSyncErrorRetries;
+			numericUpDownSyncErrorRetryDelay.Value = PerforceSyncOptions.DefaultSyncErrorRetryDelay;
 		}
 	}
 }
diff --git a/Engine/Source/Programs/UnrealGameSync/UnrealGameSyncShared/UserSettings.cs b/Engine/Source/Programs/UnrealGameSync/UnrealGameSyncShared/UserSettings.cs
--- a/Engine/Source/Programs/UnrealGameSync/UnrealGameSyncShared/UserSettings.cs
+++ b/Engine/Source/Programs/UnrealGameSync/UnrealGameSyncShared/UserSettings.cs
@@ -842,6 +842,7 @@ namespace UnrealGameSync
 			SyncOptions.MaxCommandsPerBatch = _configFile.GetOptionalIntValue("Perforce.MaxCommandsPerBatch", SyncOptions.MaxCommandsPerBatch);
 			SyncOptions.MaxSizePerBatch = _configFile.GetOptionalIntValue("Perforce.MaxSizePerBatch", SyncOptions.MaxSizePerBatch);
 			SyncOptions.NumSyncErrorRetries = _configFile.GetOptionalIntValue("Perforce.NumSyncErrorRetries", SyncOptions.NumSyncErrorRetries);
+			SyncOptions.SyncErrorRetryDelay = _configFile.GetOptionalIntValue("Perforce.SyncErrorRetryDelay", SyncOptions.SyncErrorRetryDelay);
 		}
 
 		static Dictionary<Guid, bool> GetCategorySettings(ConfigSection? section, string includedKey, string excludedKey)
diff --git a/Engine/Source/Programs/UnrealGameSync/UnrealGameSyncShared/WorkspaceUpdate.cs b/Engine/Source/Programs/UnrealGameSync/UnrealGameSyncShared/WorkspaceUpdate.cs
--- a/Engine/Source/Programs/UnrealGameSync/UnrealGameSyncShared/WorkspaceUpdate.cs
+++ b/Engine/Source/Programs/UnrealGameSync/UnrealGameSyncShared/WorkspaceUpdate.cs
@@ -65,6 +65,7 @@ namespace UnrealGameSync
 		public const int DefaultMaxCommandsPerBatch = 200;
 		public const int DefaultMaxSizePerBatch = 128 * 1024 * 1024;
 		public const int DefaultNumSyncErrorRetries = 0;
+		public const int DefaultSyncErrorRetryDelay = 0;
 
 		public int? NumThreads { get; set; }
 
@@ -72,6 +73,7 @@ namespace UnrealGameSync
 		public int? MaxSizePerBatch { get; set; }
 
 		public int? NumSyncErrorRetries { get; set; }
+		public int? SyncErrorRetryDelay { get; set; }
 
 		public PerforceSyncOptions Clone()
 		{
@@ -1738,6 +1740,8 @@ namespace UnrealGameSync
 				string statusMessage = "";
 
 				int maxRetries = context.PerforceSyncOptions?.NumSyncErrorRetries ?? PerforceSyncOptions.DefaultNumSyncErrorRetries;
+				int retryDelay = context.PerforceSyncOptions?.SyncErrorRetryDelay ?? PerforceSyncOptions.DefaultSyncErrorRetryDelay;
+				
 				for (int attempt = 0; ; attempt++)
 				{
 					// Sync the files
@@ -1747,7 +1751,13 @@ namespace UnrealGameSync
 					{
 						break;
 					}
-					threadLog.LogWarning("Sync error ({Message}); retrying... ({Count}/{MaxCount})", errorMessage ?? "unknown", attempt + 1, maxRetries);
+
+					threadLog.LogWarning("Sync error ({Message}); waiting ({RetryDelay}ms) and retrying... ({Count}/{MaxCount})", errorMessage ?? "unknown", retryDelay, attempt + 1, maxRetries);
+
+					if (retryDelay > 0)
+					{ 
+						await Task.Delay(retryDelay, cancellationToken);
+					}
 				}
 
 				// If it failed, try to set it on the state if nothing else has failed first
