[UGS]: If Horde tools are unavailable, don't kill the perforce tools

diff --git a/Engine/Source/Programs/UnrealGameSync/UnrealGameSync/ToolUpdateMonitor.cs b/Engine/Source/Programs/UnrealGameSync/UnrealGameSync/ToolUpdateMonitor.cs
--- a/Engine/Source/Programs/UnrealGameSync/UnrealGameSync/ToolUpdateMonitor.cs
+++ b/Engine/Source/Programs/UnrealGameSync/UnrealGameSync/ToolUpdateMonitor.cs
@@ -215,6 +215,7 @@ namespace UnrealGameSync
 			IPerforceConnection? perforce = null;
 			try
 			{
+				bool perforceToolsAvailable = true;
 				Stopwatch timer = Stopwatch.StartNew();
 				// Update all the available tools
 				List<ToolInfo> tools = new List<ToolInfo>();
@@ -229,9 +230,10 @@ namespace UnrealGameSync
 					{
 						LastStatus = Tuple.Create(false, $"Error while polling Perforce for available tools: {ex.Message}");
 						_logger.LogWarning(ex, "Error while polling Perforce for available tools: {Message}", ex.Message);
-						return;
+						perforceToolsAvailable = false;
 					}
 				}
+				bool hordeToolsAvailable = true;
 				using (HordeHttpClient? hordeHttpClient = _serviceProvider.GetService<HordeHttpClient>())
 				{
 					if (hordeHttpClient != null)
@@ -244,11 +246,16 @@ namespace UnrealGameSync
 						{
 							LastStatus = Tuple.Create(false, $"Error while polling Horde for available tools: {ex.Message}");
 							_logger.LogWarning(ex, "Error while polling Horde for available tools: {Message}", ex.Message);
-							return;
+							hordeToolsAvailable = false;
 						}
 					}
 				}
 
+				if (!perforceToolsAvailable && !hordeToolsAvailable)
+				{
+					return;
+				}
+
 				bool hasChanged = false;
 				int previousToolsCount = _tools.Count;
 
