[Backout] - CL34036503
Original CL Desc
-----------------------------------------------------------------
[UGS]: If Horde tools are unavailable, don't kill the perforce tools

diff --git a/Engine/Source/Programs/UnrealGameSync/UnrealGameSync/ToolUpdateMonitor.cs b/Engine/Source/Programs/UnrealGameSync/UnrealGameSync/ToolUpdateMonitor.cs
--- a/Engine/Source/Programs/UnrealGameSync/UnrealGameSync/ToolUpdateMonitor.cs
+++ b/Engine/Source/Programs/UnrealGameSync/UnrealGameSync/ToolUpdateMonitor.cs
@@ -215,7 +215,6 @@ namespace UnrealGameSync
 			IPerforceConnection? perforce = null;
 			try
 			{
-				bool perforceToolsAvailable = true;
 				Stopwatch timer = Stopwatch.StartNew();
 				// Update all the available tools
 				List<ToolInfo> tools = new List<ToolInfo>();
@@ -230,10 +229,9 @@ namespace UnrealGameSync
 					{
 						LastStatus = Tuple.Create(false, $"Error while polling Perforce for available tools: {ex.Message}");
 						_logger.LogWarning(ex, "Error while polling Perforce for available tools: {Message}", ex.Message);
-						perforceToolsAvailable = false;
+						return;
 					}
 				}
-				bool hordeToolsAvailable = true;
 				using (HordeHttpClient? hordeHttpClient = _serviceProvider.GetService<HordeHttpClient>())
 				{
 					if (hordeHttpClient != null)
@@ -246,16 +244,11 @@ namespace UnrealGameSync
 						{
 							LastStatus = Tuple.Create(false, $"Error while polling Horde for available tools: {ex.Message}");
 							_logger.LogWarning(ex, "Error while polling Horde for available tools: {Message}", ex.Message);
-							hordeToolsAvailable = false;
+							return;
 						}
 					}
 				}
 
-				if (!perforceToolsAvailable && !hordeToolsAvailable)
-				{
-					return;
-				}
-
 				bool hasChanged = false;
 				int previousToolsCount = _tools.Count;
 
