UGS: Clarify type of user/pass required in password prompt dialog

When starting UGS, the Perforce password dialog can come up if it expired. But nowhere in the dialog does it tell for what service the user/pass is for. Simply mentioning "Perforce" makes this more clear.

#rnx

diff --git a/Engine/Source/Programs/UnrealGameSync/UnrealGameSync/PerforceModalTask.cs b/Engine/Source/Programs/UnrealGameSync/UnrealGameSync/PerforceModalTask.cs
--- a/Engine/Source/Programs/UnrealGameSync/UnrealGameSync/PerforceModalTask.cs
+++ b/Engine/Source/Programs/UnrealGameSync/UnrealGameSync/PerforceModalTask.cs
@@ -82,11 +82,11 @@ namespace UnrealGameSync
 						string passwordPrompt;
 						if (String.IsNullOrEmpty(password))
 						{
-							passwordPrompt = $"Enter the password for user '{perforceSettings.UserName}' on server '{perforceSettings.ServerAndPort}'.";
+							passwordPrompt = $"Enter the Perforce password for user '{perforceSettings.UserName}' on server '{perforceSettings.ServerAndPort}'.";
 						}
 						else
 						{
-							passwordPrompt = $"Authentication failed. Enter the password for user '{perforceSettings.UserName}' on server '{perforceSettings.ServerAndPort}'.";
+							passwordPrompt = $"Perforce authentication failed. Enter the password for user '{perforceSettings.UserName}' on server '{perforceSettings.ServerAndPort}'.";
 						}
 
 						using PasswordWindow passwordWindow = new PasswordWindow(passwordPrompt, password ?? String.Empty);
