Fix CA1508 'record.Action != FileAction.Branch' is always 'true'.
#rnx

diff --git a/Engine/Source/Programs/UnrealGameSync/UnrealGameSyncShared/WorkspaceUpdate.cs b/Engine/Source/Programs/UnrealGameSync/UnrealGameSyncShared/WorkspaceUpdate.cs
--- a/Engine/Source/Programs/UnrealGameSync/UnrealGameSyncShared/WorkspaceUpdate.cs
+++ b/Engine/Source/Programs/UnrealGameSync/UnrealGameSyncShared/WorkspaceUpdate.cs
@@ -841,7 +841,7 @@ namespace UnrealGameSync
 								OpenedRecord record = response.Data;
 								if (!String.IsNullOrEmpty(record.DepotFile) && !String.IsNullOrEmpty(record.ClientFile))
 								{
-									if (record.Action != FileAction.Add || record.Action != FileAction.Branch || record.Action != FileAction.MoveAdd)
+									if (record.Action != FileAction.Add && record.Action != FileAction.Branch && record.Action != FileAction.MoveAdd)
 									{
 										string relativePath = PerforceUtils.GetClientRelativePath(record.ClientFile);
 										syncFiles.Add(new SyncFile(record.DepotFile, relativePath, 0));
