Add ability to opt in to always deleting files

diff --git a/Engine/Source/Programs/UnrealGameSync/UnrealGameSync/Controls/WorkspaceControl.cs b/Engine/Source/Programs/UnrealGameSync/UnrealGameSync/Controls/WorkspaceControl.cs
--- a/Engine/Source/Programs/UnrealGameSync/UnrealGameSync/Controls/WorkspaceControl.cs
+++ b/Engine/Source/Programs/UnrealGameSync/UnrealGameSync/Controls/WorkspaceControl.cs
@@ -1070,7 +1070,12 @@ namespace UnrealGameSync
 
 				FlashWindow(ParentForm.Handle, true);
 
-				if (context.DeleteFiles.Count > 2000)
+				if (_settings.Global.AlwaysDeleteFiles)
+				{
+					StartWorkspaceUpdate(context, _updateCallback);
+					return;
+				}
+				else if (context.DeleteFiles.Count > 2000)
 				{
 					if (MessageBox.Show($"{context.DeleteFiles.Count:n0} files will be removed from your workspace by the current sync filter. Would you like to continue?", "Delete Files", MessageBoxButtons.OKCancel) == DialogResult.OK)
 					{
@@ -4792,6 +4797,7 @@ namespace UnrealGameSync
 		{
 			OptionsContextMenu_AutoResolveConflicts.Checked = _settings.Global.AutoResolveConflicts;
 			OptionsContextMenu_AlwaysClobberFiles.Checked = _settings.Global.AlwaysClobberFiles;
+			OptionsContextMenu_AlwaysDeleteFiles.Checked = _settings.Global.AlwaysDeleteFiles;
 
 			OptionsContextMenu_SyncPrecompiledBinaries.DropDownItems.Clear();
 
@@ -5302,6 +5308,13 @@ namespace UnrealGameSync
 			_settings.Save(_logger);
 		}
 
+		private void OptionsContextMenu_AlwaysDeleteFiles_Click(object sender, EventArgs e)
+		{
+			OptionsContextMenu_AlwaysDeleteFiles.Checked ^= true;
+			_settings.Global.AlwaysDeleteFiles = OptionsContextMenu_AlwaysDeleteFiles.Checked;
+			_settings.Save(_logger);
+		}
+
 		private void OptionsContextMenu_EditorArguments_Click(object sender, EventArgs e)
 		{
 			ModifyEditorArguments();
diff --git a/Engine/Source/Programs/UnrealGameSync/UnrealGameSyncShared/GlobalSettings.cs b/Engine/Source/Programs/UnrealGameSync/UnrealGameSyncShared/GlobalSettings.cs
--- a/Engine/Source/Programs/UnrealGameSync/UnrealGameSyncShared/GlobalSettings.cs
+++ b/Engine/Source/Programs/UnrealGameSync/UnrealGameSyncShared/GlobalSettings.cs
@@ -57,6 +57,7 @@ namespace UnrealGameSync
 		public FilterSettings Filter { get; set; } = new FilterSettings();
 		public bool AutoResolveConflicts { get; set; } = true;
 		public bool AlwaysClobberFiles { get; set; } = false;
+		public bool AlwaysDeleteFiles { get; set; } = false;
 	}
 
 	public class GlobalSettingsFile
