Parallelize file exists calls. Saves ~1:20 on 3m files
#rnx

diff --git a/Engine/Source/Programs/UnrealGameSync/UnrealGameSyncShared/WorkspaceUpdate.cs b/Engine/Source/Programs/UnrealGameSync/UnrealGameSyncShared/WorkspaceUpdate.cs
--- a/Engine/Source/Programs/UnrealGameSync/UnrealGameSyncShared/WorkspaceUpdate.cs
+++ b/Engine/Source/Programs/UnrealGameSync/UnrealGameSyncShared/WorkspaceUpdate.cs
@@ -850,26 +850,32 @@ namespace UnrealGameSync
 							}
 
 							// Enumerate all the files to be synced. NOTE: depotPath is escaped, whereas clientPath is not.
+							List<string> syncRelativePaths = new();
 							foreach (SyncFile syncRecord in syncFiles)
 							{
 								if (filter.Matches(syncRecord.RelativePath))
 								{
 									syncTree.IncludeFile(PerforceUtils.EscapePath(syncRecord.RelativePath), syncRecord.Size, logger);
 									syncDepotPaths.Add(syncRecord.DepotFile);
+									syncRelativePaths.Add(syncRecord.RelativePath);
 									requiredFreeSpace += syncRecord.Size;
-
-									// If the file exists the required free space can be reduced as those bytes will be replaced.
-									FileInfo localFileInfo = FileReference.Combine(project.LocalRootPath, syncRecord.RelativePath).ToFileInfo();
-									if (localFileInfo.Exists)
-									{
-										requiredFreeSpace -= localFileInfo.Length;
-									}
 								}
 								else
 								{
 									syncTree.ExcludeFile(PerforceUtils.EscapePath(syncRecord.RelativePath));
 								}
 							}
+
+							Parallel.ForEach(syncRelativePaths, syncRelativePath =>
+							{
+								// If the file exists the required free space can be reduced as those bytes will be replaced.
+								FileInfo localFileInfo = FileReference.Combine(project.LocalRootPath, syncRelativePath).ToFileInfo();
+								if (localFileInfo.Exists)
+								{
+									Interlocked.Add(ref requiredFreeSpace, -localFileInfo.Length);
+								}
+							});
+
 						}
 
 						try
