Add a Clear option to the log panel
#rnx

diff --git a/Engine/Source/Programs/UnrealGameSync/UnrealGameSync/Controls/LogControl.cs b/Engine/Source/Programs/UnrealGameSync/UnrealGameSync/Controls/LogControl.cs
--- a/Engine/Source/Programs/UnrealGameSync/UnrealGameSync/Controls/LogControl.cs
+++ b/Engine/Source/Programs/UnrealGameSync/UnrealGameSync/Controls/LogControl.cs
@@ -138,6 +138,8 @@ namespace UnrealGameSync
 			ContextMenuStrip.Items.Add("Copy", null, new EventHandler(ContextMenu_CopySelection));
 			ContextMenuStrip.Items.Add("-");
 			ContextMenuStrip.Items.Add("Select All", null, new EventHandler(ContextMenu_SelectAll));
+			ContextMenuStrip.Items.Add("-");
+			ContextMenuStrip.Items.Add("Clear", null, new EventHandler(ContextMenu_Clear));
 
 			SetStyle(ControlStyles.Selectable, true);
 		}
@@ -247,6 +249,7 @@ namespace UnrealGameSync
 				}
 				catch { }
 			}
+			Invalidate();
 		}
 
 		public void ScrollToEnd()
@@ -398,6 +401,11 @@ namespace UnrealGameSync
 			SelectAll();
 		}
 
+		private void ContextMenu_Clear(object? sender, EventArgs e)
+		{
+			Clear();
+		}
+
 		public void SelectAll()
 		{
 			if (_lines.Count > 0)
