Allow the client root to be specified on the commandline

diff --git a/Engine/Source/Programs/UnrealGameSync/UnrealGameSyncCmd/Program.cs b/Engine/Source/Programs/UnrealGameSync/UnrealGameSyncCmd/Program.cs
--- a/Engine/Source/Programs/UnrealGameSync/UnrealGameSyncCmd/Program.cs
+++ b/Engine/Source/Programs/UnrealGameSync/UnrealGameSyncCmd/Program.cs
@@ -183,6 +183,9 @@ namespace UnrealGameSyncCmd
 
 			[CommandLine("-Project=")]
 			public string? ProjectName { get; set; }
+
+			[CommandLine("-ClientRoot=")]
+			public string? ClientRoot { get; set; }
 		}
 
 		class UpdateState
@@ -619,7 +622,7 @@ namespace UnrealGameSyncCmd
 				StreamRecord stream = streamResponse.Data;
 
 				// Get the new directory for the client
-				DirectoryReference clientDir = DirectoryReference.Combine(DirectoryReference.GetCurrentDirectory(), stream.Stream.Replace('/', '+'));
+				DirectoryReference clientDir = DirectoryReference.FromString(options.ClientRoot) ?? DirectoryReference.Combine(DirectoryReference.GetCurrentDirectory(), stream.Stream.Replace('/', '+'));
 				DirectoryReference.CreateDirectory(clientDir);
 
 				// Make up a new client name 
