Set and persist the automation port setting (lost in CL# 30118307)

diff --git a/Engine/Source/Programs/UnrealGameSync/UnrealGameSync/Forms/ApplicationSettingsWindow.cs b/Engine/Source/Programs/UnrealGameSync/UnrealGameSync/Forms/ApplicationSettingsWindow.cs
--- a/Engine/Source/Programs/UnrealGameSync/UnrealGameSync/Forms/ApplicationSettingsWindow.cs
+++ b/Engine/Source/Programs/UnrealGameSync/UnrealGameSync/Forms/ApplicationSettingsWindow.cs
@@ -275,7 +275,8 @@ namespace UnrealGameSync
 
 			if (!String.Equals(launcherSettings.HordeServer, originalLauncherSettings.HordeServer, StringComparison.OrdinalIgnoreCase) ||
 				!String.Equals(launcherSettings.PerforceServerAndPort, originalLauncherSettings.PerforceServerAndPort, StringComparison.OrdinalIgnoreCase) ||
-				!String.Equals(launcherSettings.PerforceUserName, originalLauncherSettings.PerforceUserName, StringComparison.OrdinalIgnoreCase))
+				!String.Equals(launcherSettings.PerforceUserName, originalLauncherSettings.PerforceUserName, StringComparison.OrdinalIgnoreCase) ||
+				automationPortNumber != _initialAutomationPortNumber)
 			{
 				if (result == Result.Ok)
 				{
