Update the Current state when calling Modify

diff --git a/Engine/Source/Programs/UnrealGameSync/UnrealGameSyncShared/WorkspaceState.cs b/Engine/Source/Programs/UnrealGameSync/UnrealGameSyncShared/WorkspaceState.cs
--- a/Engine/Source/Programs/UnrealGameSync/UnrealGameSyncShared/WorkspaceState.cs
+++ b/Engine/Source/Programs/UnrealGameSync/UnrealGameSyncShared/WorkspaceState.cs
@@ -281,7 +281,8 @@ namespace UnrealGameSync
 						stream.Write(Utility.SerializeJson(state));
 						stream.SetLength(stream.Position);
 
-						return new ReadOnlyWorkspaceState(state);
+						Current = new ReadOnlyWorkspaceState(state);
+						return Current;
 					}
 				}
 				catch (IOException ex) when ((ex.HResult & 0xffff) == ERROR_SHARING_VIOLATION)
