Add -IgnoreExistingClients to init a workspace even if there is a workspace that covers the current path
Remove some dead code
#rnx

diff --git a/Engine/Source/Programs/UnrealGameSync/UnrealGameSync/Controls/WorkspaceControl.cs b/Engine/Source/Programs/UnrealGameSync/UnrealGameSync/Controls/WorkspaceControl.cs
--- a/Engine/Source/Programs/UnrealGameSync/UnrealGameSync/Controls/WorkspaceControl.cs
+++ b/Engine/Source/Programs/UnrealGameSync/UnrealGameSync/Controls/WorkspaceControl.cs
@@ -365,7 +365,6 @@ namespace UnrealGameSync
 
 			FileReference projectLogBaseName = FileReference.Combine(_workspaceDataFolder, "sync.log");
 
-			ILogger perforceLogger = _serviceProvider.GetRequiredService<ILogger<PerforceMonitor>>();
 			_perforceMonitor = new PerforceMonitor(perforceClientSettings, openProjectInfo.ProjectInfo, openProjectInfo.LatestProjectConfigFile, openProjectInfo.ProjectInfo.CacheFolder, openProjectInfo.LocalConfigFiles, _serviceProvider);
 
 			_perforceMonitor.OnUpdate += UpdateBuildListCallback;
@@ -373,7 +372,6 @@ namespace UnrealGameSync
 			_perforceMonitor.OnStreamChange += StreamChanged;
 			_perforceMonitor.OnLoginExpired += LoginExpired;
 
-			ILogger eventLogger = _serviceProvider.GetRequiredService<ILogger<EventMonitor>>();
 			_eventMonitor = new EventMonitor(_apiUrl, PerforceUtils.GetClientOrDepotDirectoryName(SelectedProjectIdentifier), openProjectInfo.PerforceSettings.UserName, _serviceProvider);
 			_eventMonitor.OnUpdatesReady += UpdateReviewsCallback;
 
diff --git a/Engine/Source/Programs/UnrealGameSync/UnrealGameSyncCmd/Program.cs b/Engine/Source/Programs/UnrealGameSync/UnrealGameSyncCmd/Program.cs
--- a/Engine/Source/Programs/UnrealGameSync/UnrealGameSyncCmd/Program.cs
+++ b/Engine/Source/Programs/UnrealGameSync/UnrealGameSyncCmd/Program.cs
@@ -186,6 +186,9 @@ namespace UnrealGameSyncCmd
 
 			[CommandLine("-ClientRoot=")]
 			public string? ClientRoot { get; set; }
+
+			[CommandLine("-IgnoreExistingClients")]
+			public bool IgnoreExistingClients { get; set; }
 		}
 
 		class UpdateState
@@ -628,17 +631,20 @@ namespace UnrealGameSyncCmd
 				// Make up a new client name 
 				string clientName = options.ClientName ?? Regex.Replace($"{perforce.Settings.UserName}_{hostName}_{stream.Stream.Trim('/')}", "[^0-9a-zA-Z_.-]", "+");
 
-				// Check there are no existing clients under the current path
-				List<ClientsRecord> clients = await FindExistingClients(perforce, hostName, clientDir);
-				if (clients.Count > 0)
+				if (!options.IgnoreExistingClients)
 				{
-					if (clients.Count == 1 && clientName.Equals(clients[0].Name, StringComparison.OrdinalIgnoreCase) && clientDir == TryParseRoot(clients[0].Root))
+					// Check there are no existing clients under the current path
+					List<ClientsRecord> clients = await FindExistingClients(perforce, hostName, clientDir);
+					if (clients.Count > 0)
 					{
-						logger.LogInformation("Reusing existing client for {ClientDir} ({ClientName})", clientDir, options.ClientName);
-					}
-					else
-					{
-						throw new UserErrorException("Current directory is already within a Perforce workspace ({ClientName})", clients[0].Name);
+						if (clients.Count == 1 && clientName.Equals(clients[0].Name, StringComparison.OrdinalIgnoreCase) && clientDir == TryParseRoot(clients[0].Root))
+						{
+							logger.LogInformation("Reusing existing client for {ClientDir} ({ClientName})", clientDir, options.ClientName);
+						}
+						else
+						{
+							throw new UserErrorException("Current directory is already within a Perforce workspace ({ClientName})", clients[0].Name);
+						}
 					}
 				}
 
