UGS: Fix enumeration of tools from Horde. We no longer register HordeHttpClient as a standalone service; it has to be obtained through IHordeClient.

diff --git a/Engine/Source/Programs/UnrealGameSync/UnrealGameSync/ToolUpdateMonitor.cs b/Engine/Source/Programs/UnrealGameSync/UnrealGameSync/ToolUpdateMonitor.cs
--- a/Engine/Source/Programs/UnrealGameSync/UnrealGameSync/ToolUpdateMonitor.cs
+++ b/Engine/Source/Programs/UnrealGameSync/UnrealGameSync/ToolUpdateMonitor.cs
@@ -232,20 +232,19 @@ namespace UnrealGameSync
 						return;
 					}
 				}
-				using (HordeHttpClient? hordeHttpClient = _serviceProvider.GetService<HordeHttpClient>())
+				IHordeClient? hordeClient = _serviceProvider.GetService<IHordeClient>();
+				if (hordeClient != null)
 				{
-					if (hordeHttpClient != null)
+					using HordeHttpClient hordeHttpClient = hordeClient.CreateHttpClient();
+					try
 					{
-						try
-						{
-							await ReadHordeToolsAsync(hordeHttpClient, tools, cancellationToken);
-						}
-						catch (Exception ex) when (ex is not OperationCanceledException)
-						{
-							LastStatus = Tuple.Create(false, $"Error while polling Horde for available tools: {ex.Message}");
-							_logger.LogWarning(ex, "Error while polling Horde for available tools: {Message}", ex.Message);
-							return;
-						}
+						await ReadHordeToolsAsync(hordeHttpClient, tools, cancellationToken);
+					}
+					catch (Exception ex) when (ex is not OperationCanceledException)
+					{
+						LastStatus = Tuple.Create(false, $"Error while polling Horde for available tools: {ex.Message}");
+						_logger.LogWarning(ex, "Error while polling Horde for available tools: {Message}", ex.Message);
+						return;
 					}
 				}
 
@@ -548,9 +547,11 @@ namespace UnrealGameSync
 				}
 				else
 				{
-					using HordeHttpClient? hordeHttpClient = _serviceProvider.GetService<HordeHttpClient>();
-					if (hordeHttpClient != null)
+					IHordeClient? hordeClient = _serviceProvider.GetService<IHordeClient>();
+					if (hordeClient != null)
 					{
+						using HordeHttpClient hordeHttpClient = hordeClient.CreateHttpClient();
+
 						string[] fields = tool.Revision.Split(',');
 						if (fields.Length != 3)
 						{
