UGS: Only include all projects in generated solution if that option is enabled or if no project is selected, correctly generate a Game specific solution, and fix pathing to open it

#rnx

diff --git a/Engine/Source/Programs/UnrealGameSync/UnrealGameSync/Controls/WorkspaceControl.cs b/Engine/Source/Programs/UnrealGameSync/UnrealGameSync/Controls/WorkspaceControl.cs
--- a/Engine/Source/Programs/UnrealGameSync/UnrealGameSync/Controls/WorkspaceControl.cs
+++ b/Engine/Source/Programs/UnrealGameSync/UnrealGameSync/Controls/WorkspaceControl.cs
@@ -5411,31 +5411,50 @@ namespace UnrealGameSync
 
 		private void OpenSolution()
 		{
+			FileReference? solutionFileName = null;
+
 			string primaryProjectName = "UE4";
 			if (FileReference.Exists(FileReference.Combine(BranchDirectoryName, "UE5.sln")))
 			{
 				primaryProjectName = "UE5";
 			}
 
-			FileReference primaryProjectNameFileName = FileReference.Combine(BranchDirectoryName, "Engine", "Intermediate", "ProjectFiles", "PrimaryProjectName.txt");
-			if (!FileReference.Exists(primaryProjectNameFileName))
-			{
-				// Old path, needs to be maintained while projects prior to UE5.1 are supported
-				primaryProjectNameFileName = FileReference.Combine(BranchDirectoryName, "Engine", "Intermediate", "ProjectFiles", "MasterProjectName.txt");
-			}
-			if (FileReference.Exists(primaryProjectNameFileName))
+			FileReference primaryProjectPathFileName = FileReference.Combine(BranchDirectoryName, "Engine", "Intermediate", "ProjectFiles", "PrimaryProjectPath.txt");
+			if (FileReference.Exists(primaryProjectPathFileName))
 			{
 				try
 				{
-					primaryProjectName = FileReference.ReadAllText(primaryProjectNameFileName).Trim();
+					solutionFileName = new FileReference(FileReference.ReadAllText(primaryProjectPathFileName).Trim() + ".sln");
 				}
 				catch (Exception ex)
 				{
-					_logger.LogError(ex, "Unable to read '{File}'", primaryProjectNameFileName);
+					_logger.LogError(ex, "Unable to read '{File}'", primaryProjectPathFileName);
 				}
 			}
+			
+			if (solutionFileName == null)
+			{
+				FileReference primaryProjectNameFileName = FileReference.Combine(BranchDirectoryName, "Engine", "Intermediate", "ProjectFiles", "PrimaryProjectName.txt");
+				if (!FileReference.Exists(primaryProjectNameFileName))
+				{
+					// Old path, needs to be maintained while projects prior to UE5.1 are supported
+					primaryProjectNameFileName = FileReference.Combine(BranchDirectoryName, "Engine", "Intermediate", "ProjectFiles", "MasterProjectName.txt");
+				}
+				if (FileReference.Exists(primaryProjectNameFileName))
+				{
+					try
+					{
+						primaryProjectName = FileReference.ReadAllText(primaryProjectNameFileName).Trim();
+					}
+					catch (Exception ex)
+					{
+						_logger.LogError(ex, "Unable to read '{File}'", primaryProjectNameFileName);
+					}
+				}
+
+				solutionFileName = FileReference.Combine(BranchDirectoryName, primaryProjectName + ".sln");
+			}
 
-			FileReference solutionFileName = FileReference.Combine(BranchDirectoryName, primaryProjectName + ".sln");
 			if (!FileReference.Exists(solutionFileName))
 			{
 				MessageBox.Show(String.Format("Couldn't find solution at {0}", solutionFileName));
diff --git a/Engine/Source/Programs/UnrealGameSync/UnrealGameSyncShared/WorkspaceUpdate.cs b/Engine/Source/Programs/UnrealGameSync/UnrealGameSyncShared/WorkspaceUpdate.cs
--- a/Engine/Source/Programs/UnrealGameSync/UnrealGameSyncShared/WorkspaceUpdate.cs
+++ b/Engine/Source/Programs/UnrealGameSync/UnrealGameSyncShared/WorkspaceUpdate.cs
@@ -127,7 +127,7 @@ namespace UnrealGameSync
 			}
 			if (workspaceSettings.Filter.AllProjects ?? globalSettings.Filter.AllProjects ?? false)
 			{
-				options |= WorkspaceUpdateOptions.SyncAllProjects | WorkspaceUpdateOptions.IncludeAllProjectsInSolution;
+				options |= WorkspaceUpdateOptions.SyncAllProjects;
 			}
 			if (workspaceSettings.Filter.AllProjectsInSln ?? globalSettings.Filter.AllProjectsInSln ?? false)
 			{
@@ -1257,11 +1257,11 @@ namespace UnrealGameSync
 
 					StringBuilder commandLine = new StringBuilder();
 					commandLine.AppendFormat("\"{0}\"", FileReference.Combine(project.LocalRootPath, $"GenerateProjectFiles.{ShellScriptExt}"));
-					if ((Context.Options & WorkspaceUpdateOptions.SyncAllProjects) == 0 && (Context.Options & WorkspaceUpdateOptions.IncludeAllProjectsInSolution) == 0)
+					if ((Context.Options & WorkspaceUpdateOptions.IncludeAllProjectsInSolution) == 0)
 					{
 						if (project.LocalFileName.HasExtension(".uproject"))
 						{
-							commandLine.AppendFormat(" \"{0}\"", project.LocalFileName);
+							commandLine.AppendFormat(" -Project=\"{0}\" -Game", project.LocalFileName);
 						}
 					}
 					commandLine.Append(" -progress");
