UGS: Upgrade UGS to NET 8.

diff --git a/Engine/Source/Programs/UnrealGameSync/UnrealGameSync/Controls/BuildListControl.cs b/Engine/Source/Programs/UnrealGameSync/UnrealGameSync/Controls/BuildListControl.cs
--- a/Engine/Source/Programs/UnrealGameSync/UnrealGameSync/Controls/BuildListControl.cs
+++ b/Engine/Source/Programs/UnrealGameSync/UnrealGameSync/Controls/BuildListControl.cs
@@ -2,6 +2,7 @@
 
 using System;
 using System.ComponentModel;
+using System.Diagnostics.CodeAnalysis;
 using System.Drawing;
 using System.Runtime.InteropServices;
 using System.Windows.Forms;
@@ -69,7 +70,7 @@ namespace UnrealGameSync
 		{
 			if (!IsHandleCreated)
 			{
-				return Height / Font.Height;
+				return Height / Font!.Height;
 			}
 			return SendMessage(Handle, LvmGetcountperpage, IntPtr.Zero, IntPtr.Zero).ToInt32();
 		}
@@ -142,6 +143,7 @@ namespace UnrealGameSync
 
 		[Browsable(false)]
 		[DesignerSerializationVisibility(DesignerSerializationVisibility.Hidden)]
+		[AllowNull]
 		public override Font Font
 		{
 			get => base.Font;
diff --git a/Engine/Source/Programs/UnrealGameSync/UnrealGameSync/Controls/LogControl.cs b/Engine/Source/Programs/UnrealGameSync/UnrealGameSync/Controls/LogControl.cs
--- a/Engine/Source/Programs/UnrealGameSync/UnrealGameSync/Controls/LogControl.cs
+++ b/Engine/Source/Programs/UnrealGameSync/UnrealGameSync/Controls/LogControl.cs
@@ -555,7 +555,7 @@ namespace UnrealGameSync
 
 			if (e.Button == MouseButtons.Right)
 			{
-				ContextMenuStrip.Show(this, e.Location);
+				ContextMenuStrip!.Show(this, e.Location);
 			}
 		}
 
diff --git a/Engine/Source/Programs/UnrealGameSync/UnrealGameSync/Controls/WorkspaceControl.cs b/Engine/Source/Programs/UnrealGameSync/UnrealGameSync/Controls/WorkspaceControl.cs
--- a/Engine/Source/Programs/UnrealGameSync/UnrealGameSync/Controls/WorkspaceControl.cs
+++ b/Engine/Source/Programs/UnrealGameSync/UnrealGameSync/Controls/WorkspaceControl.cs
@@ -1058,7 +1058,7 @@ namespace UnrealGameSync
 
 			if (result == WorkspaceUpdateResult.FilesToResolve)
 			{
-				FlashWindow(ParentForm.Handle, true);
+				FlashWindow(ParentForm!.Handle, true);
 				MessageBox.Show(ParentForm, "You have files to resolve after syncing your workspace. Please check P4.");
 			}
 			else if (result == WorkspaceUpdateResult.FilesToDelete)
@@ -1066,7 +1066,7 @@ namespace UnrealGameSync
 				DesiredTaskbarState = Tuple.Create(TaskbarState.Paused, 0.0f);
 				_owner.UpdateProgress();
 
-				FlashWindow(ParentForm.Handle, true);
+				FlashWindow(ParentForm!.Handle, true);
 
 				if (_settings.Global.AlwaysDeleteFiles)
 				{
@@ -1161,7 +1161,7 @@ namespace UnrealGameSync
 			UpdateStatusPanel();
 			UpdateSyncActionCheckboxes();
 
-			FlashWindow(ParentForm.Handle, true);
+			FlashWindow(ParentForm!.Handle, true);
 
 			// Do this last because it may result in the control being disposed
 			if (result == WorkspaceUpdateResult.FailedToSyncLoginExpired)
@@ -1601,7 +1601,7 @@ namespace UnrealGameSync
 			// Remove everything that's in a custom column
 			foreach (ColumnHeader customColumn in _customColumns)
 			{
-				ConfigObject columnConfig = (ConfigObject)customColumn.Tag;
+				ConfigObject columnConfig = (ConfigObject)customColumn.Tag!;
 				foreach (string badgeName in columnConfig.GetValue("Badges", "").Split(','))
 				{
 					badgeNameToGroup.Remove(badgeName);
@@ -2034,7 +2034,7 @@ namespace UnrealGameSync
 				{
 					BuildList.DrawTrackedBackground(e.Graphics, e.Bounds);
 				}
-				else if (((ChangesRecord)e.Item.Tag).Number == _workspace.PendingChangeNumber)
+				else if (((ChangesRecord)e.Item.Tag!).Number == _workspace.PendingChangeNumber)
 				{
 					BuildList.DrawTrackedBackground(e.Graphics, e.Bounds);
 				}
@@ -3050,7 +3050,7 @@ namespace UnrealGameSync
 			{
 				foreach (ColumnHeader customColumn in _customColumns)
 				{
-					ConfigObject config = (ConfigObject)customColumn.Tag;
+					ConfigObject config = (ConfigObject)customColumn.Tag!;
 					if (config != null)
 					{
 						List<BadgeInfo> badges = new List<BadgeInfo>();
@@ -3239,7 +3239,7 @@ namespace UnrealGameSync
 				ListViewHitTestInfo hitTest = BuildList.HitTest(args.Location);
 				if (hitTest.Item != null)
 				{
-					ChangesRecord change = (ChangesRecord)hitTest.Item.Tag;
+					ChangesRecord change = (ChangesRecord)hitTest.Item.Tag!;
 					if (change != null)
 					{
 						if (change.Number == _workspace.CurrentChangeNumber)
@@ -3713,7 +3713,7 @@ namespace UnrealGameSync
 							{
 								if (item != null)
 								{
-									ChangesRecord summary = (ChangesRecord)item.Tag;
+									ChangesRecord summary = (ChangesRecord)item.Tag!;
 									if (summary != null)
 									{
 										// No need to keep on iterating if we're past the required CL :
@@ -4055,14 +4055,14 @@ namespace UnrealGameSync
 
 		public void ShowIssueDetails(IssueData issue)
 		{
-			IssueDetailsWindow.Show(ParentForm, _issueMonitor, _perforceSettings, GetServerTimeOffset(), issue, _serviceProvider, StreamName);
+			IssueDetailsWindow.Show(ParentForm!, _issueMonitor, _perforceSettings, GetServerTimeOffset(), issue, _serviceProvider, StreamName);
 		}
 
 		private void BuildHealthContextMenu_Browse_Click(object sender, EventArgs e)
 		{
 			string? defaultFilter = GetDefaultIssueFilter();
 			Dictionary<string, Func<IssueData, bool>> customFilters = GetCustomIssueFilters();
-			IssueBrowserWindow.Show(ParentForm, _issueMonitor, _perforceSettings, GetServerTimeOffset(), _serviceProvider, StreamName, customFilters, defaultFilter);
+			IssueBrowserWindow.Show(ParentForm!, _issueMonitor, _perforceSettings, GetServerTimeOffset(), _serviceProvider, StreamName, customFilters, defaultFilter);
 		}
 
 		private void BuildHealthContextMenu_Settings_Click(object sender, EventArgs e)
@@ -4333,7 +4333,7 @@ namespace UnrealGameSync
 			{
 				if (item != null)
 				{
-					ChangesRecord summary = (ChangesRecord)item.Tag;
+					ChangesRecord summary = (ChangesRecord)item.Tag!;
 					if (summary != null && summary.Number <= changeNumber)
 					{
 						item.Selected = true;
@@ -4378,7 +4378,7 @@ namespace UnrealGameSync
 					}
 					else
 					{
-						ChangesRecord change = (ChangesRecord)hitTest.Item.Tag;
+						ChangesRecord change = (ChangesRecord)hitTest.Item.Tag!;
 						if (_workspace.PendingChangeNumber == change.Number)
 						{
 							Rectangle subItemRect = hitTest.Item.SubItems[StatusColumn.Index].Bounds;
@@ -4461,7 +4461,7 @@ namespace UnrealGameSync
 						{
 							if (customColumn.Index < hitTest.Item.SubItems.Count && hitTest.Item.SubItems[customColumn.Index] == hitTest.SubItem)
 							{
-								ChangeLayoutInfo layoutInfo = GetChangeLayoutInfo((ChangesRecord)hitTest.Item.Tag);
+								ChangeLayoutInfo layoutInfo = GetChangeLayoutInfo((ChangesRecord)hitTest.Item.Tag!);
 
 								List<BadgeInfo>? badges;
 								if (layoutInfo._customBadges.TryGetValue(customColumn.Text, out badges) && badges.Count > 0)
@@ -4652,7 +4652,7 @@ namespace UnrealGameSync
 					ChangeLayoutInfo layoutInfo = GetChangeLayoutInfo(changesRecord);
 					if (layoutInfo._descriptionBadges.Count > 0)
 					{
-						Point listLocation = GetBadgeListLocation(layoutInfo._descriptionBadges, hitTest.SubItem.Bounds, HorizontalAlign.Right, VerticalAlignment.Middle);
+						Point listLocation = GetBadgeListLocation(layoutInfo._descriptionBadges, hitTest.SubItem!.Bounds, HorizontalAlign.Right, VerticalAlignment.Middle);
 						newHoverBadgeUniqueId = HitTestBadge(e.Location, layoutInfo._descriptionBadges, listLocation)?._uniqueId;
 					}
 				}
@@ -4661,7 +4661,7 @@ namespace UnrealGameSync
 					ChangeLayoutInfo layoutInfo = GetChangeLayoutInfo(changesRecord);
 					if (layoutInfo._buildBadges.Count > 0)
 					{
-						Point buildListLocation = GetBadgeListLocation(layoutInfo._buildBadges, hitTest.SubItem.Bounds, HorizontalAlign.Center, VerticalAlignment.Middle);
+						Point buildListLocation = GetBadgeListLocation(layoutInfo._buildBadges, hitTest.SubItem!.Bounds, HorizontalAlign.Center, VerticalAlignment.Middle);
 						buildListLocation.X = Math.Max(buildListLocation.X, hitTest.SubItem.Bounds.Left);
 
 						BadgeInfo? badge = HitTestBadge(e.Location, layoutInfo._buildBadges, buildListLocation);
@@ -4689,7 +4689,7 @@ namespace UnrealGameSync
 					List<BadgeInfo>? badges;
 					if (layoutInfo._customBadges.TryGetValue(column.Text, out badges) && badges.Count > 0)
 					{
-						Point listLocation = GetBadgeListLocation(badges, hitTest.SubItem.Bounds, HorizontalAlign.Center, VerticalAlignment.Middle);
+						Point listLocation = GetBadgeListLocation(badges, hitTest.SubItem!.Bounds, HorizontalAlign.Center, VerticalAlignment.Middle);
 						newHoverBadgeUniqueId = HitTestBadge(e.Location, badges, listLocation)?._uniqueId;
 					}
 				}
@@ -4914,10 +4914,10 @@ namespace UnrealGameSync
 			_buildFont = BuildList.Font;
 
 			_selectedBuildFont?.Dispose();
-			_selectedBuildFont = new Font(_buildFont, FontStyle.Bold);
+			_selectedBuildFont = new Font(_buildFont!, FontStyle.Bold);
 
 			_badgeFont?.Dispose();
-			_badgeFont = new Font(_buildFont.FontFamily, _buildFont.SizeInPoints - 2, FontStyle.Bold);
+			_badgeFont = new Font(_buildFont!.FontFamily, _buildFont.SizeInPoints - 2, FontStyle.Bold);
 		}
 
 		public void ExecCommand(string description, string statusText, string fileName, string arguments, string workingDir, bool useLogWindow)
@@ -5184,11 +5184,11 @@ namespace UnrealGameSync
 		private void BuildList_ItemMouseHover(object sender, ListViewItemMouseHoverEventArgs args)
 		{
 			Point clientPoint = BuildList.PointToClient(Cursor.Position);
-			if (args.Item.SubItems.Count >= 6)
+			if (args.Item != null && args.Item.SubItems.Count >= 6)
 			{
 				if (args.Item.Bounds.Contains(clientPoint))
 				{
-					ChangesRecord change = (ChangesRecord)args.Item.Tag;
+					ChangesRecord change = (ChangesRecord)args.Item.Tag!;
 					if (change == null)
 					{
 						return;
@@ -5599,7 +5599,7 @@ namespace UnrealGameSync
 
 			for (int idx = 0; idx < BuildList.Items.Count; idx++)
 			{
-				ChangesRecord change = (ChangesRecord)BuildList.Items[idx].Tag;
+				ChangesRecord? change = (ChangesRecord?)BuildList.Items[idx].Tag;
 				if (change != null)
 				{
 					EventSummary? summary = _eventMonitor.GetSummaryForChange(change.Number);
@@ -5716,7 +5716,7 @@ namespace UnrealGameSync
 			string[] combinedSyncFilter = UserSettings.GetCombinedSyncFilter(GetSyncCategories(), _settings.Global.Filter, _workspaceSettings.Filter, _perforceMonitor.LatestPerforceConfigSection());
 			List<string> syncPaths = WorkspaceUpdate.GetSyncPaths(_workspace.Project, _workspaceSettings.Filter.AllProjects ?? _settings.Global.Filter.AllProjects ?? false, combinedSyncFilter);
 
-			bool performedClean = CleanWorkspaceWindow.DoClean(ParentForm, _perforceSettings, BranchDirectoryName, _workspace.Project.ClientRootPath, syncPaths, extraSafeToDeleteFolders.Split('\n'), extraSafeToDeleteExtensions.Split('\n'), _serviceProvider.GetRequiredService<ILogger<CleanWorkspaceWindow>>());
+			bool performedClean = CleanWorkspaceWindow.DoClean(ParentForm!, _perforceSettings, BranchDirectoryName, _workspace.Project.ClientRootPath, syncPaths, extraSafeToDeleteFolders.Split('\n'), extraSafeToDeleteExtensions.Split('\n'), _serviceProvider.GetRequiredService<ILogger<CleanWorkspaceWindow>>());
 
 			if (performedClean)
 			{
@@ -6125,7 +6125,7 @@ namespace UnrealGameSync
 
 			if (args.Control && args.KeyCode == Keys.C && BuildList.SelectedItems.Count > 0)
 			{
-				int selectedChange = ((ChangesRecord)BuildList.SelectedItems[0].Tag).Number;
+				int selectedChange = ((ChangesRecord)BuildList.SelectedItems[0].Tag!).Number;
 				Clipboard.SetText(String.Format("{0}", selectedChange));
 			}
 		}
@@ -6396,7 +6396,7 @@ namespace UnrealGameSync
 				{
 					if (selectedItem != null)
 					{
-						ChangesRecord change = (ChangesRecord)selectedItem.Tag;
+						ChangesRecord change = (ChangesRecord)selectedItem.Tag!;
 						changeNumberToBisectState[change.Number] = BisectState.Include;
 					}
 				}
diff --git a/Engine/Source/Programs/UnrealGameSync/UnrealGameSync/Forms/ArgumentsWindow.cs b/Engine/Source/Programs/UnrealGameSync/UnrealGameSync/Forms/ArgumentsWindow.cs
--- a/Engine/Source/Programs/UnrealGameSync/UnrealGameSync/Forms/ArgumentsWindow.cs
+++ b/Engine/Source/Programs/UnrealGameSync/UnrealGameSync/Forms/ArgumentsWindow.cs
@@ -129,7 +129,11 @@ namespace UnrealGameSync
 		private void ArgumentsList_MouseClick(object sender, MouseEventArgs e)
 		{
 			ListViewHitTestInfo info = ArgumentsList.HitTest(e.Location);
-			if (info.Item.Index == ArgumentsList.Items.Count - 1)
+			if (info.Item == null)
+			{
+				return;
+			}
+			else if (info.Item.Index == ArgumentsList.Items.Count - 1)
 			{
 				ListViewItem newItem = new ListViewItem();
 				newItem.Checked = true;
diff --git a/Engine/Source/Programs/UnrealGameSync/UnrealGameSync/Forms/CleanWorkspaceWindow.cs b/Engine/Source/Programs/UnrealGameSync/UnrealGameSync/Forms/CleanWorkspaceWindow.cs
--- a/Engine/Source/Programs/UnrealGameSync/UnrealGameSync/Forms/CleanWorkspaceWindow.cs
+++ b/Engine/Source/Programs/UnrealGameSync/UnrealGameSync/Forms/CleanWorkspaceWindow.cs
@@ -478,31 +478,31 @@ namespace UnrealGameSync
 
 		private void FolderContextMenu_SelectAll_Click(object sender, EventArgs e)
 		{
-			TreeNode node = (TreeNode)FolderContextMenu.Tag;
+			TreeNode node = (TreeNode)FolderContextMenu.Tag!;
 			SetSelected(node, SelectionType.All);
 		}
 
 		private void FolderContextMenu_SelectSafeToDelete_Click(object sender, EventArgs e)
 		{
-			TreeNode node = (TreeNode)FolderContextMenu.Tag;
+			TreeNode node = (TreeNode)FolderContextMenu.Tag!;
 			SetSelected(node, SelectionType.SafeToDelete);
 		}
 
 		private void FolderContextMenu_SelectEmptyFolder_Click(object sender, EventArgs e)
 		{
-			TreeNode node = (TreeNode)FolderContextMenu.Tag;
+			TreeNode node = (TreeNode)FolderContextMenu.Tag!;
 			SetSelected(node, SelectionType.Empty);
 		}
 
 		private void FolderContextMenu_SelectNone_Click(object sender, EventArgs e)
 		{
-			TreeNode node = (TreeNode)FolderContextMenu.Tag;
+			TreeNode node = (TreeNode)FolderContextMenu.Tag!;
 			SetSelected(node, SelectionType.None);
 		}
 
 		private void FolderContextMenu_OpenWithExplorer_Click(object sender, EventArgs e)
 		{
-			TreeNode node = (TreeNode)FolderContextMenu.Tag;
+			TreeNode node = (TreeNode)FolderContextMenu.Tag!;
 			TreeNodeData nodeData = (TreeNodeData)node.Tag;
 
 			if (nodeData._folder != null)
diff --git a/Engine/Source/Programs/UnrealGameSync/UnrealGameSync/Forms/ClobberWindow.cs b/Engine/Source/Programs/UnrealGameSync/UnrealGameSync/Forms/ClobberWindow.cs
--- a/Engine/Source/Programs/UnrealGameSync/UnrealGameSync/Forms/ClobberWindow.cs
+++ b/Engine/Source/Programs/UnrealGameSync/UnrealGameSync/Forms/ClobberWindow.cs
@@ -62,7 +62,7 @@ namespace UnrealGameSync
 			{
 				if (item != null)
 				{
-					_filesToClobber[(string)item.Tag] = item.Checked;
+					_filesToClobber[(string)item.Tag!] = item.Checked;
 				}
 			}
 		}
diff --git a/Engine/Source/Programs/UnrealGameSync/UnrealGameSync/Forms/DeleteWindow.cs b/Engine/Source/Programs/UnrealGameSync/UnrealGameSync/Forms/DeleteWindow.cs
--- a/Engine/Source/Programs/UnrealGameSync/UnrealGameSync/Forms/DeleteWindow.cs
+++ b/Engine/Source/Programs/UnrealGameSync/UnrealGameSync/Forms/DeleteWindow.cs
@@ -46,7 +46,7 @@ namespace UnrealGameSync
 		{
 			foreach (ListViewItem? item in FileList.Items)
 			{
-				_filesToDelete[(string)item!.Tag] = item.Checked;
+				_filesToDelete[(string)item!.Tag!] = item.Checked;
 			}
 		}
 	}
diff --git a/Engine/Source/Programs/UnrealGameSync/UnrealGameSync/Forms/IssueBrowserWindow.cs b/Engine/Source/Programs/UnrealGameSync/UnrealGameSync/Forms/IssueBrowserWindow.cs
--- a/Engine/Source/Programs/UnrealGameSync/UnrealGameSync/Forms/IssueBrowserWindow.cs
+++ b/Engine/Source/Programs/UnrealGameSync/UnrealGameSync/Forms/IssueBrowserWindow.cs
@@ -190,7 +190,7 @@ namespace UnrealGameSync
 						}
 
 						ListViewItem existingItem = IssueListView.Items[itemIdx];
-						IssueData existingIssue = (IssueData)existingItem.Tag;
+						IssueData existingIssue = (IssueData)existingItem.Tag!;
 						if (existingIssue == null || existingIssue.Id < issue.Id)
 						{
 							IssueList_InsertItem(itemIdx, issue, midnight);
@@ -297,7 +297,7 @@ namespace UnrealGameSync
 			ListViewHitTestInfo hitTest = IssueListView.HitTest(e.Location);
 			if (hitTest.Item != null)
 			{
-				IssueData issue = (IssueData)hitTest.Item.Tag;
+				IssueData issue = (IssueData)hitTest.Item.Tag!;
 				ShowIssue(issue);
 			}
 		}
@@ -363,7 +363,7 @@ namespace UnrealGameSync
 
 		private void ShowIssue(IssueData issue)
 		{
-			IssueDetailsWindow.Show(Owner, _issueMonitor, _perforceSettings, _serverTimeOffset, issue, _serviceProvider, _currentStream);
+			IssueDetailsWindow.Show(Owner!, _issueMonitor, _perforceSettings, _serverTimeOffset, issue, _serviceProvider, _currentStream);
 		}
 
 		private void IssueListView_DrawColumnHeader(object sender, DrawListViewColumnHeaderEventArgs e)
@@ -415,7 +415,7 @@ namespace UnrealGameSync
 				return;
 			}
 
-			IssueData issue = (IssueData)e.Item.Tag;
+			IssueData issue = (IssueData)e.Item.Tag!;
 			if (e.ColumnIndex == IconHeader.Index)
 			{
 				if (!issue.ResolvedAt.HasValue && issue.FixChange == 0)
diff --git a/Engine/Source/Programs/UnrealGameSync/UnrealGameSync/Forms/IssueDetailsWindow.cs b/Engine/Source/Programs/UnrealGameSync/UnrealGameSync/Forms/IssueDetailsWindow.cs
--- a/Engine/Source/Programs/UnrealGameSync/UnrealGameSync/Forms/IssueDetailsWindow.cs
+++ b/Engine/Source/Programs/UnrealGameSync/UnrealGameSync/Forms/IssueDetailsWindow.cs
@@ -1000,7 +1000,7 @@ namespace UnrealGameSync
 						PerforceChangeDetailsWithDescribeRecord? details;
 						_perforceWorker!.TryGetChangeDetails(change.Number, out details);
 
-						StatusLineListViewWidget typeWidget = (StatusLineListViewWidget)item.SubItems[TypeHeader.Index].Tag;
+						StatusLineListViewWidget typeWidget = (StatusLineListViewWidget)item.SubItems[TypeHeader.Index].Tag!;
 						UpdateChangeTypeWidget(typeWidget, details);
 
 						BuildListView.RedrawItems(item.Index, item.Index, true);
diff --git a/Engine/Source/Programs/UnrealGameSync/UnrealGameSync/Forms/IssueFixedWindow.cs b/Engine/Source/Programs/UnrealGameSync/UnrealGameSync/Forms/IssueFixedWindow.cs
--- a/Engine/Source/Programs/UnrealGameSync/UnrealGameSync/Forms/IssueFixedWindow.cs
+++ b/Engine/Source/Programs/UnrealGameSync/UnrealGameSync/Forms/IssueFixedWindow.cs
@@ -262,7 +262,7 @@ namespace UnrealGameSync
 
 		private void ChangesContextMenu_MoreInfo_Click(object sender, EventArgs e)
 		{
-			DescribeRecord record = (DescribeRecord)ChangesListContextMenu.Tag;
+			DescribeRecord record = (DescribeRecord)ChangesListContextMenu.Tag!;
 			Program.SpawnP4Vc(String.Format("{0} change {1}", _perforceSettings.GetArgumentsForExternalProgram(true), record.Number));
 		}
 
diff --git a/Engine/Source/Programs/UnrealGameSync/UnrealGameSync/Forms/MainWindow.cs b/Engine/Source/Programs/UnrealGameSync/UnrealGameSync/Forms/MainWindow.cs
--- a/Engine/Source/Programs/UnrealGameSync/UnrealGameSync/Forms/MainWindow.cs
+++ b/Engine/Source/Programs/UnrealGameSync/UnrealGameSync/Forms/MainWindow.cs
@@ -739,17 +739,17 @@ namespace UnrealGameSync
 			StatusLine actionLine = new StatusLine();
 			actionLine.AddLink("Retry", FontStyle.Bold | FontStyle.Underline, () =>
 			{
-				BeginInvoke(new MethodInvoker(() => TryOpenProject(project, TabControl.FindTabIndex(errorPanel))));
+				BeginInvoke(new System.Windows.Forms.MethodInvoker(() => TryOpenProject(project, TabControl.FindTabIndex(errorPanel))));
 			});
 			actionLine.AddText(" | ");
 			actionLine.AddLink("Settings", FontStyle.Bold | FontStyle.Underline, () =>
 			{
-				BeginInvoke(new MethodInvoker(() => EditSelectedProject(errorPanel)));
+				BeginInvoke(new System.Windows.Forms.MethodInvoker(() => EditSelectedProject(errorPanel)));
 			});
 			actionLine.AddText(" | ");
 			actionLine.AddLink("Close", FontStyle.Bold | FontStyle.Underline, () =>
 			{
-				BeginInvoke(new MethodInvoker(() => TabControl.RemoveTab(TabControl.FindTabIndex(errorPanel))));
+				BeginInvoke(new System.Windows.Forms.MethodInvoker(() => TabControl.RemoveTab(TabControl.FindTabIndex(errorPanel))));
 			});
 			lines.Add(actionLine);
 
diff --git a/Engine/Source/Programs/UnrealGameSync/UnrealGameSync/Forms/ModifyBuildStepsWindow.cs b/Engine/Source/Programs/UnrealGameSync/UnrealGameSync/Forms/ModifyBuildStepsWindow.cs
--- a/Engine/Source/Programs/UnrealGameSync/UnrealGameSync/Forms/ModifyBuildStepsWindow.cs
+++ b/Engine/Source/Programs/UnrealGameSync/UnrealGameSync/Forms/ModifyBuildStepsWindow.cs
@@ -74,9 +74,9 @@ namespace UnrealGameSync
 			{
 				if (item != null)
 				{
-					using BuildStepWindow editStep = new BuildStepWindow((BuildStep)item.Tag, _targetNames, _baseDirectory, _variables);
+					using BuildStepWindow editStep = new BuildStepWindow((BuildStep)item.Tag!, _targetNames, _baseDirectory, _variables);
 					editStep.ShowDialog();
-					item.Text = ((BuildStep)item.Tag).Description;
+					item.Text = ((BuildStep)item.Tag!).Description;
 					break;
 				}
 			}
@@ -104,7 +104,7 @@ namespace UnrealGameSync
 			{
 				if (item != null)
 				{
-					_steps.Add((BuildStep)item.Tag);
+					_steps.Add((BuildStep)item.Tag!);
 				}
 			}
 			Close();
@@ -116,7 +116,7 @@ namespace UnrealGameSync
 
 			bool hasSelection = (selectedIndex != -1);
 			EditStepButton.Enabled = hasSelection;
-			RemoveStepButton.Enabled = hasSelection && !_projectSteps.Contains(((BuildStep)BuildStepList.SelectedItems[0].Tag).UniqueId);
+			RemoveStepButton.Enabled = hasSelection && !_projectSteps.Contains(((BuildStep)BuildStepList.SelectedItems[0].Tag!).UniqueId);
 
 			MoveUp.Enabled = (selectedIndex >= 1);
 			MoveDown.Enabled = (selectedIndex >= 0 && selectedIndex < BuildStepList.Items.Count - 1);
@@ -140,7 +140,7 @@ namespace UnrealGameSync
 			}
 			else
 			{
-				BuildStep task = (BuildStep)e.Item.Tag;
+				BuildStep task = (BuildStep)e.Item.Tag!;
 
 				bool enabled;
 				if (e.ColumnIndex == 1)
@@ -188,7 +188,7 @@ namespace UnrealGameSync
 				int columnIndex = hitTest.Item.SubItems.IndexOf(hitTest.SubItem);
 				if (columnIndex >= 1 && columnIndex <= 3)
 				{
-					BuildStep task = (BuildStep)hitTest.Item.Tag;
+					BuildStep task = (BuildStep)hitTest.Item.Tag!;
 					if (columnIndex == 1)
 					{
 						task.NormalSync ^= true;
@@ -254,7 +254,7 @@ namespace UnrealGameSync
 			{
 				if (item != null)
 				{
-					BuildStep step = (BuildStep)item.Tag;
+					BuildStep step = (BuildStep)item.Tag!;
 					step.OrderIndex = _steps.Count;
 					_steps.Add(step);
 				}
diff --git a/Engine/Source/Programs/UnrealGameSync/UnrealGameSync/Forms/NotificationWindow.cs b/Engine/Source/Programs/UnrealGameSync/UnrealGameSync/Forms/NotificationWindow.cs
--- a/Engine/Source/Programs/UnrealGameSync/UnrealGameSync/Forms/NotificationWindow.cs
+++ b/Engine/Source/Programs/UnrealGameSync/UnrealGameSync/Forms/NotificationWindow.cs
@@ -57,7 +57,7 @@ namespace UnrealGameSync
 
 			CalculateBounds();
 
-			Rectangle workingArea = Screen.PrimaryScreen.WorkingArea;
+			Rectangle workingArea = Screen.PrimaryScreen!.WorkingArea;
 			Location = new Point(workingArea.Right - Size.Width - 16, workingArea.Bottom - Size.Height - 16);
 
 			BringToFront();
diff --git a/Engine/Source/Programs/UnrealGameSync/UnrealGameSync/Forms/SdkInfoWindow.cs b/Engine/Source/Programs/UnrealGameSync/UnrealGameSync/Forms/SdkInfoWindow.cs
--- a/Engine/Source/Programs/UnrealGameSync/UnrealGameSync/Forms/SdkInfoWindow.cs
+++ b/Engine/Source/Programs/UnrealGameSync/UnrealGameSync/Forms/SdkInfoWindow.cs
@@ -313,7 +313,7 @@ namespace UnrealGameSync
 
 			List<BadgeInfo> badges = new List<BadgeInfo>();
 
-			SdkItem sdk = (SdkItem)subItem.Tag;
+			SdkItem sdk = (SdkItem)subItem.Tag!;
 
 			foreach (SdkAction action in sdk.Actions)
 			{
diff --git a/Engine/Source/Programs/UnrealGameSync/UnrealGameSync/Forms/SelectUserWindow.cs b/Engine/Source/Programs/UnrealGameSync/UnrealGameSync/Forms/SelectUserWindow.cs
--- a/Engine/Source/Programs/UnrealGameSync/UnrealGameSync/Forms/SelectUserWindow.cs
+++ b/Engine/Source/Programs/UnrealGameSync/UnrealGameSync/Forms/SelectUserWindow.cs
@@ -67,7 +67,7 @@ namespace UnrealGameSync
 					}
 
 					UserListView.Items[nextIndex].Selected = true;
-					_selectedUserIndex = (int)UserListView.Items[nextIndex].Tag;
+					_selectedUserIndex = (int)UserListView.Items[nextIndex].Tag!;
 				}
 			}
 		}
@@ -182,7 +182,7 @@ namespace UnrealGameSync
 		{
 			if (UserListView.SelectedItems.Count > 0)
 			{
-				_selectedUserIndex = (int)UserListView.SelectedItems[0].Tag;
+				_selectedUserIndex = (int)UserListView.SelectedItems[0].Tag!;
 				DialogResult = DialogResult.OK;
 				Close();
 			}
diff --git a/Engine/Source/Programs/UnrealGameSync/UnrealGameSync/UnrealGameSync.csproj b/Engine/Source/Programs/UnrealGameSync/UnrealGameSync/UnrealGameSync.csproj
--- a/Engine/Source/Programs/UnrealGameSync/UnrealGameSync/UnrealGameSync.csproj
+++ b/Engine/Source/Programs/UnrealGameSync/UnrealGameSync/UnrealGameSync.csproj
@@ -1,7 +1,7 @@
 <Project Sdk="Microsoft.NET.Sdk">
   <PropertyGroup>
     <OutputType>WinExe</OutputType>
-    <TargetFramework>net6.0-windows</TargetFramework>
+    <TargetFramework>net8.0-windows</TargetFramework>
     <UseWindowsForms>true</UseWindowsForms>
     <Deterministic>true</Deterministic>
     <Product>UnrealGameSync</Product>
diff --git a/Engine/Source/Programs/UnrealGameSync/UnrealGameSyncCmd/UnrealGameSyncCmd.csproj b/Engine/Source/Programs/UnrealGameSync/UnrealGameSyncCmd/UnrealGameSyncCmd.csproj
--- a/Engine/Source/Programs/UnrealGameSync/UnrealGameSyncCmd/UnrealGameSyncCmd.csproj
+++ b/Engine/Source/Programs/UnrealGameSync/UnrealGameSyncCmd/UnrealGameSyncCmd.csproj
@@ -2,7 +2,7 @@
 
   <PropertyGroup>
     <OutputType>Exe</OutputType>
-    <TargetFramework>net6.0</TargetFramework>
+    <TargetFramework>net8.0</TargetFramework>
     <Nullable>enable</Nullable>
     <AssemblyName>ugs</AssemblyName>
     <RunAnalyzersDuringBuild>False</RunAnalyzersDuringBuild>
diff --git a/Engine/Source/Programs/UnrealGameSync/UnrealGameSyncShared/UnrealGameSyncShared.csproj b/Engine/Source/Programs/UnrealGameSync/UnrealGameSyncShared/UnrealGameSyncShared.csproj
--- a/Engine/Source/Programs/UnrealGameSync/UnrealGameSyncShared/UnrealGameSyncShared.csproj
+++ b/Engine/Source/Programs/UnrealGameSync/UnrealGameSyncShared/UnrealGameSyncShared.csproj
@@ -1,7 +1,7 @@
 <Project Sdk="Microsoft.NET.Sdk">
 
   <PropertyGroup>
-    <TargetFramework>net6.0</TargetFramework>
+    <TargetFramework>net8.0</TargetFramework>
     <Nullable>enable</Nullable>
     <RootNamespace>UnrealGameSync</RootNamespace>
     <RunAnalyzersDuringBuild>False</RunAnalyzersDuringBuild>
