UGS: Improve error handling when extracting data from bundles.

#rnx

diff --git a/Engine/Source/Programs/UnrealGameSync/UnrealGameSync/ArtifactDownload.cs b/Engine/Source/Programs/UnrealGameSync/UnrealGameSync/ArtifactDownload.cs
--- a/Engine/Source/Programs/UnrealGameSync/UnrealGameSync/ArtifactDownload.cs
+++ b/Engine/Source/Programs/UnrealGameSync/UnrealGameSync/ArtifactDownload.cs
@@ -138,6 +138,10 @@ namespace UnrealGameSync
 			{
 				DownloadProgressWindow.Execute((p, ctx) => DoSyncAsync(descriptor.BaseUrl, descriptor.RefName, outputDir, p, ctx), CancellationToken.None);
 			}
+			catch (OperationCanceledException)
+			{
+				return;
+			}
 			catch (Exception ex)
 			{
 				MessageBox.Show($"An error occurred while downloading the specified item.\n\n{ex}");
diff --git a/Engine/Source/Programs/UnrealGameSync/UnrealGameSync/Forms/DownloadProgressWindow.cs b/Engine/Source/Programs/UnrealGameSync/UnrealGameSync/Forms/DownloadProgressWindow.cs
--- a/Engine/Source/Programs/UnrealGameSync/UnrealGameSync/Forms/DownloadProgressWindow.cs
+++ b/Engine/Source/Programs/UnrealGameSync/UnrealGameSync/Forms/DownloadProgressWindow.cs
@@ -1,6 +1,7 @@
 // Copyright Epic Games, Inc. All Rights Reserved.
 
 using System;
+using System.Runtime.ExceptionServices;
 using System.Threading;
 using System.Threading.Tasks;
 using System.Windows.Forms;
@@ -28,6 +29,7 @@ namespace UnrealGameSync.Forms
 
 		readonly Func<IProgress<string>, CancellationToken, Task> _taskFunc;
 		readonly CancellationTokenSource _cancellationSource;
+		ExceptionDispatchInfo? _taskExceptionDispatchInfo;
 		Task _task = Task.CompletedTask;
 
 		public Task Task => _task;
@@ -51,19 +53,14 @@ namespace UnrealGameSync.Forms
 			base.Dispose(disposing);
 		}
 
-		public static bool Execute(Func<IProgress<string>, CancellationToken, Task> taskFunc, CancellationToken cancellationToken)
+		public static void Execute(Func<IProgress<string>, CancellationToken, Task> taskFunc, CancellationToken cancellationToken)
 		{
 			using DownloadProgressWindow window = new DownloadProgressWindow(taskFunc, cancellationToken);
 			window.ShowDialog();
 
-			if (window._cancellationSource.IsCancellationRequested)
-			{
-				cancellationToken.ThrowIfCancellationRequested(); // Rethrow if cancelled by caller, not if cancelled by user
-				return false;
-			}
+			cancellationToken.ThrowIfCancellationRequested();
 
-			window.Task.Wait(CancellationToken.None); // should have already finished, just rethrowing exceptions
-			return true;
+			window._taskExceptionDispatchInfo?.Throw();
 		}
 
 		private void UpdateStatus(string text)
@@ -74,10 +71,22 @@ namespace UnrealGameSync.Forms
 			}
 		}
 
+		private async Task RunTaskAsync(SynchronizationContext syncContext)
+		{
+			try
+			{
+				await _taskFunc(new ProgressReporter(this, syncContext), _cancellationSource.Token);
+			}
+			catch (Exception ex)
+			{
+				_taskExceptionDispatchInfo = ExceptionDispatchInfo.Capture(ex);
+			}
+		}
+
 		private void DownloadProcessWindow_Load(object sender, EventArgs e)
 		{
 			SynchronizationContext syncContext = SynchronizationContext.Current!;
-			_task = Task.Run(() => _taskFunc(new ProgressReporter(this, syncContext), _cancellationSource.Token), _cancellationSource.Token);
+			_task = Task.Run(() => RunTaskAsync(syncContext), _cancellationSource.Token);
 			_task.ContinueWith(x => syncContext.Post(y => Close(), null), TaskScheduler.Default);
 		}
 
