UGS: Fix some static analysis warnings.

#rnx

diff --git a/Engine/Source/Programs/UnrealGameSync/UnrealGameSync/AutomationServer.cs b/Engine/Source/Programs/UnrealGameSync/UnrealGameSync/AutomationServer.cs
--- a/Engine/Source/Programs/UnrealGameSync/UnrealGameSync/AutomationServer.cs
+++ b/Engine/Source/Programs/UnrealGameSync/UnrealGameSync/AutomationServer.cs
@@ -243,7 +243,7 @@ namespace UnrealGameSync
 
 		async Task RunTcpAsync(int portNumber, CancellationToken cancellationToken)
 		{
-			TcpListener listener = new TcpListener(IPAddress.Loopback, portNumber);
+			using TcpListener listener = new TcpListener(IPAddress.Loopback, portNumber);
 			using (IDisposable disposable = cancellationToken.Register(() => listener.Stop()))
 			{
 				listener.Start();
@@ -371,7 +371,7 @@ namespace UnrealGameSync
 
 		public async ValueTask DisposeAsync()
 		{
-			_cancellationSource.Cancel();
+			await _cancellationSource.CancelAsync();
 
 			try
 			{
diff --git a/Engine/Source/Programs/UnrealGameSync/UnrealGameSync/Forms/IssueDetailsWindow.cs b/Engine/Source/Programs/UnrealGameSync/UnrealGameSync/Forms/IssueDetailsWindow.cs
--- a/Engine/Source/Programs/UnrealGameSync/UnrealGameSync/Forms/IssueDetailsWindow.cs
+++ b/Engine/Source/Programs/UnrealGameSync/UnrealGameSync/Forms/IssueDetailsWindow.cs
@@ -540,11 +540,11 @@ namespace UnrealGameSync
 				{
 					if (_issue.ResolvedAt.HasValue)
 					{
-						status.AppendFormat("Closed as systemic issue.", _issue.FixChange);
+						status.AppendFormat("Closed as systemic issue.");
 					}
 					else
 					{
-						status.AppendFormat("Fixed as systemic issue (pending verification).", _issue.FixChange);
+						status.AppendFormat("Fixed as systemic issue (pending verification).");
 					}
 				}
 				else
diff --git a/Engine/Source/Programs/UnrealGameSync/UnrealGameSync/Forms/IssueFixedWindow.cs b/Engine/Source/Programs/UnrealGameSync/UnrealGameSync/Forms/IssueFixedWindow.cs
--- a/Engine/Source/Programs/UnrealGameSync/UnrealGameSync/Forms/IssueFixedWindow.cs
+++ b/Engine/Source/Programs/UnrealGameSync/UnrealGameSync/Forms/IssueFixedWindow.cs
@@ -49,19 +49,17 @@ namespace UnrealGameSync
 				StopAsync().Wait();
 			}
 
-			Task StopAsync()
+			async Task StopAsync()
 			{
-				Task stopTask = Task.CompletedTask;
 				if (_backgroundTask != null)
 				{
 					_onComplete = null;
 
-					_cancellationSource.Cancel();
-					stopTask = _backgroundTask;
+					await _cancellationSource.CancelAsync();
+					await _backgroundTask;
 
 					_backgroundTask = null!;
 				}
-				return stopTask;
 			}
 
 			public void FetchChanges(string userName)
diff --git a/Engine/Source/Programs/UnrealGameSync/UnrealGameSync/PerforceMonitor.cs b/Engine/Source/Programs/UnrealGameSync/UnrealGameSync/PerforceMonitor.cs
--- a/Engine/Source/Programs/UnrealGameSync/UnrealGameSync/PerforceMonitor.cs
+++ b/Engine/Source/Programs/UnrealGameSync/UnrealGameSync/PerforceMonitor.cs
@@ -464,7 +464,7 @@ namespace UnrealGameSync
 						notifyTask = Task.Delay(TimeSpan.FromSeconds(5.0), cancellationSource.Token).ContinueWith(_ => _synchronizationContext.Post(_ => OnUpdateMetadata?.Invoke(), null), cancellationSource.Token, new TaskContinuationOptions(), TaskScheduler.Default);
 					}
 				}
-				cancellationSource.Cancel();
+				await cancellationSource.CancelAsync();
 				await notifyTask.ContinueWith(_ => { }, TaskScheduler.Default); // Ignore exceptions
 			}
 
diff --git a/Engine/Source/Programs/UnrealGameSync/UnrealGameSync/UpdateMonitor.cs b/Engine/Source/Programs/UnrealGameSync/UnrealGameSync/UpdateMonitor.cs
--- a/Engine/Source/Programs/UnrealGameSync/UnrealGameSync/UpdateMonitor.cs
+++ b/Engine/Source/Programs/UnrealGameSync/UnrealGameSync/UpdateMonitor.cs
@@ -142,7 +142,7 @@ namespace UnrealGameSync
 
 			if (_workerTask != null)
 			{
-				_cancellationSource.Cancel();
+				await _cancellationSource.CancelAsync();
 
 				await _workerTask;
 				_workerTask = null;
