UGS: Fix deadlock on shutdown due to waiting on task from UI thread.

#rnx

diff --git a/Engine/Source/Programs/UnrealGameSync/UnrealGameSync/AutomationServer.cs b/Engine/Source/Programs/UnrealGameSync/UnrealGameSync/AutomationServer.cs
--- a/Engine/Source/Programs/UnrealGameSync/UnrealGameSync/AutomationServer.cs
+++ b/Engine/Source/Programs/UnrealGameSync/UnrealGameSync/AutomationServer.cs
@@ -136,10 +136,10 @@ namespace UnrealGameSync
 		readonly CancellationTokenSource _cancellationSource = new CancellationTokenSource();
 
 		const string IpcChannel = @"\.\pipe\UGSChannel";
-		readonly ConfiguredTaskAwaitable _ipcTask;
+		readonly Task? _ipcTask;
 
 		public const int DefaultPortNumber = 30422;
-		readonly ConfiguredTaskAwaitable? _tcpTask;
+		readonly Task? _tcpTask;
 
 		readonly Action<AutomationRequest> _postRequest;
 
@@ -153,7 +153,7 @@ namespace UnrealGameSync
 			try
 			{
 				// IPC named pipe
-				_ipcTask = RunIpcAsync(uri, _cancellationSource.Token).ConfigureAwait(false);
+				_ipcTask = Task.Run(() => RunIpcAsync(uri, _cancellationSource.Token), _cancellationSource.Token);
 
 				// TCP listener setup
 				int portNumber = GetPortNumber();
@@ -161,7 +161,7 @@ namespace UnrealGameSync
 				{
 					try
 					{
-						_tcpTask = RunTcpAsync(portNumber, _cancellationSource.Token).ConfigureAwait(false);
+						_tcpTask = Task.Run(() => RunTcpAsync(portNumber, _cancellationSource.Token), _cancellationSource.Token);
 					}
 					catch (Exception ex)
 					{
@@ -371,22 +371,25 @@ namespace UnrealGameSync
 
 		public async ValueTask DisposeAsync()
 		{
-			await _cancellationSource.CancelAsync();
+			await _cancellationSource.CancelAsync().ConfigureAwait(false);
 
-			try
+			if (_ipcTask != null)
 			{
-				await _ipcTask;
-			}
-			catch (Exception ex)
-			{
-				_logger.LogError(ex, "Error awaiting IPC background task");
+				try
+				{
+					await _ipcTask.ConfigureAwait(false);
+				}
+				catch (Exception ex)
+				{
+					_logger.LogError(ex, "Error awaiting IPC background task");
+				}
 			}
 
 			if (_tcpTask != null)
 			{
 				try
 				{
-					await _tcpTask.Value;
+					await _tcpTask.ConfigureAwait(false);
 				}
 				catch (Exception ex)
 				{
