UnrealGameSync: Default to a .sln in the project directory if it exists and not syncing all projects

#rnx

diff --git a/Engine/Source/Programs/UnrealGameSync/UnrealGameSync/Controls/WorkspaceControl.cs b/Engine/Source/Programs/UnrealGameSync/UnrealGameSync/Controls/WorkspaceControl.cs
--- a/Engine/Source/Programs/UnrealGameSync/UnrealGameSync/Controls/WorkspaceControl.cs
+++ b/Engine/Source/Programs/UnrealGameSync/UnrealGameSync/Controls/WorkspaceControl.cs
@@ -5411,28 +5411,34 @@ namespace UnrealGameSync
 
 		private void OpenSolution()
 		{
+			WorkspaceUpdateOptions options = WorkspaceUpdateContext.GetOptionsFromConfig(_settings.Global, _workspaceSettings);
 			FileReference? solutionFileName = null;
 
-			string primaryProjectName = "UE4";
-			if (FileReference.Exists(FileReference.Combine(BranchDirectoryName, "UE5.sln")))
+			// Default to opening the sln in the project if not including all projects in solution
+			if (!options.HasFlag(WorkspaceUpdateOptions.IncludeAllProjectsInSolution) && SelectedFileName.HasExtension(".uproject"))
 			{
-				primaryProjectName = "UE5";
+				solutionFileName = SelectedFileName.ChangeExtension(".sln");
 			}
 
-			FileReference primaryProjectPathFileName = FileReference.Combine(BranchDirectoryName, "Engine", "Intermediate", "ProjectFiles", "PrimaryProjectPath.txt");
-			if (FileReference.Exists(primaryProjectPathFileName))
+			// Check PrimaryProjectPath.txt
+			if (solutionFileName == null || !FileReference.Exists(solutionFileName))
 			{
-				try
-				{
-					solutionFileName = new FileReference(FileReference.ReadAllText(primaryProjectPathFileName).Trim() + ".sln");
-				}
-				catch (Exception ex)
+				FileReference primaryProjectPathFileName = FileReference.Combine(BranchDirectoryName, "Engine", "Intermediate", "ProjectFiles", "PrimaryProjectPath.txt");
+				if (FileReference.Exists(primaryProjectPathFileName))
 				{
-					_logger.LogError(ex, "Unable to read '{File}'", primaryProjectPathFileName);
+					try
+					{
+						solutionFileName = new FileReference(FileReference.ReadAllText(primaryProjectPathFileName).Trim() + ".sln");
+					}
+					catch (Exception ex)
+					{
+						_logger.LogError(ex, "Unable to read '{File}'", primaryProjectPathFileName);
+					}
 				}
 			}
-			
-			if (solutionFileName == null)
+
+			// Check PrimaryProjectName.txt if PrimaryProjectPath.txt doesn't exist
+			if (solutionFileName == null || !FileReference.Exists(solutionFileName))
 			{
 				FileReference primaryProjectNameFileName = FileReference.Combine(BranchDirectoryName, "Engine", "Intermediate", "ProjectFiles", "PrimaryProjectName.txt");
 				if (!FileReference.Exists(primaryProjectNameFileName))
@@ -5444,14 +5450,24 @@ namespace UnrealGameSync
 				{
 					try
 					{
-						primaryProjectName = FileReference.ReadAllText(primaryProjectNameFileName).Trim();
+						string primaryProjectName = FileReference.ReadAllText(primaryProjectNameFileName).Trim();
+						solutionFileName = FileReference.Combine(BranchDirectoryName, primaryProjectName + ".sln");
 					}
 					catch (Exception ex)
 					{
 						_logger.LogError(ex, "Unable to read '{File}'", primaryProjectNameFileName);
 					}
 				}
+			}
 
+			// Fallback to UE#.sln in the branch directory
+			if (solutionFileName == null || !FileReference.Exists(solutionFileName))
+			{
+				string primaryProjectName = "UE5";
+				if (FileReference.Exists(FileReference.Combine(BranchDirectoryName, "UE4.sln")))
+				{
+					primaryProjectName = "UE4";
+				}
 				solutionFileName = FileReference.Combine(BranchDirectoryName, primaryProjectName + ".sln");
 			}
 
