Fix free disk space check in UGS CLI for mac and linux

diff --git a/Engine/Source/Programs/UnrealGameSync/UnrealGameSyncShared/WorkspaceUpdate.cs b/Engine/Source/Programs/UnrealGameSync/UnrealGameSyncShared/WorkspaceUpdate.cs
--- a/Engine/Source/Programs/UnrealGameSync/UnrealGameSyncShared/WorkspaceUpdate.cs
+++ b/Engine/Source/Programs/UnrealGameSync/UnrealGameSyncShared/WorkspaceUpdate.cs
@@ -881,7 +881,7 @@ namespace UnrealGameSync
 						try
 						{
 							DirectoryInfo localRootInfo = project.LocalRootPath.ToDirectoryInfo();
-							DriveInfo drive = new DriveInfo(localRootInfo.Root.FullName);
+							DriveInfo drive = new DriveInfo(localRootInfo.FullName);
 
 							if (drive.AvailableFreeSpace < requiredFreeSpace)
 							{
