List the first few files when UGS warns about opened files when changing streams
- added a change stream warning windows
- gathered the opened files in the change stream function and show the window instead of a message box

diff --git a/Engine/Source/Programs/UnrealGameSync/UnrealGameSync/Controls/WorkspaceControl.cs b/Engine/Source/Programs/UnrealGameSync/UnrealGameSync/Controls/WorkspaceControl.cs
--- a/Engine/Source/Programs/UnrealGameSync/UnrealGameSync/Controls/WorkspaceControl.cs
+++ b/Engine/Source/Programs/UnrealGameSync/UnrealGameSync/Controls/WorkspaceControl.cs
@@ -22,6 +22,7 @@ using EpicGames.Perforce;
 using Microsoft.Extensions.DependencyInjection;
 using Microsoft.Extensions.Logging;
 using Microsoft.Win32;
+using UnrealGameSync.Forms;
 
 namespace UnrealGameSync
 {
@@ -4249,41 +4250,59 @@ namespace UnrealGameSync
 				}
 				else
 				{
-					for (int idx = 0; idx < 2; idx++)
+					// check if there is any file opened in the current stream
+					List<string> openedFiles = new List<string>();
+
+					async Task<bool> GetOpenedFiles(IPerforceConnection perforce, CancellationToken cancellationToken)
 					{
-						async Task<bool> SwitchFunc(IPerforceConnection perforce, CancellationToken cancellationToken)
+						await foreach (OpenedRecord record in perforce.OpenedAsync(OpenedOptions.None, FileSpecList.Any, cancellationToken))
 						{
-							if (idx == 1 || !await perforce.OpenedAsync(OpenedOptions.None, FileSpecList.Any, cancellationToken).AnyAsync(cancellationToken))
-							{
-								await perforce.SwitchClientToStreamAsync(newStreamName, SwitchClientOptions.IgnoreOpenFiles, cancellationToken);
-								return true;
-							}
-							return false;
+							openedFiles.Add(record.DepotFile);
 						}
 
-						ModalTask<bool>? switchTask = PerforceModalTask.Execute<bool>(this, "Switching streams", "Please wait...", _perforceSettings, SwitchFunc, _logger);
-						if (switchTask == null || !switchTask.Succeeded)
+						return openedFiles.Any();
+					}
+
+					ModalTask<bool>? openedFilesTask = PerforceModalTask.Execute<bool>(this, "Checking for opened files", "Please wait...", _perforceSettings, GetOpenedFiles, _logger);
+
+					if (openedFilesTask == null || !openedFilesTask.Succeeded)
+					{
+						return;
+					}
+					if (openedFilesTask.Result)
+					{
+						SwitchStreamWarningWindow switchStreamWarningWindow = new SwitchStreamWarningWindow(openedFiles);
+						switchStreamWarningWindow.FormBorderStyle = FormBorderStyle.FixedDialog;
+						if (switchStreamWarningWindow.ShowDialog() == DialogResult.Cancel)
 						{
-							break;
+							return;
 						}
-						if (switchTask.Result)
-						{
-							StatusPanel.SuspendLayout();
-							StreamChanged();
+					}
 
-							// Reset the last code change we found when switching streams
-							_workspace.ModifyState(x =>
-							{
-								x.CurrentCodeChangeNumber = -1;
-							});
+					// switch to the new stream
+					async Task<bool> SwitchFunc(IPerforceConnection perforce, CancellationToken cancellationToken)
+					{
+						await perforce.SwitchClientToStreamAsync(newStreamName, SwitchClientOptions.IgnoreOpenFiles, cancellationToken);
+						return true;
+					}
 
-							StatusPanel.ResumeLayout();
-							break;
-						}
-						if (MessageBox.Show("You have files open for edit in this workspace. If you continue, you will not be able to submit them until you switch back.\n\nContinue switching streams?", "Files checked out", MessageBoxButtons.YesNo) != DialogResult.Yes)
+					ModalTask<bool>? switchTask = PerforceModalTask.Execute<bool>(this, "Switching streams", "Please wait...", _perforceSettings, SwitchFunc, _logger);
+					if (switchTask == null || !switchTask.Succeeded)
+					{
+						return;
+					}
+					if (switchTask.Result)
+					{
+						StatusPanel.SuspendLayout();
+						StreamChanged();
+
+						// Reset the last code change we found when switching streams
+						_workspace.ModifyState(x =>
 						{
-							break;
-						}
+							x.CurrentCodeChangeNumber = -1;
+						});
+
+						StatusPanel.ResumeLayout();
 					}
 				}
 			}
diff --git a/Engine/Source/Programs/UnrealGameSync/UnrealGameSync/Forms/SwitchStreamWarningWindow.cs b/Engine/Source/Programs/UnrealGameSync/UnrealGameSync/Forms/SwitchStreamWarningWindow.cs
new file mode 100644
--- /dev/null
+++ b/Engine/Source/Programs/UnrealGameSync/UnrealGameSync/Forms/SwitchStreamWarningWindow.cs
@@ -0,0 +1,38 @@
+﻿// Copyright Epic Games, Inc. All Rights Reserved.
+
+using System;
+using System.Collections.Generic;
+using System.ComponentModel;
+using System.Data;
+using System.Drawing;
+using System.Linq;
+using System.Text;
+using System.Threading.Tasks;
+using System.Windows.Forms;
+
+namespace UnrealGameSync.Forms
+{
+	public partial class SwitchStreamWarningWindow : Form
+	{
+		public SwitchStreamWarningWindow(List<string> openedFiles)
+		{
+			InitializeComponent();
+
+			foreach (string openedFile in openedFiles)
+			{
+				OpenedFiles.Items.Add(openedFile);
+			}
+			
+		}
+
+		private void OkButton_Click(object sender, EventArgs e)
+		{
+			DialogResult = DialogResult.OK;
+		}
+
+		private void CancButton_Click(object sender, EventArgs e)
+		{
+			DialogResult = DialogResult.Cancel;
+		}
+	}
+}
