[NO NON-PII COMMIT MESSAGE]

diff --git a/Engine/Source/Programs/UnrealGameSync/UnrealGameSync/Controls/WorkspaceControl.cs b/Engine/Source/Programs/UnrealGameSync/UnrealGameSync/Controls/WorkspaceControl.cs
--- a/Engine/Source/Programs/UnrealGameSync/UnrealGameSync/Controls/WorkspaceControl.cs
+++ b/Engine/Source/Programs/UnrealGameSync/UnrealGameSync/Controls/WorkspaceControl.cs
@@ -3777,6 +3777,7 @@ namespace UnrealGameSync
 				if (_workspace.CurrentChangeNumber != -1)
 				{
 					programsLine.AddLink("Unreal Editor", FontStyle.Regular, () => { LaunchEditor(); });
+					programsLine.AddLink(" \u25BE", FontStyle.Regular, () => { ModifyEditorArguments(); });
 					programsLine.AddText("  |  ");
 				}
 
