Create UGS presets to control sync filters

diff --git a/Engine/Source/Programs/UnrealGameSync/UnrealGameSync/Controls/WorkspaceControl.cs b/Engine/Source/Programs/UnrealGameSync/UnrealGameSync/Controls/WorkspaceControl.cs
--- a/Engine/Source/Programs/UnrealGameSync/UnrealGameSync/Controls/WorkspaceControl.cs
+++ b/Engine/Source/Programs/UnrealGameSync/UnrealGameSync/Controls/WorkspaceControl.cs
@@ -907,7 +907,7 @@ namespace UnrealGameSync
 				}
 			}
 
-			string[] combinedSyncFilter = UserSettings.GetCombinedSyncFilter(GetSyncCategories(), _settings.Global.Filter, _workspaceSettings.Filter, _perforceMonitor.LatestPerforceConfigSection());
+			string[] combinedSyncFilter = UserSettings.GetCombinedSyncFilter(GetSyncCategories(), _workspaceSettings.Preset, GetPresets(), _settings.Global.Filter, _workspaceSettings.Filter, _perforceMonitor.LatestPerforceConfigSection());
 
 			WorkspaceUpdateContext context = new WorkspaceUpdateContext(changeNumber, options, GetEditorBuildConfig(), combinedSyncFilter, _projectSettings.BuildSteps, null);
 			if (options.HasFlag(WorkspaceUpdateOptions.SyncArchives))
@@ -1712,6 +1712,11 @@ namespace UnrealGameSync
 			return ConfigUtils.GetSyncCategories(_perforceMonitor.LatestProjectConfigFile);
 		}
 
+		IDictionary<string, Preset> GetPresets()
+		{
+			return ConfigUtils.GetPresets(_perforceMonitor.LatestProjectConfigFile);
+		}
+
 		static List<string> GetProjectRoots(string inUProjectDirsPath)
 		{
 			List<string> projectRoots = new List<string>();
@@ -5751,7 +5756,7 @@ namespace UnrealGameSync
 				extraSafeToDeleteExtensions = "";
 			}
 
-			string[] combinedSyncFilter = UserSettings.GetCombinedSyncFilter(GetSyncCategories(), _settings.Global.Filter, _workspaceSettings.Filter, _perforceMonitor.LatestPerforceConfigSection());
+			string[] combinedSyncFilter = UserSettings.GetCombinedSyncFilter(GetSyncCategories(), _workspaceSettings.Preset, GetPresets(), _settings.Global.Filter, _workspaceSettings.Filter, _perforceMonitor.LatestPerforceConfigSection());
 			List<string> syncPaths = WorkspaceUpdate.GetSyncPaths(_workspace.Project, _workspaceSettings.Filter.AllProjects ?? _settings.Global.Filter.AllProjects ?? false, combinedSyncFilter);
 
 			bool performedClean = CleanWorkspaceWindow.DoClean(ParentForm!, _perforceSettings, BranchDirectoryName, _workspace.Project.ClientRootPath, syncPaths, extraSafeToDeleteFolders.Split('\n'), extraSafeToDeleteExtensions.Split('\n'), _serviceProvider.GetRequiredService<ILogger<CleanWorkspaceWindow>>());
@@ -6062,7 +6067,7 @@ namespace UnrealGameSync
 
 		private void OptionsContextMenu_SyncFilter_Click(object sender, EventArgs e)
 		{
-			using SyncFilter filter = new SyncFilter(GetSyncCategories(), _settings.Global.Filter, _workspaceSettings.Filter, _perforceMonitor.LatestPerforceConfigSection());
+			using SyncFilter filter = new SyncFilter(GetSyncCategories(), _workspaceSettings.Preset, GetPresets(), _settings.Global.Filter, _workspaceSettings.Filter, _perforceMonitor.LatestPerforceConfigSection());
 			if (filter.ShowDialog() == DialogResult.OK)
 			{
 				_settings.Global.Filter = filter.GlobalFilter;
@@ -6073,6 +6078,65 @@ namespace UnrealGameSync
 			}
 		}
 
+		private void OptionsContextMenu_Presets_Click(object sender, EventArgs e)
+		{
+			void UpdateFilter(FilterSettings filter, Preset role)
+			{
+				foreach (KeyValuePair<Guid, RoleCategory> category in role.Categories)
+				{
+					bool include = category.Value.Enabled;
+					Guid guid = category.Value.Id;
+
+					if (include)
+					{
+						if (!filter.IncludeCategories.Contains(guid))
+						{
+							filter.IncludeCategories.Add(guid);
+						}
+
+						if (filter.ExcludeCategories.Contains(guid))
+						{
+							filter.ExcludeCategories.Remove(guid);
+						}
+					}
+					else
+					{
+						if (!filter.ExcludeCategories.Contains(guid))
+						{
+							filter.ExcludeCategories.Add(guid);
+						}
+
+						if (filter.IncludeCategories.Contains(guid))
+						{
+							filter.IncludeCategories.Remove(guid);
+						}
+					}
+				}
+			}
+
+			using PresetsWindow presetWindow = new PresetsWindow(
+				GetSyncCategories(),
+				_workspaceSettings.Preset,
+				GetPresets(),
+				(Preset role) =>
+				{
+					UpdateFilter(_settings.Global.Filter, role);
+					_settings.Save(_logger);
+				},
+				(Preset role) =>
+				{
+					UpdateFilter(_workspaceSettings.Filter, role);
+					_workspaceSettings.Save(_logger);
+				});
+
+			presetWindow.FormBorderStyle = FormBorderStyle.FixedDialog;
+			if (presetWindow.ShowDialog() == DialogResult.OK)
+			{
+				_workspaceSettings.Preset = presetWindow.PresetName;
+				_workspaceSettings.Save(_logger);
+			}
+		}
+
 		/// <summary>
 		/// Used to load CustomLatestChangeTypes config data into the windows form.
 		/// </summary>
diff --git a/Engine/Source/Programs/UnrealGameSync/UnrealGameSync/Forms/CombinedViewsWindow.cs b/Engine/Source/Programs/UnrealGameSync/UnrealGameSync/Forms/CombinedViewsWindow.cs
new file mode 100644
--- /dev/null
+++ b/Engine/Source/Programs/UnrealGameSync/UnrealGameSync/Forms/CombinedViewsWindow.cs
@@ -0,0 +1,88 @@
+// Copyright Epic Games, Inc. All Rights Reserved.
+
+using System;
+using System.Collections.Generic;
+using System.ComponentModel;
+using System.Data;
+using System.Drawing;
+using System.Linq;
+using System.Text;
+using System.Threading.Tasks;
+using System.Windows.Forms;
+
+namespace UnrealGameSync.Forms
+{
+	public partial class CombinedViewsWindow : Form
+	{
+		private readonly IList<string> _views;
+		public CombinedViewsWindow(IList<string> views)
+		{
+			InitializeComponent();
+
+			_views = views;
+			FilterViews(String.Empty);
+		}
+
+		private void CopyToClipBoard_Click(object sender, EventArgs e)
+		{
+			StringBuilder sb = new StringBuilder();
+
+			foreach (object? view in Views.Items)
+			{
+				if (view is string item)
+				{
+					sb.AppendLine(item);
+				}
+			}
+
+			Clipboard.SetText(sb.ToString());
+		}
+
+		private void SearchBox_TextChanged(object sender, EventArgs e)
+		{
+			FilterViews(SearchBox.Text);
+		}
+
+		private void FilterViews(string search)
+		{
+			Views.Items.Clear();
+
+			if (String.IsNullOrWhiteSpace(search))
+			{
+				foreach (string view in _views)
+				{
+					Views.Items.Add(view);
+				}
+			}
+			else
+			{
+				foreach (string view in _views)
+				{
+					if (view.Contains(search, StringComparison.OrdinalIgnoreCase))
+					{
+						Views.Items.Add(view);
+					}
+				}
+			}
+
+			// Create a Graphics object to use when determining the size of the largest item in the ListBox.
+			Graphics g = Views.CreateGraphics();
+
+			// calculate the max size
+			int maxHSize = 0;
+			for (int i = 0; i < Views.Items.Count; i++)
+			{
+				int hzSize = (int)g.MeasureString(Views.Items[i].ToString(), Views.Font).Width;
+				maxHSize = Math.Max(maxHSize, hzSize);
+			}
+
+			// Set the HorizontalExtent property.
+			Views.HorizontalExtent = maxHSize;
+		}
+
+		private void ClearSearch_Click(object sender, EventArgs e)
+		{
+			SearchBox.Clear();
+		}
+	}
+}
diff --git a/Engine/Source/Programs/UnrealGameSync/UnrealGameSync/Forms/PresetsWindow.cs b/Engine/Source/Programs/UnrealGameSync/UnrealGameSync/Forms/PresetsWindow.cs
new file mode 100644
--- /dev/null
+++ b/Engine/Source/Programs/UnrealGameSync/UnrealGameSync/Forms/PresetsWindow.cs
@@ -0,0 +1,168 @@
+﻿// Copyright Epic Games, Inc. All Rights Reserved.
+
+using System;
+using System.Collections.Generic;
+using System.Windows.Forms;
+
+namespace UnrealGameSync.Forms
+{
+	public partial class PresetsWindow : Form
+	{
+		private readonly IDictionary<Guid, WorkspaceSyncCategory> _uniqueIdToCategory;
+		private readonly IDictionary<string, Preset> _roles;
+
+		private readonly Action<Preset>? _copyToGlobal;
+		private readonly Action<Preset>? _copyToWorkspace;
+
+		public string PresetName { get; set; } = String.Empty;
+
+		public PresetsWindow(
+			IDictionary<Guid, WorkspaceSyncCategory> uniqueIdToCategory,
+			string presetName,
+			IDictionary<string, Preset> roles,
+			Action<Preset>? copyToGlobal = null,
+			Action<Preset>? copyToWorkspace = null
+			)
+		{
+			_uniqueIdToCategory = uniqueIdToCategory;
+			_roles = roles;
+			_copyToGlobal = copyToGlobal;
+			_copyToWorkspace = copyToWorkspace;
+
+			InitializeComponent();
+
+			// add an empty string to be able to clear the current role
+			RoleSelector.Items.Add(String.Empty);
+			foreach (string rolesKey in roles.Keys)
+			{
+				RoleSelector.Items.Add(rolesKey);
+			}
+
+			RoleSelector.SelectedItem = presetName;
+			RoleSelector.SelectedIndexChanged += RoleSelector_SelectedIndexChanged;
+
+			SyncFilters.SelectionMode = DataGridViewSelectionMode.CellSelect;
+
+			SetCurrentPreset(presetName);
+		}
+
+		private void RoleSelector_SelectedIndexChanged(object? sender, System.EventArgs e)
+		{
+			if (sender is ComboBox senderComboBox)
+			{
+				string? name = senderComboBox.SelectedItem?.ToString();
+				PopulateGrid(name ?? String.Empty);
+				PopulateViews(name ?? String.Empty);
+			}
+		}
+
+		private void PopulateGrid(string name)
+		{
+			if (String.IsNullOrWhiteSpace(name))
+			{
+				SyncFilters.Rows.Clear();
+			}
+			else
+			{
+				if (!_roles.TryGetValue(name, out Preset? role))
+				{
+					return;
+				}
+
+				SyncFilters.Rows.Clear();
+
+				foreach (KeyValuePair<Guid, RoleCategory> category in role.Categories)
+				{
+					if (!_uniqueIdToCategory.TryGetValue(category.Key, out WorkspaceSyncCategory? workspaceSyncCategory))
+					{
+						continue;
+					}
+
+					string[] row = new[]
+					{
+						workspaceSyncCategory.Name, category.Value.Enabled ? "True" : "False",
+						String.Join('\n', workspaceSyncCategory.Paths)
+					};
+
+					SyncFilters.Rows.Add(row);
+				}
+			}
+		}
+
+		private void PopulateViews(string name)
+		{
+			if (String.IsNullOrWhiteSpace(name))
+			{
+				Views.Items.Clear();
+			}
+			else
+			{
+				if (_roles.TryGetValue(name, out Preset? role))
+				{
+					Views.Items.Clear();
+
+					foreach (string view in role.Views)
+					{
+						ListViewItem item = new ListViewItem(view);
+						Views.Items.Add(item);
+					}
+				}
+			}
+		}
+
+		private void SelectRole_Click(object sender, EventArgs e)
+		{
+			string? newPreset = RoleSelector?.SelectedItem?.ToString();
+
+			SetCurrentPreset(newPreset);
+		}
+
+		private void SetCurrentPreset(string? newPreset)
+		{
+			PresetName = String.IsNullOrWhiteSpace(newPreset) ? String.Empty : newPreset;
+
+			PopulateGrid(PresetName);
+			PopulateViews(PresetName);
+			CurrentPreset.Text = PresetName;
+		}
+
+		private void OkButton_Click(object sender, EventArgs e)
+		{
+			DialogResult = DialogResult.OK;
+		}
+
+		private void CancButton_Click(object sender, EventArgs e)
+		{
+			DialogResult = DialogResult.Cancel;
+		}
+
+		private void CopyToGlobal_Click(object sender, EventArgs e)
+		{
+			CopyTo(_copyToGlobal);
+		}
+
+		private void CopyToWorkspace_Click(object sender, EventArgs e)
+		{
+			CopyTo(_copyToWorkspace);
+		}
+
+		private void CopyTo(Action<Preset>? action)
+		{
+			if (action == null)
+			{
+				return;
+			}
+
+			string? currentRoleName = RoleSelector?.SelectedItem?.ToString();
+			if (String.IsNullOrWhiteSpace(currentRoleName))
+			{
+				return;
+			}
+
+			if (_roles.TryGetValue(currentRoleName, out Preset? role))
+			{
+				action.Invoke(role);
+			}
+		}
+	}
+}
diff --git a/Engine/Source/Programs/UnrealGameSync/UnrealGameSync/Forms/SyncFilter.cs b/Engine/Source/Programs/UnrealGameSync/UnrealGameSync/Forms/SyncFilter.cs
--- a/Engine/Source/Programs/UnrealGameSync/UnrealGameSync/Forms/SyncFilter.cs
+++ b/Engine/Source/Programs/UnrealGameSync/UnrealGameSync/Forms/SyncFilter.cs
@@ -4,37 +4,80 @@ using System;
 using System.Collections.Generic;
 using System.Linq;
 using System.Windows.Forms;
+using UnrealGameSync.Forms;
 
 namespace UnrealGameSync
 {
 	partial class SyncFilter : Form
 	{
-		readonly Dictionary<Guid, WorkspaceSyncCategory> _uniqueIdToCategory;
+		class CheckListItem
+		{
+			public string Name => GetName();
+			public Guid UniqueId { get; set; } = Guid.Empty;
+			public string CategoryName { get; set; } = String.Empty;
+			public bool Locked { get; set; } = false;
+			public string Preset { get; set; } = String.Empty;
+			public bool Enabled { get; set; } = false;
+
+			public override string ToString()
+			{
+				return GetName();
+			}
+
+			private string GetName()
+			{
+				if (String.IsNullOrWhiteSpace(Preset))
+				{
+					return CategoryName;
+				}
+
+				string suffix = Locked ? " - Locked" : String.Empty;
+				return $"{CategoryName} [Preset: {Preset}{suffix}]";
+			}
+		}
+
+		private readonly Dictionary<Guid, WorkspaceSyncCategory> _uniqueIdToCategory;
+		private readonly string _roleName;
+		private readonly IDictionary<string, Preset> _roles;
+		private readonly ConfigSection? _perforceSection;
+		
 		public FilterSettings GlobalFilter;
 		public FilterSettings WorkspaceFilter;
-		readonly ConfigSection? _perforceSection;
-
-		public SyncFilter(Dictionary<Guid, WorkspaceSyncCategory> uniqueIdToCategory, FilterSettings globalFilter, FilterSettings workspaceFilter, ConfigSection? perforceSection)
+		
+		public SyncFilter(
+			Dictionary<Guid, WorkspaceSyncCategory> uniqueIdToCategory,
+			string roleName,
+			IDictionary<string, Preset> roles,
+			FilterSettings globalFilter,
+			FilterSettings workspaceFilter,
+			ConfigSection? perforceSection)
 		{
 			InitializeComponent();
 			Font = new System.Drawing.Font("Segoe UI", 8.25F, System.Drawing.FontStyle.Regular, System.Drawing.GraphicsUnit.Point, ((byte)(0)));
 
 			_uniqueIdToCategory = uniqueIdToCategory;
+			_roleName = roleName;
+			_roles = roles;
+			_perforceSection = perforceSection;
+			
 			GlobalFilter = globalFilter;
 			WorkspaceFilter = workspaceFilter;
-			_perforceSection = perforceSection;
 
+			Preset? role = new Preset();
+
+			_roles.TryGetValue(_roleName, out role);
+			
 			Dictionary<Guid, bool> syncCategories = WorkspaceSyncCategory.GetDefault(_uniqueIdToCategory.Values);
 
 			WorkspaceSyncCategory.ApplyDelta(syncCategories, GlobalFilter.GetCategories());
 			GlobalControl.SetView(GlobalFilter.View.ToArray());
-			SetExcludedCategories(GlobalControl.CategoriesCheckList, _uniqueIdToCategory, syncCategories);
+			SetExcludedCategories(GlobalControl.CategoriesCheckList, _uniqueIdToCategory, role, syncCategories);
 			GlobalControl.SyncAllProjects.Checked = GlobalFilter.AllProjects ?? false;
 			GlobalControl.IncludeAllProjectsInSolution.Checked = GlobalFilter.AllProjectsInSln ?? false;
 
 			WorkspaceSyncCategory.ApplyDelta(syncCategories, WorkspaceFilter.GetCategories());
 			WorkspaceControl.SetView(WorkspaceFilter.View.ToArray());
-			SetExcludedCategories(WorkspaceControl.CategoriesCheckList, _uniqueIdToCategory, syncCategories);
+			SetExcludedCategories(WorkspaceControl.CategoriesCheckList, _uniqueIdToCategory, role, syncCategories);
 			WorkspaceControl.SyncAllProjects.Checked = WorkspaceFilter.AllProjects ?? GlobalFilter.AllProjects ?? false;
 			WorkspaceControl.IncludeAllProjectsInSolution.Checked = WorkspaceFilter.AllProjectsInSln ?? GlobalFilter.AllProjectsInSln ?? false;
 
@@ -45,6 +88,16 @@ namespace UnrealGameSync
 
 		private void GlobalControl_CategoriesCheckList_ItemCheck(object? sender, ItemCheckEventArgs e)
 		{
+			// do not allow changing the value of locked categories
+			if (sender is System.Windows.Forms.CheckedListBox listBox)
+			{
+				if (listBox.Items[e.Index] is CheckListItem item && item.Locked)
+				{
+					e.NewValue = e.CurrentValue;
+					return;
+				}
+			}
+
 			WorkspaceControl.CategoriesCheckList.SetItemCheckState(e.Index, e.NewValue);
 		}
 
@@ -58,20 +111,45 @@ namespace UnrealGameSync
 			WorkspaceControl.IncludeAllProjectsInSolution.Checked = GlobalControl.IncludeAllProjectsInSolution.Checked;
 		}
 
-		private static void SetExcludedCategories(CheckedListBox listBox, Dictionary<Guid, WorkspaceSyncCategory> uniqueIdToFilter, Dictionary<Guid, bool> categoryIdToSetting)
+		private static void SetExcludedCategories(
+			CheckedListBox listBox, 
+			Dictionary<Guid, WorkspaceSyncCategory> uniqueIdToFilter,
+			Preset? role,
+			Dictionary<Guid, bool> categoryIdToSetting)
 		{
 			listBox.Items.Clear();
+			
 			foreach (WorkspaceSyncCategory filter in uniqueIdToFilter.Values)
 			{
-				if (!filter.Hidden)
+				if (filter.Hidden)
 				{
-					CheckState state = CheckState.Checked;
-					if (!categoryIdToSetting[filter.UniqueId])
-					{
-						state = CheckState.Unchecked;
-					}
-					listBox.Items.Add(filter, state);
+					continue;
 				}
+
+				CheckState state = CheckState.Checked;
+				if (!categoryIdToSetting[filter.UniqueId])
+				{
+					state = CheckState.Unchecked;
+				}
+					
+				CheckListItem item = new CheckListItem()
+				{
+					UniqueId = filter.UniqueId,
+					CategoryName = filter.Name,
+					Locked = false,
+					Preset = String.Empty,
+					Enabled = filter.Enable
+				};
+
+				if (role != null && role.Categories.TryGetValue(filter.UniqueId, out RoleCategory? category))
+				{
+					item.Preset = role.Name;
+					item.Locked = true;
+					item.Enabled = category.Enabled;
+					state = category.Enabled ? CheckState.Checked : CheckState.Unchecked;
+				}
+					
+				listBox.Items.Add(item, state);
 			}
 		}
 
@@ -101,10 +179,13 @@ namespace UnrealGameSync
 			Dictionary<Guid, bool> result = new Dictionary<Guid, bool>();
 			for (int idx = 0; idx < listBox.Items.Count; idx++)
 			{
-				Guid uniqueId = ((WorkspaceSyncCategory)listBox.Items[idx]).UniqueId;
-				if (!result.ContainsKey(uniqueId))
+				if (listBox.Items[idx] is CheckListItem item)
 				{
-					result[uniqueId] = listBox.GetItemCheckState(idx) == CheckState.Checked;
+					Guid uniqueId = item.UniqueId;
+					if (!result.ContainsKey(uniqueId))
+					{
+						result[uniqueId] = listBox.GetItemCheckState(idx) == CheckState.Checked;
+					}	
 				}
 			}
 			foreach (KeyValuePair<Guid, bool> originalSetting in originalSettings)
@@ -144,12 +225,16 @@ namespace UnrealGameSync
 		{
 			GetSettings(out FilterSettings newGlobalFilter, out FilterSettings newWorkspaceFilter);
 
-			string[] filter = UserSettings.GetCombinedSyncFilter(_uniqueIdToCategory, newGlobalFilter, newWorkspaceFilter, _perforceSection);
+			string[] filter = UserSettings.GetCombinedSyncFilter(_uniqueIdToCategory, _roleName, _roles, newGlobalFilter, newWorkspaceFilter, _perforceSection);
 			if (filter.Length == 0)
 			{
 				filter = new string[] { "All files will be synced." };
 			}
-			MessageBox.Show(String.Join("\r\n", filter), "Combined View");
+
+			CombinedViewsWindow combinedViewsWindow = new CombinedViewsWindow(filter);
+
+			combinedViewsWindow.FormBorderStyle = FormBorderStyle.FixedDialog;
+			combinedViewsWindow.ShowDialog();
 		}
 	}
 }
diff --git a/Engine/Source/Programs/UnrealGameSync/UnrealGameSyncCmd/Program.cs b/Engine/Source/Programs/UnrealGameSync/UnrealGameSyncCmd/Program.cs
--- a/Engine/Source/Programs/UnrealGameSync/UnrealGameSyncCmd/Program.cs
+++ b/Engine/Source/Programs/UnrealGameSync/UnrealGameSyncCmd/Program.cs
@@ -533,9 +533,23 @@ namespace UnrealGameSyncCmd
 		static string[] ReadSyncFilter(UserWorkspaceSettings workspaceSettings, GlobalSettingsFile userSettings, ConfigFile projectConfig)
 		{
 			Dictionary<Guid, WorkspaceSyncCategory> syncCategories = ConfigUtils.GetSyncCategories(projectConfig);
+
+			// check if any category is from the role
+			IDictionary<string, Preset> roles = ConfigUtils.GetPresets(projectConfig);
+			if (roles.TryGetValue(workspaceSettings.Preset, out Preset? role))
+			{
+				foreach (RoleCategory roleCategory in role.Categories.Values)
+				{
+					if (syncCategories.TryGetValue(roleCategory.Id, out WorkspaceSyncCategory? category))
+					{
+						category.Enable = roleCategory.Enabled;
+					}
+				}
+			}
+			
 			ConfigSection? perforceSection = projectConfig.FindSection("Perforce");
 
-			string[] combinedSyncFilter = GlobalSettingsFile.GetCombinedSyncFilter(syncCategories, userSettings.Global.Filter, workspaceSettings.Filter, perforceSection);
+			string[] combinedSyncFilter = GlobalSettingsFile.GetCombinedSyncFilter(syncCategories, workspaceSettings.Preset, roles, userSettings.Global.Filter, workspaceSettings.Filter, perforceSection);
 
 			return combinedSyncFilter;
 		}
@@ -1190,6 +1204,10 @@ namespace UnrealGameSyncCmd
 				FilterSettings globalFilter = context.UserSettings.Global.Filter;
 				FilterSettings workspaceFilter = workspaceSettings.Filter;
 
+				IDictionary<string, Preset> roles = ConfigUtils.GetPresets(projectConfig);
+
+				roles.TryGetValue(workspaceSettings.Preset, out Preset? role);
+				
 				FilterCommandOptions options = context.Arguments.ApplyTo<FilterCommandOptions>(logger);
 				context.Arguments.CheckAllArgumentsUsed(context.Logger);
 
@@ -1226,6 +1244,12 @@ namespace UnrealGameSyncCmd
 						enabled = syncCategory.Enable;
 					}
 
+					if (role != null && role.Categories.TryGetValue(syncCategory.UniqueId, out RoleCategory? roleCategory))
+					{
+						scope = $"(Preset: {role.Name})";
+						enabled = roleCategory.Enabled;
+					}
+					
 					logger.LogInformation("  {Id,30} {Enabled,3} {Scope,-9} {Name}", syncCategory.UniqueId, enabled ? "Yes" : "No", scope, syncCategory.Name);
 				}
 
@@ -1247,6 +1271,16 @@ namespace UnrealGameSyncCmd
 						logger.LogInformation("  {Line}", line);
 					}
 				}
+				if (role != null && role.Views.Count > 0)
+				{
+					logger.LogInformation("");
+					logger.LogInformation("Preset View:");
+					foreach (string line in role.Views)
+					{
+						logger.LogInformation("  {Line}", line);
+					}
+				}
+
 
 				string[] filter = ReadSyncFilter(workspaceSettings, context.UserSettings, projectConfig);
 
diff --git a/Engine/Source/Programs/UnrealGameSync/UnrealGameSyncShared/ConfigUtils.cs b/Engine/Source/Programs/UnrealGameSync/UnrealGameSyncShared/ConfigUtils.cs
--- a/Engine/Source/Programs/UnrealGameSync/UnrealGameSyncShared/ConfigUtils.cs
+++ b/Engine/Source/Programs/UnrealGameSync/UnrealGameSyncShared/ConfigUtils.cs
@@ -2,6 +2,7 @@
 
 using System;
 using System.Collections.Generic;
+using System.Data;
 using System.Diagnostics.CodeAnalysis;
 using System.IO;
 using System.Linq;
@@ -548,6 +549,90 @@ namespace UnrealGameSync
 			return uniqueIdToCategory;
 		}
 
+
+		public static IDictionary<string, Preset> GetPresets(ConfigFile? projectConfigFile)
+		{
+			Dictionary<string, Preset> roles = new(StringComparer.OrdinalIgnoreCase);
+
+			if (projectConfigFile == null)
+			{
+				return roles;
+			}
+
+			string[] lines = projectConfigFile.GetValues("Presets.Preset", []);
+			foreach (string line in lines)
+			{
+				ConfigObject obj = new ConfigObject(line);
+
+				Preset? preset = new Preset();
+
+				preset.Name = obj.GetValue("Name", String.Empty);
+
+				// do not allow empty role name
+				if (String.IsNullOrWhiteSpace(preset.Name))
+				{
+					continue;
+				}
+					
+				IEnumerable<string> categories = obj.GetValue("Categories", String.Empty)
+					.Split(';')
+					.Select(x => x.Trim())
+					.Where(x => x.Length > 0)
+					.Distinct()
+					.OrderBy(x => x);
+
+				foreach (string categoryLine in categories)
+				{
+					string[] values = categoryLine
+						.Split(',')
+						.Select(x => x.Trim())
+						.ToArray();
+
+					// all categories shall be completely defined
+					if (values.Length != 2) continue;
+
+					RoleCategory category = new();
+					if (Guid.TryParse(values[0], out Guid guid))
+					{
+						category.Id = guid;	
+					}
+
+					if (Boolean.TryParse(values[1], out bool enabled))
+					{
+						category.Enabled = enabled;
+					}
+
+					if (category.Id != Guid.Empty)
+					{
+						preset.Categories.TryAdd(category.Id, category);
+					}
+				}
+
+				IEnumerable<string> views = obj.GetValue("Views", String.Empty)
+						.Split(';')
+						.Select(x => x.Trim())
+						.Where(x => x.Length > 0)
+						.Distinct()
+					;
+
+				foreach (string view in views)
+				{
+					preset.Views.Add(view);
+				}
+
+				if (roles.ContainsKey(preset.Name))
+				{
+					roles[preset.Name].Import(preset);
+				}
+				else
+				{
+					roles.TryAdd(preset.Name, preset);	
+				}
+			}
+
+			return roles;
+		}
+
 		static IEnumerable<Guid> ParseGuids(IEnumerable<string> values)
 		{
 			foreach (string value in values)
diff --git a/Engine/Source/Programs/UnrealGameSync/UnrealGameSyncShared/GlobalSettings.cs b/Engine/Source/Programs/UnrealGameSync/UnrealGameSyncShared/GlobalSettings.cs
--- a/Engine/Source/Programs/UnrealGameSync/UnrealGameSyncShared/GlobalSettings.cs
+++ b/Engine/Source/Programs/UnrealGameSync/UnrealGameSyncShared/GlobalSettings.cs
@@ -172,14 +172,30 @@ namespace UnrealGameSync
 			return settings;
 		}
 
-		public static string[] GetCombinedSyncFilter(Dictionary<Guid, WorkspaceSyncCategory> uniqueIdToFilter, FilterSettings globalFilter, FilterSettings workspaceFilter, ConfigSection? perforceSection)
+		public static string[] GetCombinedSyncFilter(
+			Dictionary<Guid, WorkspaceSyncCategory> uniqueIdToFilter,
+			string roleName,
+			IDictionary<string, Preset> roles,
+			FilterSettings globalFilter,
+			FilterSettings workspaceFilter,
+			ConfigSection? perforceSection)
 		{
-			List<string> lines = new List<string>();
+			
+
+			ISet<string> lines = new HashSet<string>();
 			foreach (string viewLine in Enumerable.Concat(globalFilter.View, workspaceFilter.View).Select(x => x.Trim()).Where(x => x.Length > 0 && !x.StartsWith(";", StringComparison.Ordinal)))
 			{
 				lines.Add(viewLine);
 			}
 
+			if (roles.TryGetValue(roleName, out Preset? role))
+			{
+				foreach (string roleView in role.Views)
+				{
+					lines.Add(roleView);
+				}
+			}
+
 			Dictionary<Guid, bool> globalCategoryIdToSetting = globalFilter.GetCategories();
 			Dictionary<Guid, bool> workspaceCategoryIdToSetting = workspaceFilter.GetCategories();
 
@@ -188,18 +204,25 @@ namespace UnrealGameSync
 			{
 				bool enable = filter.Enable;
 
-				bool globalEnable;
-				if (globalCategoryIdToSetting.TryGetValue(filter.UniqueId, out globalEnable))
+				if (globalCategoryIdToSetting.TryGetValue(filter.UniqueId, out bool globalEnable))
 				{
 					enable = globalEnable;
 				}
 
-				bool workspaceEnable;
-				if (workspaceCategoryIdToSetting.TryGetValue(filter.UniqueId, out workspaceEnable))
+				if (workspaceCategoryIdToSetting.TryGetValue(filter.UniqueId, out bool workspaceEnable))
 				{
 					enable = workspaceEnable;
-				}
+				}	
 
+				// role override anything set at the global or workspace level
+				if (role != null)
+				{
+					if (role.Categories.TryGetValue(filter.UniqueId, out RoleCategory? roleCategory))
+					{
+						enable = roleCategory.Enabled;
+					}
+				}
+				
 				if (enable)
 				{
 					EnableFilter(filter.UniqueId, enabled, uniqueIdToFilter);
@@ -210,7 +233,10 @@ namespace UnrealGameSync
 			{
 				if (!enabled.Contains(filter.UniqueId))
 				{
-					lines.AddRange(filter.Paths.Select(x => "-" + x.Trim()));
+					foreach (string path in filter.Paths.Select(x => "-" + x.Trim()))
+					{
+						lines.Add(path);
+					}
 				}
 			}
 
@@ -219,7 +245,7 @@ namespace UnrealGameSync
 			if (lines.Count > 0 && perforceSection != null)
 			{
 				IEnumerable<string> additionalPaths = perforceSection.GetValues("AdditionalPathsToSync", Array.Empty<string>());
-				lines = lines.Concat(additionalPaths).ToList();
+				lines.UnionWith(additionalPaths);
 			}
 
 			return lines.ToArray();
diff --git a/Engine/Source/Programs/UnrealGameSync/UnrealGameSyncShared/Preset.cs b/Engine/Source/Programs/UnrealGameSync/UnrealGameSyncShared/Preset.cs
new file mode 100644
--- /dev/null
+++ b/Engine/Source/Programs/UnrealGameSync/UnrealGameSyncShared/Preset.cs
@@ -0,0 +1,43 @@
+// Copyright Epic Games, Inc. All Rights Reserved.
+
+using System;
+using System.Collections.Generic;
+using System.Security;
+
+namespace UnrealGameSync
+{
+	public class RoleCategory
+	{
+		public Guid Id { get; set; } = Guid.Empty;
+		public bool Enabled { get; set; } = false;
+	}
+	
+	public class Preset
+	{
+		public string Name { get; set; } = String.Empty;
+		public IDictionary<Guid, RoleCategory> Categories { get; } = new Dictionary<Guid, RoleCategory>();
+		public ISet<string> Views { get; } = new HashSet<string>();
+
+		public void Import(Preset preset)
+		{
+			Name = preset.Name;
+			foreach (KeyValuePair<Guid, RoleCategory> rhs in preset.Categories)
+			{
+				if (Categories.TryGetValue(rhs.Key, out RoleCategory? lhs))
+				{
+					lhs.Enabled = rhs.Value.Enabled;
+				}
+				else
+				{
+					Categories.Add(rhs.Key, rhs.Value);
+				}
+			}
+
+			foreach (string rhs in preset.Views)
+			{
+				Views.Add(rhs);
+			}
+		}
+	}
+}
+
diff --git a/Engine/Source/Programs/UnrealGameSync/UnrealGameSyncShared/UserSettings.cs b/Engine/Source/Programs/UnrealGameSync/UnrealGameSyncShared/UserSettings.cs
--- a/Engine/Source/Programs/UnrealGameSync/UnrealGameSyncShared/UserSettings.cs
+++ b/Engine/Source/Programs/UnrealGameSync/UnrealGameSyncShared/UserSettings.cs
@@ -385,6 +385,9 @@ namespace UnrealGameSync
 		// The currently selected project, relative to the root directory
 		public string ProjectPath { get; set; } = String.Empty;
 
+		// The currently selected preset
+		public string Preset { get; set; } = String.Empty;
+		
 		// Workspace specific SyncFilters
 		public FilterSettings Filter { get; set; } = new FilterSettings();
 
