UGS: Fix compilation warnings

diff --git a/Engine/Source/Programs/UnrealGameSync/UnrealGameSync/Forms/CombinedViewsWindow.cs b/Engine/Source/Programs/UnrealGameSync/UnrealGameSync/Forms/CombinedViewsWindow.cs
--- a/Engine/Source/Programs/UnrealGameSync/UnrealGameSync/Forms/CombinedViewsWindow.cs
+++ b/Engine/Source/Programs/UnrealGameSync/UnrealGameSync/Forms/CombinedViewsWindow.cs
@@ -2,12 +2,8 @@
 
 using System;
 using System.Collections.Generic;
-using System.ComponentModel;
-using System.Data;
 using System.Drawing;
-using System.Linq;
 using System.Text;
-using System.Threading.Tasks;
 using System.Windows.Forms;
 
 namespace UnrealGameSync.Forms
diff --git a/Engine/Source/Programs/UnrealGameSync/UnrealGameSync/Forms/SyncFilter.cs b/Engine/Source/Programs/UnrealGameSync/UnrealGameSync/Forms/SyncFilter.cs
--- a/Engine/Source/Programs/UnrealGameSync/UnrealGameSync/Forms/SyncFilter.cs
+++ b/Engine/Source/Programs/UnrealGameSync/UnrealGameSync/Forms/SyncFilter.cs
@@ -230,8 +230,10 @@ namespace UnrealGameSync
 			{
 				filter = new string[] { "All files will be synced." };
 			}
-
+			
+#pragma warning disable CA2000 // Dispose objects before losing scope
 			CombinedViewsWindow combinedViewsWindow = new CombinedViewsWindow(filter);
+#pragma warning restore CA2000
 
 			combinedViewsWindow.FormBorderStyle = FormBorderStyle.FixedDialog;
 			combinedViewsWindow.ShowDialog();
diff --git a/Engine/Source/Programs/UnrealGameSync/UnrealGameSyncShared/ConfigUtils.cs b/Engine/Source/Programs/UnrealGameSync/UnrealGameSyncShared/ConfigUtils.cs
--- a/Engine/Source/Programs/UnrealGameSync/UnrealGameSyncShared/ConfigUtils.cs
+++ b/Engine/Source/Programs/UnrealGameSync/UnrealGameSyncShared/ConfigUtils.cs
@@ -549,7 +549,6 @@ namespace UnrealGameSync
 			return uniqueIdToCategory;
 		}
 
-
 		public static IDictionary<string, Preset> GetPresets(ConfigFile? projectConfigFile)
 		{
 			Dictionary<string, Preset> roles = new(StringComparer.OrdinalIgnoreCase);
@@ -589,7 +588,10 @@ namespace UnrealGameSync
 						.ToArray();
 
 					// all categories shall be completely defined
-					if (values.Length != 2) continue;
+					if (values.Length != 2)
+					{
+						continue;
+					}
 
 					RoleCategory category = new();
 					if (Guid.TryParse(values[0], out Guid guid))
diff --git a/Engine/Source/Programs/UnrealGameSync/UnrealGameSyncShared/Preset.cs b/Engine/Source/Programs/UnrealGameSync/UnrealGameSyncShared/Preset.cs
--- a/Engine/Source/Programs/UnrealGameSync/UnrealGameSyncShared/Preset.cs
+++ b/Engine/Source/Programs/UnrealGameSync/UnrealGameSyncShared/Preset.cs
@@ -2,7 +2,6 @@
 
 using System;
 using System.Collections.Generic;
-using System.Security;
 
 namespace UnrealGameSync
 {
