Users would like to specify Pre-compiled Binary setting per-project

diff --git a/Engine/Source/Programs/UnrealGameSync/UnrealGameSync/Controls/WorkspaceControl.cs b/Engine/Source/Programs/UnrealGameSync/UnrealGameSync/Controls/WorkspaceControl.cs
--- a/Engine/Source/Programs/UnrealGameSync/UnrealGameSync/Controls/WorkspaceControl.cs
+++ b/Engine/Source/Programs/UnrealGameSync/UnrealGameSync/Controls/WorkspaceControl.cs
@@ -4757,7 +4757,7 @@ namespace UnrealGameSync
 			Dictionary<string, KeyValuePair<IArchiveChannel, int>> archiveTypeToSelection = new Dictionary<string, KeyValuePair<IArchiveChannel, int>>();
 			foreach (IArchiveChannel archive in archives)
 			{
-				ArchiveSettings? archiveSettings = _settings.Archives.FirstOrDefault(x => x.Type == archive.Type);
+				ArchiveSettings? archiveSettings = _workspaceSettings.Archives.FirstOrDefault(x => x.Type == archive.Type);
 				if (archiveSettings != null && archiveSettings.Enabled)
 				{
 					int preference = archiveSettings.Order.IndexOf(archive.Name);
@@ -4778,22 +4778,22 @@ namespace UnrealGameSync
 
 		private void ClearSelectedArchives()
 		{
-			foreach (ArchiveSettings archiveSettings in _settings.Archives)
+			foreach (ArchiveSettings archiveSettings in _workspaceSettings.Archives)
 			{
 				archiveSettings.Enabled = false;
 			}
-			_settings.Save(_logger);
+			_workspaceSettings.Save(_logger);
 
 			UpdateSelectedArchives();
 		}
 
 		private void SetSelectedArchive(IArchiveChannel archive, bool selected)
 		{
-			ArchiveSettings? archiveSettings = _settings.Archives.FirstOrDefault(x => x.Type == archive.Type);
+			ArchiveSettings? archiveSettings = _workspaceSettings.Archives.FirstOrDefault(x => x.Type == archive.Type);
 			if (archiveSettings == null)
 			{
 				archiveSettings = new ArchiveSettings(selected, archive.Type, new string[] { archive.Name });
-				_settings.Archives.Add(archiveSettings);
+				_workspaceSettings.Archives.Add(archiveSettings);
 			}
 			else if (selected)
 			{
@@ -4805,7 +4805,7 @@ namespace UnrealGameSync
 			{
 				archiveSettings.Enabled = false;
 			}
-			_settings.Save(_logger);
+			_workspaceSettings.Save(_logger);
 
 			UpdateSelectedArchives();
 		}
@@ -5994,7 +5994,7 @@ namespace UnrealGameSync
 			}
 		}
 
-		private bool ShouldSyncPrecompiledEditor => _settings.Archives.Any(x => x.Enabled && x.Type == IArchiveChannel.EditorArchiveType) && GetArchiveChannels().Any(x => x.Type == "Editor");
+		private bool ShouldSyncPrecompiledEditor => _workspaceSettings.Archives.Any(x => x.Enabled && x.Type == IArchiveChannel.EditorArchiveType) && GetArchiveChannels().Any(x => x.Type == "Editor");
 
 		public BuildConfig GetEditorBuildConfig()
 		{
diff --git a/Engine/Source/Programs/UnrealGameSync/UnrealGameSyncShared/OpenProjectInfo.cs b/Engine/Source/Programs/UnrealGameSync/UnrealGameSyncShared/OpenProjectInfo.cs
--- a/Engine/Source/Programs/UnrealGameSync/UnrealGameSyncShared/OpenProjectInfo.cs
+++ b/Engine/Source/Programs/UnrealGameSync/UnrealGameSyncShared/OpenProjectInfo.cs
@@ -216,6 +216,12 @@ namespace UnrealGameSync
 				string projectPath = newSelectedClientFileName.Substring(branchClientPath.Length);
 				UserWorkspaceSettings userWorkspaceSettings = userSettings.FindOrAddWorkspaceSettings(branchDirectoryName, perforceSettings.ServerAndPort, perforceSettings.UserName, perforceSettings.ClientName!, branchPath, projectPath, logger);
 
+				// make sure that the default value is in the list
+				if (userSettings.SyncPrecompiledEditor && !userWorkspaceSettings.Archives.Any(a => a.Type.Equals("Editor")))
+				{
+					userWorkspaceSettings.Archives.Add(new ArchiveSettings(true, "Editor", Array.Empty<string>()));
+				}
+
 				// Now compute the updated project info
 				ProjectInfo projectInfo = await ProjectInfo.CreateAsync(perforceClient, userWorkspaceSettings, cancellationToken);
 
diff --git a/Engine/Source/Programs/UnrealGameSync/UnrealGameSyncShared/UserSettings.cs b/Engine/Source/Programs/UnrealGameSync/UnrealGameSyncShared/UserSettings.cs
--- a/Engine/Source/Programs/UnrealGameSync/UnrealGameSyncShared/UserSettings.cs
+++ b/Engine/Source/Programs/UnrealGameSync/UnrealGameSyncShared/UserSettings.cs
@@ -196,9 +196,14 @@ namespace UnrealGameSync
 
 	public class ArchiveSettings
 	{
-		public bool Enabled { get; set; }
-		public string Type { get; set; }
-		public List<string> Order { get; init; }
+		public bool Enabled { get; set; } = false;
+		public string Type { get; set; } = String.Empty;
+		public List<string> Order { get; set; } = new List<string>();
+
+		public ArchiveSettings()
+		{
+
+		}
 
 		public ArchiveSettings(bool enabled, string type, IEnumerable<string> order)
 		{
@@ -391,6 +396,9 @@ namespace UnrealGameSync
 		// Workspace specific SyncFilters
 		public FilterSettings Filter { get; set; } = new FilterSettings();
 
+		// Precompiled binaries
+		public List<ArchiveSettings> Archives { get; set; } = new List<ArchiveSettings>();
+
 		[JsonIgnore]
 		public FileReference ConfigFile => GetConfigFile(RootDir);
 
@@ -555,9 +563,7 @@ namespace UnrealGameSync
 		public BuildConfig CompiledEditorBuildConfig { get; set; } // NB: This assumes not using precompiled editor. See CurrentBuildConfig.
 		public TabLabels TabLabels { get; set; }
 		public long NextLauncherVersionCheck { get; set; }
-
-		// Precompiled binaries
-		public List<ArchiveSettings> Archives { get; init; } = new List<ArchiveSettings>();
+		public bool SyncPrecompiledEditor { get; set; }
 
 		// OIDC Settings
 		public Dictionary<string, string> ProviderToRefreshTokens { get; init; } = new Dictionary<string, string>();
@@ -676,10 +682,11 @@ namespace UnrealGameSync
 			Version = (UserSettingsVersion)_configFile.GetValue("General.Version", (int)UserSettingsVersion.Initial);
 			BuildAfterSync = (_configFile.GetValue("General.BuildAfterSync", "1") != "0");
 			RunAfterSync = (_configFile.GetValue("General.RunAfterSync", "1") != "0");
-			bool syncPrecompiledEditor = (_configFile.GetValue("General.SyncPrecompiledEditor", "0") != "0");
 			OpenSolutionAfterSync = (_configFile.GetValue("General.OpenSolutionAfterSync", "0") != "0");
 			ShowLogWindow = (_configFile.GetValue("General.ShowLogWindow", false));
 
+			SyncPrecompiledEditor = (_configFile.GetValue("General.SyncPrecompiledEditor", "0") != "0");
+
 			string? autoResolveConflicts = _configFile.GetValue("General.AutoResolveConflicts", "");
 			if (!String.IsNullOrEmpty(autoResolveConflicts))
 			{
@@ -775,22 +782,6 @@ namespace UnrealGameSync
 			}
 			EditorArgumentsPrompt = _configFile.GetValue("General.EditorArgumentsPrompt", false);
 
-			// Precompiled binaries
-			string[] archiveValues = _configFile.GetValues("PrecompiledBinaries.Archives", Array.Empty<string>());
-			foreach (string archiveValue in archiveValues)
-			{
-				ArchiveSettings? settings;
-				if (ArchiveSettings.TryParseConfigEntry(archiveValue, out settings))
-				{
-					Archives.Add(settings);
-				}
-			}
-
-			if (syncPrecompiledEditor)
-			{
-				Archives.Add(new ArchiveSettings(true, "Editor", Array.Empty<string>()));
-			}
-
 			// OIDC Settings
 			string[] tokens = _configFile.GetValues("OIDCProviders.Tokens", Array.Empty<string>());
 			foreach (string tokenValue in tokens)
@@ -1101,10 +1092,6 @@ namespace UnrealGameSync
 			scheduleSection.SetValue("AnyOpenProject", ScheduleAnyOpenProject);
 			scheduleSection.SetValues("Projects", ScheduleProjects.Select(x => x.ToConfigEntry()).ToArray());
 
-			// Precompiled binaries
-			ConfigSection archivesSection = _configFile.FindOrAddSection("PrecompiledBinaries");
-			archivesSection.SetValues("Archives", Archives.Select(x => x.ToConfigEntry()).ToArray());
-
 			// OIDC Settings
 			ConfigSection oidcSection = _configFile.FindOrAddSection("OIDCProviders");
 			List<ConfigObject> tokenObjects = new List<ConfigObject>();
