Have a column in UGS for cl reviewer and preflight

diff --git a/Engine/Programs/UnrealGameSync/UnrealGameSync.ini b/Engine/Programs/UnrealGameSync/UnrealGameSync.ini
--- a/Engine/Programs/UnrealGameSync/UnrealGameSync.ini
+++ b/Engine/Programs/UnrealGameSync/UnrealGameSync.ini
@@ -45,6 +45,8 @@ BuildUnrealHeaderTool=0
 EditorTarget=UnrealEditor
 +DefaultEditorArgument=(Name="-log", Enabled=false, Locked=false)
 +DefaultEditorArgument=(Name="-fastload", Enabled=false, Locked=false)
+ReviewTag=#rb
+PreflightTag=#preflight
 
 ; Set this to a network share which contains the SDK installers for your site
 SdkInstallerDir= 
diff --git a/Engine/Source/Programs/UnrealGameSync/UnrealGameSync/Controls/WorkspaceControl.cs b/Engine/Source/Programs/UnrealGameSync/UnrealGameSync/Controls/WorkspaceControl.cs
--- a/Engine/Source/Programs/UnrealGameSync/UnrealGameSync/Controls/WorkspaceControl.cs
+++ b/Engine/Source/Programs/UnrealGameSync/UnrealGameSync/Controls/WorkspaceControl.cs
@@ -256,6 +256,12 @@ namespace UnrealGameSync
 		bool _updateBuildMetadataPosted;
 		bool _updateReviewsPosted;
 
+		private readonly bool _showReviewerColumn = true;
+		private readonly bool _showPreflightColumn = true;
+
+		private ISet<string> _reviewTags = new HashSet<string>(StringComparer.OrdinalIgnoreCase);
+		private ISet<string> _preflightTags = new HashSet<string>(StringComparer.OrdinalIgnoreCase);
+
 		WorkspaceUpdateCallback? _updateCallback;
 
 		System.Threading.Timer? _startupTimer;
@@ -300,8 +306,44 @@ namespace UnrealGameSync
 					_projectSettings.RequiredBadges.Clear();
 					_projectSettings.RequiredBadges.AddRange(requiredBadges.Split(new char[] { ',' }, StringSplitOptions.RemoveEmptyEntries | StringSplitOptions.TrimEntries).Distinct().ToList());
 				}
+
+				if (TryGetProjectSetting(projectConfigFile, "ReviewTag", out string? reviewTags))
+				{
+					foreach (string tag in reviewTags.Split(","))
+					{
+						_reviewTags.Add(tag.Trim());
+					}
+				}
+
+				if (TryGetProjectSetting(projectConfigFile, "PreflightTag", out string? preflightTags))
+				{
+					foreach (string tag in preflightTags.Split(","))
+					{
+						_preflightTags.Add(tag.Trim());
+					}
+				}
 			}
 
+			_showPreflightColumn = _preflightTags.Any() &&_hordeClient != null;
+			
+
+			if (_showPreflightColumn)
+			{
+				BuildListContextMenu_OpenPreflight.Enabled = true;
+			}
+			else
+			{
+				BuildList.Columns.Remove(PreflightColumn);
+				BuildListContextMenu_OpenPreflight.Enabled = false;
+			}
+
+			_showReviewerColumn = _reviewTags.Any();
+			if (!_showReviewerColumn)
+			{
+				BuildList.Columns.Remove(ReviewerColumn);
+			}
+
+
 			DesiredTaskbarState = Tuple.Create(TaskbarState.NoProgress, 0.0f);
 
 			System.Reflection.PropertyInfo doubleBufferedProperty = typeof(Control).GetProperty("DoubleBuffered", System.Reflection.BindingFlags.NonPublic | System.Reflection.BindingFlags.Instance)!;
@@ -555,6 +597,11 @@ namespace UnrealGameSync
 					}
 				}
 
+				if (_showPreflightColumn == false)
+				{
+					BuildList.Columns.Remove(PreflightColumn);
+				}
+
 				_columnWidths = new float[BuildList.Columns.Count];
 				for (int idx = 0; idx < BuildList.Columns.Count; idx++)
 				{
@@ -576,6 +623,17 @@ namespace UnrealGameSync
 				_desiredColumnWidths[TimeColumn.Index] = _minColumnWidths[TimeColumn.Index];
 				_desiredColumnWidths[ChangeColumn.Index] = _minColumnWidths[ChangeColumn.Index];
 				_desiredColumnWidths[AuthorColumn.Index] = (int)(120 * dpiScaleX);
+
+				if (_showReviewerColumn)
+				{
+					_desiredColumnWidths[ReviewerColumn.Index] = (int)(120 * dpiScaleX);
+				}
+
+				if (_showPreflightColumn)
+				{
+					_desiredColumnWidths[PreflightColumn.Index] = (int)(120 * dpiScaleX);
+				}
+					
 				_desiredColumnWidths[CISColumn.Index] = (int)(200 * dpiScaleX);
 				_desiredColumnWidths[StatusColumn.Index] = (int)(300 * dpiScaleX);
 
@@ -1403,6 +1461,66 @@ namespace UnrealGameSync
 			columns[AuthorColumn.Index] = userName;
 			columns[DescriptionColumn.Index] = change.Description!.Replace('\n', ' ');
 
+			if (_showReviewerColumn || _showPreflightColumn)
+			{
+				StringReader reader = new StringReader(change.Description);
+				while (reader.Peek() != -1)
+				{
+					string? line = reader.ReadLine();
+					if (String.IsNullOrWhiteSpace(line))
+					{
+						continue;
+					}
+
+					if (_showReviewerColumn)
+					{
+						foreach (string reviewTag in _reviewTags)
+						{
+							if (String.IsNullOrWhiteSpace(reviewTag) || !line.StartsWith(reviewTag, StringComparison.OrdinalIgnoreCase))
+							{
+								continue;
+							}
+
+							if (String.IsNullOrWhiteSpace(columns[ReviewerColumn.Index]))
+							{
+								columns[ReviewerColumn.Index] = line.Replace($"{reviewTag} ", String.Empty);
+							}
+							else
+							{
+								columns[ReviewerColumn.Index] += ", " + line.Replace($"{reviewTag} ", String.Empty);
+							}
+						}
+					}
+
+
+					if (_showPreflightColumn)
+					{
+						foreach (string preflightTag in _preflightTags)
+						{
+							if (String.IsNullOrWhiteSpace(preflightTag)
+								|| String.IsNullOrWhiteSpace(line)
+								|| !line.StartsWith(preflightTag, StringComparison.OrdinalIgnoreCase))
+							{
+								continue;
+							}
+
+							string id = line
+								.Replace($"{_hordeClient?.ServerUrl.ToString().Trim().TrimEnd('/')}/job/", String.Empty, StringComparison.OrdinalIgnoreCase)
+								.Replace($"{preflightTag} ", String.Empty, StringComparison.OrdinalIgnoreCase);
+
+							if (String.IsNullOrWhiteSpace(columns[PreflightColumn.Index]))
+							{
+								columns[PreflightColumn.Index] = id;
+							}
+							else
+							{
+								columns[PreflightColumn.Index] += ", " + id;
+							}
+						}
+					}
+				}
+			}
+
 			for (int columnIdx = 1; columnIdx < BuildList.Columns.Count; columnIdx++)
 			{
 				item.SubItems.Add(new ListViewItem.ListViewSubItem(item, columns[columnIdx] ?? ""));
@@ -2096,8 +2214,8 @@ namespace UnrealGameSync
 				}
 
 				idx++;
-				if (idx >= _sortedChangeNumbers.Count 
-					|| !_perforceMonitor.TryGetChangeDetails(_sortedChangeNumbers[idx], out details) 
+				if (idx >= _sortedChangeNumbers.Count
+					|| !_perforceMonitor.TryGetChangeDetails(_sortedChangeNumbers[idx], out details)
 					|| details.ContainsCode)
 				{
 					return null;
@@ -2452,6 +2570,14 @@ namespace UnrealGameSync
 			{
 				TextRenderer.DrawText(e.Graphics, e.SubItem.Text, currentFont, e.Bounds, textColor, TextFormatFlags.EndEllipsis | TextFormatFlags.SingleLine | TextFormatFlags.VerticalCenter | TextFormatFlags.NoPrefix);
 			}
+			else if (_showReviewerColumn && e.ColumnIndex == ReviewerColumn.Index)
+			{
+				TextRenderer.DrawText(e.Graphics, e.SubItem.Text, currentFont, e.Bounds, textColor, TextFormatFlags.EndEllipsis | TextFormatFlags.SingleLine | TextFormatFlags.VerticalCenter | TextFormatFlags.NoPrefix);
+			}
+			else if (_showPreflightColumn && e.ColumnIndex == PreflightColumn.Index)
+			{
+				TextRenderer.DrawText(e.Graphics, e.SubItem.Text, currentFont, e.Bounds, textColor, TextFormatFlags.EndEllipsis | TextFormatFlags.SingleLine | TextFormatFlags.VerticalCenter | TextFormatFlags.NoPrefix);
+			}
 			else if (e.ColumnIndex == DescriptionColumn.Index)
 			{
 				ChangeLayoutInfo layout = GetChangeLayoutInfo(change);
@@ -5114,6 +5240,58 @@ namespace UnrealGameSync
 			}
 		}
 
+		private void BuildListContextMenu_OpenPreflight_Click(object sender, EventArgs e)
+		{
+			if (_contextMenuChange != null && _showPreflightColumn)
+			{
+				try
+				{
+					StringReader reader = new StringReader(_contextMenuChange.Description);
+					while (reader.Peek() != -1)
+					{
+						string? line = reader.ReadLine();
+						if (String.IsNullOrWhiteSpace(line))
+						{
+							continue;
+						}
+
+						List<string> preflightGuids = new List<string>();
+
+						foreach (string preflightTag in _preflightTags)
+						{
+							if (String.IsNullOrWhiteSpace(preflightTag)
+								|| String.IsNullOrWhiteSpace(line)
+								|| !line.StartsWith(preflightTag, StringComparison.OrdinalIgnoreCase))
+							{
+								continue;
+							}
+
+							string id = line
+								.Replace($"{_hordeClient?.ServerUrl.ToString().Trim().TrimEnd('/')}/job/", String.Empty, StringComparison.OrdinalIgnoreCase)
+								.Replace($"{preflightTag} ", String.Empty, StringComparison.OrdinalIgnoreCase);
+
+							if (id.Equals("none", StringComparison.OrdinalIgnoreCase) || id.Equals("skip", StringComparison.OrdinalIgnoreCase))
+							{
+								continue;
+							}
+
+							preflightGuids.Add(id);
+						}
+
+						foreach (string preflightGuid in preflightGuids)
+						{
+							Utility.OpenUrl($"{_hordeClient?.ServerUrl.ToString().Trim().TrimEnd('/')}/job/{preflightGuid}");
+						}
+					}
+				}
+				catch
+				{
+					MessageBox.Show("Unable to open Horde server.");
+				}
+			}
+		}
+
+
 		private void BuildListContextMenu_AddStar_Click(object sender, EventArgs e)
 		{
 			if (_contextMenuChange != null)
@@ -5350,7 +5528,7 @@ namespace UnrealGameSync
 			List<string> defaultEditorArgumentDefinitions = new List<string>();
 			ConfigUtils.GetProjectSettings(_perforceMonitor.LatestProjectConfigFile, SelectedProjectIdentifier, "DefaultEditorArgument", defaultEditorArgumentDefinitions);
 
-			List<LockableEditorArgument> defaultEditorArguments = new List<LockableEditorArgument>(); 
+			List<LockableEditorArgument> defaultEditorArguments = new List<LockableEditorArgument>();
 			foreach (string editorArgumentDefinition in defaultEditorArgumentDefinitions.Distinct())
 			{
 				LockableEditorArgument? editorArgument;
@@ -5371,9 +5549,9 @@ namespace UnrealGameSync
 
 			foreach (LockableEditorArgument editorArgument in _settings.EditorArguments)
 			{
-				currentEditorArguments.Add( new LockableEditorArgument(editorArgument.Name, editorArgument.Enabled) );
+				currentEditorArguments.Add(new LockableEditorArgument(editorArgument.Name, editorArgument.Enabled));
 			}
-			
+
 			foreach (LockableEditorArgument defaultEditorArgument in defaultEditorArguments)
 			{
 				// Check to see if the user already has this default argument in their list
