Moved the default parameters away from the defaut ini file.
Changed the columns settings to be in the correct method instead of the constructor.

diff --git a/Engine/Programs/UnrealGameSync/UnrealGameSync.ini b/Engine/Programs/UnrealGameSync/UnrealGameSync.ini
--- a/Engine/Programs/UnrealGameSync/UnrealGameSync.ini
+++ b/Engine/Programs/UnrealGameSync/UnrealGameSync.ini
@@ -45,8 +45,6 @@ BuildUnrealHeaderTool=0
 EditorTarget=UnrealEditor
 +DefaultEditorArgument=(Name="-log", Enabled=false, Locked=false)
 +DefaultEditorArgument=(Name="-fastload", Enabled=false, Locked=false)
-ReviewTag=#rb
-PreflightTag=#preflight
 
 ; Set this to a network share which contains the SDK installers for your site
 SdkInstallerDir= 
diff --git a/Engine/Source/Programs/UnrealGameSync/UnrealGameSync/Controls/WorkspaceControl.cs b/Engine/Source/Programs/UnrealGameSync/UnrealGameSync/Controls/WorkspaceControl.cs
--- a/Engine/Source/Programs/UnrealGameSync/UnrealGameSync/Controls/WorkspaceControl.cs
+++ b/Engine/Source/Programs/UnrealGameSync/UnrealGameSync/Controls/WorkspaceControl.cs
@@ -256,8 +256,8 @@ namespace UnrealGameSync
 		bool _updateBuildMetadataPosted;
 		bool _updateReviewsPosted;
 
-		private readonly bool _showReviewerColumn = true;
-		private readonly bool _showPreflightColumn = true;
+		private bool _showReviewerColumn = true;
+		private bool _showPreflightColumn = true;
 
 		private ISet<string> _reviewTags = new HashSet<string>(StringComparer.OrdinalIgnoreCase);
 		private ISet<string> _preflightTags = new HashSet<string>(StringComparer.OrdinalIgnoreCase);
@@ -297,6 +297,9 @@ namespace UnrealGameSync
 			_projectSettings = settings.FindOrAddProjectSettings(openProjectInfo.ProjectInfo, openProjectInfo.WorkspaceSettings, _logger);
 			_hordeClient = serviceProvider.GetService<IHordeClient>();
 
+			_reviewTags.Clear();
+			_preflightTags.Clear();
+
 			ConfigFile projectConfigFile = openProjectInfo.LatestProjectConfigFile;
 			if (projectConfigFile != null)
 			{
@@ -324,26 +327,6 @@ namespace UnrealGameSync
 				}
 			}
 
-			_showPreflightColumn = _preflightTags.Any() &&_hordeClient != null;
-			
-
-			if (_showPreflightColumn)
-			{
-				BuildListContextMenu_OpenPreflight.Enabled = true;
-			}
-			else
-			{
-				BuildList.Columns.Remove(PreflightColumn);
-				BuildListContextMenu_OpenPreflight.Enabled = false;
-			}
-
-			_showReviewerColumn = _reviewTags.Any();
-			if (!_showReviewerColumn)
-			{
-				BuildList.Columns.Remove(ReviewerColumn);
-			}
-
-
 			DesiredTaskbarState = Tuple.Create(TaskbarState.NoProgress, 0.0f);
 
 			System.Reflection.PropertyInfo doubleBufferedProperty = typeof(Control).GetProperty("DoubleBuffered", System.Reflection.BindingFlags.NonPublic | System.Reflection.BindingFlags.Instance)!;
@@ -553,6 +536,25 @@ namespace UnrealGameSync
 					BuildList.Columns.Remove(customColumn);
 				}
 
+				// check the reviewer and preflight columns
+				_showPreflightColumn = _preflightTags.Any() && _hordeClient != null;
+
+				if (_showPreflightColumn)
+				{
+					BuildListContextMenu_OpenPreflight.Enabled = true;
+				}
+				else
+				{
+					BuildList.Columns.Remove(PreflightColumn);
+					BuildListContextMenu_OpenPreflight.Enabled = false;
+				}
+
+				_showReviewerColumn = _reviewTags.Any();
+				if (!_showReviewerColumn)
+				{
+					BuildList.Columns.Remove(ReviewerColumn);
+				}
+
 				Dictionary<string, ColumnHeader> nameToColumn = new Dictionary<string, ColumnHeader>();
 				foreach (ColumnHeader? column in BuildList.Columns)
 				{
