[NO NON-PII COMMIT MESSAGE]

diff --git a/Engine/Source/Programs/UnrealGameSync/UnrealGameSyncShared/OpenProjectInfo.cs b/Engine/Source/Programs/UnrealGameSync/UnrealGameSyncShared/OpenProjectInfo.cs
--- a/Engine/Source/Programs/UnrealGameSync/UnrealGameSyncShared/OpenProjectInfo.cs
+++ b/Engine/Source/Programs/UnrealGameSync/UnrealGameSyncShared/OpenProjectInfo.cs
@@ -217,9 +217,9 @@ namespace UnrealGameSync
 				UserWorkspaceSettings userWorkspaceSettings = userSettings.FindOrAddWorkspaceSettings(branchDirectoryName, perforceSettings.ServerAndPort, perforceSettings.UserName, perforceSettings.ClientName!, branchPath, projectPath, logger);
 
 				// make sure that the default value is in the list
-				if (userSettings.SyncPrecompiledEditor && !userWorkspaceSettings.Archives.Any(a => a.Type.Equals("Editor")))
+				if (!userWorkspaceSettings.Archives.Any(a => a.Type.Equals("Editor")))
 				{
-					userWorkspaceSettings.Archives.Add(new ArchiveSettings(true, "Editor", Array.Empty<string>()));
+					userWorkspaceSettings.Archives.Add(new ArchiveSettings(userSettings.SyncPrecompiledEditor, "Editor", Array.Empty<string>()));
 				}
 
 				// Now compute the updated project info
diff --git a/Engine/Source/Programs/UnrealGameSync/UnrealGameSyncShared/UserSettings.cs b/Engine/Source/Programs/UnrealGameSync/UnrealGameSyncShared/UserSettings.cs
--- a/Engine/Source/Programs/UnrealGameSync/UnrealGameSyncShared/UserSettings.cs
+++ b/Engine/Source/Programs/UnrealGameSync/UnrealGameSyncShared/UserSettings.cs
@@ -685,7 +685,7 @@ namespace UnrealGameSync
 			OpenSolutionAfterSync = (_configFile.GetValue("General.OpenSolutionAfterSync", "0") != "0");
 			ShowLogWindow = (_configFile.GetValue("General.ShowLogWindow", false));
 
-			SyncPrecompiledEditor = (_configFile.GetValue("General.SyncPrecompiledEditor", "0") != "0");
+			SyncPrecompiledEditor = _configFile.GetValue("General.SyncPrecompiledEditor", "1") == "1";
 
 			string? autoResolveConflicts = _configFile.GetValue("General.AutoResolveConflicts", "");
 			if (!String.IsNullOrEmpty(autoResolveConflicts))
