UnrealGameSync: Add checkbox for including engine programs in the solution

#rnx

diff --git a/Engine/Source/Programs/UnrealBuildTool/ProjectFiles/ProjectFileGenerator.cs b/Engine/Source/Programs/UnrealBuildTool/ProjectFiles/ProjectFileGenerator.cs
--- a/Engine/Source/Programs/UnrealBuildTool/ProjectFiles/ProjectFileGenerator.cs
+++ b/Engine/Source/Programs/UnrealBuildTool/ProjectFiles/ProjectFileGenerator.cs
@@ -1545,9 +1545,12 @@ namespace UnrealBuildTool
 								bGatherThirdPartySource = true;
 								break;
 
+							case "-PROGRAMS":
+								bIncludeEnginePrograms = true;
+								break;
+
 							case "-NOPROGRAMS":
 								bIncludeEnginePrograms = false;
-								bIncludeDotNetPrograms = false;
 								break;
 
 							case "-GAME":
diff --git a/Engine/Source/Programs/UnrealGameSync/UnrealGameSync/Forms/SyncFilter.cs b/Engine/Source/Programs/UnrealGameSync/UnrealGameSync/Forms/SyncFilter.cs
--- a/Engine/Source/Programs/UnrealGameSync/UnrealGameSync/Forms/SyncFilter.cs
+++ b/Engine/Source/Programs/UnrealGameSync/UnrealGameSync/Forms/SyncFilter.cs
@@ -74,12 +74,14 @@ namespace UnrealGameSync
 			SetExcludedCategories(GlobalControl.CategoriesCheckList, _uniqueIdToCategory, role, syncCategories);
 			GlobalControl.SyncAllProjects.Checked = GlobalFilter.AllProjects ?? false;
 			GlobalControl.IncludeAllProjectsInSolution.Checked = GlobalFilter.AllProjectsInSln ?? false;
+			GlobalControl.IncludeEngineProgramsInSolution.Checked = GlobalFilter.EngineProgramsInSln ?? false;
 
 			WorkspaceSyncCategory.ApplyDelta(syncCategories, WorkspaceFilter.GetCategories());
 			WorkspaceControl.SetView(WorkspaceFilter.View.ToArray());
 			SetExcludedCategories(WorkspaceControl.CategoriesCheckList, _uniqueIdToCategory, role, syncCategories);
 			WorkspaceControl.SyncAllProjects.Checked = WorkspaceFilter.AllProjects ?? GlobalFilter.AllProjects ?? false;
 			WorkspaceControl.IncludeAllProjectsInSolution.Checked = WorkspaceFilter.AllProjectsInSln ?? GlobalFilter.AllProjectsInSln ?? false;
+			WorkspaceControl.IncludeEngineProgramsInSolution.Checked = WorkspaceFilter.EngineProgramsInSln ?? GlobalFilter.EngineProgramsInSln ?? false;
 
 			GlobalControl.CategoriesCheckList.ItemCheck += GlobalControl_CategoriesCheckList_ItemCheck;
 			GlobalControl.SyncAllProjects.CheckStateChanged += GlobalControl_SyncAllProjects_CheckStateChanged;
@@ -161,6 +163,7 @@ namespace UnrealGameSync
 			newGlobalFilter.View.AddRange(GlobalControl.GetView());
 			newGlobalFilter.AllProjects = GlobalControl.SyncAllProjects.Checked;
 			newGlobalFilter.AllProjectsInSln = GlobalControl.IncludeAllProjectsInSolution.Checked;
+			newGlobalFilter.EngineProgramsInSln = GlobalControl.IncludeEngineProgramsInSolution.Checked;
 
 			Dictionary<Guid, bool> globalSyncCategories = GetCategorySettings(GlobalControl.CategoriesCheckList, GlobalFilter.GetCategories());
 			newGlobalFilter.SetCategories(WorkspaceSyncCategory.GetDelta(defaultSyncCategories, globalSyncCategories));
@@ -169,6 +172,7 @@ namespace UnrealGameSync
 			newWorkspaceFilter.View.AddRange(WorkspaceControl.GetView());
 			newWorkspaceFilter.AllProjects = (WorkspaceControl.SyncAllProjects.Checked == newGlobalFilter.AllProjects) ? (bool?)null : WorkspaceControl.SyncAllProjects.Checked;
 			newWorkspaceFilter.AllProjectsInSln = (WorkspaceControl.IncludeAllProjectsInSolution.Checked == newGlobalFilter.AllProjectsInSln) ? (bool?)null : WorkspaceControl.IncludeAllProjectsInSolution.Checked;
+			newWorkspaceFilter.EngineProgramsInSln = (WorkspaceControl.IncludeEngineProgramsInSolution.Checked == newGlobalFilter.EngineProgramsInSln) ? (bool?)null : WorkspaceControl.IncludeEngineProgramsInSolution.Checked;
 
 			Dictionary<Guid, bool> workspaceSyncCategories = GetCategorySettings(WorkspaceControl.CategoriesCheckList, WorkspaceFilter.GetCategories());
 			newWorkspaceFilter.SetCategories(WorkspaceSyncCategory.GetDelta(globalSyncCategories, workspaceSyncCategories));
diff --git a/Engine/Source/Programs/UnrealGameSync/UnrealGameSyncCmd/Program.cs b/Engine/Source/Programs/UnrealGameSync/UnrealGameSyncCmd/Program.cs
--- a/Engine/Source/Programs/UnrealGameSync/UnrealGameSyncCmd/Program.cs
+++ b/Engine/Source/Programs/UnrealGameSync/UnrealGameSyncCmd/Program.cs
@@ -1183,6 +1183,10 @@ namespace UnrealGameSyncCmd
 			[CommandLine("-GpfOnlyCurrent", Value = "false")]
 			public bool? AllProjectsInSln { get; set; } = null;
 
+			[CommandLine("-GpfEnginePrograms", Value = "true")]
+			[CommandLine("-GpfNoEnginePrograms", Value = "false")]
+			public bool? EngineProgramsInSln { get; set; } = null;
+
 			[CommandLine("-Global")]
 			public bool Global { get; set; }
 		}
@@ -1332,6 +1336,7 @@ namespace UnrealGameSyncCmd
 
 				settings.AllProjects = commandOptions.AllProjects ?? settings.AllProjects;
 				settings.AllProjectsInSln = commandOptions.AllProjectsInSln ?? settings.AllProjectsInSln;
+				settings.EngineProgramsInSln = commandOptions.EngineProgramsInSln ?? settings.EngineProgramsInSln;
 			}
 
 			static Guid GetCategoryId(string text, IEnumerable<WorkspaceSyncCategory> syncCategories)
diff --git a/Engine/Source/Programs/UnrealGameSync/UnrealGameSyncShared/GlobalSettings.cs b/Engine/Source/Programs/UnrealGameSync/UnrealGameSyncShared/GlobalSettings.cs
--- a/Engine/Source/Programs/UnrealGameSync/UnrealGameSyncShared/GlobalSettings.cs
+++ b/Engine/Source/Programs/UnrealGameSync/UnrealGameSyncShared/GlobalSettings.cs
@@ -15,6 +15,7 @@ namespace UnrealGameSync
 		public List<string> View { get; init; } = new List<string>();
 		public bool? AllProjects { get; set; }
 		public bool? AllProjectsInSln { get; set; }
+		public bool? EngineProgramsInSln { get; set; }
 
 		public void Reset()
 		{
@@ -23,6 +24,7 @@ namespace UnrealGameSync
 			View.Clear();
 			AllProjects = null;
 			AllProjectsInSln = null;
+			EngineProgramsInSln = null;
 		}
 
 		public void SetCategories(Dictionary<Guid, bool> categories)
diff --git a/Engine/Source/Programs/UnrealGameSync/UnrealGameSyncShared/UserSettings.cs b/Engine/Source/Programs/UnrealGameSync/UnrealGameSyncShared/UserSettings.cs
--- a/Engine/Source/Programs/UnrealGameSync/UnrealGameSyncShared/UserSettings.cs
+++ b/Engine/Source/Programs/UnrealGameSync/UnrealGameSyncShared/UserSettings.cs
@@ -661,6 +661,7 @@ namespace UnrealGameSync
 				coreSettingsData.Filter.SetCategories(GetCategorySettings(configFile.FindSection("General"), "SyncIncludedCategories", "SyncExcludedCategories"));
 				coreSettingsData.Filter.AllProjects = configFile.GetValue("General.SyncAllProjects", false);
 				coreSettingsData.Filter.AllProjectsInSln = configFile.GetValue("General.IncludeAllProjectsInSolution", false);
+				coreSettingsData.Filter.EngineProgramsInSln = configFile.GetValue("General.IncludeEngineProgramsInSolution", false);
 			}
 
 			return new UserSettings(fileName, configFile, coreFileName, coreSettingsData);
@@ -1031,6 +1032,9 @@ namespace UnrealGameSync
 
 				int includeAllProjectsInSolution = workspaceSection.GetValue("IncludeAllProjectsInSolution", -1);
 				currentWorkspace.Filter.AllProjectsInSln = (includeAllProjectsInSolution == 0) ? (bool?)false : (includeAllProjectsInSolution == 1) ? (bool?)true : (bool?)null;
+
+				int includeEngineProgramsInSolution = workspaceSection.GetValue("IncludeEngineProgramsInSolution", -1);
+				currentWorkspace.Filter.EngineProgramsInSln = (includeEngineProgramsInSolution == 0) ? (bool?)false : (includeEngineProgramsInSolution == 1) ? (bool?)true : (bool?)null;
 			}
 		}
 
diff --git a/Engine/Source/Programs/UnrealGameSync/UnrealGameSyncShared/WorkspaceUpdate.cs b/Engine/Source/Programs/UnrealGameSync/UnrealGameSyncShared/WorkspaceUpdate.cs
--- a/Engine/Source/Programs/UnrealGameSync/UnrealGameSyncShared/WorkspaceUpdate.cs
+++ b/Engine/Source/Programs/UnrealGameSync/UnrealGameSyncShared/WorkspaceUpdate.cs
@@ -40,6 +40,7 @@ namespace UnrealGameSync
 		RemoveFilteredFiles = 0x4000,
 		Clobber = 0x8000,
 		Refilter = 0x10000,
+		IncludeEngineProgramsInSolution = 0x20000,
 	}
 
 	public enum WorkspaceUpdateResult
@@ -133,6 +134,10 @@ namespace UnrealGameSync
 			{
 				options |= WorkspaceUpdateOptions.IncludeAllProjectsInSolution;
 			}
+			if (workspaceSettings.Filter.EngineProgramsInSln ?? globalSettings.Filter.EngineProgramsInSln ?? false)
+			{
+				options |= WorkspaceUpdateOptions.IncludeEngineProgramsInSolution;
+			}
 			return options;
 		}
 	}
@@ -1257,13 +1262,21 @@ namespace UnrealGameSync
 
 					StringBuilder commandLine = new StringBuilder();
 					commandLine.AppendFormat("\"{0}\"", FileReference.Combine(project.LocalRootPath, $"GenerateProjectFiles.{ShellScriptExt}"));
-					if ((Context.Options & WorkspaceUpdateOptions.IncludeAllProjectsInSolution) == 0)
+					if (!Context.Options.HasFlag(WorkspaceUpdateOptions.IncludeAllProjectsInSolution))
 					{
 						if (project.LocalFileName.HasExtension(".uproject"))
 						{
 							commandLine.AppendFormat(" -Project=\"{0}\" -Game", project.LocalFileName);
 						}
 					}
+					if (Context.Options.HasFlag(WorkspaceUpdateOptions.IncludeEngineProgramsInSolution))
+					{
+						commandLine.Append(" -Programs");
+					}
+					else
+					{
+						commandLine.Append(" -NoPrograms");
+					}
 					commandLine.Append(" -progress");
 					logger.LogInformation("Generating project files...");
 					logger.LogInformation("gpf> Running {Arguments}", commandLine);
