UGS: Fix static analysis warnings

diff --git a/Engine/Source/Programs/UnrealGameSync/UnrealGameSync/Controls/WorkspaceControl.cs b/Engine/Source/Programs/UnrealGameSync/UnrealGameSync/Controls/WorkspaceControl.cs
--- a/Engine/Source/Programs/UnrealGameSync/UnrealGameSync/Controls/WorkspaceControl.cs
+++ b/Engine/Source/Programs/UnrealGameSync/UnrealGameSync/Controls/WorkspaceControl.cs
@@ -259,8 +259,8 @@ namespace UnrealGameSync
 		private bool _showReviewerColumn = true;
 		private bool _showPreflightColumn = true;
 
-		private ISet<string> _reviewTags = new HashSet<string>(StringComparer.OrdinalIgnoreCase);
-		private ISet<string> _preflightTags = new HashSet<string>(StringComparer.OrdinalIgnoreCase);
+		private readonly ISet<string> _reviewTags = new HashSet<string>(StringComparer.OrdinalIgnoreCase);
+		private readonly ISet<string> _preflightTags = new HashSet<string>(StringComparer.OrdinalIgnoreCase);
 
 		WorkspaceUpdateCallback? _updateCallback;
 
@@ -1485,16 +1485,15 @@ namespace UnrealGameSync
 
 							if (String.IsNullOrWhiteSpace(columns[ReviewerColumn.Index]))
 							{
-								columns[ReviewerColumn.Index] = line.Replace($"{reviewTag} ", String.Empty);
+								columns[ReviewerColumn.Index] = line.Replace($"{reviewTag} ", String.Empty, StringComparison.InvariantCulture);
 							}
 							else
 							{
-								columns[ReviewerColumn.Index] += ", " + line.Replace($"{reviewTag} ", String.Empty);
+								columns[ReviewerColumn.Index] += ", " + line.Replace($"{reviewTag} ", String.Empty, StringComparison.InvariantCulture);
 							}
 						}
 					}
 
-
 					if (_showPreflightColumn)
 					{
 						foreach (string preflightTag in _preflightTags)
@@ -5293,7 +5292,6 @@ namespace UnrealGameSync
 			}
 		}
 
-
 		private void BuildListContextMenu_AddStar_Click(object sender, EventArgs e)
 		{
 			if (_contextMenuChange != null)
