Fixed a bug where an error would occur when moving UGS between monitors with different DPIs.
Also fixed a bug where we would not immediately adjust the UI based on the new monitor DPI.

diff --git a/Engine/Source/Programs/UnrealGameSync/UnrealGameSync/Controls/BuildListControl.cs b/Engine/Source/Programs/UnrealGameSync/UnrealGameSync/Controls/BuildListControl.cs
--- a/Engine/Source/Programs/UnrealGameSync/UnrealGameSync/Controls/BuildListControl.cs
+++ b/Engine/Source/Programs/UnrealGameSync/UnrealGameSync/Controls/BuildListControl.cs
@@ -149,5 +149,13 @@ namespace UnrealGameSync
 			get => base.Font;
 			set => base.Font = value;
 		}
+
+		protected override void Dispose(bool disposing)
+		{
+			Font?.Dispose();
+			Font = null;
+
+			base.Dispose(disposing);
+		}
 	}
 }
diff --git a/Engine/Source/Programs/UnrealGameSync/UnrealGameSync/Controls/WorkspaceControl.cs b/Engine/Source/Programs/UnrealGameSync/UnrealGameSync/Controls/WorkspaceControl.cs
--- a/Engine/Source/Programs/UnrealGameSync/UnrealGameSync/Controls/WorkspaceControl.cs
+++ b/Engine/Source/Programs/UnrealGameSync/UnrealGameSync/Controls/WorkspaceControl.cs
@@ -5065,7 +5065,7 @@ namespace UnrealGameSync
 		private void BuildList_FontChanged(object? sender, EventArgs e)
 		{
 			_buildFont?.Dispose();
-			_buildFont = BuildList.Font;
+			_buildFont = new Font(BuildList.Font!, FontStyle.Regular);
 
 			_selectedBuildFont?.Dispose();
 			_selectedBuildFont = new Font(_buildFont!, FontStyle.Bold);
@@ -5074,6 +5074,11 @@ namespace UnrealGameSync
 			_badgeFont = new Font(_buildFont!.FontFamily, _buildFont.SizeInPoints - 2, FontStyle.Bold);
 		}
 
+		private void WorkspaceControl_DpiChangedAfterParent(object sender, System.EventArgs e)
+		{
+			UpdateBuildMetadataCallback();
+		}
+
 		public void ExecCommand(string description, string statusText, string fileName, string arguments, string workingDir, bool useLogWindow)
 		{
 			if (_workspace != null && !_workspace.IsBusy())
