UnrealGameSync: Don't generate a game specific solution for a project with no Source

#rnx

diff --git a/Engine/Source/Programs/UnrealGameSync/UnrealGameSyncShared/WorkspaceUpdate.cs b/Engine/Source/Programs/UnrealGameSync/UnrealGameSyncShared/WorkspaceUpdate.cs
--- a/Engine/Source/Programs/UnrealGameSync/UnrealGameSyncShared/WorkspaceUpdate.cs
+++ b/Engine/Source/Programs/UnrealGameSync/UnrealGameSyncShared/WorkspaceUpdate.cs
@@ -1266,7 +1266,13 @@ namespace UnrealGameSync
 					{
 						if (project.LocalFileName.HasExtension(".uproject"))
 						{
-							commandLine.AppendFormat(" -Project=\"{0}\" -Game", project.LocalFileName);
+							commandLine.AppendFormat(" -Project=\"{0}\"", project.LocalFileName);
+
+							// Game specific solutions are only valid if a Source folder exists
+							if (DirectoryReference.Exists(DirectoryReference.Combine(project.LocalFileName.Directory, "Source")))
+							{
+								commandLine.Append(" -Game");
+							}
 						}
 					}
 					if (Context.Options.HasFlag(WorkspaceUpdateOptions.IncludeEngineProgramsInSolution))
