UnrealBuildTool: Apply bPrimaryProjectNameFromFolder to -Game generated solutions by adding a suffix

#rnx

diff --git a/Engine/Source/Programs/UnrealBuildTool/ProjectFiles/ProjectFileGenerator.cs b/Engine/Source/Programs/UnrealBuildTool/ProjectFiles/ProjectFileGenerator.cs
--- a/Engine/Source/Programs/UnrealBuildTool/ProjectFiles/ProjectFileGenerator.cs
+++ b/Engine/Source/Programs/UnrealBuildTool/ProjectFiles/ProjectFileGenerator.cs
@@ -951,6 +951,12 @@ namespace UnrealBuildTool
 					}
 				}
 				IntermediateProjectFilesPath = DirectoryReference.Combine(PrimaryProjectPath, "Intermediate", "ProjectFiles");
+
+				// Modify the primary project name from the root folder name
+				if (bPrimaryProjectNameFromFolder)
+				{
+					PrimaryProjectName = $"{PrimaryProjectName}_{Unreal.RootDirectory.GetDirectoryName()}";
+				}
 			}
 			else
 			{
diff --git a/Engine/Source/Programs/UnrealGameSync/UnrealGameSync/Controls/WorkspaceControl.cs b/Engine/Source/Programs/UnrealGameSync/UnrealGameSync/Controls/WorkspaceControl.cs
--- a/Engine/Source/Programs/UnrealGameSync/UnrealGameSync/Controls/WorkspaceControl.cs
+++ b/Engine/Source/Programs/UnrealGameSync/UnrealGameSync/Controls/WorkspaceControl.cs
@@ -5624,26 +5624,17 @@ namespace UnrealGameSync
 			WorkspaceUpdateOptions options = WorkspaceUpdateContext.GetOptionsFromConfig(_settings.Global, _workspaceSettings);
 			FileReference? solutionFileName = null;
 
-			// Default to opening the sln in the project if not including all projects in solution
-			if (!options.HasFlag(WorkspaceUpdateOptions.IncludeAllProjectsInSolution) && SelectedFileName.HasExtension(".uproject"))
-			{
-				solutionFileName = SelectedFileName.ChangeExtension(".sln");
-			}
-
 			// Check PrimaryProjectPath.txt
-			if (solutionFileName == null || !FileReference.Exists(solutionFileName))
+			FileReference primaryProjectPathFileName = FileReference.Combine(BranchDirectoryName, "Engine", "Intermediate", "ProjectFiles", "PrimaryProjectPath.txt");
+			if (FileReference.Exists(primaryProjectPathFileName))
 			{
-				FileReference primaryProjectPathFileName = FileReference.Combine(BranchDirectoryName, "Engine", "Intermediate", "ProjectFiles", "PrimaryProjectPath.txt");
-				if (FileReference.Exists(primaryProjectPathFileName))
+				try
 				{
-					try
-					{
-						solutionFileName = new FileReference(FileReference.ReadAllText(primaryProjectPathFileName).Trim() + ".sln");
-					}
-					catch (Exception ex)
-					{
-						_logger.LogError(ex, "Unable to read '{File}'", primaryProjectPathFileName);
-					}
+					solutionFileName = new FileReference(FileReference.ReadAllText(primaryProjectPathFileName).Trim() + ".sln");
+				}
+				catch (Exception ex)
+				{
+					_logger.LogError(ex, "Unable to read '{File}'", primaryProjectPathFileName);
 				}
 			}
 
@@ -5670,6 +5661,15 @@ namespace UnrealGameSync
 				}
 			}
 
+			// Fallback to opening the sln in the project if not including all projects in solution
+			if (solutionFileName == null || !FileReference.Exists(solutionFileName))
+			{
+				if (!options.HasFlag(WorkspaceUpdateOptions.IncludeAllProjectsInSolution) && SelectedFileName.HasExtension(".uproject"))
+				{
+					solutionFileName = SelectedFileName.ChangeExtension(".sln");
+				}
+			}
+
 			// Fallback to UE#.sln in the branch directory
 			if (solutionFileName == null || !FileReference.Exists(solutionFileName))
 			{
