UnrealGameSync: Resolve remaining intellisense warnings and errors

#rnx

diff --git a/Engine/Source/Programs/UnrealGameSync/PostBadgeStatus/Program.cs b/Engine/Source/Programs/UnrealGameSync/PostBadgeStatus/Program.cs
--- a/Engine/Source/Programs/UnrealGameSync/PostBadgeStatus/Program.cs
+++ b/Engine/Source/Programs/UnrealGameSync/PostBadgeStatus/Program.cs
@@ -2,13 +2,10 @@
 
 using System;
 using System.Collections.Generic;
-using System.Data.SqlClient;
 using System.IO;
-using System.Linq;
 using System.Net;
 using System.Text;
 using System.Threading;
-using System.Threading.Tasks;
 using System.Web.Script.Serialization;
 
 namespace WriteBadgeStatus
@@ -43,7 +40,7 @@ namespace WriteBadgeStatus
 				Project = Project,
 				ArchivePath = ""
 			};
-			if (!int.TryParse(Change, out Build.ChangeNumber))
+			if (!Int32.TryParse(Change, out Build.ChangeNumber))
 			{
 				Console.WriteLine("Change must be an integer!");
 				return 1;
@@ -53,7 +50,7 @@ namespace WriteBadgeStatus
 				Console.WriteLine("Change must be Starting, Failure, Warning, Success, or Skipped!");
 				return 1;
 			}
-			if (!string.IsNullOrWhiteSpace(Metadata))
+			if (!String.IsNullOrWhiteSpace(Metadata))
 			{
 				Build.Metadata = new JavaScriptSerializer().Deserialize<object>(Metadata);
 			}
@@ -68,7 +65,7 @@ namespace WriteBadgeStatus
 				{
 					if (++NumRetries <= 3)
 					{
-						Console.WriteLine(string.Format("An exception was thrown attempting to send the request: {0}, retrying...", ex.Message));
+						Console.WriteLine(String.Format("An exception was thrown attempting to send the request: {0}, retrying...", ex.Message));
 						// Wait 5 seconds and retry;
 						Thread.Sleep(5000);
 					}
@@ -115,21 +112,15 @@ namespace WriteBadgeStatus
 			public string ArchivePath;
 			public object Metadata;
 
-			public bool IsSuccess
-			{
-				get { return Result == BuildDataResult.Success || Result == BuildDataResult.Warning; }
-			}
+			public bool IsSuccess => Result == BuildDataResult.Success || Result == BuildDataResult.Warning;
 
-			public bool IsFailure
-			{
-				get { return Result == BuildDataResult.Failure; }
-			}
+			public bool IsFailure => Result == BuildDataResult.Failure;
 		}
 
 		static int SendRequest(string URI, string Resource, string Method, string RequestBody = null, params string[] QueryParams)
 		{
 			// set up the query string
-			StringBuilder TargetURI = new StringBuilder(string.Format("{0}/api/{1}", URI, Resource));
+			StringBuilder TargetURI = new StringBuilder(String.Format("{0}/api/{1}", URI, Resource));
 			if (QueryParams.Length != 0)
 			{
 				TargetURI.Append("?");
@@ -146,9 +137,8 @@ namespace WriteBadgeStatus
 			Request.ContentType = "application/json";
 			Request.Method = Method;
 
-
 			// Add json to request body
-			if (!string.IsNullOrEmpty(RequestBody))
+			if (!String.IsNullOrEmpty(RequestBody))
 			{
 				if (Method == "POST")
 				{
@@ -171,22 +161,21 @@ namespace WriteBadgeStatus
 						return ((int)Response.StatusCode >= 200 && (int)Response.StatusCode <= 299) ? 0 : 1;
 					}
 				}
-				
 			}
 			catch (WebException ex)
 			{
 				if (ex.Response != null)
 				{
-					throw new Exception(string.Format("Request returned status: {0}, message: {1}", ((HttpWebResponse)ex.Response).StatusCode, ex.Message));
+					throw new Exception(String.Format("Request returned status: {0}, message: {1}", ((HttpWebResponse)ex.Response).StatusCode, ex.Message));
 				}
 				else
 				{
-					throw new Exception(string.Format("Request returned message: {0}", ex.InnerException.Message));
+					throw new Exception(String.Format("Request returned message: {0}", ex.InnerException.Message));
 				}
 			}
 			catch (Exception ex)
 			{
-				throw new Exception(string.Format("Couldn't complete the request, error: {0}", ex.Message));
+				throw new Exception(String.Format("Couldn't complete the request, error: {0}", ex.Message));
 			}
 		}
 	}
diff --git a/Engine/Source/Programs/UnrealGameSync/PostBadgeStatus/Properties/AssemblyInfo.cs b/Engine/Source/Programs/UnrealGameSync/PostBadgeStatus/Properties/AssemblyInfo.cs
--- a/Engine/Source/Programs/UnrealGameSync/PostBadgeStatus/Properties/AssemblyInfo.cs
+++ b/Engine/Source/Programs/UnrealGameSync/PostBadgeStatus/Properties/AssemblyInfo.cs
@@ -1,7 +1,6 @@
-﻿// Copyright Epic Games, Inc. All Rights Reserved.
+// Copyright Epic Games, Inc. All Rights Reserved.
 
 using System.Reflection;
-using System.Runtime.CompilerServices;
 using System.Runtime.InteropServices;
 
 // General Information about an assembly is controlled through the following 
diff --git a/Engine/Source/Programs/UnrealGameSync/UnrealGameSync/UriHandlers/P4Handler.cs b/Engine/Source/Programs/UnrealGameSync/UnrealGameSync/UriHandlers/P4Handler.cs
--- a/Engine/Source/Programs/UnrealGameSync/UnrealGameSync/UriHandlers/P4Handler.cs
+++ b/Engine/Source/Programs/UnrealGameSync/UnrealGameSync/UriHandlers/P4Handler.cs
@@ -1,5 +1,7 @@
 // Copyright Epic Games, Inc. All Rights Reserved.
 
+using System;
+
 namespace UnrealGameSync
 {
 
@@ -11,11 +13,11 @@ namespace UnrealGameSync
 		[UriHandler(true)]
 		public static UriResult P4V(string depotPath)
 		{
-			string commandLine = string.Format("-s \"{0}\"", depotPath);
+			string commandLine = String.Format("-s \"{0}\"", depotPath);
 
 			if (!Utility.SpawnHiddenProcess("p4v.exe", commandLine))
 			{
-				return new UriResult() { Error = string.Format("Error spawning p4v.exe with command line: {0}", commandLine) };
+				return new UriResult() { Error = String.Format("Error spawning p4v.exe with command line: {0}", commandLine) };
 			}
 
 			return new UriResult() { Success = true };
@@ -24,7 +26,7 @@ namespace UnrealGameSync
 		[UriHandler(true)]
 		public static UriResult Timelapse(string depotPath, int line = -1)
 		{
-			string commandLine = string.Format("timelapse {0}{1}", line == -1 ? "" : string.Format(" -l {0} ", line), depotPath);
+			string commandLine = String.Format("timelapse {0}{1}", line == -1 ? "" : String.Format(" -l {0} ", line), depotPath);
 
 			Program.SpawnP4Vc(commandLine);
 
