Do not limit the archives search to the BuildList when syncing to a specific changelist.

diff --git a/Engine/Source/Programs/UnrealGameSync/UnrealGameSync/Controls/WorkspaceControl.cs b/Engine/Source/Programs/UnrealGameSync/UnrealGameSync/Controls/WorkspaceControl.cs
--- a/Engine/Source/Programs/UnrealGameSync/UnrealGameSync/Controls/WorkspaceControl.cs
+++ b/Engine/Source/Programs/UnrealGameSync/UnrealGameSync/Controls/WorkspaceControl.cs
@@ -917,12 +917,12 @@ namespace UnrealGameSync
 			}
 		}
 
-		void StartSync(int changeNumber)
+		void StartSync(int changeNumber, bool skipKnownChangeLists)
 		{
-			StartSync(changeNumber, false, null);
+			StartSync(changeNumber, false, null, skipKnownChangeLists);
 		}
 
-		void StartSync(int changeNumber, bool syncOnly, WorkspaceUpdateCallback? callback)
+		void StartSync(int changeNumber, bool syncOnly, WorkspaceUpdateCallback? callback, bool skipKnownChangeLists)
 		{
 			WorkspaceUpdateOptions options = WorkspaceUpdateOptions.Sync | WorkspaceUpdateOptions.SyncArchives | WorkspaceUpdateOptions.GenerateProjectFiles;
 			if (!syncOnly)
@@ -940,15 +940,15 @@ namespace UnrealGameSync
 					options |= WorkspaceUpdateOptions.OpenSolutionAfterSync;
 				}
 			}
-			StartWorkspaceUpdate(changeNumber, options, callback);
+			StartWorkspaceUpdate(changeNumber, options, callback, skipKnownChangeLists);
 		}
 
-		void StartWorkspaceUpdate(int changeNumber, WorkspaceUpdateOptions options)
+		void StartWorkspaceUpdate(int changeNumber, WorkspaceUpdateOptions options, bool skipKnownChangeLists)
 		{
-			StartWorkspaceUpdate(changeNumber, options, null);
+			StartWorkspaceUpdate(changeNumber, options, null, skipKnownChangeLists);
 		}
 
-		void StartWorkspaceUpdate(int changeNumber, WorkspaceUpdateOptions options, WorkspaceUpdateCallback? callback)
+		void StartWorkspaceUpdate(int changeNumber, WorkspaceUpdateOptions options, WorkspaceUpdateCallback? callback, bool skipKnownChangeLists)
 		{
 			if ((options & (WorkspaceUpdateOptions.Sync | WorkspaceUpdateOptions.Build)) != 0 && GetProcessesRunningInWorkspace().Length > 0)
 			{
@@ -981,7 +981,7 @@ namespace UnrealGameSync
 				List<IArchiveChannel> selectedArchiveChannels = GetSelectedArchiveChannels(archives);
 				foreach (IArchiveChannel archiveChannel in selectedArchiveChannels)
 				{
-					IArchive? archivePath = GetArchiveForChangeNumber(archiveChannel, changeNumber);
+					IArchive? archivePath = GetArchiveForChangeNumber(archiveChannel, changeNumber, skipKnownChangeLists);
 					if (archivePath == null)
 					{
 						MessageBox.Show(String.Format("There are no compiled \"{0}\" binaries for this change. To sync it, you must disable syncing of precompiled editor binaries.", archiveChannel.Name));
@@ -1004,7 +1004,7 @@ namespace UnrealGameSync
 
 				if ((selectedArchiveChannels.Count == 1) && (!selectedArchiveChannels[0].IgnoreRequiredBadges) && (_projectSettings.RequiredBadges.Count > 0))
 				{
-					if (!CanSyncChange(changeNumber))
+					if (!CanSyncChange(changeNumber, skipKnownChangeLists))
 					{
 						MessageBox.Show(String.Format("This change has not passed the required tests \"{0}\". To sync it, you must disable syncing of precompiled editor binaries.", String.Join(", ", _projectSettings.RequiredBadges)));
 						return;
@@ -2170,7 +2170,7 @@ namespace UnrealGameSync
 			}
 		}
 
-		private IArchive? GetArchiveForChangeNumber(IArchiveChannel archiveChannel, int changeNumber)
+		private IArchive? GetArchiveForChangeNumber(IArchiveChannel archiveChannel, int changeNumber, bool skipKnownChangeLists)
 		{
 			Dictionary<int, IArchive?>? changeNumberToArchive;
 			if (!_archiveToChangeNumberToArchiveKey.TryGetValue(archiveChannel.Name, out changeNumberToArchive))
@@ -2178,15 +2178,18 @@ namespace UnrealGameSync
 				changeNumberToArchive = new Dictionary<int, IArchive?>();
 				_archiveToChangeNumberToArchiveKey[archiveChannel.Name] = changeNumberToArchive;
 			}
-			return GetArchiveForChangeNumber(archiveChannel, changeNumber, changeNumber, changeNumberToArchive);
+			return GetArchiveForChangeNumber(archiveChannel, changeNumber, changeNumber, changeNumberToArchive, skipKnownChangeLists);
 		}
 
-		private IArchive? GetArchiveForChangeNumber(IArchiveChannel archiveChannel, int changeNumber, int maxChangeNumber, Dictionary<int, IArchive?> changeNumberToArchive)
+		private IArchive? GetArchiveForChangeNumber(IArchiveChannel archiveChannel, int changeNumber, int maxChangeNumber, Dictionary<int, IArchive?> changeNumberToArchive, bool skipKnownChangeLists)
 		{
 			IArchive? archive;
 			if (!changeNumberToArchive.TryGetValue(changeNumber, out archive))
 			{
-				archive = GetArchiveForChangeNumberUncached(archiveChannel, changeNumber, maxChangeNumber);
+				archive = skipKnownChangeLists
+				? archiveChannel.TryGetArchiveForChangeNumber(changeNumber, maxChangeNumber)
+				: GetArchiveForChangeNumberUncached(archiveChannel, changeNumber, maxChangeNumber);
+
 				changeNumberToArchive.Add(changeNumber, archive);
 			}
 			return archive;
@@ -2229,7 +2232,7 @@ namespace UnrealGameSync
 			return Color.FromArgb((int)(first.R + (second.R - first.R) * t), (int)(first.G + (second.G - first.G) * t), (int)(first.B + (second.B - first.B) * t));
 		}
 
-		private bool CanSyncChange(int changeNumber)
+		private bool CanSyncChange(int changeNumber, bool skipKnownChangeLists)
 		{
 			if (_perforceMonitor == null)
 			{
@@ -2237,7 +2240,7 @@ namespace UnrealGameSync
 			}
 
 			List<IArchiveChannel> selectedArchives = GetSelectedArchiveChannels(GetArchiveChannels());
-			bool returnValue = selectedArchives.Count == 0 || selectedArchives.All(x => GetArchiveForChangeNumber(x, changeNumber) != null);
+			bool returnValue = selectedArchives.Count == 0 || selectedArchives.All(x => GetArchiveForChangeNumber(x, changeNumber, skipKnownChangeLists) != null);
 
 			// If we can sync this change and we're using archives, let's double check for any required badges
 			if (returnValue && (selectedArchives.Count == 1) && (!selectedArchives[0].IgnoreRequiredBadges) && (_projectSettings.RequiredBadges.Count > 0))
@@ -2474,7 +2477,7 @@ namespace UnrealGameSync
 
 			Font currentFont = (change.Number == _workspace.PendingChangeNumber || change.Number == _workspace.CurrentChangeNumber) ? _selectedBuildFont! : _buildFont!;
 
-			bool allowSync = CanSyncChange(change.Number);
+			bool allowSync = CanSyncChange(change.Number, false);
 			int badgeAlpha = allowSync ? 255 : 128;
 			Color textColor = (allowSync || change.Number == _workspace.PendingChangeNumber || change.Number == _workspace.CurrentChangeNumber || _workspace.State.AdditionalChangeNumbers.Contains(change.Number)) ? SystemColors.WindowText : Blend(SystemColors.Window, SystemColors.WindowText, 0.25f);
 
@@ -3381,7 +3384,7 @@ namespace UnrealGameSync
 						}
 						else
 						{
-							StartSync(change.Number, false, null);
+							StartSync(change.Number, false, null, false);
 						}
 					}
 				}
@@ -4569,9 +4572,9 @@ namespace UnrealGameSync
 						else
 						{
 							Rectangle syncBadgeRectangle = GetSyncBadgeRectangle(hitTest.Item.SubItems[StatusColumn.Index].Bounds);
-							if (syncBadgeRectangle.Contains(args.Location) && CanSyncChange(change.Number))
+							if (syncBadgeRectangle.Contains(args.Location) && CanSyncChange(change.Number, false))
 							{
-								StartSync(change.Number, false, null);
+								StartSync(change.Number, false, null, false);
 							}
 						}
 
@@ -5095,13 +5098,13 @@ namespace UnrealGameSync
 			}
 		}
 
-		public void SyncChange(int changeNumber, bool syncOnly, WorkspaceUpdateCallback? callback)
+		public void SyncChange(int changeNumber, bool syncOnly, WorkspaceUpdateCallback? callback, bool skipKnownChangeLists)
 		{
 			if (_workspace != null)
 			{
 				_owner.ShowAndActivate();
 				SelectChange(changeNumber);
-				StartSync(changeNumber, syncOnly, callback);
+				StartSync(changeNumber, syncOnly, callback, skipKnownChangeLists);
 			}
 		}
 
@@ -5134,7 +5137,7 @@ namespace UnrealGameSync
 				}
 				else if (changeNumber >= _perforceMonitor.LastChangeByCurrentUser || MessageBox.Show(String.Format("The changelist that would be synced is before the last change you submitted.\n\nIf you continue, your changes submitted after CL {0} will be locally removed from your workspace until you can sync past them again.\n\nAre you sure you want to continue?", changeNumber), "Local changes will be removed", MessageBoxButtons.YesNo) == DialogResult.Yes)
 				{
-					SyncChange(changeNumber, false, callback);
+					SyncChange(changeNumber, false, callback, false);
 				}
 			}
 		}
@@ -5167,7 +5170,7 @@ namespace UnrealGameSync
 		{
 			if (_contextMenuChange != null)
 			{
-				StartSync(_contextMenuChange.Number);
+				StartSync(_contextMenuChange.Number, false);
 			}
 		}
 
@@ -5175,7 +5178,7 @@ namespace UnrealGameSync
 		{
 			if (_contextMenuChange != null)
 			{
-				StartWorkspaceUpdate(_contextMenuChange.Number, WorkspaceUpdateOptions.Sync | WorkspaceUpdateOptions.ContentOnly);
+				StartWorkspaceUpdate(_contextMenuChange.Number, WorkspaceUpdateOptions.Sync | WorkspaceUpdateOptions.ContentOnly, false);
 			}
 		}
 
@@ -5183,23 +5186,23 @@ namespace UnrealGameSync
 		{
 			if (_contextMenuChange != null)
 			{
-				StartWorkspaceUpdate(_contextMenuChange.Number, WorkspaceUpdateOptions.SyncSingleChange);
+				StartWorkspaceUpdate(_contextMenuChange.Number, WorkspaceUpdateOptions.SyncSingleChange, false);
 			}
 		}
 
 		private void BuildListContextMenu_Build_Click(object sender, EventArgs e)
 		{
-			StartWorkspaceUpdate(_workspace.CurrentChangeNumber, WorkspaceUpdateOptions.SyncArchives | WorkspaceUpdateOptions.Build);
+			StartWorkspaceUpdate(_workspace.CurrentChangeNumber, WorkspaceUpdateOptions.SyncArchives | WorkspaceUpdateOptions.Build, false);
 		}
 
 		private void BuildListContextMenu_Rebuild_Click(object sender, EventArgs e)
 		{
-			StartWorkspaceUpdate(_workspace.CurrentChangeNumber, WorkspaceUpdateOptions.SyncArchives | WorkspaceUpdateOptions.Build | WorkspaceUpdateOptions.Clean);
+			StartWorkspaceUpdate(_workspace.CurrentChangeNumber, WorkspaceUpdateOptions.SyncArchives | WorkspaceUpdateOptions.Build | WorkspaceUpdateOptions.Clean, false);
 		}
 
 		private void BuildListContextMenu_GenerateProjectFiles_Click(object sender, EventArgs e)
 		{
-			StartWorkspaceUpdate(_workspace.CurrentChangeNumber, WorkspaceUpdateOptions.GenerateProjectFiles);
+			StartWorkspaceUpdate(_workspace.CurrentChangeNumber, WorkspaceUpdateOptions.GenerateProjectFiles, false);
 		}
 
 		private void BuildListContextMenu_CancelSync_Click(object sender, EventArgs e)
@@ -5767,7 +5770,7 @@ namespace UnrealGameSync
 					{
 						options |= WorkspaceUpdateOptions.AutoResolveChanges;
 					}
-					StartWorkspaceUpdate(changeNumber, options);
+					StartWorkspaceUpdate(changeNumber, options, false);
 				}
 			}
 		}
@@ -5829,7 +5832,7 @@ namespace UnrealGameSync
 				if (change != null)
 				{
 					EventSummary? summary = _eventMonitor.GetSummaryForChange(change.Number);
-					if (CanSyncChange(change.Number)
+					if (CanSyncChange(change.Number, false)
 						&& CanSyncChangeType(changeType, change.Number, summary))
 					{
 						if (changeType.FindNewestGoodContent)
@@ -5874,7 +5877,7 @@ namespace UnrealGameSync
 					break;
 				}
 
-				if (!CanSyncChange(nextChangeNumber))
+				if (!CanSyncChange(nextChangeNumber, false))
 				{
 					break;
 				}
@@ -6397,7 +6400,7 @@ namespace UnrealGameSync
 			using ChangelistWindow childWindow = new ChangelistWindow((_workspace == null) ? -1 : _workspace.CurrentChangeNumber);
 			if (childWindow.ShowDialog() == DialogResult.OK)
 			{
-				StartSync(childWindow.ChangeNumber);
+				StartSync(childWindow.ChangeNumber, true);
 			}
 		}
 
@@ -6744,7 +6747,7 @@ namespace UnrealGameSync
 			{
 				_owner.ShowAndActivate();
 				SelectChange(bisectChange);
-				StartSync(bisectChange);
+				StartSync(bisectChange, false);
 			}
 		}
 
diff --git a/Engine/Source/Programs/UnrealGameSync/UnrealGameSync/Forms/MainWindow.cs b/Engine/Source/Programs/UnrealGameSync/UnrealGameSync/Forms/MainWindow.cs
--- a/Engine/Source/Programs/UnrealGameSync/UnrealGameSync/Forms/MainWindow.cs
+++ b/Engine/Source/Programs/UnrealGameSync/UnrealGameSync/Forms/MainWindow.cs
@@ -297,7 +297,8 @@ namespace UnrealGameSync
 				}
 				else
 				{
-					workspace.SyncChange(changelist, true, result => StartExecCommandAfterSync(workspace, result, command));
+					// I believe we want to bypass the known CL here
+					workspace.SyncChange(changelist, true, result => StartExecCommandAfterSync(workspace, result, command), true);
 				}
 			}
 		}
