UnrealGameSync: Undo all the .sln changes and add a new option to pass -Game which is off by default

#rnx

diff --git a/Engine/Source/Programs/UnrealGameSync/UnrealGameSync/Forms/SyncFilter.cs b/Engine/Source/Programs/UnrealGameSync/UnrealGameSync/Forms/SyncFilter.cs
--- a/Engine/Source/Programs/UnrealGameSync/UnrealGameSync/Forms/SyncFilter.cs
+++ b/Engine/Source/Programs/UnrealGameSync/UnrealGameSync/Forms/SyncFilter.cs
@@ -74,18 +74,19 @@ namespace UnrealGameSync
 			SetExcludedCategories(GlobalControl.CategoriesCheckList, _uniqueIdToCategory, role, syncCategories);
 			GlobalControl.SyncAllProjects.Checked = GlobalFilter.AllProjects ?? false;
 			GlobalControl.IncludeAllProjectsInSolution.Checked = GlobalFilter.AllProjectsInSln ?? false;
-			GlobalControl.IncludeEngineProgramsInSolution.Checked = GlobalFilter.EngineProgramsInSln ?? false;
+			GlobalControl.GenerateUprojectSpecificSolution.Checked = GlobalFilter.UprojectSpecificSln ?? false;
 
 			WorkspaceSyncCategory.ApplyDelta(syncCategories, WorkspaceFilter.GetCategories());
 			WorkspaceControl.SetView(WorkspaceFilter.View.ToArray());
 			SetExcludedCategories(WorkspaceControl.CategoriesCheckList, _uniqueIdToCategory, role, syncCategories);
 			WorkspaceControl.SyncAllProjects.Checked = WorkspaceFilter.AllProjects ?? GlobalFilter.AllProjects ?? false;
 			WorkspaceControl.IncludeAllProjectsInSolution.Checked = WorkspaceFilter.AllProjectsInSln ?? GlobalFilter.AllProjectsInSln ?? false;
-			WorkspaceControl.IncludeEngineProgramsInSolution.Checked = WorkspaceFilter.EngineProgramsInSln ?? GlobalFilter.EngineProgramsInSln ?? false;
+			WorkspaceControl.GenerateUprojectSpecificSolution.Checked = WorkspaceFilter.UprojectSpecificSln ?? GlobalFilter.UprojectSpecificSln ?? false;
 
 			GlobalControl.CategoriesCheckList.ItemCheck += GlobalControl_CategoriesCheckList_ItemCheck;
 			GlobalControl.SyncAllProjects.CheckStateChanged += GlobalControl_SyncAllProjects_CheckStateChanged;
 			GlobalControl.IncludeAllProjectsInSolution.CheckStateChanged += GlobalControl_IncludeAllProjectsInSolution_CheckStateChanged;
+			GlobalControl.GenerateUprojectSpecificSolution.CheckStateChanged += GlobalControl_GenerateUprojectSpecificSolution_CheckStateChanged;
 		}
 
 		private void GlobalControl_CategoriesCheckList_ItemCheck(object? sender, ItemCheckEventArgs e)
@@ -113,6 +114,12 @@ namespace UnrealGameSync
 			WorkspaceControl.IncludeAllProjectsInSolution.Checked = GlobalControl.IncludeAllProjectsInSolution.Checked;
 		}
 
+		private void GlobalControl_GenerateUprojectSpecificSolution_CheckStateChanged(object? sender, EventArgs e)
+		{
+			WorkspaceControl.GenerateUprojectSpecificSolution.Checked = GlobalControl.GenerateUprojectSpecificSolution.Checked;
+		}
+
+
 		private static void SetExcludedCategories(
 			CheckedListBox listBox, 
 			Dictionary<Guid, WorkspaceSyncCategory> uniqueIdToFilter,
@@ -163,7 +170,7 @@ namespace UnrealGameSync
 			newGlobalFilter.View.AddRange(GlobalControl.GetView());
 			newGlobalFilter.AllProjects = GlobalControl.SyncAllProjects.Checked;
 			newGlobalFilter.AllProjectsInSln = GlobalControl.IncludeAllProjectsInSolution.Checked;
-			newGlobalFilter.EngineProgramsInSln = GlobalControl.IncludeEngineProgramsInSolution.Checked;
+			newGlobalFilter.UprojectSpecificSln = GlobalControl.GenerateUprojectSpecificSolution.Checked;
 
 			Dictionary<Guid, bool> globalSyncCategories = GetCategorySettings(GlobalControl.CategoriesCheckList, GlobalFilter.GetCategories());
 			newGlobalFilter.SetCategories(WorkspaceSyncCategory.GetDelta(defaultSyncCategories, globalSyncCategories));
@@ -172,7 +179,7 @@ namespace UnrealGameSync
 			newWorkspaceFilter.View.AddRange(WorkspaceControl.GetView());
 			newWorkspaceFilter.AllProjects = (WorkspaceControl.SyncAllProjects.Checked == newGlobalFilter.AllProjects) ? (bool?)null : WorkspaceControl.SyncAllProjects.Checked;
 			newWorkspaceFilter.AllProjectsInSln = (WorkspaceControl.IncludeAllProjectsInSolution.Checked == newGlobalFilter.AllProjectsInSln) ? (bool?)null : WorkspaceControl.IncludeAllProjectsInSolution.Checked;
-			newWorkspaceFilter.EngineProgramsInSln = (WorkspaceControl.IncludeEngineProgramsInSolution.Checked == newGlobalFilter.EngineProgramsInSln) ? (bool?)null : WorkspaceControl.IncludeEngineProgramsInSolution.Checked;
+			newWorkspaceFilter.UprojectSpecificSln = (WorkspaceControl.GenerateUprojectSpecificSolution.Checked == newGlobalFilter.UprojectSpecificSln) ? (bool?)null : WorkspaceControl.GenerateUprojectSpecificSolution.Checked;
 
 			Dictionary<Guid, bool> workspaceSyncCategories = GetCategorySettings(WorkspaceControl.CategoriesCheckList, WorkspaceFilter.GetCategories());
 			newWorkspaceFilter.SetCategories(WorkspaceSyncCategory.GetDelta(globalSyncCategories, workspaceSyncCategories));
diff --git a/Engine/Source/Programs/UnrealGameSync/UnrealGameSyncCmd/Program.cs b/Engine/Source/Programs/UnrealGameSync/UnrealGameSyncCmd/Program.cs
--- a/Engine/Source/Programs/UnrealGameSync/UnrealGameSyncCmd/Program.cs
+++ b/Engine/Source/Programs/UnrealGameSync/UnrealGameSyncCmd/Program.cs
@@ -1183,9 +1183,9 @@ namespace UnrealGameSyncCmd
 			[CommandLine("-GpfOnlyCurrent", Value = "false")]
 			public bool? AllProjectsInSln { get; set; } = null;
 
-			[CommandLine("-GpfEnginePrograms", Value = "true")]
-			[CommandLine("-GpfNoEnginePrograms", Value = "false")]
-			public bool? EngineProgramsInSln { get; set; } = null;
+			[CommandLine("-GpfMinimalSln", Value = "true")]
+			[CommandLine("-GpfFullSln", Value = "false")]
+			public bool? UprojectSpecificSln { get; set; } = null;
 
 			[CommandLine("-Global")]
 			public bool Global { get; set; }
@@ -1336,7 +1336,7 @@ namespace UnrealGameSyncCmd
 
 				settings.AllProjects = commandOptions.AllProjects ?? settings.AllProjects;
 				settings.AllProjectsInSln = commandOptions.AllProjectsInSln ?? settings.AllProjectsInSln;
-				settings.EngineProgramsInSln = commandOptions.EngineProgramsInSln ?? settings.EngineProgramsInSln;
+				settings.UprojectSpecificSln = commandOptions.UprojectSpecificSln ?? settings.UprojectSpecificSln;
 			}
 
 			static Guid GetCategoryId(string text, IEnumerable<WorkspaceSyncCategory> syncCategories)
diff --git a/Engine/Source/Programs/UnrealGameSync/UnrealGameSyncShared/GlobalSettings.cs b/Engine/Source/Programs/UnrealGameSync/UnrealGameSyncShared/GlobalSettings.cs
--- a/Engine/Source/Programs/UnrealGameSync/UnrealGameSyncShared/GlobalSettings.cs
+++ b/Engine/Source/Programs/UnrealGameSync/UnrealGameSyncShared/GlobalSettings.cs
@@ -15,7 +15,7 @@ namespace UnrealGameSync
 		public List<string> View { get; init; } = new List<string>();
 		public bool? AllProjects { get; set; }
 		public bool? AllProjectsInSln { get; set; }
-		public bool? EngineProgramsInSln { get; set; }
+		public bool? UprojectSpecificSln { get; set; }
 
 		public void Reset()
 		{
@@ -24,7 +24,7 @@ namespace UnrealGameSync
 			View.Clear();
 			AllProjects = null;
 			AllProjectsInSln = null;
-			EngineProgramsInSln = null;
+			UprojectSpecificSln = null;
 		}
 
 		public void SetCategories(Dictionary<Guid, bool> categories)
diff --git a/Engine/Source/Programs/UnrealGameSync/UnrealGameSyncShared/UserSettings.cs b/Engine/Source/Programs/UnrealGameSync/UnrealGameSyncShared/UserSettings.cs
--- a/Engine/Source/Programs/UnrealGameSync/UnrealGameSyncShared/UserSettings.cs
+++ b/Engine/Source/Programs/UnrealGameSync/UnrealGameSyncShared/UserSettings.cs
@@ -661,7 +661,7 @@ namespace UnrealGameSync
 				coreSettingsData.Filter.SetCategories(GetCategorySettings(configFile.FindSection("General"), "SyncIncludedCategories", "SyncExcludedCategories"));
 				coreSettingsData.Filter.AllProjects = configFile.GetValue("General.SyncAllProjects", false);
 				coreSettingsData.Filter.AllProjectsInSln = configFile.GetValue("General.IncludeAllProjectsInSolution", false);
-				coreSettingsData.Filter.EngineProgramsInSln = configFile.GetValue("General.IncludeEngineProgramsInSolution", false);
+				coreSettingsData.Filter.UprojectSpecificSln = configFile.GetValue("General.UprojectSpecificSolution", false);
 			}
 
 			return new UserSettings(fileName, configFile, coreFileName, coreSettingsData);
@@ -1033,8 +1033,8 @@ namespace UnrealGameSync
 				int includeAllProjectsInSolution = workspaceSection.GetValue("IncludeAllProjectsInSolution", -1);
 				currentWorkspace.Filter.AllProjectsInSln = (includeAllProjectsInSolution == 0) ? (bool?)false : (includeAllProjectsInSolution == 1) ? (bool?)true : (bool?)null;
 
-				int includeEngineProgramsInSolution = workspaceSection.GetValue("IncludeEngineProgramsInSolution", -1);
-				currentWorkspace.Filter.EngineProgramsInSln = (includeEngineProgramsInSolution == 0) ? (bool?)false : (includeEngineProgramsInSolution == 1) ? (bool?)true : (bool?)null;
+				int uprojectSpecificSolution = workspaceSection.GetValue("UprojectSpecificSolution", -1);
+				currentWorkspace.Filter.UprojectSpecificSln = (uprojectSpecificSolution == 0) ? (bool?)false : (uprojectSpecificSolution == 1) ? (bool?)true : (bool?)null;
 			}
 		}
 
diff --git a/Engine/Source/Programs/UnrealGameSync/UnrealGameSyncShared/WorkspaceUpdate.cs b/Engine/Source/Programs/UnrealGameSync/UnrealGameSyncShared/WorkspaceUpdate.cs
--- a/Engine/Source/Programs/UnrealGameSync/UnrealGameSyncShared/WorkspaceUpdate.cs
+++ b/Engine/Source/Programs/UnrealGameSync/UnrealGameSyncShared/WorkspaceUpdate.cs
@@ -40,7 +40,7 @@ namespace UnrealGameSync
 		RemoveFilteredFiles = 0x4000,
 		Clobber = 0x8000,
 		Refilter = 0x10000,
-		IncludeEngineProgramsInSolution = 0x20000,
+		UprojectSpecificSolution = 0x20000,
 	}
 
 	public enum WorkspaceUpdateResult
@@ -134,9 +134,9 @@ namespace UnrealGameSync
 			{
 				options |= WorkspaceUpdateOptions.IncludeAllProjectsInSolution;
 			}
-			if (workspaceSettings.Filter.EngineProgramsInSln ?? globalSettings.Filter.EngineProgramsInSln ?? false)
+			if (workspaceSettings.Filter.UprojectSpecificSln ?? globalSettings.Filter.UprojectSpecificSln ?? false)
 			{
-				options |= WorkspaceUpdateOptions.IncludeEngineProgramsInSolution;
+				options |= WorkspaceUpdateOptions.UprojectSpecificSolution;
 			}
 			return options;
 		}
@@ -1268,21 +1268,13 @@ namespace UnrealGameSync
 						{
 							commandLine.AppendFormat(" -Project=\"{0}\"", project.LocalFileName);
 
-							// Game specific solutions are only valid if a Source folder exists
-							if (DirectoryReference.Exists(DirectoryReference.Combine(project.LocalFileName.Directory, "Source")))
+							// Uproject specific solutions are only valid if a Source folder exists
+							if (Context.Options.HasFlag(WorkspaceUpdateOptions.UprojectSpecificSolution) && DirectoryReference.Exists(DirectoryReference.Combine(project.LocalFileName.Directory, "Source")))
 							{
 								commandLine.Append(" -Game");
 							}
 						}
 					}
-					if (Context.Options.HasFlag(WorkspaceUpdateOptions.IncludeEngineProgramsInSolution))
-					{
-						commandLine.Append(" -Programs");
-					}
-					else
-					{
-						commandLine.Append(" -NoPrograms");
-					}
 					commandLine.Append(" -progress");
 					logger.LogInformation("Generating project files...");
 					logger.LogInformation("gpf> Running {Arguments}", commandLine);
