UGS: Add an install/uninstall command for UGS-CLI on Windows which adds it to the PATH environment variable.

#rnx

diff --git a/Engine/Source/Programs/UnrealGameSync/UnrealGameSyncCmd/Program.cs b/Engine/Source/Programs/UnrealGameSync/UnrealGameSyncCmd/Program.cs
--- a/Engine/Source/Programs/UnrealGameSync/UnrealGameSyncCmd/Program.cs
+++ b/Engine/Source/Programs/UnrealGameSync/UnrealGameSyncCmd/Program.cs
@@ -124,6 +124,10 @@ namespace UnrealGameSyncCmd
 				null,
 				null
 			),
+			new CommandInfo("uninstall", typeof(UninstallCommand), null,
+				null,
+				null
+			),
 			new CommandInfo("upgrade", typeof(UpgradeCommand), typeof(UpgradeCommandOptions),
 				"ugs upgrade",
 				"Upgrades the current installation with the latest build of UGS."
@@ -1566,43 +1570,89 @@ namespace UnrealGameSyncCmd
 		{
 			public override async Task ExecuteAsync(CommandContext context)
 			{
-				ILogger logger = context.Logger;
+				await UpdateInstallAsync(true, context.Logger);
+			}
+		}
 
-				if (RuntimeInformation.IsOSPlatform(OSPlatform.OSX))
+		class UninstallCommand : Command
+		{
+			public override async Task ExecuteAsync(CommandContext context)
+			{
+				await UpdateInstallAsync(false, context.Logger);
+			}
+		}
+
+		static async Task UpdateInstallAsync(bool install, ILogger logger)
+		{
+			if (OperatingSystem.IsWindows())
+			{
+				const string EnvVarName = "PATH";
+
+				FileReference assemblyFile = new FileReference(Assembly.GetExecutingAssembly().GetOriginalLocation());
+				DirectoryReference assemblyDir = assemblyFile.Directory;
+
+				string? pathVar = Environment.GetEnvironmentVariable(EnvVarName, EnvironmentVariableTarget.User);
+				pathVar ??= String.Empty;
+
+				List<string> paths = new List<string>(pathVar.Split(Path.PathSeparator, StringSplitOptions.RemoveEmptyEntries));
+
+				int changes = paths.RemoveAll(x => x.Equals(assemblyDir.FullName, StringComparison.OrdinalIgnoreCase));
+				if (install)
 				{
-					DirectoryReference? userDir = DirectoryReference.GetSpecialFolder(Environment.SpecialFolder.UserProfile);
-					if (userDir != null)
-					{
-						FileReference configFile = FileReference.Combine(userDir, ".zshrc");
-						await AddAliasAsync(configFile, logger);
-					}
+					paths.Add(assemblyDir.FullName);
+					changes++;
 				}
-				else if (RuntimeInformation.IsOSPlatform(OSPlatform.Linux))
+				if (changes > 0)
 				{
-					DirectoryReference? userDir = DirectoryReference.GetSpecialFolder(Environment.SpecialFolder.UserProfile);
-					if (userDir != null)
-					{
-						FileReference configFile = FileReference.Combine(userDir, ".bashrc");
-						await AddAliasAsync(configFile, logger);
-					}
+					pathVar = String.Join(Path.PathSeparator, paths);
+					Environment.SetEnvironmentVariable(EnvVarName, pathVar, EnvironmentVariableTarget.User);
 				}
 			}
-
-			static async Task AddAliasAsync(FileReference configFile, ILogger logger)
+			else if (OperatingSystem.IsMacOS())
 			{
-				DirectoryReference currentDir = new FileReference(Assembly.GetExecutingAssembly().Location).Directory;
-
-				List<string> lines = new List<string>();
-				if (FileReference.Exists(configFile))
+				DirectoryReference? userDir = DirectoryReference.GetSpecialFolder(Environment.SpecialFolder.UserProfile);
+				if (userDir != null)
 				{
-					lines.AddRange(await FileReference.ReadAllLinesAsync(configFile));
-					lines.RemoveAll(x => Regex.IsMatch(x, @"^\s*alias\s+ugs\s*="));
+					FileReference configFile = FileReference.Combine(userDir, ".zshrc");
+					await UpdateAliasAsync(configFile, install, logger);
 				}
+			}
+			else if (OperatingSystem.IsLinux())
+			{
+				DirectoryReference? userDir = DirectoryReference.GetSpecialFolder(Environment.SpecialFolder.UserProfile);
+				if (userDir != null)
+				{
+					FileReference configFile = FileReference.Combine(userDir, ".bashrc");
+					await UpdateAliasAsync(configFile, install, logger);
+				}
+			}
+		}
+
+		static async Task UpdateAliasAsync(FileReference configFile, bool install, ILogger logger)
+		{
+			DirectoryReference currentDir = new FileReference(Assembly.GetExecutingAssembly().Location).Directory;
+
+			List<string> lines = new List<string>();
+			if (FileReference.Exists(configFile))
+			{
+				lines.AddRange(await FileReference.ReadAllLinesAsync(configFile));
+				lines.RemoveAll(x => Regex.IsMatch(x, @"^\s*alias\s+ugs\s*="));
+			}
+			if (install)
+			{
 				lines.Add($"alias ugs={FileReference.Combine(currentDir, "ugs")}");
+			}
 
-				await FileReference.WriteAllLinesAsync(configFile, lines);
+			await FileReference.WriteAllLinesAsync(configFile, lines);
+
+			if (install)
+			{
 				logger.LogInformation("Added 'ugs' alias to {ConfigFile}", configFile);
 			}
+			else
+			{
+				logger.LogInformation("Removed 'ugs' alias from {ConfigFile}", configFile);
+			}
 		}
 
 		class UpgradeCommandOptions
diff --git a/Engine/Source/Programs/UnrealGameSync/UnrealGameSyncCmd/Toolbox.json b/Engine/Source/Programs/UnrealGameSync/UnrealGameSyncCmd/Toolbox.json
new file mode 100644
--- /dev/null
+++ b/Engine/Source/Programs/UnrealGameSync/UnrealGameSyncCmd/Toolbox.json
@@ -0,0 +1,14 @@
+{
+    "installCommand": {
+        "fileName": "ugs.exe",
+        "arguments": [
+            "install"
+        ]
+    },
+    "uninstallCommand": {
+        "fileName": "ugs.exe",
+        "arguments": [
+            "uninstall"
+        ]
+    }
+}
\ No newline at end of file
diff --git a/Engine/Source/Programs/UnrealGameSync/UnrealGameSyncCmd/UnrealGameSyncCmd.csproj b/Engine/Source/Programs/UnrealGameSync/UnrealGameSyncCmd/UnrealGameSyncCmd.csproj
--- a/Engine/Source/Programs/UnrealGameSync/UnrealGameSyncCmd/UnrealGameSyncCmd.csproj
+++ b/Engine/Source/Programs/UnrealGameSync/UnrealGameSyncCmd/UnrealGameSyncCmd.csproj
@@ -17,6 +17,12 @@
     <DocumentationFile></DocumentationFile>
   </PropertyGroup>
 
+  <ItemGroup>
+    <Content Include="Toolbox.json">
+      <CopyToOutputDirectory>Always</CopyToOutputDirectory>
+    </Content>
+  </ItemGroup>
+
   <ItemGroup>
     <PackageReference Include="Serilog" Version="2.10.0" />
     <PackageReference Include="Serilog.Sinks.Console" Version="4.0.1" />
