adding project overrride for presets

diff --git a/Engine/Source/Programs/UnrealGameSync/UnrealGameSync/Controls/WorkspaceControl.cs b/Engine/Source/Programs/UnrealGameSync/UnrealGameSync/Controls/WorkspaceControl.cs
--- a/Engine/Source/Programs/UnrealGameSync/UnrealGameSync/Controls/WorkspaceControl.cs
+++ b/Engine/Source/Programs/UnrealGameSync/UnrealGameSync/Controls/WorkspaceControl.cs
@@ -1833,7 +1833,7 @@ namespace UnrealGameSync
 
 		IDictionary<string, Preset> GetPresets()
 		{
-			return ConfigUtils.GetPresets(_perforceMonitor.LatestProjectConfigFile);
+			return ConfigUtils.GetPresets(_perforceMonitor.LatestProjectConfigFile, SelectedProjectIdentifier);
 		}
 
 		static List<string> GetProjectRoots(string inUProjectDirsPath)
diff --git a/Engine/Source/Programs/UnrealGameSync/UnrealGameSyncCmd/Program.cs b/Engine/Source/Programs/UnrealGameSync/UnrealGameSyncCmd/Program.cs
--- a/Engine/Source/Programs/UnrealGameSync/UnrealGameSyncCmd/Program.cs
+++ b/Engine/Source/Programs/UnrealGameSync/UnrealGameSyncCmd/Program.cs
@@ -534,12 +534,12 @@ namespace UnrealGameSyncCmd
 			return ConnectAsync(settings.ServerAndPort, settings.UserName, settings.ClientName, loggerFactory);
 		}
 
-		static string[] ReadSyncFilter(UserWorkspaceSettings workspaceSettings, GlobalSettingsFile userSettings, ConfigFile projectConfig)
+		static string[] ReadSyncFilter(UserWorkspaceSettings workspaceSettings, GlobalSettingsFile userSettings, ConfigFile projectConfig, string projectIdentifier)
 		{
 			Dictionary<Guid, WorkspaceSyncCategory> syncCategories = ConfigUtils.GetSyncCategories(projectConfig);
 
 			// check if any category is from the role
-			IDictionary<string, Preset> roles = ConfigUtils.GetPresets(projectConfig);
+			IDictionary<string, Preset> roles = ConfigUtils.GetPresets(projectConfig, projectIdentifier);
 			if (roles.TryGetValue(workspaceSettings.Preset, out Preset? role))
 			{
 				foreach (RoleCategory roleCategory in role.Categories.Values)
@@ -882,7 +882,7 @@ namespace UnrealGameSyncCmd
 				options |= WorkspaceUpdateContext.GetOptionsFromConfig(context.UserSettings.Global, settings);
 				options |= WorkspaceUpdateOptions.RemoveFilteredFiles;
 
-				string[] syncFilter = ReadSyncFilter(settings, context.UserSettings, projectConfig);
+				string[] syncFilter = ReadSyncFilter(settings, context.UserSettings, projectConfig, state.Current.ProjectIdentifier);
 
 				using WorkspaceLock? workspaceLock = CreateWorkspaceLock(settings.RootDir);
 				if (workspaceLock != null && !await workspaceLock.TryAcquireAsync())
@@ -1212,7 +1212,7 @@ namespace UnrealGameSyncCmd
 				FilterSettings globalFilter = context.UserSettings.Global.Filter;
 				FilterSettings workspaceFilter = workspaceSettings.Filter;
 
-				IDictionary<string, Preset> roles = ConfigUtils.GetPresets(projectConfig);
+				IDictionary<string, Preset> roles = ConfigUtils.GetPresets(projectConfig, workspaceState.Current.ProjectIdentifier);
 
 				roles.TryGetValue(workspaceSettings.Preset, out Preset? role);
 				
@@ -1289,7 +1289,7 @@ namespace UnrealGameSyncCmd
 					}
 				}
 
-				string[] filter = ReadSyncFilter(workspaceSettings, context.UserSettings, projectConfig);
+				string[] filter = ReadSyncFilter(workspaceSettings, context.UserSettings, projectConfig, workspaceState.Current.ProjectIdentifier);
 
 				logger.LogInformation("");
 				logger.LogInformation("Combined view:");
diff --git a/Engine/Source/Programs/UnrealGameSync/UnrealGameSyncShared/ConfigUtils.cs b/Engine/Source/Programs/UnrealGameSync/UnrealGameSyncShared/ConfigUtils.cs
--- a/Engine/Source/Programs/UnrealGameSync/UnrealGameSyncShared/ConfigUtils.cs
+++ b/Engine/Source/Programs/UnrealGameSync/UnrealGameSyncShared/ConfigUtils.cs
@@ -2,7 +2,6 @@
 
 using System;
 using System.Collections.Generic;
-using System.Data;
 using System.Diagnostics.CodeAnalysis;
 using System.IO;
 using System.Linq;
@@ -549,19 +548,35 @@ namespace UnrealGameSync
 			return uniqueIdToCategory;
 		}
 
-		public static IDictionary<string, Preset> GetPresets(ConfigFile? projectConfigFile)
+		public static IDictionary<string, Preset> GetPresets(ConfigFile? projectConfigFile, string projectIdentifier)
 		{
-			Dictionary<string, Preset> roles = new(StringComparer.OrdinalIgnoreCase);
+			Dictionary<string, Preset> presetsDictionary = new(StringComparer.OrdinalIgnoreCase);
 
 			if (projectConfigFile == null)
 			{
-				return roles;
+				return presetsDictionary;
 			}
+			List<string> presetDefinitions = new List<string>();
 
-			string[] lines = projectConfigFile.GetValues("Presets.Preset", []);
-			foreach (string line in lines)
+			presetDefinitions.AddRange(projectConfigFile.GetValues("Presets.Preset", []));
+
+			if (TryGetProjectSetting(projectConfigFile, projectIdentifier, "Preset", out string? projectPresets))
+			{
+				using StringReader sr = new StringReader(projectPresets);
+				while (sr.Peek() != -1)
+				{
+					string? line = sr.ReadLine();
+					if (String.IsNullOrWhiteSpace(line))
+					{
+						continue;
+					}
+					presetDefinitions.Add(line);
+				}
+			}
+
+			foreach (string presetDefinition in presetDefinitions)
 			{
-				ConfigObject obj = new ConfigObject(line);
+				ConfigObject obj = new ConfigObject(presetDefinition);
 
 				Preset? preset = new Preset();
 
@@ -572,7 +587,7 @@ namespace UnrealGameSync
 				{
 					continue;
 				}
-					
+
 				IEnumerable<string> categories = obj.GetValue("Categories", String.Empty)
 					.Split(';')
 					.Select(x => x.Trim())
@@ -596,7 +611,7 @@ namespace UnrealGameSync
 					RoleCategory category = new();
 					if (Guid.TryParse(values[0], out Guid guid))
 					{
-						category.Id = guid;	
+						category.Id = guid;
 					}
 
 					if (Boolean.TryParse(values[1], out bool enabled))
@@ -622,17 +637,17 @@ namespace UnrealGameSync
 					preset.Views.Add(view);
 				}
 
-				if (roles.ContainsKey(preset.Name))
+				if (presetsDictionary.ContainsKey(preset.Name))
 				{
-					roles[preset.Name].Import(preset);
+					presetsDictionary[preset.Name].Import(preset);
 				}
 				else
 				{
-					roles.TryAdd(preset.Name, preset);	
+					presetsDictionary.TryAdd(preset.Name, preset);
 				}
 			}
 
-			return roles;
+			return presetsDictionary;
 		}
 
 		static IEnumerable<Guid> ParseGuids(IEnumerable<string> values)
