Option to Disable Project Generation after Syncing

diff --git a/Engine/Source/Programs/UnrealGameSync/UnrealGameSync/Controls/WorkspaceControl.cs b/Engine/Source/Programs/UnrealGameSync/UnrealGameSync/Controls/WorkspaceControl.cs
--- a/Engine/Source/Programs/UnrealGameSync/UnrealGameSync/Controls/WorkspaceControl.cs
+++ b/Engine/Source/Programs/UnrealGameSync/UnrealGameSync/Controls/WorkspaceControl.cs
@@ -635,7 +635,7 @@ namespace UnrealGameSync
 				{
 					_desiredColumnWidths[PreflightColumn.Index] = (int)(120 * dpiScaleX);
 				}
-					
+
 				_desiredColumnWidths[CISColumn.Index] = (int)(200 * dpiScaleX);
 				_desiredColumnWidths[StatusColumn.Index] = (int)(300 * dpiScaleX);
 
@@ -924,9 +924,13 @@ namespace UnrealGameSync
 
 		void StartSync(int changeNumber, bool syncOnly, WorkspaceUpdateCallback? callback, bool skipKnownChangeLists)
 		{
-			WorkspaceUpdateOptions options = WorkspaceUpdateOptions.Sync | WorkspaceUpdateOptions.SyncArchives | WorkspaceUpdateOptions.GenerateProjectFiles;
+			WorkspaceUpdateOptions options = WorkspaceUpdateOptions.Sync | WorkspaceUpdateOptions.SyncArchives;
 			if (!syncOnly)
 			{
+				if (_settings.GenerateProjectAfterSync)
+				{
+					options |= WorkspaceUpdateOptions.GenerateProjectFiles;
+				}
 				if (_settings.BuildAfterSync)
 				{
 					options |= WorkspaceUpdateOptions.Build;
@@ -5019,6 +5023,14 @@ namespace UnrealGameSync
 			return new List<IArchiveChannel>();
 		}
 
+		private void GenerateAfterSyncCheckBox_CheckedChanged(object sender, EventArgs e)
+		{
+			_settings.GenerateProjectAfterSync = GenerateAfterSyncCheckBox.Checked;
+			_settings.Save(_logger);
+
+			UpdateSyncActionCheckboxes();
+		}
+
 		private void BuildAfterSyncCheckBox_CheckedChanged(object sender, EventArgs e)
 		{
 			_settings.BuildAfterSync = BuildAfterSyncCheckBox.Checked;
@@ -5045,12 +5057,14 @@ namespace UnrealGameSync
 
 		private void UpdateSyncActionCheckboxes()
 		{
+			GenerateAfterSyncCheckBox.Checked = _settings.GenerateProjectAfterSync;
 			BuildAfterSyncCheckBox.Checked = _settings.BuildAfterSync;
 			RunAfterSyncCheckBox.Checked = _settings.RunAfterSync;
 			OpenSolutionAfterSyncCheckBox.Checked = _settings.OpenSolutionAfterSync;
 
 			if (_workspace == null || _workspace.IsBusy())
 			{
+				GenerateAfterSyncCheckBox.Enabled = false;
 				AfterSyncingLabel.Enabled = false;
 				BuildAfterSyncCheckBox.Enabled = false;
 				RunAfterSyncCheckBox.Enabled = false;
@@ -5059,6 +5073,7 @@ namespace UnrealGameSync
 			else
 			{
 				AfterSyncingLabel.Enabled = true;
+				GenerateAfterSyncCheckBox.Enabled = _hasBuildSteps;
 				BuildAfterSyncCheckBox.Enabled = _hasBuildSteps;
 				RunAfterSyncCheckBox.Enabled = BuildAfterSyncCheckBox.Checked || ShouldSyncPrecompiledEditor;
 				OpenSolutionAfterSyncCheckBox.Enabled = !ShouldSyncPrecompiledEditor;
diff --git a/Engine/Source/Programs/UnrealGameSync/UnrealGameSyncShared/UserSettings.cs b/Engine/Source/Programs/UnrealGameSync/UnrealGameSyncShared/UserSettings.cs
--- a/Engine/Source/Programs/UnrealGameSync/UnrealGameSyncShared/UserSettings.cs
+++ b/Engine/Source/Programs/UnrealGameSync/UnrealGameSyncShared/UserSettings.cs
@@ -536,6 +536,7 @@ namespace UnrealGameSync
 
 		// General settings
 		public UserSettingsVersion Version { get; set; } = UserSettingsVersion.Latest;
+		public bool GenerateProjectAfterSync { get; set; }
 		public bool BuildAfterSync { get; set; }
 		public bool RunAfterSync { get; set; }
 		public bool OpenSolutionAfterSync { get; set; }
@@ -675,6 +676,7 @@ namespace UnrealGameSync
 
 			// General settings
 			Version = (UserSettingsVersion)_configFile.GetValue("General.Version", (int)UserSettingsVersion.Initial);
+			GenerateProjectAfterSync = (_configFile.GetValue("General.GenerateProjectAfterSync", "1") != "0");
 			BuildAfterSync = (_configFile.GetValue("General.BuildAfterSync", "1") != "0");
 			RunAfterSync = (_configFile.GetValue("General.RunAfterSync", "1") != "0");
 			bool syncPrecompiledEditor = (_configFile.GetValue("General.SyncPrecompiledEditor", "0") != "0");
@@ -1061,6 +1063,7 @@ namespace UnrealGameSync
 			ConfigSection generalSection = _configFile.FindOrAddSection("General");
 			generalSection.Clear();
 			generalSection.SetValue("Version", (int)Version);
+			generalSection.SetValue("GenerateProjectAfterSync", GenerateProjectAfterSync);
 			generalSection.SetValue("BuildAfterSync", BuildAfterSync);
 			generalSection.SetValue("RunAfterSync", RunAfterSync);
 			generalSection.SetValue("OpenSolutionAfterSync", OpenSolutionAfterSync);
