UGS: Copy UgsCmd to a temp folder when installing in order to avoid invalidating shell instances that don't pick up the PATH environment variable change. Also include version number in the help text.

#rnx

diff --git a/Engine/Source/Programs/UnrealGameSync/UnrealGameSyncCmd/Program.cs b/Engine/Source/Programs/UnrealGameSync/UnrealGameSyncCmd/Program.cs
--- a/Engine/Source/Programs/UnrealGameSync/UnrealGameSyncCmd/Program.cs
+++ b/Engine/Source/Programs/UnrealGameSync/UnrealGameSyncCmd/Program.cs
@@ -355,6 +355,16 @@ namespace UnrealGameSyncCmd
 
 		static void PrintHelp()
 		{
+			string appName = "UnrealGameSync Command-Line Tool";
+
+			string? productVersion = FileVersionInfo.GetVersionInfo(Assembly.GetExecutingAssembly().Location).ProductVersion;
+			if (productVersion != null)
+			{
+				appName = $"{appName} ({productVersion})";
+			}
+
+			Console.WriteLine(appName);
+			Console.WriteLine("");
 			Console.WriteLine("Usage:");
 			foreach (CommandInfo command in _commands)
 			{
@@ -1584,22 +1594,29 @@ namespace UnrealGameSyncCmd
 
 		static async Task UpdateInstallAsync(bool install, ILogger logger)
 		{
+			DirectoryReference? installDir = GetInstallFolder();
+			if (installDir != null)
+			{
+				UpdateInstalledFiles(install, installDir, logger);
+			}
+			else
+			{
+				installDir = new FileReference(Assembly.GetExecutingAssembly().GetOriginalLocation()).Directory;
+			}
+
 			if (OperatingSystem.IsWindows())
 			{
 				const string EnvVarName = "PATH";
 
-				FileReference assemblyFile = new FileReference(Assembly.GetExecutingAssembly().GetOriginalLocation());
-				DirectoryReference assemblyDir = assemblyFile.Directory;
-
 				string? pathVar = Environment.GetEnvironmentVariable(EnvVarName, EnvironmentVariableTarget.User);
 				pathVar ??= String.Empty;
 
 				List<string> paths = new List<string>(pathVar.Split(Path.PathSeparator, StringSplitOptions.RemoveEmptyEntries));
 
-				int changes = paths.RemoveAll(x => x.Equals(assemblyDir.FullName, StringComparison.OrdinalIgnoreCase));
+				int changes = paths.RemoveAll(x => x.Equals(installDir.FullName, StringComparison.OrdinalIgnoreCase));
 				if (install)
 				{
-					paths.Add(assemblyDir.FullName);
+					paths.Add(installDir.FullName);
 					changes++;
 				}
 				if (changes > 0)
@@ -1607,6 +1624,15 @@ namespace UnrealGameSyncCmd
 					pathVar = String.Join(Path.PathSeparator, paths);
 					Environment.SetEnvironmentVariable(EnvVarName, pathVar, EnvironmentVariableTarget.User);
 				}
+
+				if (install)
+				{
+					logger.LogInformation("Added {Path} to PATH environment variable", installDir);
+				}
+				else
+				{
+					logger.LogInformation("Removed {Path} from PATH environment variable", installDir);
+				}
 			}
 			else if (OperatingSystem.IsMacOS())
 			{
@@ -1614,7 +1640,7 @@ namespace UnrealGameSyncCmd
 				if (userDir != null)
 				{
 					FileReference configFile = FileReference.Combine(userDir, ".zshrc");
-					await UpdateAliasAsync(configFile, install, logger);
+					await UpdateAliasAsync(configFile, install, installDir, logger);
 				}
 			}
 			else if (OperatingSystem.IsLinux())
@@ -1623,15 +1649,60 @@ namespace UnrealGameSyncCmd
 				if (userDir != null)
 				{
 					FileReference configFile = FileReference.Combine(userDir, ".bashrc");
-					await UpdateAliasAsync(configFile, install, logger);
+					await UpdateAliasAsync(configFile, install, installDir, logger);
 				}
 			}
 		}
 
-		static async Task UpdateAliasAsync(FileReference configFile, bool install, ILogger logger)
+		static DirectoryReference? GetInstallFolder()
 		{
-			DirectoryReference currentDir = new FileReference(Assembly.GetExecutingAssembly().Location).Directory;
+			if (OperatingSystem.IsWindows())
+			{
+				DirectoryReference? installDir = DirectoryReference.GetSpecialFolder(Environment.SpecialFolder.LocalApplicationData);
+				if (installDir != null)
+				{
+					return DirectoryReference.Combine(installDir, "Epic Games", "UgsCmd");
+				}
+			}
+			return null;
+		}
 
+		static void UpdateInstalledFiles(bool install, DirectoryReference installDir, ILogger logger)
+		{
+			FileReference assemblyFile = new FileReference(Assembly.GetExecutingAssembly().GetOriginalLocation());
+			DirectoryReference sourceDir = assemblyFile.Directory;
+
+			DirectoryReference TempDir = DirectoryReference.Combine(installDir.ParentDirectory!, "~" + installDir.GetDirectoryName());
+			if (DirectoryReference.Exists(TempDir))
+			{
+				DirectoryReference.Delete(TempDir, true);
+			}
+
+			if (DirectoryReference.Exists(installDir))
+			{
+				logger.LogInformation("Removing application files from {Dir}", installDir);
+
+				Directory.Move(installDir.FullName, TempDir.FullName);
+				DirectoryReference.Delete(TempDir, true);
+			}
+
+			if (install)
+			{
+				logger.LogInformation("Copying application files to {Dir}", installDir);
+
+				DirectoryReference.CreateDirectory(TempDir);
+				foreach (FileReference SourceFile in DirectoryReference.EnumerateFiles(sourceDir, "*", SearchOption.AllDirectories))
+				{
+					FileReference TargetFile = FileReference.Combine(TempDir, SourceFile.MakeRelativeTo(sourceDir));
+					DirectoryReference.CreateDirectory(TargetFile.Directory);
+					FileReference.Copy(SourceFile, TargetFile, true);
+				}
+				Directory.Move(TempDir.FullName, installDir.FullName);
+			}
+		}
+
+		static async Task UpdateAliasAsync(FileReference configFile, bool install, DirectoryReference installDir, ILogger logger)
+		{
 			List<string> lines = new List<string>();
 			if (FileReference.Exists(configFile))
 			{
@@ -1640,7 +1711,7 @@ namespace UnrealGameSyncCmd
 			}
 			if (install)
 			{
-				lines.Add($"alias ugs={FileReference.Combine(currentDir, "ugs")}");
+				lines.Add($"alias ugs={FileReference.Combine(installDir, "ugs")}");
 			}
 
 			await FileReference.WriteAllLinesAsync(configFile, lines);
