[UE-194386] - UGS - The "More..." menu needs a way to be better organized

Nests menu options in the "More..." context menu. Nesting is determined by the build step's "Description" parameter, using the "|" character to delimit groups. e.g. "Group|SubGroup|ItemName" and "Group|SubGroup|ItemName2" will cause "ItemName" and "ItemName2" to be siblings under the "Group -> SubGroup" menu nesting.

diff --git a/Engine/Source/Programs/UnrealGameSync/UnrealGameSync/Controls/WorkspaceControl.cs b/Engine/Source/Programs/UnrealGameSync/UnrealGameSync/Controls/WorkspaceControl.cs
--- a/Engine/Source/Programs/UnrealGameSync/UnrealGameSync/Controls/WorkspaceControl.cs
+++ b/Engine/Source/Programs/UnrealGameSync/UnrealGameSync/Controls/WorkspaceControl.cs
@@ -6004,7 +6004,8 @@ namespace UnrealGameSync
 					Dictionary<Guid, ConfigObject> projectBuildStepObjects = GetProjectBuildStepObjects(projectConfigFile);
 					List<BuildStep> userSteps = GetUserBuildSteps(projectBuildStepObjects);
 
-					Dictionary<string, ContextMenuStrip> nameToMenu = new Dictionary<string, ContextMenuStrip>();
+					Dictionary<string, ContextMenuStrip>  nameToMenu         = new Dictionary<string, ContextMenuStrip>(StringComparer.OrdinalIgnoreCase);
+					Dictionary<string, ToolStripMenuItem> overflowNameToMenu = new Dictionary<string, ToolStripMenuItem>(StringComparer.OrdinalIgnoreCase);
 					foreach (BuildStep step in userSteps)
 					{
 						if (!String.IsNullOrEmpty(step.StatusPanelLink) && CanRunStep(step))
@@ -6016,20 +6017,62 @@ namespace UnrealGameSync
 							}
 							else
 							{
-								string menuName = step.StatusPanelLink.Substring(0, baseMenuIdx);
-								string itemName = step.StatusPanelLink.Substring(baseMenuIdx + 1).Replace("&", "&&", StringComparison.Ordinal);
+								// At least one segment exists
+								string[] segments       = step.StatusPanelLink.Split('|');
+								int lastSegmentStartIdx = step.StatusPanelLink.LastIndexOf('|');
 
 #pragma warning disable CA2000 // warning CA2000: Use recommended dispose pattern to ensure that object created by 'new ToolStripMenuItem(itemName)' is disposed on all paths. If possible, wrap the creation within a 'using' statement or a 'using' declaration. Otherwise, use a try-finally pattern, with a dedicated local variable declared before the try region and an unconditional Dispose invocation on non-null value in the 'finally' region, say 'x?.Dispose()'. If the object is explicitly disposed within the try region or the dispose ownership is transfered to another object or method, assign 'null' to the local variable just after such an operation to prevent double dispose in 'finally'.
+
+								// Create the menu item we'll add to the view
+								string itemName = segments[segments.Length - 1].Replace("&", "&&", StringComparison.OrdinalIgnoreCase); // Tool "Name" will be the last part of the segment always
 								ToolStripMenuItem newMenuItem = new ToolStripMenuItem(itemName);
 								newMenuItem.Click += new EventHandler((sender, e) => { RunCustomTool(step, userSteps); });
 
-								if (menuName == "More...")
+								// Where the item ends up depends on the first segment text
+								if (segments[0] == "More...")
 								{
-									MoreToolsContextMenu.Items.Insert(moreToolsItemCount++, newMenuItem);
+									// Overflow Menu
+									ToolStripMenuItem? curParentItem = null;
+
+									// Walk the segments to find the parent menu, creating missing records along the way
+									string combinedSegmentKey = "";
+									for( int idx = 1; idx < segments.Length - 1; ++idx) // Skip the first segment since we know it's `More...`, and stop before the last segment since it's the name
+									{
+										combinedSegmentKey += segments[idx];
+										string segmentItemName = segments[idx].Replace("&", "&&", StringComparison.OrdinalIgnoreCase);
+
+										ToolStripMenuItem? foundParentItem = null;
+										if (!overflowNameToMenu.TryGetValue(combinedSegmentKey, out foundParentItem))
+										{
+											// Missing this segment
+											foundParentItem = new ToolStripMenuItem(segmentItemName);
+											overflowNameToMenu.Add(combinedSegmentKey, foundParentItem);
+											if(curParentItem != null)
+											{
+												curParentItem.DropDownItems.Add(foundParentItem);
+											}
+											else
+											{
+												MoreToolsContextMenu.Items.Insert(moreToolsItemCount++, foundParentItem);
+											}
+										}
+										curParentItem = foundParentItem;
+									}
+
+									if (curParentItem != null)
+									{
+										curParentItem.DropDownItems.Add(newMenuItem);
+									}
+									else
+									{
+										MoreToolsContextMenu.Items.Insert(moreToolsItemCount++, newMenuItem); // No parent item, add it to the root
+									}
 								}
 								else
 								{
+									// Not Overflow Menu
 									ContextMenuStrip? menu;
+									string menuName = segments[0]; // Not in the overflow menu, so the first segment is the menu name
 									if (!nameToMenu.TryGetValue(menuName, out menu))
 									{
 										menu = new ContextMenuStrip();
@@ -6039,8 +6082,8 @@ namespace UnrealGameSync
 										components.Add(menu);
 									}
 									menu.Items.Add(newMenuItem);
-								}
 #pragma warning restore CA2000
+								}
 							}
 						}
 
