[UE-70845] - Set revertunchanged on new workspaces by default

Default functionality can be defined in the UnrealGameSync.ini using the "General.ClientSubmitOptions" key

diff --git a/Engine/Source/Programs/UnrealGameSync/UnrealGameSync/Forms/AutomatedBuildWindow.cs b/Engine/Source/Programs/UnrealGameSync/UnrealGameSync/Forms/AutomatedBuildWindow.cs
--- a/Engine/Source/Programs/UnrealGameSync/UnrealGameSync/Forms/AutomatedBuildWindow.cs
+++ b/Engine/Source/Programs/UnrealGameSync/UnrealGameSync/Forms/AutomatedBuildWindow.cs
@@ -32,14 +32,16 @@ namespace UnrealGameSync.Forms
 		string? _serverAndPortOverride;
 		string? _userNameOverride;
 		readonly IPerforceSettings _defaultPerforceSettings;
+		readonly UserSettings _settings;
 
 		BuildInfo? _result;
 
-		private AutomatedBuildWindow(string streamName, int changelist, string command, IPerforceSettings defaultPerforceSettings, string? defaultWorkspaceName, string? defaultProjectPath, IServiceProvider serviceProvider)
+		private AutomatedBuildWindow(string streamName, int changelist, string command, IPerforceSettings defaultPerforceSettings, string? defaultWorkspaceName, string? defaultProjectPath, UserSettings settings, IServiceProvider serviceProvider)
 		{
 			_streamName = streamName;
 			_defaultPerforceSettings = defaultPerforceSettings;
 			_serviceProvider = serviceProvider;
+			_settings = settings;
 
 			InitializeComponent();
 			Font = new System.Drawing.Font("Segoe UI", 8.25F, System.Drawing.FontStyle.Regular, System.Drawing.GraphicsUnit.Point, ((byte)(0)));
@@ -93,7 +95,7 @@ namespace UnrealGameSync.Forms
 				}
 			}
 
-			using AutomatedBuildWindow window = new AutomatedBuildWindow(streamName, changelist, command, defaultPerforceSettings, defaultWorkspaceName, defaultProjectPath, loggerFactory);
+			using AutomatedBuildWindow window = new AutomatedBuildWindow(streamName, changelist, command, defaultPerforceSettings, defaultWorkspaceName, defaultProjectPath, settings, loggerFactory);
 			if (window.ShowDialog() == DialogResult.OK)
 			{
 				buildInfo = window._result!;
@@ -122,7 +124,7 @@ namespace UnrealGameSync.Forms
 		private void WorkspaceNameNewBtn_Click(object sender, EventArgs e)
 		{
 			string? workspaceName;
-			if (NewWorkspaceWindow.ShowModal(this, Perforce, _streamName, WorkspaceNameTextBox.Text, _serviceProvider, out workspaceName))
+			if (NewWorkspaceWindow.ShowModal(this, Perforce, _streamName, WorkspaceNameTextBox.Text, _settings, _serviceProvider, out workspaceName))
 			{
 				WorkspaceNameTextBox.Text = workspaceName;
 				UpdateOkButton();
diff --git a/Engine/Source/Programs/UnrealGameSync/UnrealGameSync/Forms/AutomatedSyncWindow.cs b/Engine/Source/Programs/UnrealGameSync/UnrealGameSync/Forms/AutomatedSyncWindow.cs
--- a/Engine/Source/Programs/UnrealGameSync/UnrealGameSync/Forms/AutomatedSyncWindow.cs
+++ b/Engine/Source/Programs/UnrealGameSync/UnrealGameSync/Forms/AutomatedSyncWindow.cs
@@ -98,13 +98,15 @@ namespace UnrealGameSync
 		string? _serverAndPortOverride;
 		string? _userNameOverride;
 		readonly IPerforceSettings _defaultPerforceSettings;
+		readonly UserSettings _settings;
 		WorkspaceInfo? _selectedWorkspaceInfo;
 
-		private AutomatedSyncWindow(string streamName, string projectPath, string? workspaceName, IPerforceSettings defaultPerforceSettings, IServiceProvider serviceProvider)
+		private AutomatedSyncWindow(string streamName, string projectPath, string? workspaceName, IPerforceSettings defaultPerforceSettings, UserSettings settings, IServiceProvider serviceProvider)
 		{
 			_streamName = streamName;
 			_defaultPerforceSettings = defaultPerforceSettings;
 			_serviceProvider = serviceProvider;
+			_settings = settings;
 
 			InitializeComponent();
 			Font = new System.Drawing.Font("Segoe UI", 8.25F, System.Drawing.FontStyle.Regular, System.Drawing.GraphicsUnit.Point, ((byte)(0)));
@@ -135,11 +137,11 @@ namespace UnrealGameSync
 			return (workspaceTask != null && workspaceTask.Succeeded) ? workspaceTask.Result : null;
 		}
 
-		public static bool ShowModal(IWin32Window owner, IPerforceSettings defaultPerforceSettings, string streamName, string projectPath, [NotNullWhen(true)] out WorkspaceInfo? workspaceInfo, IServiceProvider serviceProvider)
+		public static bool ShowModal(IWin32Window owner, IPerforceSettings defaultPerforceSettings, string streamName, string projectPath, UserSettings settings, [NotNullWhen(true)] out WorkspaceInfo? workspaceInfo, IServiceProvider serviceProvider)
 		{
 			string? workspaceName = FindDefaultWorkspace(owner, defaultPerforceSettings, streamName, serviceProvider);
 
-			using AutomatedSyncWindow window = new AutomatedSyncWindow(streamName, projectPath, workspaceName, defaultPerforceSettings, serviceProvider);
+			using AutomatedSyncWindow window = new AutomatedSyncWindow(streamName, projectPath, workspaceName, defaultPerforceSettings, settings, serviceProvider);
 			if (window.ShowDialog() == DialogResult.OK)
 			{
 				workspaceInfo = window._selectedWorkspaceInfo!;
@@ -169,7 +171,7 @@ namespace UnrealGameSync
 		private void WorkspaceNameNewBtn_Click(object sender, EventArgs e)
 		{
 			string? workspaceName;
-			if (NewWorkspaceWindow.ShowModal(this, _defaultPerforceSettings, _streamName, WorkspaceNameTextBox.Text, _serviceProvider, out workspaceName))
+			if (NewWorkspaceWindow.ShowModal(this, _defaultPerforceSettings, _streamName, WorkspaceNameTextBox.Text, _settings, _serviceProvider, out workspaceName))
 			{
 				WorkspaceNameTextBox.Text = workspaceName;
 				UpdateOkButton();
diff --git a/Engine/Source/Programs/UnrealGameSync/UnrealGameSync/Forms/MainWindow.cs b/Engine/Source/Programs/UnrealGameSync/UnrealGameSync/Forms/MainWindow.cs
--- a/Engine/Source/Programs/UnrealGameSync/UnrealGameSync/Forms/MainWindow.cs
+++ b/Engine/Source/Programs/UnrealGameSync/UnrealGameSync/Forms/MainWindow.cs
@@ -321,7 +321,7 @@ namespace UnrealGameSync
 			string projectPath = reader.ReadString();
 
 			AutomatedSyncWindow.WorkspaceInfo? workspaceInfo;
-			if (!AutomatedSyncWindow.ShowModal(this, _defaultPerforceSettings, streamName, projectPath, out workspaceInfo, _serviceProvider))
+			if (!AutomatedSyncWindow.ShowModal(this, _defaultPerforceSettings, streamName, projectPath, _settings, out workspaceInfo, _serviceProvider))
 			{
 				return new AutomationRequestOutput(AutomationRequestResult.Canceled);
 			}
diff --git a/Engine/Source/Programs/UnrealGameSync/UnrealGameSync/Forms/NewWorkspaceWindow.cs b/Engine/Source/Programs/UnrealGameSync/UnrealGameSync/Forms/NewWorkspaceWindow.cs
--- a/Engine/Source/Programs/UnrealGameSync/UnrealGameSync/Forms/NewWorkspaceWindow.cs
+++ b/Engine/Source/Programs/UnrealGameSync/UnrealGameSync/Forms/NewWorkspaceWindow.cs
@@ -63,7 +63,7 @@ namespace UnrealGameSync
 
 		static class NewWorkspaceTask
 		{
-			public static async Task RunAsync(IPerforceConnection perforce, NewWorkspaceSettings settings, string owner, string hostName, CancellationToken cancellationToken)
+			public static async Task RunAsync(IPerforceConnection perforce, NewWorkspaceSettings settings, string owner, string hostName, ClientSubmitOptions clientSubmitOptions, CancellationToken cancellationToken)
 			{
 				PerforceResponseList<ClientsRecord> response = await perforce.TryGetClientsAsync(ClientsOptions.None, settings.Name, -1, null, null, cancellationToken);
 				if (!response.Succeeded)
@@ -79,6 +79,7 @@ namespace UnrealGameSync
 				client.Host = hostName;
 				client.Stream = settings.Stream;
 				client.Options = ClientOptions.Rmdir;
+				client.SubmitOptions = clientSubmitOptions;
 				await perforce.CreateClientAsync(client, cancellationToken);
 			}
 		}
@@ -89,8 +90,9 @@ namespace UnrealGameSync
 		readonly IServiceProvider _serviceProvider;
 		NewWorkspaceSettings? _settings;
 		readonly DirectoryReference? _defaultRootPath;
+		readonly UserSettings _userSettings;
 
-		private NewWorkspaceWindow(IPerforceSettings perforceSettings, string? forceStream, string? defaultStream, InfoRecord info, List<ClientsRecord> clients, IServiceProvider serviceProvider)
+		private NewWorkspaceWindow(IPerforceSettings perforceSettings, string? forceStream, string? defaultStream, InfoRecord info, List<ClientsRecord> clients, IServiceProvider serviceProvider, UserSettings settings)
 		{
 			InitializeComponent();
 			Font = new System.Drawing.Font("Segoe UI", 8.25F, System.Drawing.FontStyle.Regular, System.Drawing.GraphicsUnit.Point, ((byte)(0)));
@@ -99,6 +101,7 @@ namespace UnrealGameSync
 			_info = info;
 			_clients = clients;
 			_serviceProvider = serviceProvider;
+			_userSettings = settings;
 
 			Dictionary<DirectoryReference, int> rootPathToCount = new Dictionary<DirectoryReference, int>();
 			foreach (ClientsRecord client in clients)
@@ -158,7 +161,7 @@ namespace UnrealGameSync
 			UpdateRootDirCueBanner();
 		}
 
-		public static bool ShowModal(IWin32Window owner, IPerforceSettings perforceSettings, string? forceStreamName, string currentWorkspaceName, IServiceProvider serviceProvider, [NotNullWhen(true)] out string? workspaceName)
+		public static bool ShowModal(IWin32Window owner, IPerforceSettings perforceSettings, string? forceStreamName, string currentWorkspaceName, UserSettings settings, IServiceProvider serviceProvider, [NotNullWhen(true)] out string? workspaceName)
 		{
 			ModalTask<FindWorkspaceSettingsTask>? task = PerforceModalTask.Execute(owner, "Checking settings", "Checking settings, please wait...", perforceSettings, (p, c) => FindWorkspaceSettingsTask.RunAsync(p, currentWorkspaceName, c), serviceProvider.GetRequiredService<ILogger<FindWorkspaceSettingsTask>>());
 			if (task == null || !task.Succeeded)
@@ -167,7 +170,7 @@ namespace UnrealGameSync
 				return false;
 			}
 
-			using NewWorkspaceWindow window = new NewWorkspaceWindow(perforceSettings, forceStreamName, task.Result.CurrentStream, task.Result.Info, task.Result.Clients, serviceProvider);
+			using NewWorkspaceWindow window = new NewWorkspaceWindow(perforceSettings, forceStreamName, task.Result.CurrentStream, task.Result.Info, task.Result.Clients, serviceProvider, settings);
 			if (window.ShowDialog(owner) == DialogResult.OK)
 			{
 				workspaceName = window._settings!.Name;
@@ -324,7 +327,7 @@ namespace UnrealGameSync
 
 				ILogger logger = _serviceProvider.GetRequiredService<ILogger<NewWorkspaceWindow>>();
 
-				ModalTask? result = PerforceModalTask.Execute(Owner, "Creating workspace", "Creating workspace, please wait...", _perforceSettings, (p, c) => NewWorkspaceTask.RunAsync(p, _settings, _info.UserName ?? "", _info.ClientHost ?? "", c), logger);
+				ModalTask? result = PerforceModalTask.Execute(Owner, "Creating workspace", "Creating workspace, please wait...", _perforceSettings, (p, c) => NewWorkspaceTask.RunAsync(p, _settings, _info.UserName ?? "", _info.ClientHost ?? "", _userSettings.ClientSubmitOptions, c), logger);
 				if (result == null || !result.Succeeded)
 				{
 					return;
diff --git a/Engine/Source/Programs/UnrealGameSync/UnrealGameSync/Forms/OpenProjectWindow.cs b/Engine/Source/Programs/UnrealGameSync/UnrealGameSync/Forms/OpenProjectWindow.cs
--- a/Engine/Source/Programs/UnrealGameSync/UnrealGameSync/Forms/OpenProjectWindow.cs
+++ b/Engine/Source/Programs/UnrealGameSync/UnrealGameSync/Forms/OpenProjectWindow.cs
@@ -183,7 +183,7 @@ namespace UnrealGameSync
 			WorkspaceRadioBtn.Checked = true;
 
 			string? workspaceName;
-			if (NewWorkspaceWindow.ShowModal(this, Perforce, null, WorkspaceNameTextBox.Text, _serviceProvider, out workspaceName))
+			if (NewWorkspaceWindow.ShowModal(this, Perforce, null, WorkspaceNameTextBox.Text, _settings, _serviceProvider, out workspaceName))
 			{
 				WorkspaceNameTextBox.Text = workspaceName;
 				UpdateOkButton();
diff --git a/Engine/Source/Programs/UnrealGameSync/UnrealGameSyncShared/UserSettings.cs b/Engine/Source/Programs/UnrealGameSync/UnrealGameSyncShared/UserSettings.cs
--- a/Engine/Source/Programs/UnrealGameSync/UnrealGameSyncShared/UserSettings.cs
+++ b/Engine/Source/Programs/UnrealGameSync/UnrealGameSyncShared/UserSettings.cs
@@ -10,6 +10,7 @@ using System.Runtime.InteropServices;
 using System.Text;
 using System.Text.Json.Serialization;
 using EpicGames.Core;
+using EpicGames.Perforce;
 using Microsoft.Extensions.Logging;
 
 namespace UnrealGameSync
@@ -590,6 +591,7 @@ namespace UnrealGameSync
 
 		// Perforce settings
 		public PerforceSyncOptions SyncOptions => base.Global.Perforce;
+		public ClientSubmitOptions ClientSubmitOptions { get; set; }
 
 		private List<UserSelectedProjectSettings> ReadProjectList(string settingName, string legacySettingName)
 		{
@@ -732,6 +734,7 @@ namespace UnrealGameSync
 			OpenProjects = ReadProjectList("General.OpenProjects", "General.OpenProjectFileNames");
 			RecentProjects = ReadProjectList("General.RecentProjects", "General.OtherProjectFileNames");
 			SyncTypeId = _configFile.GetValue("General.SyncTypeID", "");
+			ClientSubmitOptions = _configFile.GetEnumValue("General.ClientSubmitOptions", ClientSubmitOptions.RevertUnchanged);
 			string? oldSyncTye = _configFile.GetValue("General.SyncType", null);
 			if (oldSyncTye != null)
 			{
@@ -1083,6 +1086,7 @@ namespace UnrealGameSync
 			generalSection.SetValue("FilterIndex", FilterIndex);
 			generalSection.SetValues("RecentProjects", RecentProjects.Select(x => x.ToConfigEntry()).ToArray());
 			generalSection.SetValue("SyncTypeID", SyncTypeId);
+			generalSection.SetValue("ClientSubmitOptions", ClientSubmitOptions.ToString());
 			generalSection.SetValue("NextLauncherVersionCheck", NextLauncherVersionCheck);
 
 			// Build configuration
