[UE-113623] - Ability to filter more than one person via Filter > Author

Use `+` or `;` to search for multiple names. Any changelist with an author string that contains any of the provided search terms will be displayed.

diff --git a/Engine/Source/Programs/UnrealGameSync/UnrealGameSync/Controls/WorkspaceControl.cs b/Engine/Source/Programs/UnrealGameSync/UnrealGameSync/Controls/WorkspaceControl.cs
--- a/Engine/Source/Programs/UnrealGameSync/UnrealGameSync/Controls/WorkspaceControl.cs
+++ b/Engine/Source/Programs/UnrealGameSync/UnrealGameSync/Controls/WorkspaceControl.cs
@@ -269,6 +269,7 @@ namespace UnrealGameSync
 
 		// When an author filter is applied this will be non-empty and != AuthorFilterPlaceholderText
 		string _authorFilterText = "";
+		List<string> _authorsFilter = new List<string>();
 
 		// Placeholder text that is in the control and cleared when the user starts editing.
 		const string AuthorFilterPlaceholderText = "<username>";
@@ -1579,11 +1580,13 @@ namespace UnrealGameSync
 			}
 
 			// if filtering by user, only show changes where the author contains the filter text
-			if (!String.IsNullOrEmpty(_authorFilterText)
-				&& _authorFilterText != AuthorFilterPlaceholderText
-				&& change.User!.IndexOf(_authorFilterText, StringComparison.OrdinalIgnoreCase) < 0)
+			if( _authorsFilter.Any() )
 			{
-				return false;
+				// If none of the (potentially partial) username strings in the query are contained in the change's user, ignore this change
+				if (!_authorsFilter.Any(username => change.User.Contains(username, StringComparison.OrdinalIgnoreCase)))
+				{
+					return false;
+				}
 			}
 
 			// If this is a robomerge change, check if any filters will cause it to be hidden
@@ -7009,6 +7012,12 @@ namespace UnrealGameSync
 		private void FilterContextMenu_Author_Name_Changed(object sender, EventArgs e)
 		{
 			_authorFilterText = FilterContextMenu_Author_Name.Text;
+			if (!String.IsNullOrWhiteSpace(_authorFilterText) && _authorFilterText != AuthorFilterPlaceholderText)
+			{
+				_authorsFilter.Clear();
+				_authorsFilter.AddAll( _authorFilterText.Split([';', '+'], StringSplitOptions.TrimEntries));
+			}
+
 			UpdateBuildListFilter();
 		}
 
