Adding Fast Sync option to the perforce settings.
With Fast Sync the file considered are only the one that have been updated between the current CL and the target CL.

diff --git a/Engine/Source/Programs/UnrealGameSync/UnrealGameSync/Controls/WorkspaceControl.cs b/Engine/Source/Programs/UnrealGameSync/UnrealGameSync/Controls/WorkspaceControl.cs
--- a/Engine/Source/Programs/UnrealGameSync/UnrealGameSync/Controls/WorkspaceControl.cs
+++ b/Engine/Source/Programs/UnrealGameSync/UnrealGameSync/Controls/WorkspaceControl.cs
@@ -926,6 +926,12 @@ namespace UnrealGameSync
 		void StartSync(int changeNumber, bool syncOnly, WorkspaceUpdateCallback? callback, bool skipKnownChangeLists)
 		{
 			WorkspaceUpdateOptions options = WorkspaceUpdateOptions.Sync | WorkspaceUpdateOptions.SyncArchives;
+			
+			if (_settings.SyncOptions.UseFastSync is true)
+			{
+				options |= WorkspaceUpdateOptions.FastSync;
+			}
+			
 			if (!syncOnly)
 			{
 				if (_settings.GenerateProjectAfterSync)
diff --git a/Engine/Source/Programs/UnrealGameSync/UnrealGameSync/Forms/PerforceSyncSettingsWindow.cs b/Engine/Source/Programs/UnrealGameSync/UnrealGameSync/Forms/PerforceSyncSettingsWindow.cs
--- a/Engine/Source/Programs/UnrealGameSync/UnrealGameSync/Forms/PerforceSyncSettingsWindow.cs
+++ b/Engine/Source/Programs/UnrealGameSync/UnrealGameSync/Forms/PerforceSyncSettingsWindow.cs
@@ -27,6 +27,8 @@ namespace UnrealGameSync
 			numericUpDownMaxSizePerBatch.Value = (syncOptions.MaxSizePerBatch ?? PerforceSyncOptions.DefaultMaxSizePerBatch) / 1024 / 1024;
 			numericUpDownRetriesOnSyncError.Value = syncOptions.NumSyncErrorRetries ?? PerforceSyncOptions.DefaultNumSyncErrorRetries;
 			numericUpDownSyncErrorRetryDelay.Value = syncOptions.SyncErrorRetryDelay ?? PerforceSyncOptions.DefaultSyncErrorRetryDelay;
+
+			FastSyncCheckBox.Checked = syncOptions.UseFastSync ?? false;
 		}
 
 		private void OkButton_Click(object sender, EventArgs e)
@@ -35,6 +37,7 @@ namespace UnrealGameSync
 			_settings.SyncOptions.MaxSizePerBatch = GetValueIfNotDefault((int)numericUpDownMaxSizePerBatch.Value, PerforceSyncOptions.DefaultMaxSizePerBatch) * 1024 * 1024;
 			_settings.SyncOptions.NumSyncErrorRetries = GetValueIfNotDefault((int)numericUpDownRetriesOnSyncError.Value, PerforceSyncOptions.DefaultNumSyncErrorRetries);
 			_settings.SyncOptions.SyncErrorRetryDelay = GetValueIfNotDefault((int)numericUpDownSyncErrorRetryDelay.Value, PerforceSyncOptions.DefaultSyncErrorRetryDelay);
+			_settings.SyncOptions.UseFastSync = FastSyncCheckBox.Checked;
 			_settings.Save(_logger);
 
 			DialogResult = System.Windows.Forms.DialogResult.OK;
@@ -58,6 +61,7 @@ namespace UnrealGameSync
 			numericUpDownMaxSizePerBatch.Value = PerforceSyncOptions.DefaultMaxSizePerBatch / 1024 / 1024;
 			numericUpDownRetriesOnSyncError.Value = PerforceSyncOptions.DefaultNumSyncErrorRetries;
 			numericUpDownSyncErrorRetryDelay.Value = PerforceSyncOptions.DefaultSyncErrorRetryDelay;
+			FastSyncCheckBox.Checked = PerforceSyncOptions.DefaultUseFastSync;
 		}
 	}
 }
diff --git a/Engine/Source/Programs/UnrealGameSync/UnrealGameSyncShared/WorkspaceUpdate.cs b/Engine/Source/Programs/UnrealGameSync/UnrealGameSyncShared/WorkspaceUpdate.cs
--- a/Engine/Source/Programs/UnrealGameSync/UnrealGameSyncShared/WorkspaceUpdate.cs
+++ b/Engine/Source/Programs/UnrealGameSync/UnrealGameSyncShared/WorkspaceUpdate.cs
@@ -41,6 +41,7 @@ namespace UnrealGameSync
 		Clobber = 0x8000,
 		Refilter = 0x10000,
 		UprojectSpecificSolution = 0x20000,
+		FastSync = 0x40000
 	}
 
 	public enum WorkspaceUpdateResult
@@ -68,6 +69,8 @@ namespace UnrealGameSync
 		public const int DefaultNumSyncErrorRetries = 0;
 		public const int DefaultSyncErrorRetryDelay = 0;
 
+		public const bool DefaultUseFastSync = false;
+
 		public int? NumThreads { get; set; }
 
 		public int? MaxCommandsPerBatch { get; set; }
@@ -76,6 +79,8 @@ namespace UnrealGameSync
 		public int? NumSyncErrorRetries { get; set; }
 		public int? SyncErrorRetryDelay { get; set; }
 
+		public bool? UseFastSync { get; set; }
+
 		public PerforceSyncOptions Clone()
 		{
 			return (PerforceSyncOptions)MemberwiseClone();
@@ -443,15 +448,17 @@ namespace UnrealGameSync
 		class SyncTree
 		{
 			public bool CanUseWildcard;
+			private readonly bool CanContainExcludedFiles;
 			public int TotalIncludedFiles;
 			public long TotalSize;
 			public int TotalExcludedFiles;
 			public Dictionary<string, long> IncludedFiles = new Dictionary<string, long>();
 			public Dictionary<string, SyncTree> NameToSubTree = new Dictionary<string, SyncTree>(StringComparer.OrdinalIgnoreCase);
 
-			public SyncTree(bool canUseWildcard)
+			public SyncTree(bool canUseWildcard, bool canContainExcludedFiles)
 			{
 				CanUseWildcard = canUseWildcard;
+				CanContainExcludedFiles = canContainExcludedFiles;
 			}
 
 			public SyncTree FindOrAddSubTree(string name)
@@ -459,7 +466,7 @@ namespace UnrealGameSync
 				SyncTree? result;
 				if (!NameToSubTree.TryGetValue(name, out result))
 				{
-					result = new SyncTree(CanUseWildcard);
+					result = new SyncTree(CanUseWildcard, CanContainExcludedFiles);
 					NameToSubTree.Add(name, result);
 				}
 				return result;
@@ -488,6 +495,11 @@ namespace UnrealGameSync
 
 			public void ExcludeFile(string path)
 			{
+				if (!CanContainExcludedFiles)
+				{
+					return;
+				}
+
 				int idx = path.IndexOf('/', StringComparison.Ordinal);
 				if (idx != -1)
 				{
@@ -499,19 +511,54 @@ namespace UnrealGameSync
 
 			public void GetOptimizedSyncCommands(string prefix, int changeNumber, SyncBatchBuilder builder)
 			{
-				if (CanUseWildcard && TotalExcludedFiles == 0 && TotalSize < builder.MaxSizePerList)
+				if (CanContainExcludedFiles)
+				{
+					GetOptimizedSyncCommandsWithExcludedFiles(prefix, changeNumber, builder);
+				}
+				else
+				{
+					GetOptimizedSyncCommandsWithoutExcludedFiles(prefix, changeNumber, builder);
+				}
+			}
+
+			private void GetOptimizedSyncCommandsWithExcludedFiles(string prefix, int changeNumber, SyncBatchBuilder builder)
+			{
+				if (CanUseWildcard
+				 && TotalExcludedFiles == 0
+				 && TotalSize < builder.MaxSizePerList)
+				{
+					builder.Add($"{prefix}/...@{changeNumber}", TotalSize);
+				}
+				else
+				{
+					foreach (KeyValuePair<string, long> file in IncludedFiles)
+					{
+						builder.Add($"{prefix}/{file.Key}@{changeNumber}", file.Value);
+					}
+					foreach (KeyValuePair<string, SyncTree> pair in NameToSubTree)
+					{
+						pair.Value.GetOptimizedSyncCommandsWithExcludedFiles($"{prefix}/{PerforceUtils.EscapePath(pair.Key)}", changeNumber, builder);
+					}
+				}
+			}
+
+			private void GetOptimizedSyncCommandsWithoutExcludedFiles(string prefix, int changeNumber, SyncBatchBuilder builder)
+			{
+				if (CanUseWildcard
+				 && IncludedFiles.Any()
+				 && TotalSize < builder.MaxSizePerList)
 				{
-					builder.Add(String.Format("{0}/...@{1}", prefix, changeNumber), TotalSize);
+					builder.Add($"{prefix}/...@{changeNumber}", TotalSize);
 				}
 				else
 				{
 					foreach (KeyValuePair<string, long> file in IncludedFiles)
 					{
-						builder.Add(String.Format("{0}/{1}@{2}", prefix, file.Key, changeNumber), file.Value);
+						builder.Add($"{prefix}/{file.Key}@{changeNumber}", file.Value);
 					}
 					foreach (KeyValuePair<string, SyncTree> pair in NameToSubTree)
 					{
-						pair.Value.GetOptimizedSyncCommands(String.Format("{0}/{1}", prefix, PerforceUtils.EscapePath(pair.Key)), changeNumber, builder);
+						pair.Value.GetOptimizedSyncCommandsWithoutExcludedFiles($"{prefix}/{PerforceUtils.EscapePath(pair.Key)}", changeNumber, builder);
 					}
 				}
 			}
@@ -755,6 +802,8 @@ namespace UnrealGameSync
 
 							// Update the sync filter hash. We've removed any files we need to at this point.
 							state = stateMgr.Modify(x => x.CurrentSyncFilterHash = nextSyncFilterHash);
+
+							times.Add(new Tuple<string, TimeSpan>("Sync Filters Changed", filterStopwatch.Stop("Success")));
 						}
 					}
 
@@ -768,7 +817,8 @@ namespace UnrealGameSync
 					}
 
 					// Create a tree to store the sync path
-					SyncTree syncTree = new SyncTree(false);
+					bool canContainExcludedFiles = !Context.Options.HasFlag(WorkspaceUpdateOptions.FastSync);
+					SyncTree syncTree = new SyncTree(false, canContainExcludedFiles);
 					if (!Context.Options.HasFlag(WorkspaceUpdateOptions.SyncSingleChange))
 					{
 						foreach (string relativeSyncPath in relativeSyncPaths)
@@ -794,111 +844,131 @@ namespace UnrealGameSync
 					List<string> syncDepotPaths = new List<string>();
 					using (RecordCounter counter = new RecordCounter(Progress, "Filtering files..."))
 					{
-						// Track the total new bytes that will be required on disk when syncing. Add an extra 100MB for padding.
-						long requiredFreeSpace = 100 * 1024 * 1024;
-
-						foreach (string syncPath in syncPaths)
+						using (TelemetryStopwatch filterStopwatch = new TelemetryStopwatch("Workspace_Sync_FilteringFiles", project.TelemetryProjectIdentifier))
 						{
-							string syncFilter = Context.Options.HasFlag(WorkspaceUpdateOptions.Sync) ? $"{Context.ChangeNumber}" : $"={Context.ChangeNumber}";
+							filterStopwatch.AddData(new { MachineName = System.Net.Dns.GetHostName(), DomainName = Environment.UserDomainName, ServerAndPort = perforce.Settings.ServerAndPort, ClientName = perforce.Settings.ClientName, UserName = perforce.Settings.UserName });
 
-							List<SyncFile> syncFiles = new List<SyncFile>();
-							await foreach (PerforceResponse<SyncRecord> response in perforce.TrySyncAsync(SyncOptions.PreviewOnly, -1, 0, -1, -1, -1, -1, $"{syncPath}@{syncFilter}", cancellationToken))
-							{
-								if (!response.Succeeded)
-								{
-									return (WorkspaceUpdateResult.FailedToSync, $"Couldn't enumerate changes matching {syncPath}.");
-								}
-								if (response.Info != null)
-								{
-									logger.LogInformation("Note: {Note}", response.Info.Data);
-									continue;
-								}
+							// Track the total new bytes that will be required on disk when syncing. Add an extra 100MB for padding.
+							long requiredFreeSpace = 100 * 1024 * 1024;
 
-								SyncRecord record = response.Data;
+							string syncFilter = Context.Options.HasFlag(WorkspaceUpdateOptions.Sync) ? $"{Context.ChangeNumber}" : $"={Context.ChangeNumber}";
 
-								// Get the full local path
-								string relativePath;
-								try
+							if (Context.Options.HasFlag(WorkspaceUpdateOptions.Sync) && Context.Options.HasFlag(WorkspaceUpdateOptions.FastSync))
+							{
+								if (state.CurrentChangeNumber > 0 && Context.ChangeNumber > state.CurrentChangeNumber)
 								{
-									FileReference syncFile = new FileReference(record.Path.ToString());
-									relativePath = PerforceUtils.GetClientRelativePath(project.LocalRootPath, syncFile);
+									syncFilter = $"{state.CurrentChangeNumber},{Context.ChangeNumber}";
+									logger.LogInformation("Using Fast Sync between current changelist {CurrentChangelist} and target changelist {TargetChangelist}", state.CurrentChangeNumber, Context.ChangeNumber);
 								}
-								catch (PathTooLongException)
+								else
 								{
-									logger.LogInformation("The local path for {Path} exceeds the maximum allowed by Windows. Re-sync your workspace to a directory with a shorter name, or delete the file from the server.", record.Path);
-									return (WorkspaceUpdateResult.FailedToSync, "File exceeds maximum path length allowed by Windows.");
+									logger.LogInformation("Fast Sync cannot be used with current changelist {CurrentChangelist} and target changelist {TargetChangelist}", state.CurrentChangeNumber, Context.ChangeNumber);
 								}
-
-								// Create the sync record
-								long syncSize = (record.Action == SyncAction.Deleted) ? 0 : record.FileSize;
-								syncFiles.Add(new SyncFile(record.DepotFile.ToString(), relativePath, syncSize));
-								counter.Increment();
 							}
 
-							// Also sync the currently open files
-							await foreach (PerforceResponse<OpenedRecord> response in perforce.TryOpenedAsync(OpenedOptions.None, FileSpecList.Any, cancellationToken))
+							foreach (string syncPath in syncPaths)
 							{
-								if (!response.Succeeded)
+								List<SyncFile> syncFiles = new List<SyncFile>();
+								await foreach (PerforceResponse<SyncRecord> response in perforce.TrySyncAsync(SyncOptions.PreviewOnly, -1, 0, -1, -1, -1, -1, $"{syncPath}@{syncFilter}", cancellationToken))
 								{
-									return (WorkspaceUpdateResult.FailedToSync, $"Couldn't enumerate changes matching {syncPath}.");
+									if (!response.Succeeded)
+									{
+										return (WorkspaceUpdateResult.FailedToSync, $"Couldn't enumerate changes matching {syncPath}.");
+									}
+									if (response.Info != null)
+									{
+										logger.LogInformation("Note: {Note}", response.Info.Data);
+										continue;
+									}
+
+									SyncRecord record = response.Data;
+
+									// Get the full local path
+									string relativePath;
+									try
+									{
+										FileReference syncFile = new FileReference(record.Path.ToString());
+										relativePath = PerforceUtils.GetClientRelativePath(project.LocalRootPath, syncFile);
+									}
+									catch (PathTooLongException)
+									{
+										logger.LogInformation("The local path for {Path} exceeds the maximum allowed by Windows. Re-sync your workspace to a directory with a shorter name, or delete the file from the server.", record.Path);
+										return (WorkspaceUpdateResult.FailedToSync, "File exceeds maximum path length allowed by Windows.");
+									}
+
+									// Create the sync record
+									long syncSize = (record.Action == SyncAction.Deleted) ? 0 : record.FileSize;
+									syncFiles.Add(new SyncFile(record.DepotFile.ToString(), relativePath, syncSize));
+									counter.Increment();
 								}
 
-								OpenedRecord record = response.Data;
-								if (!String.IsNullOrEmpty(record.DepotFile) && !String.IsNullOrEmpty(record.ClientFile))
+								// Also sync the currently open files
+								await foreach (PerforceResponse<OpenedRecord> response in perforce.TryOpenedAsync(OpenedOptions.None, FileSpecList.Any, cancellationToken))
 								{
-									if (record.Action != FileAction.Add && record.Action != FileAction.Branch && record.Action != FileAction.MoveAdd)
+									if (!response.Succeeded)
 									{
-										string relativePath = PerforceUtils.GetClientRelativePath(record.ClientFile);
-										syncFiles.Add(new SyncFile(record.DepotFile, relativePath, 0));
+										return (WorkspaceUpdateResult.FailedToSync, $"Couldn't enumerate changes matching {syncPath}.");
+									}
+
+									OpenedRecord record = response.Data;
+									if (!String.IsNullOrEmpty(record.DepotFile) && !String.IsNullOrEmpty(record.ClientFile))
+									{
+										if (record.Action != FileAction.Add && record.Action != FileAction.Branch && record.Action != FileAction.MoveAdd)
+										{
+											string relativePath = PerforceUtils.GetClientRelativePath(record.ClientFile);
+											syncFiles.Add(new SyncFile(record.DepotFile, relativePath, 0));
+										}
 									}
 								}
-							}
 
-							// Enumerate all the files to be synced. NOTE: depotPath is escaped, whereas clientPath is not.
-							List<string> syncRelativePaths = new();
-							foreach (SyncFile syncRecord in syncFiles)
-							{
-								if (filter.Matches(syncRecord.RelativePath))
+								// Enumerate all the files to be synced. NOTE: depotPath is escaped, whereas clientPath is not.
+								List<string> syncRelativePaths = new();
+								foreach (SyncFile syncRecord in syncFiles)
 								{
-									syncTree.IncludeFile(PerforceUtils.EscapePath(syncRecord.RelativePath), syncRecord.Size, logger);
-									syncDepotPaths.Add(syncRecord.DepotFile);
-									syncRelativePaths.Add(syncRecord.RelativePath);
-									requiredFreeSpace += syncRecord.Size;
+									if (filter.Matches(syncRecord.RelativePath))
+									{
+										syncTree.IncludeFile(PerforceUtils.EscapePath(syncRecord.RelativePath), syncRecord.Size, logger);
+										syncDepotPaths.Add(syncRecord.DepotFile);
+										syncRelativePaths.Add(syncRecord.RelativePath);
+										requiredFreeSpace += syncRecord.Size;
+									}
+									else
+									{
+										syncTree.ExcludeFile(PerforceUtils.EscapePath(syncRecord.RelativePath));
+									}
 								}
-								else
+
+								Parallel.ForEach(syncRelativePaths, syncRelativePath =>
 								{
-									syncTree.ExcludeFile(PerforceUtils.EscapePath(syncRecord.RelativePath));
-								}
+									// If the file exists the required free space can be reduced as those bytes will be replaced.
+									FileInfo localFileInfo = FileReference.Combine(project.LocalRootPath, syncRelativePath).ToFileInfo();
+									if (localFileInfo.Exists)
+									{
+										Interlocked.Add(ref requiredFreeSpace, -localFileInfo.Length);
+									}
+								});
 							}
 
-							Parallel.ForEach(syncRelativePaths, syncRelativePath =>
+							try
 							{
-								// If the file exists the required free space can be reduced as those bytes will be replaced.
-								FileInfo localFileInfo = FileReference.Combine(project.LocalRootPath, syncRelativePath).ToFileInfo();
-								if (localFileInfo.Exists)
+								DirectoryInfo localRootInfo = project.LocalRootPath.ToDirectoryInfo();
+								DriveInfo drive = new DriveInfo(localRootInfo.FullName);
+
+								if (drive.AvailableFreeSpace < requiredFreeSpace)
 								{
-									Interlocked.Add(ref requiredFreeSpace, -localFileInfo.Length);
+									logger.LogInformation("Syncing requires {RequiredSpace} which exceeds the {AvailableSpace} available free space on {Drive}.", StringUtils.FormatBytesString(requiredFreeSpace), StringUtils.FormatBytesString(drive.AvailableFreeSpace), drive.Name);
+									return (WorkspaceUpdateResult.FailedToSync, "Not enough available free space.");
 								}
-							});
-
-						}
-
-						try
-						{
-							DirectoryInfo localRootInfo = project.LocalRootPath.ToDirectoryInfo();
-							DriveInfo drive = new DriveInfo(localRootInfo.FullName);
-
-							if (drive.AvailableFreeSpace < requiredFreeSpace)
+							}
+							catch (SystemException)
 							{
-								logger.LogInformation("Syncing requires {RequiredSpace} which exceeds the {AvailableSpace} available free space on {Drive}.", StringUtils.FormatBytesString(requiredFreeSpace), StringUtils.FormatBytesString(drive.AvailableFreeSpace), drive.Name);
-								return (WorkspaceUpdateResult.FailedToSync, "Not enough available free space.");
+								logger.LogInformation("Unable to check available free space for {RootPath}.", project.LocalRootPath);
 							}
-						}
-						catch (SystemException)
-						{
-							logger.LogInformation("Unable to check available free space for {RootPath}.", project.LocalRootPath);
+
+							times.Add(new Tuple<string, TimeSpan>("Sync Filtering Files", filterStopwatch.Stop("Success")));
 						}
 					}
+
 					syncTree.GetOptimizedSyncCommands(project.ClientRootPath, Context.ChangeNumber, batchBuilder);
 
 					// Clear the current sync changelist, in case we cancel
@@ -926,8 +996,10 @@ namespace UnrealGameSync
 							return syncResult;
 						}
 
-						transferStopwatch.Stop("Ok");
+						TimeSpan stop = transferStopwatch.Stop("Ok");
 						transferStopwatch.AddData(new { TransferRate = syncTree.TotalSize / Math.Max(transferStopwatch.Elapsed.TotalSeconds, 0.0001f) });
+
+						times.Add(new Tuple<string, TimeSpan>("Sync Files Transfer", stop));
 					}
 
 					int versionChangeNumber = -1;
@@ -970,49 +1042,62 @@ namespace UnrealGameSync
 						int codeChangeNumber = Context.CodeChangeNumber ?? 0;
 						if (codeChangeNumber == 0)
 						{
-							logger.LogInformation("Finding last code change for CL {Number}...", Context.ChangeNumber);
-
-							// If we are syncing to a newer change than the last code change we found (and it is not the first sync in a workspace)
-							// go head and use the last code change we found as the minimum change in our query
-							int? minChangeNumber = null;
-							if ((Context.ChangeNumber >= state.CurrentCodeChangeNumber) && (state.CurrentCodeChangeNumber > 0))
+							using (TelemetryStopwatch findingLastCodeChangeStopwatch = new TelemetryStopwatch("Workspace_Sync_FindingLastCodeChange", project.TelemetryProjectIdentifier))
 							{
-								minChangeNumber = state.CurrentCodeChangeNumber;
-							}
+								findingLastCodeChangeStopwatch.AddData(new
+								{
+									MachineName = System.Net.Dns.GetHostName(),
+									DomainName = Environment.UserDomainName,
+									ServerAndPort = perforce.Settings.ServerAndPort,
+									ClientName = perforce.Settings.ClientName,
+									UserName = perforce.Settings.UserName
+								});
+
+								logger.LogInformation("Finding last code change for CL {Number}...", Context.ChangeNumber);
+
+								// If we are syncing to a newer change than the last code change we found (and it is not the first sync in a workspace)
+								// go head and use the last code change we found as the minimum change in our query
+								int? minChangeNumber = null;
+								if ((Context.ChangeNumber >= state.CurrentCodeChangeNumber) && (state.CurrentCodeChangeNumber > 0))
+								{
+									minChangeNumber = state.CurrentCodeChangeNumber;
+								}
 
-							string[] codeRules = Utility.GetCodeFilter(Context.ProjectConfigFile);
+								string[] codeRules = Utility.GetCodeFilter(Context.ProjectConfigFile);
 
-							try
-							{
-								await foreach (PerforceChangeDetails details in Utility.EnumerateChangeDetails(perforce, minChangeNumber, maxChangeNumber: Context.ChangeNumber, syncPaths, codeRules, cancellationToken))
+								try
 								{
-									if (details.ContainsCode)
+									await foreach (PerforceChangeDetails details in Utility.EnumerateChangeDetails(perforce, minChangeNumber, maxChangeNumber: Context.ChangeNumber, syncPaths, codeRules, cancellationToken))
 									{
-										codeChangeNumber = details.Number;
-										break;
+										if (details.ContainsCode)
+										{
+											codeChangeNumber = details.Number;
+											break;
+										}
 									}
 								}
-							}
-							catch (EpicGames.Perforce.PerforceException)
-							{
-								logger.LogInformation("Falling back to the slow way of finding last code change for CL {Number}...", Context.ChangeNumber);
-
-								await foreach (PerforceChangeDetails details in Utility.EnumerateChangeDetails(perforce, null, maxChangeNumber: Context.ChangeNumber, syncPaths, codeRules, cancellationToken))
+								catch (EpicGames.Perforce.PerforceException)
 								{
-									if (details.ContainsCode)
+									logger.LogInformation("Falling back to the slow way of finding last code change for CL {Number}...", Context.ChangeNumber);
+
+									await foreach (PerforceChangeDetails details in Utility.EnumerateChangeDetails(perforce, null, maxChangeNumber: Context.ChangeNumber, syncPaths, codeRules, cancellationToken))
 									{
-										codeChangeNumber = details.Number;
-										break;
+										if (details.ContainsCode)
+										{
+											codeChangeNumber = details.Number;
+											break;
+										}
 									}
 								}
-							}
 
-							if (codeChangeNumber == 0)
-							{
-								return (WorkspaceUpdateResult.FailedToSync, $"Could not find any code changes before CL {Context.ChangeNumber}.");
-							}
+								if (codeChangeNumber == 0)
+								{
+									return (WorkspaceUpdateResult.FailedToSync, $"Could not find any code changes before CL {Context.ChangeNumber}.");
+								}
 
-							logger.LogInformation("Using code CL {Number}", codeChangeNumber);
+								logger.LogInformation("Using code CL {Number}", codeChangeNumber);
+								times.Add(new Tuple<string, TimeSpan>("Sync Finding Last Code Change", findingLastCodeChangeStopwatch.Stop("Success")));
+							}
 						}
 
 						// Set the version change
@@ -1110,32 +1195,52 @@ namespace UnrealGameSync
 						string[]? postSyncSteps = Context.ProjectConfigFile.GetValues("Sync.Step", null);
 						if (postSyncSteps != null)
 						{
-							logger.LogInformation("");
-							logger.LogInformation("Executing post-sync steps...");
-
-							Dictionary<string, string> postSyncVariables = ConfigUtils.GetWorkspaceVariables(project, Context.ChangeNumber, versionChangeNumber, null, Context.ProjectConfigFile, perforceSettings);
-							foreach (string postSyncStep in postSyncSteps.Select(x => x.Trim()))
+							using (TelemetryStopwatch postSyncStopwatch = new TelemetryStopwatch("Workspace_Sync_PostSyncSteps", project.TelemetryProjectIdentifier))
 							{
-								ConfigObject postSyncStepObject = new ConfigObject(postSyncStep);
-
-								string toolFileName = Utility.ExpandVariables(postSyncStepObject.GetValue("FileName", ""), postSyncVariables);
-								if (toolFileName != null)
+								postSyncStopwatch.AddData(new
 								{
-									string toolArguments = Utility.ExpandVariables(postSyncStepObject.GetValue("Arguments", ""), postSyncVariables);
+									MachineName = System.Net.Dns.GetHostName(),
+									DomainName = Environment.UserDomainName,
+									ServerAndPort = perforce.Settings.ServerAndPort,
+									ClientName = perforce.Settings.ClientName,
+									UserName = perforce.Settings.UserName
+								});
 
-									logger.LogInformation("post-sync> Running {FileName} {Arguments}", toolFileName, toolArguments);
 
-									if (!File.Exists(toolFileName))
-									{
-										return (WorkspaceUpdateResult.FailedToSync, $"Unable to find {toolFileName}. You may have an incomplete sync; try cleaning your workspace.");
-									}
+								logger.LogInformation("");
+								logger.LogInformation("Executing post-sync steps...");
+
+								Dictionary<string, string> postSyncVariables = ConfigUtils.GetWorkspaceVariables(project, Context.ChangeNumber,
+									versionChangeNumber, null, Context.ProjectConfigFile, perforceSettings);
+								foreach (string postSyncStep in postSyncSteps.Select(x => x.Trim()))
+								{
+									ConfigObject postSyncStepObject = new ConfigObject(postSyncStep);
 
-									int resultFromTool = await Utility.ExecuteProcessAsync(toolFileName, null, toolArguments, line => ProcessOutput(line, "post-sync> ", Progress, logger), cancellationToken);
-									if (resultFromTool != 0)
+									string toolFileName = Utility.ExpandVariables(postSyncStepObject.GetValue("FileName", ""), postSyncVariables);
+									if (toolFileName != null)
 									{
-										return (WorkspaceUpdateResult.FailedToSync, $"Post-sync step terminated with exit code {resultFromTool}.");
+										string toolArguments =
+											Utility.ExpandVariables(postSyncStepObject.GetValue("Arguments", ""), postSyncVariables);
+
+										logger.LogInformation("post-sync> Running {FileName} {Arguments}", toolFileName, toolArguments);
+
+										if (!File.Exists(toolFileName))
+										{
+											return (WorkspaceUpdateResult.FailedToSync,
+												$"Unable to find {toolFileName}. You may have an incomplete sync; try cleaning your workspace.");
+										}
+
+										int resultFromTool = await Utility.ExecuteProcessAsync(toolFileName, null, toolArguments,
+											line => ProcessOutput(line, "post-sync> ", Progress, logger), cancellationToken);
+										if (resultFromTool != 0)
+										{
+											return (WorkspaceUpdateResult.FailedToSync,
+												$"Post-sync step terminated with exit code {resultFromTool}.");
+										}
 									}
 								}
+
+								times.Add(new Tuple<string, TimeSpan>("Sync Post Sync Steps", postSyncStopwatch.Stop("Success")));
 							}
 						}
 					}
@@ -1155,7 +1260,7 @@ namespace UnrealGameSync
 					});
 
 					// Update the timing info
-					times.Add(new Tuple<string, TimeSpan>("Sync", syncTelemetryStopwatch.Stop("Success")));
+					times.Add(new Tuple<string, TimeSpan>("Sync Total", syncTelemetryStopwatch.Stop("Success")));
 
 					// Save the number of files synced
 					numFilesSynced = syncDepotPaths.Count;
