UGS: Fix compiler warnings

diff --git a/Engine/Source/Programs/UnrealGameSync/UnrealGameSync/Controls/WorkspaceControl.cs b/Engine/Source/Programs/UnrealGameSync/UnrealGameSync/Controls/WorkspaceControl.cs
--- a/Engine/Source/Programs/UnrealGameSync/UnrealGameSync/Controls/WorkspaceControl.cs
+++ b/Engine/Source/Programs/UnrealGameSync/UnrealGameSync/Controls/WorkspaceControl.cs
@@ -269,7 +269,7 @@ namespace UnrealGameSync
 
 		// When an author filter is applied this will be non-empty and != AuthorFilterPlaceholderText
 		string _authorFilterText = "";
-		List<string> _authorsFilter = new List<string>();
+		readonly List<string> _authorsFilter = new List<string>();
 
 		// Placeholder text that is in the control and cleared when the user starts editing.
 		const string AuthorFilterPlaceholderText = "<username>";
@@ -6019,8 +6019,8 @@ namespace UnrealGameSync
 					{
 						if (!String.IsNullOrEmpty(step.StatusPanelLink) && CanRunStep(step))
 						{
-							int baseMenuIdx = step.StatusPanelLink.IndexOf('|', StringComparison.Ordinal);
-							if (baseMenuIdx == -1)
+							bool containsSegment = step.StatusPanelLink.Contains('|', StringComparison.Ordinal);
+							if (!containsSegment)
 							{
 								_customStatusPanelLinks.Add((step.StatusPanelLink, (p, r) => RunCustomTool(step, userSteps)));
 							}
