[UE-229159] It is impossible to filter by Author in UGS

- Add missing event handlers
- Properly clear filters when removing text from the box
- Properly set text colours when editing
- Restore missing placeholder text

diff --git a/Engine/Source/Programs/UnrealGameSync/UnrealGameSync/Controls/WorkspaceControl.cs b/Engine/Source/Programs/UnrealGameSync/UnrealGameSync/Controls/WorkspaceControl.cs
--- a/Engine/Source/Programs/UnrealGameSync/UnrealGameSync/Controls/WorkspaceControl.cs
+++ b/Engine/Source/Programs/UnrealGameSync/UnrealGameSync/Controls/WorkspaceControl.cs
@@ -6990,6 +6990,7 @@ namespace UnrealGameSync
 			if (FilterContextMenu_Author_Name.Text == AuthorFilterPlaceholderText)
 			{
 				FilterContextMenu_Author_Name.Text = "";
+				FilterContextMenu_Author_Name.ForeColor = Color.Black;
 				FilterContextMenu_Author_Name.TextBox.ForeColor = Color.Black;
 			}
 		}
@@ -7005,6 +7006,7 @@ namespace UnrealGameSync
 			if (String.IsNullOrEmpty(FilterContextMenu_Author_Name.Text))
 			{
 				FilterContextMenu_Author_Name.Text = AuthorFilterPlaceholderText;
+				FilterContextMenu_Author_Name.ForeColor = Color.DarkGray;
 				FilterContextMenu_Author_Name.TextBox.ForeColor = Color.DarkGray;
 			}
 		}
@@ -7023,6 +7025,10 @@ namespace UnrealGameSync
 				_authorsFilter.Clear();
 				_authorsFilter.AddAll( _authorFilterText.Split([';', '+'], StringSplitOptions.TrimEntries));
 			}
+			else
+			{
+				_authorsFilter.Clear();
+			}
 
 			UpdateBuildListFilter();
 		}
