UGS: Increase log verbosity when downloading artifacts.

#rnx

diff --git a/Engine/Source/Programs/Shared/EpicGames.Horde/Storage/Backends/HttpStorageBackend.cs b/Engine/Source/Programs/Shared/EpicGames.Horde/Storage/Backends/HttpStorageBackend.cs
--- a/Engine/Source/Programs/Shared/EpicGames.Horde/Storage/Backends/HttpStorageBackend.cs
+++ b/Engine/Source/Programs/Shared/EpicGames.Horde/Storage/Backends/HttpStorageBackend.cs
@@ -364,11 +364,12 @@ namespace EpicGames.Horde.Storage.Backends
 			}
 		}
 
-		static async Task EnsureSuccessAsync(HttpResponseMessage message, CancellationToken cancellationToken)
+		async Task EnsureSuccessAsync(HttpResponseMessage message, CancellationToken cancellationToken)
 		{
 			if (!message.IsSuccessStatusCode)
 			{
 				string response = await message.Content.ReadAsStringAsync(cancellationToken);
+				_logger.LogError("Http response {StatusCode} (Content: {Response})", message.StatusCode, response);
 				throw new StorageException($"Http response {message.StatusCode} (Content: {response})");
 			}
 		}
diff --git a/Engine/Source/Programs/UnrealGameSync/UnrealGameSync/ArtifactDownload.cs b/Engine/Source/Programs/UnrealGameSync/UnrealGameSync/ArtifactDownload.cs
--- a/Engine/Source/Programs/UnrealGameSync/UnrealGameSync/ArtifactDownload.cs
+++ b/Engine/Source/Programs/UnrealGameSync/UnrealGameSync/ArtifactDownload.cs
@@ -45,7 +45,11 @@ namespace UnrealGameSync
 				Uri serverUrl = new Uri(baseUri.GetLeftPart(UriPartial.Authority));
 
 				ServiceCollection services = new ServiceCollection();
-				services.AddLogging(builder => builder.AddProvider(loggerProvider));
+				services.AddLogging(builder =>
+				{
+					builder.AddProvider(loggerProvider);
+					builder.SetMinimumLevel(LogLevel.Trace);
+				});
 				services.AddHorde(options =>
 				{
 					options.ServerUrl = serverUrl;
@@ -62,7 +66,11 @@ namespace UnrealGameSync
 				IBlobRef<DirectoryNode>? blobRef = await storageNamespace.ReadRefAsync<DirectoryNode>(refName, cancellationToken: cancellationToken);
 
 				progress.Report("Starting...");
+				ILogger logger = serviceProvider.GetRequiredService<ILogger<ArtifactDownload>>();
+
 				await blobRef.ExtractAsync(outputDir.ToDirectoryInfo(), new CopyProgressAdapter(progress), TimeSpan.FromSeconds(0.2), serviceProvider.GetRequiredService<ILogger<ArtifactDownload>>(), cancellationToken);
+
+				logger.LogInformation("Download complete");
 			}
 		}
 
diff --git a/Engine/Source/Programs/UnrealGameSync/UnrealGameSync/Logging.cs b/Engine/Source/Programs/UnrealGameSync/UnrealGameSync/Logging.cs
--- a/Engine/Source/Programs/UnrealGameSync/UnrealGameSync/Logging.cs
+++ b/Engine/Source/Programs/UnrealGameSync/UnrealGameSync/Logging.cs
@@ -13,6 +13,7 @@ namespace UnrealGameSync
 		{
 			Serilog.ILogger logger = new LoggerConfiguration()
 				.Enrich.FromLogContext()
+				.MinimumLevel.Verbose()
 				.WriteTo.File(new ExpressionTemplate("[{@t:HH:mm:ss.fff} {@l:w3}] {Concat('[', Substring(SourceContext, LastIndexOf(SourceContext, '.') + 1, 14), ']'),16} {Indent}{@m}{@x}\n"), file.FullName, rollingInterval: RollingInterval.Day, rollOnFileSizeLimit: true, fileSizeLimitBytes: 20 * 1024 * 1024, retainedFileCountLimit: 10)
 				.CreateLogger();
 
