UGS: Add dark mode (and generic theming support)

PR #12322

diff --git a/Engine/Source/Programs/UnrealGameSync/UnrealGameSync/ApplicationThemeManager.cs b/Engine/Source/Programs/UnrealGameSync/UnrealGameSync/ApplicationThemeManager.cs
new file mode 100644
--- /dev/null
+++ b/Engine/Source/Programs/UnrealGameSync/UnrealGameSync/ApplicationThemeManager.cs
@@ -0,0 +1,329 @@
+// Copyright Epic Games, Inc. All Rights Reserved.
+
+using System;
+using System.Drawing;
+using System.Runtime.InteropServices;
+using System.Windows.Forms;
+using Microsoft.Win32;
+using UnrealGameSync.Controls;
+
+namespace UnrealGameSync
+{
+	public sealed class InvertibleImage : IDisposable
+	{
+		bool _disposeImage;
+		Image? _standardImage;
+		Image? _invertedImage;
+
+		public InvertibleImage(Image image, bool disposeImage)
+		{
+			_disposeImage = disposeImage;
+			_standardImage = image;
+
+			InitInvertedImage();
+		}
+
+		public Image? ResolveImage()
+		{
+			bool isInverted = ProgramApplicationContext.GetTheme().ShouldInvertColorBrightness;
+			return isInverted ? _invertedImage : _standardImage;
+		}
+		
+		public void Dispose()
+		{
+			if (_disposeImage && _standardImage != null)
+			{
+				_standardImage.Dispose();
+			}
+			_standardImage = null;
+
+			_invertedImage?.Dispose();
+			_invertedImage = null;
+		}
+		
+		private void InitInvertedImage()
+		{
+			if (_standardImage != null)
+			{
+				Bitmap invertedLogoBitmap = new Bitmap(_standardImage);
+
+				for (int x = 0; x < invertedLogoBitmap.Width; ++x)
+				{
+					for (int y = 0; y < invertedLogoBitmap.Height; ++y)
+					{
+						invertedLogoBitmap.SetPixel(x, y, ApplicationTheme.InvertColorBrightness(invertedLogoBitmap.GetPixel(x, y)));
+					}
+				}
+
+				_invertedImage = invertedLogoBitmap;
+			}
+		}
+	}
+	
+	public class InvertibleColor
+	{
+		public Color StandardColor { get; init; }
+		public Color InvertedColor { get; init; }
+
+		public InvertibleColor(Color color)
+		{
+			StandardColor = color;
+			InvertedColor = ApplicationTheme.InvertColorBrightness(color);
+		}
+
+		public Color ResolveColor()
+		{
+			bool isInverted = ProgramApplicationContext.GetTheme().ShouldInvertColorBrightness;
+			return isInverted ? InvertedColor : StandardColor;
+		}
+	}
+	
+	interface ICustomThemeHandler
+	{
+		public void ApplyTheme(ApplicationThemeManager themeManager);
+	}
+
+	public sealed class ApplicationThemeManager : IDisposable
+	{
+		public ApplicationTheme Theme => ProgramApplicationContext.GetTheme();
+		
+		readonly InvertibleImage _dropListIcon = new InvertibleImage(global::UnrealGameSync.Properties.Resources.DropList, false);
+
+		public ApplicationThemeManager()
+		{
+			UpdateDefaultTheme();
+		}
+		
+		public void Dispose()
+		{
+			_dropListIcon.Dispose();
+		}
+		
+		[DllImport("dwmapi.dll")]
+		private static extern int DwmSetWindowAttribute(IntPtr hwnd, int attr, ref int attrValue, int attrSize);
+
+		[DllImport("uxtheme.dll", CharSet = CharSet.Unicode)]
+		private static extern int SetWindowTheme(IntPtr hWnd, string? pszSubAppName, string? pszSubIdList);
+		
+		private static void UpdateDefaultTheme()
+		{
+			// Set the default theme based on the user's system settings
+			int? lightThemeValue = (int?)Registry.GetValue("HKEY_CURRENT_USER\\SOFTWARE\\Microsoft\\Windows\\CurrentVersion\\Themes\\Personalize", "AppsUseLightTheme", -1);
+			bool isDarkDefault = (lightThemeValue ?? -1) == 0;
+
+			ApplicationTheme.DefaultThemeId = isDarkDefault ? ApplicationTheme.DarkThemeId : ApplicationTheme.LightThemeId;
+		}
+		
+		private void ApplyDarkMode(Control control)
+		{
+			if (!control.IsHandleCreated)
+			{
+				return;
+			}
+			
+			if (Theme.UseDarkMode)
+			{
+				string darkTheme = "darkmode_explorer";
+
+				switch (control)
+				{
+					// special case for the custom list since this makes the group header legible (but the scrollbar is white)
+					// where without this, "darkmode_explorer" will make the group header a dark blue (but have a dark scrollbox)
+					case CustomListViewControl:
+						darkTheme = "darkmode_itemsview";
+						break;
+					// Properly themes the full dropdown
+					case ComboBox:
+						darkTheme = "darkmode_cfd";
+						break;
+				}
+				
+				SetWindowTheme(control.Handle, darkTheme, null);
+			}
+			else
+			{
+				SetWindowTheme(control.Handle, null, null);
+			}
+
+			if (Environment.OSVersion.Version.Major >= 10 && Environment.OSVersion.Version.Build >= 17763)
+			{
+				const int DWMWA_USE_IMMERSIVE_DARK_MODE_BEFORE_20H1 = 19;
+				const int DWMWA_USE_IMMERSIVE_DARK_MODE = 20;
+
+				int attribute = (Environment.OSVersion.Version.Build >= 18985) ? DWMWA_USE_IMMERSIVE_DARK_MODE : DWMWA_USE_IMMERSIVE_DARK_MODE_BEFORE_20H1;
+				int useImmersiveDarkMode = Theme.UseDarkMode ? 1 : 0;
+
+				DwmSetWindowAttribute(control.Handle, attribute, ref useImmersiveDarkMode, sizeof(int));
+			}
+		}
+
+		public void ApplyThemeRecursively()
+		{
+			foreach (Form form in Application.OpenForms)
+			{
+				ApplyThemeRecursively(form);
+			}
+		}
+		public void ApplyThemeRecursively(Control control)
+		{
+			control.HandleCreated += (_, __) =>
+			{
+				ApplyThemeRecursively(control);
+			};
+			control.EnabledChanged +=  (_, __) => ApplyThemeRecursively(control);
+			control.ControlAdded += (_, args) =>
+			{
+				if (args.Control != null)
+				{
+					ApplyThemeRecursively(args.Control);
+				}
+			};
+
+			if (control is ICustomThemeHandler themeHandler)
+			{
+				themeHandler.ApplyTheme(this);
+			}
+			else
+			{
+				DefaultApplyThemeToControl(control);
+			}
+
+			ApplyThemeRecursively(control.Controls);
+		}
+
+		public void ApplyThemeRecursively(Control.ControlCollection controls)
+		{
+			foreach (Control control in controls)
+			{
+				ApplyThemeRecursively(control);
+			}
+		}
+
+		public void ApplyThemeRecursively(ToolStripItemCollection items)
+		{
+			foreach (ToolStripItem item in items)
+			{
+				ApplyThemeRecursively(item);
+			}
+		}
+
+		public void ApplyThemeRecursively(ToolStripItem item)
+		{
+			item.BackColor = Theme.BackgroundColor;
+			item.ForeColor = Theme.TextColor;
+
+			if (item is ToolStripDropDownItem dropDownItem)
+			{
+				ApplyThemeRecursively(dropDownItem.DropDown);
+			}
+		}
+
+		public void ApplyThemeRecursively(ToolTip item)
+		{
+			item.BackColor = Theme.BackgroundColor;
+			item.ForeColor = Theme.TextColor;
+		}
+
+		public void DefaultApplyThemeToControl(Control control)
+		{
+			control.BackColor = Theme.BackgroundColor;
+			control.ForeColor = Theme.ForegroundColor;
+			ApplyDarkMode(control);
+
+			switch (control)
+			{
+				case UpDownBase upDownControl:
+					upDownControl.BorderStyle = Theme.UseFlatControls ? BorderStyle.None : BorderStyle.Fixed3D;
+					break;
+				case LinkLabel linkLabel:
+					control.ForeColor = control.Enabled ? Theme.TextColor : Theme.DimmedText;
+					linkLabel.LinkColor = Theme.LinkColor;
+					linkLabel.ActiveLinkColor = Theme.LinkColorActive;
+					break;
+				case Label:
+				case GroupBox:
+					control.ForeColor = control.Enabled ? Theme.TextColor : Theme.DimmedText;
+					break;
+				case Button button:
+					control.BackColor = Theme.InputBackgroundColor;
+					control.ForeColor = control.Enabled ? Theme.TextColor : Theme.DimmedText;
+
+					button.FlatStyle = Theme.UseFlatControls ? FlatStyle.Flat : FlatStyle.Standard;
+					button.FlatAppearance.BorderColor = Theme.InputBorderColor;
+
+					if ("dropdown".Equals(button.Tag))
+					{
+						button.Image = _dropListIcon.ResolveImage();
+					}
+					break;
+				case TextBox textBox:
+					control.BackColor = Theme.InputBackgroundColor;
+					control.ForeColor = control.Enabled ? Theme.TextColor : Theme.DimmedText;
+					textBox.BorderStyle = Theme.UseFlatControls ? BorderStyle.FixedSingle : BorderStyle.Fixed3D;
+					break;
+				case ComboBox:
+					control.BackColor = Theme.InputBackgroundColor;
+					control.ForeColor = control.Enabled ? Theme.TextColor : Theme.DimmedText;
+					break;
+				case ReadOnlyCheckbox readOnlyCheckbox:
+					readOnlyCheckbox.ForeColor = readOnlyCheckbox.ReadOnly ? Theme.DimmedText : Theme.TextColor;
+					break;
+				case CheckBox:
+					control.ForeColor = control.Enabled ? Theme.TextColor : Theme.DimmedText;
+					break;
+				case ListBox:
+					control.BackColor = Theme.InputBackgroundColor;
+					control.ForeColor = control.Enabled ? Theme.TextColor : Theme.DimmedText;
+					break;
+				case ToolStrip toolStrip:
+					control.ForeColor = control.Enabled ? Theme.TextColor : Theme.DimmedText;
+					if (toolStrip.Renderer is not ThemedToolStripRenderer)
+					{
+						toolStrip.Renderer = new ThemedToolStripRenderer(this);
+					}
+					ApplyThemeRecursively(toolStrip.Items);
+					break;
+				case ListView:
+					control.BackColor = Theme.InputBackgroundColor;
+					control.ForeColor = control.Enabled ? Theme.TextColor : Theme.DimmedText;
+					break;
+				case Panel panel:
+					panel.BorderStyle = BorderStyle.None;
+					break;
+			}
+		}
+
+		private class ThemedToolStripRenderer(ApplicationThemeManager themeManager) : ToolStripProfessionalRenderer(new ThemedColorTable(themeManager))
+		{
+			protected override void OnRenderImageMargin(ToolStripRenderEventArgs e)
+			{
+				Color imageMargeBrushColor = themeManager.Theme.BackgroundColor;
+				using (Brush imageMarginBrush = new SolidBrush(imageMargeBrushColor))
+				{
+					e.Graphics.FillRectangle(imageMarginBrush, e.AffectedBounds);
+				}
+			}
+
+			protected override void OnRenderArrow(ToolStripArrowRenderEventArgs e)
+			{
+				if (e.Item != null)
+				{
+					e.ArrowColor = e.Item.Enabled ? themeManager.Theme.TextColor : themeManager.Theme.DimmedText;
+				}
+				
+				base.OnRenderArrow(e);
+			}
+			
+			private class ThemedColorTable(ApplicationThemeManager themeManager) : ProfessionalColorTable
+			{
+				public override Color MenuItemSelectedGradientBegin => themeManager.Theme.MenuHighlightBGColor;
+				public override Color MenuItemSelectedGradientEnd => themeManager.Theme.MenuHighlightBGColor;
+				public override Color MenuItemBorder => themeManager.Theme.MenuHighlightBorderColor;
+				public override Color CheckBackground => themeManager.Theme.MenuHighlightBGColor;
+				public override Color CheckPressedBackground => themeManager.Theme.MenuHighlightBGColor;
+				public override Color CheckSelectedBackground => themeManager.Theme.MenuHighlightBGColor;
+				public override Color ButtonSelectedBorder => themeManager.Theme.MenuHighlightBorderColor;
+			}
+		}
+	}
+}
diff --git a/Engine/Source/Programs/UnrealGameSync/UnrealGameSync/Controls/CustomListViewControl.cs b/Engine/Source/Programs/UnrealGameSync/UnrealGameSync/Controls/CustomListViewControl.cs
--- a/Engine/Source/Programs/UnrealGameSync/UnrealGameSync/Controls/CustomListViewControl.cs
+++ b/Engine/Source/Programs/UnrealGameSync/UnrealGameSync/Controls/CustomListViewControl.cs
@@ -2,6 +2,7 @@
 
 using System;
 using System.Drawing;
+using System.Runtime.InteropServices;
 using System.Windows.Forms;
 using System.Windows.Forms.VisualStyles;
 
@@ -190,7 +191,7 @@ namespace UnrealGameSync
 		}
 	}
 
-	partial class CustomListViewControl : ListView
+	partial class CustomListViewControl : ListView, ICustomThemeHandler
 	{
 		readonly VisualStyleRenderer? _selectedItemRenderer;
 		readonly VisualStyleRenderer? _trackedItemRenderer;
@@ -200,6 +201,8 @@ namespace UnrealGameSync
 		CustomListViewWidget? _mouseOverWidget;
 		CustomListViewWidget? _mouseDownWidget;
 
+		Color _textColor = SystemColors.WindowText;
+
 		public CustomListViewControl()
 		{
 			if (Application.RenderWithVisualStyles)
@@ -399,7 +402,7 @@ namespace UnrealGameSync
 		{
 			if (e.SubItem != null)
 			{
-				DrawText(e.Graphics, e.SubItem.Bounds, Columns[e.ColumnIndex].TextAlign, SystemColors.WindowText, e.SubItem.Text);
+				DrawText(e.Graphics, e.SubItem.Bounds, Columns[e.ColumnIndex].TextAlign, _textColor, e.SubItem.Text);
 			}
 		}
 
@@ -441,32 +444,81 @@ namespace UnrealGameSync
 
 		public void DrawDefaultBackground(Graphics graphics, Rectangle bounds)
 		{
+			bounds.Inflate(0, 1);
 			_ = this;
-			graphics.FillRectangle(SystemBrushes.Window, bounds);
+			using SolidBrush inputBGBrush = new SolidBrush(BackColor);
+			graphics.FillRectangle(inputBGBrush, bounds);
 		}
 
 		public void DrawSelectedBackground(Graphics graphics, Rectangle bounds)
 		{
+			bounds.Inflate(0, 1);
 			if (_selectedItemRenderer != null)
 			{
 				_selectedItemRenderer.DrawBackground(graphics, bounds);
 			}
 			else
 			{
-				graphics.FillRectangle(SystemBrushes.ButtonFace, bounds);
+				using SolidBrush inputBGBrush = new SolidBrush(BackColor);
+				graphics.FillRectangle(inputBGBrush, bounds);
 			}
 		}
 
 		public void DrawTrackedBackground(Graphics graphics, Rectangle bounds)
 		{
+			bounds.Inflate(0, 1);
 			if (_trackedItemRenderer != null)
 			{
 				_trackedItemRenderer.DrawBackground(graphics, bounds);
 			}
 			else
 			{
-				graphics.FillRectangle(SystemBrushes.Window, bounds);
+				using SolidBrush inputBGBrush = new SolidBrush(BackColor);
+				graphics.FillRectangle(inputBGBrush, bounds);
+			}
+		}
+
+		protected override void OnDrawColumnHeader(DrawListViewColumnHeaderEventArgs e)
+		{
+			ApplicationTheme theme = ProgramApplicationContext.GetTheme();
+			e.DrawDefault = false;
+
+			using (SolidBrush backBrush = new SolidBrush(theme.InputBackgroundColor))
+			{
+				e.Graphics.FillRectangle(backBrush, e.Bounds);
+			}
+
+			using (Pen backgroundPen = new Pen(theme.InputBackgroundColor))
+			using (Pen linePen = new Pen(theme.BackgroundColor))
+			{
+				Rectangle r = e.Bounds;
+				r.Width -= 1;
+				r.Height -= 1;
+
+				e.Graphics.DrawRectangle(backgroundPen, r);
+
+				r.Width -= 1;
+				r.Height -= 1;
+
+				e.Graphics.DrawLine(linePen, r.X, r.Y, r.Right, r.Y);
+				e.Graphics.DrawLine(linePen, r.X, r.Y, r.X, r.Bottom);
+
+				e.Graphics.DrawLine(backgroundPen, r.X + 1, r.Bottom, r.Right, r.Bottom);
+				e.Graphics.DrawLine(backgroundPen, r.Right, r.Y + 1, r.Right, r.Bottom);
+
+				Rectangle newBounds = Rectangle.Inflate(e.Bounds, -TextRenderer.MeasureText(" ", e.Font).Width, 0);
+				TextRenderer.DrawText(e.Graphics, e.Header?.Text, e.Font, newBounds, _textColor, TextFormatFlags.Left | TextFormatFlags.WordEllipsis | TextFormatFlags.VerticalCenter);
 			}
 		}
+
+		public void ApplyTheme(ApplicationThemeManager themeManager)
+		{
+			themeManager.DefaultApplyThemeToControl(this);
+
+			_textColor = themeManager.Theme.TextColor;
+
+			BackColor = themeManager.Theme.InputBackgroundColor;
+			ForeColor = themeManager.Theme.ForegroundColor;
+		}
 	}
 }
diff --git a/Engine/Source/Programs/UnrealGameSync/UnrealGameSync/Controls/LogControl.cs b/Engine/Source/Programs/UnrealGameSync/UnrealGameSync/Controls/LogControl.cs
--- a/Engine/Source/Programs/UnrealGameSync/UnrealGameSync/Controls/LogControl.cs
+++ b/Engine/Source/Programs/UnrealGameSync/UnrealGameSync/Controls/LogControl.cs
@@ -16,7 +16,7 @@ using Microsoft.Extensions.Logging;
 
 namespace UnrealGameSync
 {
-	public class LogControl : UserControl
+	public class LogControl : UserControl, ICustomThemeHandler
 	{
 		[Flags]
 		public enum ScrollInfoMask : int
@@ -119,6 +119,9 @@ namespace UnrealGameSync
 		TextSelection? _selection;
 
 		Size _fontSize;
+		Color _textColor = Color.Black;
+		Color _warningColor = Color.FromArgb(128, 128, 0);
+		Color _errorColor = Color.FromArgb(189, 54, 47);
 
 		Timer? _selectionScrollTimer;
 		int _autoScrollRate = 0;
@@ -556,6 +559,7 @@ namespace UnrealGameSync
 			if (e.Button == MouseButtons.Right)
 			{
 				ContextMenuStrip!.Show(this, e.Location);
+				ProgramApplicationContext.GetThemeManager()?.ApplyThemeRecursively(ContextMenuStrip);
 			}
 		}
 
@@ -763,15 +767,15 @@ namespace UnrealGameSync
 				Color textColor;
 				if (Regex.IsMatch(_lines[idx], "(?<!\\w)error[: ]", RegexOptions.IgnoreCase))
 				{
-					textColor = Color.FromArgb(189, 54, 47);
+					textColor = _errorColor;
 				}
 				else if (Regex.IsMatch(_lines[idx], "(?<!\\w)warning[: ]", RegexOptions.IgnoreCase))
 				{
-					textColor = Color.FromArgb(128, 128, 0);
+					textColor = _warningColor;
 				}
 				else
 				{
-					textColor = Color.Black;
+					textColor = _textColor;
 				}
 
 				int textY = (idx - _scrollLine) * _fontSize.Height;
@@ -917,6 +921,17 @@ namespace UnrealGameSync
 				}
 			}
 		}
+
+		public void ApplyTheme(ApplicationThemeManager themeManager)
+		{
+			themeManager.DefaultApplyThemeToControl(this);
+
+			BackColor = themeManager.Theme.InputBackgroundColor;
+			ForeColor = themeManager.Theme.TextColor;
+			_textColor = themeManager.Theme.TextColor;
+			_warningColor= themeManager.Theme.WarningColor;
+			_errorColor = themeManager.Theme.ErrorColor;
+		}
 	}
 
 	public class LogControlTextWriter : ILogger
diff --git a/Engine/Source/Programs/UnrealGameSync/UnrealGameSync/Controls/LogSplitContainer.cs b/Engine/Source/Programs/UnrealGameSync/UnrealGameSync/Controls/LogSplitContainer.cs
--- a/Engine/Source/Programs/UnrealGameSync/UnrealGameSync/Controls/LogSplitContainer.cs
+++ b/Engine/Source/Programs/UnrealGameSync/UnrealGameSync/Controls/LogSplitContainer.cs
@@ -7,7 +7,7 @@ using System.Windows.Forms;
 
 namespace UnrealGameSync
 {
-	public class LogSplitContainer : SplitContainer
+	public class LogSplitContainer : SplitContainer, ICustomThemeHandler
 	{
 		const int WmSetcursor = 0x0020;
 
@@ -21,6 +21,13 @@ namespace UnrealGameSync
 
 		int _logHeight = 200;
 
+		Color _textColor = SystemColors.ControlText;
+		Color _borderColor = SystemColors.ControlDark;
+		Color _activeCaption = SystemColors.ActiveCaption;
+		Color _inactiveCaption = SystemColors.InactiveCaption;
+		Color _inactiveCaptionText = SystemColors.InactiveCaptionText;
+		InvertibleImage _logIcon = new InvertibleImage(Properties.Resources.Log, false);
+
 		public event Action<bool>? OnVisibilityChanged;
 
 		[Browsable(true)]
@@ -42,6 +49,16 @@ namespace UnrealGameSync
 
 		protected static Font CaptionFont => SystemFonts.IconTitleFont!;
 
+		protected override void Dispose(bool disposing)
+		{
+			base.Dispose(disposing);
+
+			if (disposing)
+			{
+				_logIcon.Dispose();
+			}
+		}
+
 		protected override void OnCreateControl()
 		{
 			base.OnCreateControl();
@@ -231,7 +248,7 @@ namespace UnrealGameSync
 				buttonSize -= CaptionPadding;
 			}
 
-			using (Pen captionBorderPen = new Pen(SystemColors.ControlDark, 1.0f))
+			using (Pen captionBorderPen = new Pen(_borderColor, 1.0f))
 			{
 				e.Graphics.DrawRectangle(captionBorderPen, 0, captionMinY, ClientRectangle.Width - 1, captionMaxY - captionMinY - 1);
 			}
@@ -247,29 +264,36 @@ namespace UnrealGameSync
 			int crossY = (captionMinY + captionMaxY) / 2;
 			if (_hoverOverClose)
 			{
-				e.Graphics.FillRectangle(SystemBrushes.ActiveCaption, crossX - buttonSize / 2, crossY - buttonSize / 2, buttonSize, buttonSize);
+				using (Brush activeCaptionBrush = new SolidBrush(_activeCaption))
+				{
+					e.Graphics.FillRectangle(activeCaptionBrush, crossX - buttonSize / 2, crossY - buttonSize / 2, buttonSize, buttonSize);
+				}
 			}
 			else if (_mouseDownOverClose)
 			{
-				e.Graphics.FillRectangle(SystemBrushes.InactiveCaption, crossX - buttonSize / 2, crossY - buttonSize / 2, buttonSize - 2, buttonSize - 2);
+				using (Brush inactiveCaptionBrush = new SolidBrush(_inactiveCaption))
+				{
+					e.Graphics.FillRectangle(inactiveCaptionBrush, crossX - buttonSize / 2, crossY - buttonSize / 2, buttonSize - 2,buttonSize - 2);
+				}
 			}
 			if (_hoverOverClose || _mouseDownOverClose)
 			{
-				using (Pen borderPen = new Pen(SystemColors.InactiveCaptionText, 1.0f))
+				using (Pen borderPen = new Pen(_inactiveCaptionText, 1.0f))
 				{
 					e.Graphics.DrawRectangle(borderPen, crossX - buttonSize / 2, crossY - buttonSize / 2, buttonSize, buttonSize);
 				}
 			}
+
 			if (SplitterDistance >= Height - SplitterWidth)
 			{
-				e.Graphics.DrawImage(Properties.Resources.Log, new Rectangle(crossX - (buttonSize / 2) - 1, crossY - (buttonSize / 2) - 1, buttonSize + 2, buttonSize + 2), new Rectangle(16, 0, 16, 16), GraphicsUnit.Pixel);
+				e.Graphics.DrawImage(_logIcon.ResolveImage() ?? Properties.Resources.Log, new Rectangle(crossX - (buttonSize / 2) - 1, crossY - (buttonSize / 2) - 1, buttonSize + 2, buttonSize + 2), new Rectangle(16, 0, 16, 16), GraphicsUnit.Pixel);
 			}
 			else
 			{
-				e.Graphics.DrawImage(Properties.Resources.Log, new Rectangle(crossX - (buttonSize / 2) - 1, crossY - (buttonSize / 2) - 1, buttonSize + 2, buttonSize + 2), new Rectangle(0, 0, 16, 16), GraphicsUnit.Pixel);
+				e.Graphics.DrawImage(_logIcon.ResolveImage() ?? Properties.Resources.Log, new Rectangle(crossX - (buttonSize / 2) - 1, crossY - (buttonSize / 2) - 1, buttonSize + 2, buttonSize + 2), new Rectangle(0, 0, 16, 16), GraphicsUnit.Pixel);
 			}
 
-			TextRenderer.DrawText(e.Graphics, Caption, CaptionFont, new Rectangle(2, captionMinY, ClientRectangle.Width - 20, captionMaxY - captionMinY), SystemColors.ControlText, TextFormatFlags.Left | TextFormatFlags.VerticalCenter);
+			TextRenderer.DrawText(e.Graphics, Caption, CaptionFont, new Rectangle(2, captionMinY, ClientRectangle.Width - 20, captionMaxY - captionMinY), _textColor, TextFormatFlags.Left | TextFormatFlags.VerticalCenter);
 		}
 
 		protected override void WndProc(ref Message message)
@@ -288,6 +312,20 @@ namespace UnrealGameSync
 			}
 		}
 
+		public void ApplyTheme(ApplicationThemeManager themeManager)
+		{
+			themeManager.DefaultApplyThemeToControl(this);
+
+			BackColor = themeManager.Theme.InputBackgroundColor;
+			ForeColor = themeManager.Theme.TextColor;
+
+			_textColor = themeManager.Theme.TextColor;
+			_borderColor = themeManager.Theme.InputBorderColor;
+			_activeCaption = themeManager.Theme.ActiveCaption;
+			_inactiveCaption = themeManager.Theme.InactiveCaption;
+			_inactiveCaptionText = themeManager.Theme.InactiveCaptionText;
+		}
+
 		Rectangle CloseButtonRectangle => new Rectangle(ClientRectangle.Right - 20, SplitterDistance, 20, SplitterWidth);
 	}
 }
diff --git a/Engine/Source/Programs/UnrealGameSync/UnrealGameSync/Controls/ReadOnlyCheckbox.cs b/Engine/Source/Programs/UnrealGameSync/UnrealGameSync/Controls/ReadOnlyCheckbox.cs
new file mode 100644
--- /dev/null
+++ b/Engine/Source/Programs/UnrealGameSync/UnrealGameSync/Controls/ReadOnlyCheckbox.cs
@@ -0,0 +1,48 @@
+﻿// Copyright Epic Games, Inc. All Rights Reserved.
+
+using System;
+using System.Windows.Forms;
+
+namespace UnrealGameSync.Controls;
+
+public class ReadOnlyCheckbox : CheckBox
+{
+	private bool _readOnly = false;
+	[System.ComponentModel.Category("Behavior")]
+	[System.ComponentModel.DefaultValue(false)]
+	public bool ReadOnly
+	{
+		get => _readOnly;
+		set
+		{
+			_readOnly = value;
+			Cursor = _readOnly ? Cursors.No : Cursors.Default;
+			ProgramApplicationContext.GetThemeManager()?.ApplyThemeRecursively(this);
+		}
+	}
+
+	protected override void OnMouseEnter(EventArgs e)
+	{
+		// Disable highlight when the cursor is over the CheckBox
+		if (!ReadOnly)
+		{
+			base.OnMouseEnter(e);
+		}
+	}
+	protected override void OnMouseDown(MouseEventArgs e)
+	{
+		// Disable reacting (logically or visibly) to a mouse click
+		if (!ReadOnly)
+		{
+			base.OnMouseDown(e);
+		}
+	}
+	protected override void OnKeyDown(KeyEventArgs e)
+	{
+		// Suppress space key to disable checking/unchecking
+		if (!ReadOnly || e.KeyData != Keys.Space)
+		{
+			base.OnKeyDown(e);
+		}
+	}
+}
\ No newline at end of file
diff --git a/Engine/Source/Programs/UnrealGameSync/UnrealGameSync/Controls/StatusPanel.cs b/Engine/Source/Programs/UnrealGameSync/UnrealGameSync/Controls/StatusPanel.cs
--- a/Engine/Source/Programs/UnrealGameSync/UnrealGameSync/Controls/StatusPanel.cs
+++ b/Engine/Source/Programs/UnrealGameSync/UnrealGameSync/Controls/StatusPanel.cs
@@ -178,14 +178,14 @@ namespace UnrealGameSync
 
 		public override void Draw(Graphics graphics, StatusElementResources resources)
 		{
-			Color textColor = SystemColors.HotTrack;
+			Color textColor = ProgramApplicationContext.GetTheme().LinkColor;
 			if (MouseDown)
 			{
-				textColor = Color.FromArgb(textColor.B / 2, textColor.G / 2, textColor.R);
+				textColor = ProgramApplicationContext.GetTheme().LinkColorActive;
 			}
 			else if (MouseOver)
 			{
-				textColor = Color.FromArgb(textColor.B, textColor.G, textColor.R);
+				textColor = ProgramApplicationContext.GetTheme().LinkColorHover;
 			}
 			TextRenderer.DrawText(graphics, _text, resources.FindOrAddFont(_style), Bounds.Location, textColor, TextFormatFlags.NoPadding | TextFormatFlags.NoPrefix);
 		}
@@ -248,6 +248,7 @@ namespace UnrealGameSync
 			SmoothingMode prevSmoothingMode = graphics.SmoothingMode;
 			graphics.SmoothingMode = SmoothingMode.HighQuality;
 
+			Color badgeColor = MouseOver ? _hoverBackgroundColor : _backgroundColor;
 			using (GraphicsPath path = new GraphicsPath())
 			{
 				path.StartFigure();
@@ -257,14 +258,15 @@ namespace UnrealGameSync
 				path.AddLine(Bounds.Right - 2 + (MergeRight ? 1 : 0), Bounds.Top, Bounds.Left + (MergeLeft ? 1 : 0), Bounds.Top);
 				path.CloseFigure();
 
-				using (SolidBrush brush = new SolidBrush(MouseOver ? _hoverBackgroundColor : _backgroundColor))
+				using (SolidBrush brush = new SolidBrush(badgeColor))
 				{
 					graphics.FillPath(brush, path);
 				}
 			}
 
+			Color textColor = ApplicationTheme.IsColorLight(badgeColor) ? Color.Black : Color.White;
 			graphics.SmoothingMode = prevSmoothingMode;
-			TextRenderer.DrawText(graphics, _name, resources.BadgeFont, Bounds, Color.White, TextFormatFlags.HorizontalCenter | TextFormatFlags.VerticalCenter | TextFormatFlags.SingleLine | TextFormatFlags.NoPrefix | TextFormatFlags.PreserveGraphicsClipping);
+			TextRenderer.DrawText(graphics, _name, resources.BadgeFont, Bounds, textColor, TextFormatFlags.HorizontalCenter | TextFormatFlags.VerticalCenter | TextFormatFlags.SingleLine | TextFormatFlags.NoPrefix | TextFormatFlags.PreserveGraphicsClipping);
 		}
 	}
 
@@ -351,7 +353,8 @@ namespace UnrealGameSync
 
 		public void AddText(string inText, FontStyle inStyle = FontStyle.Regular)
 		{
-			_elements.Add(new TextStatusElement(inText, SystemColors.ControlText, inStyle));
+			Color textColor = ProgramApplicationContext.GetTheme().TextColor;
+			_elements.Add(new TextStatusElement(inText, textColor, inStyle));
 			_modified = true;
 		}
 
@@ -435,12 +438,11 @@ namespace UnrealGameSync
 		}
 	}
 
-	class StatusPanel : Panel
+	class StatusPanel : Panel, ICustomThemeHandler
 	{
 		const float LineSpacing = 1.35f;
 
-		Image? _projectLogo;
-		bool _disposeProjectLogo;
+		InvertibleImage _projectLogo = new InvertibleImage(Properties.Resources.DefaultProjectLogo, false);
 		Rectangle _projectLogoBounds;
 		StatusElementResources? _resources;
 		readonly List<StatusLine> _lines = new List<StatusLine>();
@@ -448,7 +450,7 @@ namespace UnrealGameSync
 		Pen? _alertDividerPen;
 		int _alertDividerY;
 		StatusLine? _alert;
-		Color? _tintColor;
+		InvertibleColor? _tintColor;
 		Point? _mouseOverLocation;
 		StatusElement? _mouseOverElement;
 		Point? _mouseDownLocation;
@@ -495,14 +497,9 @@ namespace UnrealGameSync
 					_alertDividerPen.Dispose();
 					_alertDividerPen = null;
 				}
-				if (_projectLogo != null)
-				{
-					if (_disposeProjectLogo)
-					{
-						_projectLogo.Dispose();
-					}
-					_projectLogo = null;
-				}
+				
+				_projectLogo.Dispose();
+				
 				ResetFontCache();
 			}
 		}
@@ -519,15 +516,9 @@ namespace UnrealGameSync
 
 		public void SetProjectLogo(Image newProjectLogo, bool dispose)
 		{
-			if (_projectLogo != null)
-			{
-				if (_disposeProjectLogo)
-				{
-					_projectLogo.Dispose();
-				}
-			}
-			_projectLogo = newProjectLogo;
-			_disposeProjectLogo = dispose;
+			_projectLogo.Dispose();
+			_projectLogo = new InvertibleImage(newProjectLogo, dispose);
+			
 			Invalidate();
 		}
 
@@ -541,10 +532,11 @@ namespace UnrealGameSync
 		public void Set(IEnumerable<StatusLine> newLines, StatusLine? newCaption, StatusLine? newAlert, Color? newTintColor)
 		{
 			_resources ??= new StatusElementResources(Font);
-
-			if (_tintColor != newTintColor)
+			
+			if (((_tintColor == null) != (newTintColor == null)) || (_tintColor != null && _tintColor.StandardColor != newTintColor))
 			{
 				Invalidate();
+				_tintColor = newTintColor.HasValue ? new InvertibleColor(newTintColor.Value) : null;
 			}
 
 			InvalidateElements();
@@ -552,7 +544,6 @@ namespace UnrealGameSync
 			_lines.AddRange(newLines);
 			_caption = newCaption;
 			_alert = newAlert;
-			_tintColor = newTintColor;
 			LayoutElements();
 			InvalidateElements();
 
@@ -636,7 +627,7 @@ namespace UnrealGameSync
 			}
 
 			// Get the logo size
-			Image drawProjectLogo = _projectLogo ?? Properties.Resources.DefaultProjectLogo;
+			Image drawProjectLogo = _projectLogo.ResolveImage() ?? Properties.Resources.DefaultProjectLogo;
 			float logoScale = Math.Min((float)bodyHeight - ((_caption != null) ? Font.Height : 0) / drawProjectLogo.Height, graphics.DpiY / 96.0f);
 			int logoWidth = (int)(drawProjectLogo.Width * logoScale);
 			int logoHeight = (int)(drawProjectLogo.Height * logoScale);
@@ -722,10 +713,10 @@ namespace UnrealGameSync
 		{
 			base.OnPaintBackground(e);
 
-			if (_tintColor.HasValue)
+			if (_tintColor != null)
 			{
 				int tintSize = Width / 2;
-				using (LinearGradientBrush backgroundBrush = new LinearGradientBrush(new Point(Width, 0), new Point(Width - tintSize, tintSize), _tintColor.Value, BackColor))
+				using (LinearGradientBrush backgroundBrush = new LinearGradientBrush(new Point(Width, 0), new Point(Width - tintSize, tintSize), _tintColor.ResolveColor(), BackColor))
 				{
 					backgroundBrush.WrapMode = WrapMode.TileFlipXY;
 					using (GraphicsPath path = new GraphicsPath())
@@ -747,7 +738,7 @@ namespace UnrealGameSync
 
 			if (_suspendDisplayCount == 0)
 			{
-				e.Graphics.DrawImage(_projectLogo ?? Properties.Resources.DefaultProjectLogo, _projectLogoBounds);
+				e.Graphics.DrawImage(_projectLogo.ResolveImage() ?? Properties.Resources.DefaultProjectLogo, _projectLogoBounds);
 
 				e.Graphics.SmoothingMode = System.Drawing.Drawing2D.SmoothingMode.None;
 
@@ -823,5 +814,13 @@ namespace UnrealGameSync
 				}
 			}
 		}
+
+		public void ApplyTheme(ApplicationThemeManager themeManager)
+		{
+			themeManager.DefaultApplyThemeToControl(this);
+
+			BackColor = themeManager.Theme.InputBackgroundColor;
+			ForeColor = themeManager.Theme.TextColor;
+		}
 	}
 }
diff --git a/Engine/Source/Programs/UnrealGameSync/UnrealGameSync/Controls/TabControl.cs b/Engine/Source/Programs/UnrealGameSync/UnrealGameSync/Controls/TabControl.cs
--- a/Engine/Source/Programs/UnrealGameSync/UnrealGameSync/Controls/TabControl.cs
+++ b/Engine/Source/Programs/UnrealGameSync/UnrealGameSync/Controls/TabControl.cs
@@ -9,7 +9,7 @@ using System.Windows.Forms;
 
 namespace UnrealGameSync
 {
-	class TabControl : Control
+	class TabControl : Control, ICustomThemeHandler
 	{
 		const int TabPadding = 13;//18;
 		const int TabCloseButtonWidth = 8 + 13;
@@ -52,7 +52,15 @@ namespace UnrealGameSync
 		int _selectedTabIdx;
 		int _hoverTabIdx;
 		int _highlightTabIdx;
+		bool _isHoveringCloseBtn = false;
 		TabDragData? _dragState;
+		Color _textColor = SystemColors.WindowText;
+		Color _selectedColor = SystemColors.HotTrack;
+		Color _borderColor = SystemColors.Window;
+		Color _hoverColor = SystemColors.Window;
+		Color _hoverStripeColor = Color.FromArgb(192, 192, 192);
+		Color _closeColor = SystemColors.HotTrack;
+		Color _closeColorHovered = Color.FromArgb(SystemColors.HotTrack.B, SystemColors.HotTrack.G, SystemColors.HotTrack.R);
 
 		public delegate void OnTabChangedDelegate(object? newTabData);
 		public event OnTabChangedDelegate? OnTabChanged;
@@ -460,6 +468,14 @@ namespace UnrealGameSync
 					}
 				}
 			}
+
+			bool isNowHoveringCloseBtn = _hoverTabIdx == _selectedTabIdx
+			                      && ((e.Location.X) > (_tabs[_hoverTabIdx]._minX + _tabs[_hoverTabIdx]._width - _tabs[_hoverTabIdx]._closeButtonWidth));
+			if (_isHoveringCloseBtn != isNowHoveringCloseBtn)
+			{
+				_isHoveringCloseBtn = isNowHoveringCloseBtn;
+				Invalidate();
+			}
 		}
 
 		protected override void OnMouseCaptureChanged(EventArgs e)
@@ -479,18 +495,20 @@ namespace UnrealGameSync
 
 			LayoutTabs();
 
-			using (SolidBrush hoverBrush = new SolidBrush(Color.FromArgb(192, 192, 192)))
+			using (SolidBrush backBrush = new SolidBrush(BackColor))
+			using (SolidBrush hoverBrush = new SolidBrush(_hoverColor))
+			using (SolidBrush hoverStripeBrush = new SolidBrush(_hoverStripeColor))
 			{
 				for (int idx = 0; idx < _tabs.Count; idx++)
 				{
 					TabData tab = _tabs[idx];
 					if (idx == _hoverTabIdx || idx == _highlightTabIdx)
 					{
-						DrawBackground(e.Graphics, tab, SystemBrushes.Window, hoverBrush, tab._highlight);
+						DrawBackground(e.Graphics, tab, hoverBrush, hoverStripeBrush, tab._highlight);
 					}
 					else
 					{
-						DrawBackground(e.Graphics, tab, SystemBrushes.Control, SystemBrushes.Control, tab._highlight);
+						DrawBackground(e.Graphics, tab, backBrush, backBrush, tab._highlight);
 					}
 				}
 			}
@@ -527,9 +545,12 @@ namespace UnrealGameSync
 			{
 				TabData selectedTab = _tabs[_selectedTabIdx];
 
-				using (SolidBrush selectedBrush = new SolidBrush(Color.FromArgb(0, 136, 204)))
+				Color tabBackColor = _hoverColor;
+
+				using (SolidBrush backBrush = new SolidBrush(tabBackColor))
+				using (SolidBrush selectedBrush = new SolidBrush(_selectedColor))
 				{
-					DrawBackground(e.Graphics, selectedTab, SystemBrushes.Window, selectedBrush, selectedTab._highlight);
+					DrawBackground(e.Graphics, selectedTab, backBrush, selectedBrush, selectedTab._highlight);
 				}
 
 				DrawText(e.Graphics, selectedTab);
@@ -544,10 +565,14 @@ namespace UnrealGameSync
 				int closeMidX = selectedTab._minX + selectedTab._width - (int)((TabPadding + TabCloseButtonWidth) * dpiScaleX / 2);// + (13 / 2);
 				int closeMidY = (ClientSize.Height - (int)(4 * dpiScaleY)) / 2;
 
+				Color closeColor = _isHoveringCloseBtn ? _closeColorHovered : _closeColor;
 				Rectangle closeButton = new Rectangle(closeMidX - (int)((13 / 2) * dpiScaleX), closeMidY - (int)((13 / 2) * dpiScaleY), (int)(13 * dpiScaleX), (int)(13 * dpiScaleY));
-				e.Graphics.FillEllipse(SystemBrushes.ControlDark, closeButton);
+				using (SolidBrush backBrush = new SolidBrush(closeColor))
+				{
+					e.Graphics.FillEllipse(backBrush, closeButton);
+				}
 
-				using (Pen crossPen = new Pen(SystemBrushes.Window, 2.0f * dpiScaleX))
+				using(Pen crossPen = new Pen(tabBackColor, 2.0f * dpiScaleX))
 				{
 					float indentX = 3.5f * dpiScaleX;
 					float indentY = 3.5f * dpiScaleY;
@@ -557,13 +582,19 @@ namespace UnrealGameSync
 				e.Graphics.SmoothingMode = prevSmoothingMode;
 
 				// Border
-				e.Graphics.DrawLine(SystemPens.ControlDark, selectedTab._minX, 0, selectedTab._minX + selectedTab._width, 0);
-				e.Graphics.DrawLine(SystemPens.ControlDark, selectedTab._minX, 0, selectedTab._minX, ClientSize.Height - 2);
-				e.Graphics.DrawLine(SystemPens.ControlDark, selectedTab._minX + selectedTab._width, 0, selectedTab._minX + selectedTab._width, ClientSize.Height - 2);
+				using (Pen borderPen = new Pen(_borderColor))
+				{
+					e.Graphics.DrawLine(borderPen, selectedTab._minX, 0, selectedTab._minX + selectedTab._width, 0);
+					e.Graphics.DrawLine(borderPen, selectedTab._minX, 0, selectedTab._minX, ClientSize.Height - 2);
+					e.Graphics.DrawLine(borderPen, selectedTab._minX + selectedTab._width, 0, selectedTab._minX + selectedTab._width, ClientSize.Height - 2);
+				}
 			}
 
-			e.Graphics.DrawLine(SystemPens.ControlDarkDark, 0, ClientSize.Height - 2, ClientSize.Width, ClientSize.Height - 2);
-			e.Graphics.DrawLine(SystemPens.ControlLightLight, 0, ClientSize.Height - 1, ClientSize.Width, ClientSize.Height - 1);
+			using (Pen borderPen = new Pen(_borderColor))
+			{
+				e.Graphics.DrawLine(SystemPens.ControlDarkDark, 0, ClientSize.Height - 2, ClientSize.Width, ClientSize.Height - 2);
+				e.Graphics.DrawLine(borderPen, 0, ClientSize.Height - 1, ClientSize.Width, ClientSize.Height - 1);
+			}
 		}
 
 		void DrawBackground(Graphics graphics, TabData tab, Brush backgroundBrush, Brush stripeBrush, Tuple<Color, float>? highlight)
@@ -597,7 +628,25 @@ namespace UnrealGameSync
 		{
 			float dpiScaleX = graphics.DpiX / 96.0f;
 			float dpiScaleY = graphics.DpiY / 96.0f;
-			TextRenderer.DrawText(graphics, tab._name, Font, new Rectangle(tab._minX + (int)(TabPadding * dpiScaleX), 0, tab._textSize.Width, ClientSize.Height - (int)(3 * dpiScaleY)), Color.Black, TextFormatFlags.NoPadding | TextFormatFlags.EndEllipsis | TextFormatFlags.VerticalCenter);
+			TextRenderer.DrawText(graphics, tab._name, Font, new Rectangle(tab._minX + (int)(TabPadding * dpiScaleX), 0, tab._textSize.Width, ClientSize.Height - (int)(3 * dpiScaleY)), _textColor, TextFormatFlags.NoPadding | TextFormatFlags.EndEllipsis | TextFormatFlags.VerticalCenter);
+		}
+
+		public void ApplyTheme(ApplicationThemeManager themeManager)
+		{
+			_textColor = themeManager.Theme.TextColor;
+			_hoverColor = themeManager.Theme.InputBackgroundColor;
+			_selectedColor = themeManager.Theme.LinkColor;
+			_borderColor = themeManager.Theme.InputBorderColor;
+			_closeColor = themeManager.Theme.LinkColor;
+			_closeColorHovered = themeManager.Theme.LinkColorHover;
+
+			float hoverWeight = 0.4f;
+			_hoverStripeColor = Color.FromArgb(
+				(int)Math.Round(_hoverColor.R * (1 - hoverWeight) + Color.Gray.R * hoverWeight),
+				(int)Math.Round(_hoverColor.G * (1 - hoverWeight) + Color.Gray.G * hoverWeight),
+				(int)Math.Round(_hoverColor.B * (1 - hoverWeight) + Color.Gray.B * hoverWeight));
+
+			themeManager.DefaultApplyThemeToControl(this);
 		}
 	}
 }
diff --git a/Engine/Source/Programs/UnrealGameSync/UnrealGameSync/Controls/WorkspaceControl.cs b/Engine/Source/Programs/UnrealGameSync/UnrealGameSync/Controls/WorkspaceControl.cs
--- a/Engine/Source/Programs/UnrealGameSync/UnrealGameSync/Controls/WorkspaceControl.cs
+++ b/Engine/Source/Programs/UnrealGameSync/UnrealGameSync/Controls/WorkspaceControl.cs
@@ -26,7 +26,7 @@ using UnrealGameSync.Forms;
 
 namespace UnrealGameSync
 {
-	partial class WorkspaceControl : UserControl, IMainWindowTabPanel
+	partial class WorkspaceControl : UserControl, IMainWindowTabPanel, ICustomThemeHandler
 	{
 		enum HorizontalAlignment
 		{
@@ -717,8 +717,6 @@ namespace UnrealGameSync
 		{
 			base.OnLoad(e);
 
-			_buildListWidth = BuildList.Width;
-
 			// Find the default widths 
 			_columnWidths = new float[BuildList.Columns.Count];
 			for (int idx = 0; idx < BuildList.Columns.Count; idx++)
@@ -2079,11 +2077,6 @@ namespace UnrealGameSync
 			}
 		}
 
-		private void BuildList_DrawColumnHeader(object sender, DrawListViewColumnHeaderEventArgs e)
-		{
-			e.DrawDefault = true;
-		}
-
 		class ExpandRowLayout
 		{
 			public string? _mainText;
@@ -2129,12 +2122,13 @@ namespace UnrealGameSync
 
 		private void DrawExpandRow(Graphics graphics, ExpandRowLayout layout)
 		{
-			TextRenderer.DrawText(graphics, layout._mainText, _buildFont, layout._mainRect, SystemColors.WindowText, TextFormatFlags.SingleLine | TextFormatFlags.VerticalCenter | TextFormatFlags.NoPrefix | TextFormatFlags.NoPadding);
+			Color textColor = ProgramApplicationContext.GetTheme().TextColor;
+			TextRenderer.DrawText(graphics, layout._mainText, _buildFont, layout._mainRect, textColor, TextFormatFlags.SingleLine | TextFormatFlags.VerticalCenter | TextFormatFlags.NoPrefix | TextFormatFlags.NoPadding);
 
-			Color linkColor = SystemColors.HotTrack;
+			Color linkColor = ProgramApplicationContext.GetTheme().LinkColor;
 			if (_mouseOverExpandLink)
 			{
-				linkColor = Color.FromArgb(linkColor.B, linkColor.G, linkColor.R);
+				linkColor = ProgramApplicationContext.GetTheme().LinkColorHover;
 			}
 
 			TextRenderer.DrawText(graphics, layout._linkText, _buildFont, layout._linkRect, linkColor, TextFormatFlags.SingleLine | TextFormatFlags.VerticalCenter | TextFormatFlags.NoPrefix | TextFormatFlags.NoPadding);
@@ -2158,9 +2152,6 @@ namespace UnrealGameSync
 			if (e.Item == _expandItem)
 			{
 				BuildList.DrawDefaultBackground(e.Graphics, e.Bounds);
-
-				ExpandRowLayout expandLayout = LayoutExpandRow(e.Graphics, e.Bounds);
-				DrawExpandRow(e.Graphics, expandLayout);
 			}
 			else if (_workspace != null)
 			{
@@ -2471,7 +2462,15 @@ namespace UnrealGameSync
 			{
 				return;
 			}
-
+			
+			if (e.Item == _expandItem && e.ColumnIndex == 0)
+			{
+				// Draw the expand item only once, centered in the full row
+				ExpandRowLayout expandLayout = LayoutExpandRow(e.Graphics, e.Item.Bounds);
+				DrawExpandRow(e.Graphics, expandLayout);
+				return;
+			}
+			
 			ChangesRecord? change = (ChangesRecord?)e.Item.Tag;
 			if (change == null)
 			{
@@ -2492,14 +2491,16 @@ namespace UnrealGameSync
 
 			bool allowSync = CanSyncChange(change.Number, false);
 			int badgeAlpha = allowSync ? 255 : 128;
-			Color textColor = (allowSync || change.Number == _workspace.PendingChangeNumber || change.Number == _workspace.CurrentChangeNumber || _workspace.State.AdditionalChangeNumbers.Contains(change.Number)) ? SystemColors.WindowText : Blend(SystemColors.Window, SystemColors.WindowText, 0.25f);
+			Color baseTextColor = ProgramApplicationContext.GetTheme().TextColor;
+			Color backgroundColor = ProgramApplicationContext.GetTheme().InputBackgroundColor;
+			Color textColor = (allowSync || change.Number == _workspace.PendingChangeNumber || change.Number == _workspace.CurrentChangeNumber || _workspace.State.AdditionalChangeNumbers.Contains(change.Number)) ? baseTextColor : Blend(backgroundColor, baseTextColor, 0.25f);
 
 			const int FadeRange = 6;
 			if (e.ItemIndex >= BuildList.Items.Count - FadeRange && _numChanges >= _perforceMonitor.CurrentMaxChanges && !IsBisectModeEnabled())
 			{
 				float opacity = (float)(BuildList.Items.Count - e.ItemIndex - 0.9f) / FadeRange;
 				badgeAlpha = (int)(badgeAlpha * opacity);
-				textColor = Blend(SystemColors.Window, textColor, opacity);
+				textColor = Blend(backgroundColor, textColor, opacity);
 			}
 
 			if (e.ColumnIndex == IconColumn.Index)
@@ -2679,7 +2680,7 @@ namespace UnrealGameSync
 					else if (e.ItemIndex == BuildList.HoverItem && allowSync)
 					{
 						Rectangle syncBadgeRectangle = GetSyncBadgeRectangle(e.SubItem.Bounds);
-						DrawBadge(e.Graphics, syncBadgeRectangle, "Sync", _hoverSync ? Color.FromArgb(140, 180, 230) : Color.FromArgb(112, 146, 190), true, true);
+						DrawBadge(e.Graphics, syncBadgeRectangle, "Sync", _hoverSync ? ProgramApplicationContext.GetTheme().SyncBadgeColor : ProgramApplicationContext.GetTheme().SyncHoverBadgeColor, true, true);
 						maxX = syncBadgeRectangle.Left - (int)(2 * dpiScaleX);
 					}
 
@@ -3056,11 +3057,13 @@ namespace UnrealGameSync
 			{
 				if (details.ContainsCode)
 				{
-					badges.Add(new BadgeInfo("Code", "ChangeType", Color.FromArgb(116, 185, 255)));
+					Color badgeColor = ProgramApplicationContext.GetTheme().CodeBadgeColor;
+					badges.Add(new BadgeInfo("Code", "ChangeType", badgeColor));
 				}
 				if (details.ContainsContent)
 				{
-					badges.Add(new BadgeInfo("Content", "ChangeType", Color.FromArgb(162, 155, 255)));
+					Color badgeColor = ProgramApplicationContext.GetTheme().ContentBadgeColor;
+					badges.Add(new BadgeInfo("Content", "ChangeType", badgeColor));
 				}
 			}
 			if (badges.Count == 0)
@@ -3164,6 +3167,7 @@ namespace UnrealGameSync
 					}
 
 					BadgeContextMenu.Show(control, point);
+					ProgramApplicationContext.GetThemeManager()?.ApplyThemeRecursively(BadgeContextMenu);
 				};
 			}
 
@@ -3269,7 +3273,8 @@ namespace UnrealGameSync
 					}
 				}
 
-				TextRenderer.DrawText(graphics, badgeText, _badgeFont, badgeRect, Color.White, TextFormatFlags.HorizontalCenter | TextFormatFlags.VerticalCenter | TextFormatFlags.SingleLine | TextFormatFlags.NoPrefix | TextFormatFlags.PreserveGraphicsClipping);
+				Color textColor = ApplicationTheme.IsColorLight(badgeColor) ? Color.Black : Color.White;
+				TextRenderer.DrawText(graphics, badgeText, _badgeFont, badgeRect, textColor, TextFormatFlags.HorizontalCenter | TextFormatFlags.VerticalCenter | TextFormatFlags.SingleLine | TextFormatFlags.NoPrefix | TextFormatFlags.PreserveGraphicsClipping);
 			}
 		}
 
@@ -4186,6 +4191,7 @@ namespace UnrealGameSync
 			}
 
 			BuildHealthContextMenu.Show(StatusPanel, new Point(bounds.Right, bounds.Bottom), ToolStripDropDownDirection.BelowLeft);
+			ProgramApplicationContext.GetThemeManager()?.ApplyThemeRecursively(BuildHealthContextMenu);
 		}
 
 		private void BuildHealthContextMenu_Issue_Click(IssueData issue)
@@ -4485,12 +4491,14 @@ namespace UnrealGameSync
 		private void ShowToolContextMenu(Rectangle bounds, ContextMenuStrip menuStrip)
 		{
 			menuStrip.Show(StatusPanel, new Point(bounds.Left, bounds.Bottom), ToolStripDropDownDirection.BelowRight);
+			ProgramApplicationContext.GetThemeManager()?.ApplyThemeRecursively(menuStrip);
 		}
 
 		private void ShowActionsMenu(Rectangle bounds)
 		{
 			MoreToolsContextMenu_CleanWorkspace.Image = s_hackToolStripMenuGutter;
 			MoreToolsContextMenu.Show(StatusPanel, new Point(bounds.Left, bounds.Bottom), ToolStripDropDownDirection.BelowRight);
+			ProgramApplicationContext.GetThemeManager()?.ApplyThemeRecursively(MoreToolsContextMenu);
 		}
 
 		private void SelectChange(int changeNumber)
@@ -4664,6 +4672,7 @@ namespace UnrealGameSync
 						BuildListMultiContextMenu_ShowServerTimes.Checked = !_settings.ShowLocalTimes;
 
 						BuildListMultiContextMenu.Show(BuildList, args.Location);
+						ProgramApplicationContext.GetThemeManager()?.ApplyThemeRecursively(BuildListMultiContextMenu);
 					}
 					else
 					{
@@ -4798,6 +4807,7 @@ namespace UnrealGameSync
 				BuildListContextMenu_ViewInSwarm.Visible = projectConfigFile != null && TryGetProjectSetting(projectConfigFile, "SwarmURL", out swarmUrl);
 
 				BuildListContextMenu.Show(parentControl, location, direction);
+				ProgramApplicationContext.GetThemeManager()?.ApplyThemeRecursively(BuildListContextMenu);
 			}
 		}
 
@@ -4842,6 +4852,7 @@ namespace UnrealGameSync
 							if (badge?._toolTip != null && badge._backgroundColor.A != 0)
 							{
 								BuildListToolTip.Show(badge._toolTip, BuildList, new Point(buildListLocation.X + badge._offset, hitTest.Item.Bounds.Bottom + 2));
+								ProgramApplicationContext.GetThemeManager()?.ApplyThemeRecursively(BuildListToolTip);
 							}
 							else
 							{
@@ -5020,6 +5031,7 @@ namespace UnrealGameSync
 			OptionsContextMenu_TabNames_WorkspaceRoot.Checked = _settings.TabLabels == TabLabels.WorkspaceRoot;
 			OptionsContextMenu_TabNames_ProjectFile.Checked = _settings.TabLabels == TabLabels.ProjectFile;
 			OptionsContextMenu.Show(OptionsButton, new Point(OptionsButton.Width - OptionsContextMenu.Size.Width, OptionsButton.Height));
+			ProgramApplicationContext.GetThemeManager()?.ApplyThemeRecursively(OptionsContextMenu);
 		}
 
 		private IReadOnlyList<IArchiveChannel> GetArchiveChannels()
@@ -5056,6 +5068,13 @@ namespace UnrealGameSync
 			UpdateSyncActionCheckboxes();
 		}
 
+		private void BuildListToolTip_Draw(object sender, DrawToolTipEventArgs e)
+		{
+			e.DrawBackground();
+			e.DrawBorder();
+			e.DrawText(TextFormatFlags.LeftAndRightPadding | TextFormatFlags.VerticalCenter);
+		}
+
 		private void OpenSolutionAfterSyncCheckBox_CheckedChanged(object sender, EventArgs e)
 		{
 			_settings.OpenSolutionAfterSync = OpenSolutionAfterSyncCheckBox.Checked;
@@ -5071,21 +5090,23 @@ namespace UnrealGameSync
 			RunAfterSyncCheckBox.Checked = _settings.RunAfterSync;
 			OpenSolutionAfterSyncCheckBox.Checked = _settings.OpenSolutionAfterSync;
 
+			ApplicationTheme theme = ProgramApplicationContext.GetTheme();
+
 			if (_workspace == null || _workspace.IsBusy())
 			{
-				GenerateAfterSyncCheckBox.Enabled = false;
-				AfterSyncingLabel.Enabled = false;
-				BuildAfterSyncCheckBox.Enabled = false;
-				RunAfterSyncCheckBox.Enabled = false;
-				OpenSolutionAfterSyncCheckBox.Enabled = false;
+				AfterSyncingLabel.ForeColor = theme.DimmedText;
+				GenerateAfterSyncCheckBox.ReadOnly = true;
+				BuildAfterSyncCheckBox.ReadOnly = true;
+				RunAfterSyncCheckBox.ReadOnly = true;
+				OpenSolutionAfterSyncCheckBox.ReadOnly = true;
 			}
 			else
 			{
-				AfterSyncingLabel.Enabled = true;
-				GenerateAfterSyncCheckBox.Enabled = _hasBuildSteps;
-				BuildAfterSyncCheckBox.Enabled = _hasBuildSteps;
-				RunAfterSyncCheckBox.Enabled = BuildAfterSyncCheckBox.Checked || ShouldSyncPrecompiledEditor;
-				OpenSolutionAfterSyncCheckBox.Enabled = !ShouldSyncPrecompiledEditor;
+				AfterSyncingLabel.ForeColor = theme.TextColor;
+				GenerateAfterSyncCheckBox.ReadOnly = !_hasBuildSteps;
+				BuildAfterSyncCheckBox.ReadOnly = !_hasBuildSteps;
+				RunAfterSyncCheckBox.ReadOnly = !BuildAfterSyncCheckBox.Checked && !ShouldSyncPrecompiledEditor;
+				OpenSolutionAfterSyncCheckBox.ReadOnly = ShouldSyncPrecompiledEditor;
 			}
 		}
 
@@ -5449,6 +5470,7 @@ namespace UnrealGameSync
 							if (layoutInfo._descriptionBadges.Count == 0 || clientPoint.X < GetBadgeListLocation(layoutInfo._descriptionBadges, descriptionBounds, HorizontalAlign.Right, VerticalAlignment.Middle).X - 2)
 							{
 								BuildListToolTip.Show(change.Description, BuildList, new Point(descriptionBounds.Left, descriptionBounds.Bottom + 2));
+								ProgramApplicationContext.GetThemeManager()?.ApplyThemeRecursively(BuildListToolTip);
 								return;
 							}
 						}
@@ -5484,6 +5506,7 @@ namespace UnrealGameSync
 						{
 							Rectangle summaryBounds = args.Item.SubItems[StatusColumn.Index].Bounds;
 							BuildListToolTip.Show(summaryText.ToString(), BuildList, new Point(summaryBounds.Left, summaryBounds.Bottom));
+							ProgramApplicationContext.GetThemeManager()?.ApplyThemeRecursively(BuildListToolTip);
 							return;
 						}
 					}
@@ -6449,6 +6472,7 @@ namespace UnrealGameSync
 			}
 
 			SyncContextMenu.Show(StatusPanel, new Point(bounds.Left, bounds.Bottom), ToolStripDropDownDirection.BelowRight);
+			ProgramApplicationContext.GetThemeManager()?.ApplyThemeRecursively(SyncContextMenu);
 		}
 
 		private void SyncContextMenu_LatestChangeType_Click(string syncTypeId)
@@ -6527,6 +6551,7 @@ namespace UnrealGameSync
 
 			RecentMenu_Separator.Visible = (_settings.RecentProjects.Count > 0);
 			RecentMenu.Show(StatusPanel, new Point(bounds.Left, bounds.Bottom), ToolStripDropDownDirection.BelowRight);
+			ProgramApplicationContext.GetThemeManager()?.ApplyThemeRecursively(RecentMenu);
 		}
 
 		private void RecentMenu_Browse_Click(object sender, EventArgs e)
@@ -6912,6 +6937,7 @@ namespace UnrealGameSync
 			}
 
 			FilterContextMenu.Show(FilterButton, new Point(0, FilterButton.Height));
+			ProgramApplicationContext.GetThemeManager()?.ApplyThemeRecursively(FilterContextMenu);
 		}
 
 		private void FilterContextMenu_Badge_Click(string badgeName)
@@ -7194,6 +7220,16 @@ namespace UnrealGameSync
 				_owner.UpdateTintColors();
 			}
 		}
+
+		public void ApplyTheme(ApplicationThemeManager themeManager)
+		{
+			themeManager.DefaultApplyThemeToControl(this);
+
+			BuildListToolTip.BackColor = themeManager.Theme.InputBackgroundColor;
+			BuildListToolTip.ForeColor = themeManager.Theme.TextColor;
+			
+			UpdateSyncActionCheckboxes();
+		}
 	}
 
 	public class UncontrolledChangelist
diff --git a/Engine/Source/Programs/UnrealGameSync/UnrealGameSync/Forms/ApplicationSettingsWindow.cs b/Engine/Source/Programs/UnrealGameSync/UnrealGameSync/Forms/ApplicationSettingsWindow.cs
--- a/Engine/Source/Programs/UnrealGameSync/UnrealGameSync/Forms/ApplicationSettingsWindow.cs
+++ b/Engine/Source/Programs/UnrealGameSync/UnrealGameSync/Forms/ApplicationSettingsWindow.cs
@@ -14,7 +14,7 @@ using Microsoft.Win32;
 
 namespace UnrealGameSync
 {
-	partial class ApplicationSettingsWindow : Form
+	partial class ApplicationSettingsWindow : ThemedForm
 	{
 		public enum Result
 		{
@@ -111,6 +111,16 @@ namespace UnrealGameSync
 			AutomaticallyRunAtStartupCheckBox.Checked = IsAutomaticallyRunAtStartup();
 			KeepInTrayCheckBox.Checked = settings.KeepInTray;
 
+			foreach (KeyValuePair<string, string> themeIdNamePair in settings.ThemeList)
+			{
+				ThemeComboBox.Items.Add(themeIdNamePair);
+
+				if (_settings.Theme?.ThemeId == themeIdNamePair.Key)
+				{
+					ThemeComboBox.SelectedItem = themeIdNamePair;
+				}
+			}
+
 			HordeServerTextBox.Text = launcherSettings.HordeServer;
 			HordeServerTextBox.Select(HordeServerTextBox.TextLength, 0);
 			HordeServerTextBox.CueBanner = launcherSettings.HordeServer ?? String.Empty;
@@ -273,10 +283,14 @@ namespace UnrealGameSync
 				automationPortNumber = -1;
 			}
 
+			string themeId = (ThemeComboBox.SelectedItem != null) ? ((KeyValuePair<string, string>)ThemeComboBox.SelectedItem).Key : "";
+			bool didThemeChange = (_settings.Theme == null || _settings.Theme.ThemeId != themeId);
+			
 			if (!String.Equals(launcherSettings.HordeServer, originalLauncherSettings.HordeServer, StringComparison.OrdinalIgnoreCase) ||
 				!String.Equals(launcherSettings.PerforceServerAndPort, originalLauncherSettings.PerforceServerAndPort, StringComparison.OrdinalIgnoreCase) ||
 				!String.Equals(launcherSettings.PerforceUserName, originalLauncherSettings.PerforceUserName, StringComparison.OrdinalIgnoreCase) ||
-				automationPortNumber != _initialAutomationPortNumber)
+				automationPortNumber != _initialAutomationPortNumber ||
+				didThemeChange)
 			{
 				if (result == Result.Ok)
 				{
@@ -312,6 +326,12 @@ namespace UnrealGameSync
 				_settings.Save(_logger);
 			}
 
+			if (didThemeChange)
+			{
+				_settings.Theme = _settings.LoadTheme(themeId);
+				_settings.Save(_logger);
+			}
+
 			if (_toolUpdateStatus.Item1)
 			{
 				List<Guid> newEnabledTools = new List<Guid>();
diff --git a/Engine/Source/Programs/UnrealGameSync/UnrealGameSync/Forms/ArgumentsWindow.cs b/Engine/Source/Programs/UnrealGameSync/UnrealGameSync/Forms/ArgumentsWindow.cs
--- a/Engine/Source/Programs/UnrealGameSync/UnrealGameSync/Forms/ArgumentsWindow.cs
+++ b/Engine/Source/Programs/UnrealGameSync/UnrealGameSync/Forms/ArgumentsWindow.cs
@@ -9,7 +9,7 @@ using System.Windows.Forms.VisualStyles;
 
 namespace UnrealGameSync
 {
-	public partial class ArgumentsWindow : Form
+	public partial class ArgumentsWindow : ThemedForm
 	{
 		const int LvmFirst = 0x1000;
 		const int LvmEditlabelw = (LvmFirst + 118);
@@ -170,11 +170,11 @@ namespace UnrealGameSync
 				CheckBoxState state = e.Item.Checked ? CheckBoxState.CheckedNormal : CheckBoxState.UncheckedNormal;
 				Size checkSize = CheckBoxRenderer.GetGlyphSize(e.Graphics, state);
 				CheckBoxRenderer.DrawCheckBox(e.Graphics, new Point(e.Bounds.Left + 4, e.Bounds.Top + (e.Bounds.Height - checkSize.Height) / 2), state);
-				DrawItemLabel(e.Graphics, SystemColors.WindowText, e.Item);
+				DrawItemLabel(e.Graphics, ForeColor, e.Item);
 			}
 			else
 			{
-				DrawItemLabel(e.Graphics, SystemColors.GrayText, e.Item);
+				DrawItemLabel(e.Graphics, Theme.DimmedText, e.Item);
 			}
 		}
 
diff --git a/Engine/Source/Programs/UnrealGameSync/UnrealGameSync/Forms/AutomatedBuildWindow.cs b/Engine/Source/Programs/UnrealGameSync/UnrealGameSync/Forms/AutomatedBuildWindow.cs
--- a/Engine/Source/Programs/UnrealGameSync/UnrealGameSync/Forms/AutomatedBuildWindow.cs
+++ b/Engine/Source/Programs/UnrealGameSync/UnrealGameSync/Forms/AutomatedBuildWindow.cs
@@ -8,7 +8,7 @@ using EpicGames.Perforce;
 
 namespace UnrealGameSync.Forms
 {
-	partial class AutomatedBuildWindow : Form
+	partial class AutomatedBuildWindow : ThemedForm
 	{
 		public class BuildInfo
 		{
diff --git a/Engine/Source/Programs/UnrealGameSync/UnrealGameSync/Forms/AutomatedSyncWindow.cs b/Engine/Source/Programs/UnrealGameSync/UnrealGameSync/Forms/AutomatedSyncWindow.cs
--- a/Engine/Source/Programs/UnrealGameSync/UnrealGameSync/Forms/AutomatedSyncWindow.cs
+++ b/Engine/Source/Programs/UnrealGameSync/UnrealGameSync/Forms/AutomatedSyncWindow.cs
@@ -17,7 +17,7 @@ using Microsoft.Extensions.Logging;
 
 namespace UnrealGameSync
 {
-	partial class AutomatedSyncWindow : Form
+	partial class AutomatedSyncWindow : ThemedForm
 	{
 		static class FindDefaultWorkspaceTask
 		{
diff --git a/Engine/Source/Programs/UnrealGameSync/UnrealGameSync/Forms/BuildStepWindow.cs b/Engine/Source/Programs/UnrealGameSync/UnrealGameSync/Forms/BuildStepWindow.cs
--- a/Engine/Source/Programs/UnrealGameSync/UnrealGameSync/Forms/BuildStepWindow.cs
+++ b/Engine/Source/Programs/UnrealGameSync/UnrealGameSync/Forms/BuildStepWindow.cs
@@ -10,7 +10,7 @@ using EpicGames.Core;
 
 namespace UnrealGameSync
 {
-	partial class BuildStepWindow : Form
+	partial class BuildStepWindow : ThemedForm
 	{
 		readonly BuildStep _step;
 		readonly List<string> _targetNames;
@@ -97,7 +97,7 @@ namespace UnrealGameSync
 			OtherWorkingDirTextBox.Enabled = isOther;
 			OtherWorkingDirButton.Enabled = isOther;
 			OtherArgumentsTextBox.Enabled = isOther;
-			OtherUseLogWindowCheckBox.Enabled = isOther;
+			OtherUseLogWindowCheckBox.ReadOnly = !isOther;
 		}
 
 		private void CompileRadioButton_CheckedChanged(object sender, EventArgs e)
diff --git a/Engine/Source/Programs/UnrealGameSync/UnrealGameSync/Forms/ChangelistWindow.cs b/Engine/Source/Programs/UnrealGameSync/UnrealGameSync/Forms/ChangelistWindow.cs
--- a/Engine/Source/Programs/UnrealGameSync/UnrealGameSync/Forms/ChangelistWindow.cs
+++ b/Engine/Source/Programs/UnrealGameSync/UnrealGameSync/Forms/ChangelistWindow.cs
@@ -5,7 +5,7 @@ using System.Windows.Forms;
 
 namespace UnrealGameSync
 {
-	public partial class ChangelistWindow : Form
+	public partial class ChangelistWindow : ThemedForm
 	{
 		public int ChangeNumber { get; private set; }
 
diff --git a/Engine/Source/Programs/UnrealGameSync/UnrealGameSync/Forms/CleanWorkspaceWindow.cs b/Engine/Source/Programs/UnrealGameSync/UnrealGameSync/Forms/CleanWorkspaceWindow.cs
--- a/Engine/Source/Programs/UnrealGameSync/UnrealGameSync/Forms/CleanWorkspaceWindow.cs
+++ b/Engine/Source/Programs/UnrealGameSync/UnrealGameSync/Forms/CleanWorkspaceWindow.cs
@@ -16,7 +16,7 @@ using Microsoft.Extensions.Logging;
 
 namespace UnrealGameSync
 {
-	partial class CleanWorkspaceWindow : Form
+	partial class CleanWorkspaceWindow : ThemedForm
 	{
 		enum TreeNodeAction
 		{
diff --git a/Engine/Source/Programs/UnrealGameSync/UnrealGameSync/Forms/ClobberWindow.cs b/Engine/Source/Programs/UnrealGameSync/UnrealGameSync/Forms/ClobberWindow.cs
--- a/Engine/Source/Programs/UnrealGameSync/UnrealGameSync/Forms/ClobberWindow.cs
+++ b/Engine/Source/Programs/UnrealGameSync/UnrealGameSync/Forms/ClobberWindow.cs
@@ -8,7 +8,7 @@ using System.Windows.Forms;
 
 namespace UnrealGameSync
 {
-	public partial class ClobberWindow : Form
+	public partial class ClobberWindow : ThemedForm
 	{
 		readonly Dictionary<string, bool> _filesToClobber;
 
diff --git a/Engine/Source/Programs/UnrealGameSync/UnrealGameSync/Forms/CombinedViewsWindow.cs b/Engine/Source/Programs/UnrealGameSync/UnrealGameSync/Forms/CombinedViewsWindow.cs
--- a/Engine/Source/Programs/UnrealGameSync/UnrealGameSync/Forms/CombinedViewsWindow.cs
+++ b/Engine/Source/Programs/UnrealGameSync/UnrealGameSync/Forms/CombinedViewsWindow.cs
@@ -8,7 +8,7 @@ using System.Windows.Forms;
 
 namespace UnrealGameSync.Forms
 {
-	public partial class CombinedViewsWindow : Form
+	public partial class CombinedViewsWindow : ThemedForm
 	{
 		private readonly IList<string> _views;
 		public CombinedViewsWindow(IList<string> views)
diff --git a/Engine/Source/Programs/UnrealGameSync/UnrealGameSync/Forms/ConnectWindow.cs b/Engine/Source/Programs/UnrealGameSync/UnrealGameSync/Forms/ConnectWindow.cs
--- a/Engine/Source/Programs/UnrealGameSync/UnrealGameSync/Forms/ConnectWindow.cs
+++ b/Engine/Source/Programs/UnrealGameSync/UnrealGameSync/Forms/ConnectWindow.cs
@@ -1,12 +1,13 @@
 ﻿// Copyright Epic Games, Inc. All Rights Reserved.
 
 using System;
+using System.Drawing;
 using System.Windows.Forms;
 using EpicGames.Perforce;
 
 namespace UnrealGameSync
 {
-	partial class ConnectWindow : Form
+	partial class ConnectWindow : ThemedForm
 	{
 		readonly IPerforceSettings _defaultPerforceSettings;
 		string? _serverAndPortOverride;
@@ -97,6 +98,13 @@ namespace UnrealGameSync
 
 			UserNameLabel.Enabled = !useDefaultSettings;
 			UserNameTextBox.Enabled = !useDefaultSettings;
+			
+			Color activeTextColor = Theme.TextColor;
+			Color inactiveTextColor = Theme.DimmedText;
+			ServerAndPortLabel.ForeColor = ServerAndPortLabel.Enabled ? activeTextColor : inactiveTextColor;
+			ServerAndPortTextBox.ForeColor = ServerAndPortTextBox.Enabled ? activeTextColor : inactiveTextColor;
+			UserNameLabel.ForeColor = UserNameLabel.Enabled ? activeTextColor : inactiveTextColor;
+			UserNameTextBox.ForeColor = UserNameTextBox.Enabled ? activeTextColor : inactiveTextColor;
 		}
 
 		private void BrowseUserBtn_Click(object sender, EventArgs e)
diff --git a/Engine/Source/Programs/UnrealGameSync/UnrealGameSync/Forms/DeleteWindow.cs b/Engine/Source/Programs/UnrealGameSync/UnrealGameSync/Forms/DeleteWindow.cs
--- a/Engine/Source/Programs/UnrealGameSync/UnrealGameSync/Forms/DeleteWindow.cs
+++ b/Engine/Source/Programs/UnrealGameSync/UnrealGameSync/Forms/DeleteWindow.cs
@@ -6,7 +6,7 @@ using System.Windows.Forms;
 
 namespace UnrealGameSync
 {
-	public partial class DeleteWindow : Form
+	public partial class DeleteWindow : ThemedForm
 	{
 		readonly Dictionary<string, bool> _filesToDelete;
 
diff --git a/Engine/Source/Programs/UnrealGameSync/UnrealGameSync/Forms/DiagnosticsWindow.cs b/Engine/Source/Programs/UnrealGameSync/UnrealGameSync/Forms/DiagnosticsWindow.cs
--- a/Engine/Source/Programs/UnrealGameSync/UnrealGameSync/Forms/DiagnosticsWindow.cs
+++ b/Engine/Source/Programs/UnrealGameSync/UnrealGameSync/Forms/DiagnosticsWindow.cs
@@ -11,7 +11,7 @@ using EpicGames.Core;
 
 namespace UnrealGameSync
 {
-	public partial class DiagnosticsWindow : Form
+	public partial class DiagnosticsWindow : ThemedForm
 	{
 		readonly DirectoryReference _appDataFolder;
 		readonly DirectoryReference _workspaceDataFolder;
diff --git a/Engine/Source/Programs/UnrealGameSync/UnrealGameSync/Forms/DownloadProgressWindow.cs b/Engine/Source/Programs/UnrealGameSync/UnrealGameSync/Forms/DownloadProgressWindow.cs
--- a/Engine/Source/Programs/UnrealGameSync/UnrealGameSync/Forms/DownloadProgressWindow.cs
+++ b/Engine/Source/Programs/UnrealGameSync/UnrealGameSync/Forms/DownloadProgressWindow.cs
@@ -8,7 +8,7 @@ using System.Windows.Forms;
 
 namespace UnrealGameSync.Forms
 {
-	public partial class DownloadProgressWindow : Form
+	public partial class DownloadProgressWindow : ThemedForm
 	{
 		class ProgressReporter : IProgress<string>
 		{
diff --git a/Engine/Source/Programs/UnrealGameSync/UnrealGameSync/Forms/DownloadSettingsWindow.cs b/Engine/Source/Programs/UnrealGameSync/UnrealGameSync/Forms/DownloadSettingsWindow.cs
--- a/Engine/Source/Programs/UnrealGameSync/UnrealGameSync/Forms/DownloadSettingsWindow.cs
+++ b/Engine/Source/Programs/UnrealGameSync/UnrealGameSync/Forms/DownloadSettingsWindow.cs
@@ -7,7 +7,7 @@ using EpicGames.Core;
 
 namespace UnrealGameSync.Forms
 {
-	public partial class DownloadSettingsWindow : Form
+	public partial class DownloadSettingsWindow : ThemedForm
 	{
 		public DownloadSettingsWindow()
 		{
diff --git a/Engine/Source/Programs/UnrealGameSync/UnrealGameSync/Forms/FailedToDeleteWindow.cs b/Engine/Source/Programs/UnrealGameSync/UnrealGameSync/Forms/FailedToDeleteWindow.cs
--- a/Engine/Source/Programs/UnrealGameSync/UnrealGameSync/Forms/FailedToDeleteWindow.cs
+++ b/Engine/Source/Programs/UnrealGameSync/UnrealGameSync/Forms/FailedToDeleteWindow.cs
@@ -4,7 +4,7 @@ using System.Windows.Forms;
 
 namespace UnrealGameSync
 {
-	public partial class FailedToDeleteWindow : Form
+	public partial class FailedToDeleteWindow : ThemedForm
 	{
 		public FailedToDeleteWindow()
 		{
diff --git a/Engine/Source/Programs/UnrealGameSync/UnrealGameSync/Forms/IssueAlertWindow.cs b/Engine/Source/Programs/UnrealGameSync/UnrealGameSync/Forms/IssueAlertWindow.cs
--- a/Engine/Source/Programs/UnrealGameSync/UnrealGameSync/Forms/IssueAlertWindow.cs
+++ b/Engine/Source/Programs/UnrealGameSync/UnrealGameSync/Forms/IssueAlertWindow.cs
@@ -9,7 +9,7 @@ using System.Windows.Forms;
 
 namespace UnrealGameSync
 {
-	partial class IssueAlertWindow : Form
+	partial class IssueAlertWindow : ThemedForm
 	{
 		static readonly IntPtr HwndTopmost = new IntPtr(-1);
 
diff --git a/Engine/Source/Programs/UnrealGameSync/UnrealGameSync/Forms/IssueBrowserWindow.cs b/Engine/Source/Programs/UnrealGameSync/UnrealGameSync/Forms/IssueBrowserWindow.cs
--- a/Engine/Source/Programs/UnrealGameSync/UnrealGameSync/Forms/IssueBrowserWindow.cs
+++ b/Engine/Source/Programs/UnrealGameSync/UnrealGameSync/Forms/IssueBrowserWindow.cs
@@ -12,7 +12,7 @@ using EpicGames.Perforce;
 
 namespace UnrealGameSync
 {
-	partial class IssueBrowserWindow : Form
+	partial class IssueBrowserWindow : ThemedForm
 	{
 		readonly IssueMonitor _issueMonitor;
 		readonly IPerforceSettings _perforceSettings;
@@ -353,6 +353,7 @@ namespace UnrealGameSync
 			}
 
 			FilterMenu.Show(FilterBtn, new Point(FilterBtn.Left, FilterBtn.Bottom));
+			ThemeManager?.ApplyThemeRecursively(FilterMenu);
 		}
 
 		private void FilterMenu_ShowAll_Click(object sender, EventArgs e)
@@ -366,11 +367,6 @@ namespace UnrealGameSync
 			IssueDetailsWindow.Show(Owner!, _issueMonitor, _perforceSettings, _serverTimeOffset, issue, _serviceProvider, _currentStream);
 		}
 
-		private void IssueListView_DrawColumnHeader(object sender, DrawListViewColumnHeaderEventArgs e)
-		{
-			e.DrawDefault = true;
-		}
-
 		private void IssueListView_DrawItem(object sender, DrawListViewItemEventArgs e)
 		{
 			if (e.Item.Selected)
@@ -386,15 +382,15 @@ namespace UnrealGameSync
 				Color backgroundColor;
 				if (issue.ResolvedAt.HasValue)
 				{
-					backgroundColor = SystemColors.Window;//Color.FromArgb(248, 254, 246);
+					backgroundColor = Theme.InputBackgroundColor;
 				}
 				else if (issue.FixChange > 0)
 				{
-					backgroundColor = Color.FromArgb(245, 245, 245);
+					backgroundColor = Theme.BackgroundColor;
 				}
 				else
 				{
-					backgroundColor = Color.FromArgb(254, 248, 246);
+					backgroundColor = Theme.ErrorNotificationBGTint;
 				}
 
 				using (SolidBrush brush = new SolidBrush(backgroundColor))
diff --git a/Engine/Source/Programs/UnrealGameSync/UnrealGameSync/Forms/IssueDetailsWindow.cs b/Engine/Source/Programs/UnrealGameSync/UnrealGameSync/Forms/IssueDetailsWindow.cs
--- a/Engine/Source/Programs/UnrealGameSync/UnrealGameSync/Forms/IssueDetailsWindow.cs
+++ b/Engine/Source/Programs/UnrealGameSync/UnrealGameSync/Forms/IssueDetailsWindow.cs
@@ -19,7 +19,7 @@ using Microsoft.Extensions.Logging;
 
 namespace UnrealGameSync
 {
-	partial class IssueDetailsWindow : Form
+	partial class IssueDetailsWindow : ThemedForm
 	{
 		class BuildGroup
 		{
@@ -245,7 +245,7 @@ namespace UnrealGameSync
 
 			public override void Draw(Graphics graphics, StatusElementResources resources)
 			{
-				Color textColor = Color.Gray;
+				Color textColor = Theme.TextColor;
 				FontStyle style = FontStyle.Italic;
 				if (MouseDown)
 				{
@@ -699,6 +699,8 @@ namespace UnrealGameSync
 			// Get the filter type
 			bool onlyShowCodeChanges = (FilterTypeComboBox.SelectedIndex == 1);
 			bool onlyShowContentChanges = (FilterTypeComboBox.SelectedIndex == 2);
+			
+			Color textColor = Theme.TextColor;
 
 			// Create rows for all the ranges
 			foreach (PerforceChangeRange range in _selectedStreamRanges)
@@ -716,11 +718,11 @@ namespace UnrealGameSync
 						fetchingWidget.HorizontalAlignment = HorizontalAlignment.Left;
 						if (range.ErrorMessage != null)
 						{
-							fetchingWidget.Line.AddText(range.ErrorMessage, Color.Gray, FontStyle.Italic);
+							fetchingWidget.Line.AddText(range.ErrorMessage, textColor, FontStyle.Italic);
 						}
 						else
 						{
-							fetchingWidget.Line.AddText("Fetching changes, please wait...", Color.Gray, FontStyle.Italic);
+							fetchingWidget.Line.AddText("Fetching changes, please wait...", textColor, FontStyle.Italic);
 						}
 
 						fetchingItem.SubItems.Add(new ListViewItem.ListViewSubItem(fetchingItem, "") { Tag = fetchingWidget });
@@ -890,11 +892,11 @@ namespace UnrealGameSync
 			{
 				if (details.ContainsCode)
 				{
-					typeWidget.Line.AddBadge("Code", Color.FromArgb(116, 185, 255), null);
+					typeWidget.Line.AddBadge("Code", Theme.CodeBadgeColor, null);
 				}
 				if (details.ContainsContent)
 				{
-					typeWidget.Line.AddBadge("Content", Color.FromArgb(162, 155, 255), null);
+					typeWidget.Line.AddBadge("Content", Theme.ContentBadgeColor, null);
 				}
 			}
 		}
@@ -1111,19 +1113,19 @@ namespace UnrealGameSync
 				Color backgroundColor;
 				if (buildGroup.Outcome == IssueBuildOutcome.Error)
 				{
-					backgroundColor = Color.FromArgb(254, 248, 246);
+					backgroundColor = Theme.ErrorNotificationBGTint;
 				}
 				else if (buildGroup.Outcome == IssueBuildOutcome.Warning)
 				{
-					backgroundColor = Color.FromArgb(254, 254, 246);
+					backgroundColor = Theme.WarningNotificationBGTint;
 				}
 				else if (buildGroup.Outcome == IssueBuildOutcome.Success)
 				{
-					backgroundColor = Color.FromArgb(248, 254, 246);
+					backgroundColor = Theme.SuccessNotificationBGTint;
 				}
 				else
 				{
-					backgroundColor = Color.FromArgb(245, 245, 245);
+					backgroundColor = Theme.StandardNotificationBGTint;
 				}
 
 				using (SolidBrush brush = new SolidBrush(backgroundColor))
@@ -1146,7 +1148,7 @@ namespace UnrealGameSync
 
 			Font currentFont = Font;//BuildFont;(Change.Number == Workspace.PendingChangeNumber || Change.Number == Workspace.CurrentChangeNumber)? SelectedBuildFont : BuildFont;
 
-			Color textColor = SystemColors.WindowText;
+			Color textColor = Theme.TextColor;
 
 			if (e.SubItem.Tag is CustomListViewWidget)
 			{
@@ -1272,6 +1274,7 @@ namespace UnrealGameSync
 					{
 						BuildListContextMenu_Assign.Text = String.Format("Assign to {0}", _contextMenuChange.User);
 						BuildListContextMenu.Show(BuildListView, args.Location);
+						ThemeManager?.ApplyThemeRecursively(BuildListContextMenu);
 					}
 				}
 			}
diff --git a/Engine/Source/Programs/UnrealGameSync/UnrealGameSync/Forms/IssueFixedWindow.cs b/Engine/Source/Programs/UnrealGameSync/UnrealGameSync/Forms/IssueFixedWindow.cs
--- a/Engine/Source/Programs/UnrealGameSync/UnrealGameSync/Forms/IssueFixedWindow.cs
+++ b/Engine/Source/Programs/UnrealGameSync/UnrealGameSync/Forms/IssueFixedWindow.cs
@@ -17,7 +17,7 @@ using Microsoft.Extensions.Logging;
 
 namespace UnrealGameSync
 {
-	partial class IssueFixedWindow : Form
+	partial class IssueFixedWindow : ThemedForm
 	{
 		class FindChangesWorker : Component
 		{
@@ -248,6 +248,7 @@ namespace UnrealGameSync
 					{
 						ChangesListContextMenu.Tag = record;
 						ChangesListContextMenu.Show(ChangesListView, args.Location);
+						ThemeManager?.ApplyThemeRecursively(ChangesListContextMenu);
 					}
 				}
 			}
diff --git a/Engine/Source/Programs/UnrealGameSync/UnrealGameSync/Forms/IssueSettingsWindow.cs b/Engine/Source/Programs/UnrealGameSync/UnrealGameSync/Forms/IssueSettingsWindow.cs
--- a/Engine/Source/Programs/UnrealGameSync/UnrealGameSync/Forms/IssueSettingsWindow.cs
+++ b/Engine/Source/Programs/UnrealGameSync/UnrealGameSync/Forms/IssueSettingsWindow.cs
@@ -7,7 +7,7 @@ using Microsoft.Extensions.Logging;
 
 namespace UnrealGameSync
 {
-	partial class IssueSettingsWindow : Form
+	partial class IssueSettingsWindow : ThemedForm
 	{
 		readonly UserSettings _settings;
 		readonly ILogger _logger;
diff --git a/Engine/Source/Programs/UnrealGameSync/UnrealGameSync/Forms/LauncherUpdateWindow.cs b/Engine/Source/Programs/UnrealGameSync/UnrealGameSync/Forms/LauncherUpdateWindow.cs
--- a/Engine/Source/Programs/UnrealGameSync/UnrealGameSync/Forms/LauncherUpdateWindow.cs
+++ b/Engine/Source/Programs/UnrealGameSync/UnrealGameSync/Forms/LauncherUpdateWindow.cs
@@ -4,7 +4,7 @@ using System.Windows.Forms;
 
 namespace UnrealGameSync.Forms
 {
-	public partial class LauncherUpdateWindow : Form
+	public partial class LauncherUpdateWindow : ThemedForm
 	{
 		public LauncherUpdateWindow()
 		{
diff --git a/Engine/Source/Programs/UnrealGameSync/UnrealGameSync/Forms/LeaveCommentWindow.cs b/Engine/Source/Programs/UnrealGameSync/UnrealGameSync/Forms/LeaveCommentWindow.cs
--- a/Engine/Source/Programs/UnrealGameSync/UnrealGameSync/Forms/LeaveCommentWindow.cs
+++ b/Engine/Source/Programs/UnrealGameSync/UnrealGameSync/Forms/LeaveCommentWindow.cs
@@ -4,7 +4,7 @@ using System.Windows.Forms;
 
 namespace UnrealGameSync
 {
-	public partial class LeaveCommentWindow : Form
+	public partial class LeaveCommentWindow : ThemedForm
 	{
 		public LeaveCommentWindow()
 		{
diff --git a/Engine/Source/Programs/UnrealGameSync/UnrealGameSync/Forms/LogWindow.cs b/Engine/Source/Programs/UnrealGameSync/UnrealGameSync/Forms/LogWindow.cs
--- a/Engine/Source/Programs/UnrealGameSync/UnrealGameSync/Forms/LogWindow.cs
+++ b/Engine/Source/Programs/UnrealGameSync/UnrealGameSync/Forms/LogWindow.cs
@@ -4,7 +4,7 @@ using System.Windows.Forms;
 
 namespace UnrealGameSync
 {
-	public partial class LogWindow : Form
+	public partial class LogWindow : ThemedForm
 	{
 		public LogWindow(string text)
 		{
diff --git a/Engine/Source/Programs/UnrealGameSync/UnrealGameSync/Forms/MainWindow.cs b/Engine/Source/Programs/UnrealGameSync/UnrealGameSync/Forms/MainWindow.cs
--- a/Engine/Source/Programs/UnrealGameSync/UnrealGameSync/Forms/MainWindow.cs
+++ b/Engine/Source/Programs/UnrealGameSync/UnrealGameSync/Forms/MainWindow.cs
@@ -53,7 +53,7 @@ namespace UnrealGameSync
 		}
 	}
 
-	partial class MainWindow : Form, IWorkspaceControlOwner
+	partial class MainWindow : ThemedForm, IWorkspaceControlOwner
 	{
 		class WorkspaceIssueMonitor
 		{
@@ -557,6 +557,7 @@ namespace UnrealGameSync
 				TabMenu_TabNames_WorkspaceRoot.Checked = _settings.TabLabels == TabLabels.WorkspaceRoot;
 				TabMenu_TabNames_ProjectFile.Checked = _settings.TabLabels == TabLabels.ProjectFile;
 				TabMenu.Show(TabControl, location);
+				ThemeManager?.ApplyThemeRecursively(TabMenu);
 
 				TabControl.LockHover();
 			}
@@ -701,7 +702,7 @@ namespace UnrealGameSync
 			ErrorPanel errorPanel = new ErrorPanel(project);
 			errorPanel.Parent = TabPanel;
 			errorPanel.BorderStyle = BorderStyle.FixedSingle;
-			errorPanel.BackColor = System.Drawing.Color.FromArgb(((int)(((byte)(250)))), ((int)(((byte)(250)))), ((int)(((byte)(250)))));
+			errorPanel.BackColor = Theme.BackgroundColor;
 			errorPanel.Location = new Point(0, 0);
 			errorPanel.Dock = DockStyle.Fill;
 			errorPanel.Hide();
@@ -1150,6 +1151,7 @@ namespace UnrealGameSync
 					TabControl.SelectTab(newTabIdx);
 					SaveTabSettings();
 					UpdateRecentProjectsList(openProjectInfo.SelectedProject);
+					ThemeManager?.ApplyThemeRecursively(this);
 				}
 			}
 		}
@@ -1304,6 +1306,7 @@ namespace UnrealGameSync
 			if (replaceTabIdx == -1)
 			{
 				int newTabIdx = TabControl.InsertTab(-1, newTabName, newWorkspace, newWorkspace.TintColor);
+				ThemeManager?.ApplyThemeRecursively(newWorkspace);
 				_logger.LogInformation("  Inserted tab {TabIdx}", newTabIdx);
 				return newTabIdx;
 			}
@@ -1311,6 +1314,7 @@ namespace UnrealGameSync
 			{
 				_logger.LogInformation("  Replacing tab {TabIdx}", replaceTabIdx);
 				TabControl.InsertTab(replaceTabIdx + 1, newTabName, newWorkspace, newWorkspace.TintColor);
+				ThemeManager?.ApplyThemeRecursively(newWorkspace);
 				TabControl.RemoveTab(replaceTabIdx);
 				return replaceTabIdx;
 			}
@@ -1479,18 +1483,18 @@ namespace UnrealGameSync
 				if (desiredTaskbarState.Item1 == TaskbarState.Error)
 				{
 					state = TaskbarState.Error;
-					TabControl.SetHighlight(idx, Tuple.Create(Color.FromArgb(204, 64, 64), 1.0f));
+					TabControl.SetHighlight(idx, Tuple.Create(Theme.ErrorColor, 1.0f));
 				}
 				else if (desiredTaskbarState.Item1 == TaskbarState.Paused && state != TaskbarState.Error)
 				{
 					state = TaskbarState.Paused;
-					TabControl.SetHighlight(idx, Tuple.Create(Color.FromArgb(255, 242, 0), 1.0f));
+					TabControl.SetHighlight(idx, Tuple.Create(Theme.WarningColor, 1.0f));
 				}
 				else if (desiredTaskbarState.Item1 == TaskbarState.Normal && state != TaskbarState.Error && state != TaskbarState.Paused)
 				{
 					state = TaskbarState.Normal;
 					progress = Math.Max(progress, desiredTaskbarState.Item2);
-					TabControl.SetHighlight(idx, Tuple.Create(Color.FromArgb(28, 180, 64), desiredTaskbarState.Item2));
+					TabControl.SetHighlight(idx, Tuple.Create(Theme.SuccessColor, desiredTaskbarState.Item2));
 				}
 				else
 				{
diff --git a/Engine/Source/Programs/UnrealGameSync/UnrealGameSync/Forms/ModalTaskWindow.cs b/Engine/Source/Programs/UnrealGameSync/UnrealGameSync/Forms/ModalTaskWindow.cs
--- a/Engine/Source/Programs/UnrealGameSync/UnrealGameSync/Forms/ModalTaskWindow.cs
+++ b/Engine/Source/Programs/UnrealGameSync/UnrealGameSync/Forms/ModalTaskWindow.cs
@@ -7,7 +7,7 @@ using System.Windows.Forms;
 
 namespace UnrealGameSync
 {
-	partial class ModalTaskWindow : Form
+	partial class ModalTaskWindow : ThemedForm
 	{
 		readonly Task _task;
 		readonly CancellationTokenSource _cancellationSource;
diff --git a/Engine/Source/Programs/UnrealGameSync/UnrealGameSync/Forms/ModifyBuildStepsWindow.cs b/Engine/Source/Programs/UnrealGameSync/UnrealGameSync/Forms/ModifyBuildStepsWindow.cs
--- a/Engine/Source/Programs/UnrealGameSync/UnrealGameSync/Forms/ModifyBuildStepsWindow.cs
+++ b/Engine/Source/Programs/UnrealGameSync/UnrealGameSync/Forms/ModifyBuildStepsWindow.cs
@@ -9,7 +9,7 @@ using EpicGames.Core;
 
 namespace UnrealGameSync
 {
-	partial class ModifyBuildStepsWindow : Form
+	partial class ModifyBuildStepsWindow : ThemedForm
 	{
 		readonly List<string> _targetNames;
 		readonly List<BuildStep> _steps;
@@ -129,7 +129,14 @@ namespace UnrealGameSync
 
 		private void BuildStepList_DrawColumnHeader(object sender, DrawListViewColumnHeaderEventArgs e)
 		{
-			e.DrawDefault = true;
+			if (e.Header == null)
+			{
+				e.DrawDefault = true;
+			}
+			else
+			{
+				TextRenderer.DrawText(e.Graphics, e.Header.Text, BuildStepList.Font, e.Bounds, ForeColor, TextFormatFlags.Left | TextFormatFlags.VerticalCenter);
+			}
 		}
 
 		private void BuildStepList_DrawSubItem(object sender, DrawListViewSubItemEventArgs e)
@@ -152,8 +159,10 @@ namespace UnrealGameSync
 					enabled = task.ScheduledSync;
 				}
 
-				bool selected = BuildStepList.SelectedItems.Contains(e.Item);
-				e.Graphics.FillRectangle(selected ? SystemBrushes.Highlight : SystemBrushes.Window, e.Bounds);
+				if (BuildStepList.SelectedItems.Contains(e.Item))
+				{
+					e.Graphics.FillRectangle(SystemBrushes.Highlight, e.Bounds);
+				}
 
 				CheckBoxState state;
 				if (enabled)
diff --git a/Engine/Source/Programs/UnrealGameSync/UnrealGameSync/Forms/NewWorkspaceWindow.cs b/Engine/Source/Programs/UnrealGameSync/UnrealGameSync/Forms/NewWorkspaceWindow.cs
--- a/Engine/Source/Programs/UnrealGameSync/UnrealGameSync/Forms/NewWorkspaceWindow.cs
+++ b/Engine/Source/Programs/UnrealGameSync/UnrealGameSync/Forms/NewWorkspaceWindow.cs
@@ -17,7 +17,7 @@ using Microsoft.Extensions.Logging;
 
 namespace UnrealGameSync
 {
-	partial class NewWorkspaceWindow : Form
+	partial class NewWorkspaceWindow : ThemedForm
 	{
 		class NewWorkspaceSettings
 		{
diff --git a/Engine/Source/Programs/UnrealGameSync/UnrealGameSync/Forms/NotificationWindow.cs b/Engine/Source/Programs/UnrealGameSync/UnrealGameSync/Forms/NotificationWindow.cs
--- a/Engine/Source/Programs/UnrealGameSync/UnrealGameSync/Forms/NotificationWindow.cs
+++ b/Engine/Source/Programs/UnrealGameSync/UnrealGameSync/Forms/NotificationWindow.cs
@@ -14,13 +14,13 @@ namespace UnrealGameSync
 		Error,
 	}
 
-	public partial class NotificationWindow : Form
+	public partial class NotificationWindow : ThemedForm
 	{
 		NotificationType _type;
 
 		// Logo
 		Rectangle _logoBounds;
-		Bitmap? _logoBitmap;
+		InvertibleImage _logoImage;
 
 		// Caption
 		Rectangle _captionBounds;
@@ -42,7 +42,7 @@ namespace UnrealGameSync
 
 		public NotificationWindow(Image inLogo)
 		{
-			_logoBitmap = new Bitmap(inLogo);
+			_logoImage = new InvertibleImage(new Bitmap(inLogo), true);
 			InitializeComponent();
 			Font = new System.Drawing.Font("Segoe UI", 8.25F, System.Drawing.FontStyle.Regular, System.Drawing.GraphicsUnit.Point, ((byte)(0)));
 		}
@@ -73,11 +73,7 @@ namespace UnrealGameSync
 					_captionFont.Dispose();
 					_captionFont = null;
 				}
-				if (_logoBitmap != null)
-				{
-					_logoBitmap.Dispose();
-					_logoBitmap = null;
-				}
+				_logoImage.Dispose();
 				if (components != null)
 				{
 					components.Dispose();
@@ -91,11 +87,7 @@ namespace UnrealGameSync
 		{
 			base.OnClosed(e);
 
-			if (_logoBitmap != null)
-			{
-				_logoBitmap.Dispose();
-				_logoBitmap = null;
-			}
+			_logoImage.Dispose();
 		}
 
 		public void CalculateBounds()
@@ -120,7 +112,8 @@ namespace UnrealGameSync
 
 				// Get the logo dimensions
 				int logoH = messageY + messageH - gutterH;
-				int logoW = (_logoBitmap == null) ? 16 : (int)(_logoBitmap.Width * (float)logoH / (float)_logoBitmap.Height);
+				Image? logoImage = _logoImage.ResolveImage();
+				int logoW = (logoImage == null) ? 16 : (int)(logoImage.Width * (float)logoH / (float)logoImage.Height);
 
 				// Set the window size
 				int textX = gutterW + logoW + referenceSize8 * 2;
@@ -213,28 +206,36 @@ namespace UnrealGameSync
 			base.OnPaintBackground(e);
 
 			Color tintColor;
-			switch (_type)
+			switch(_type)
 			{
 				case NotificationType.Warning:
-					tintColor = Color.FromArgb(240, 240, 220);
+					tintColor = Theme.WarningNotificationBGTint;
 					break;
 				case NotificationType.Error:
-					tintColor = Color.FromArgb(240, 220, 220);
+					tintColor = Theme.ErrorNotificationBGTint;
 					break;
 				default:
-					tintColor = Color.FromArgb(220, 220, 240);
+					tintColor = Theme.StandardNotificationBGTint;
 					break;
 			}
 
-			using (Brush backgroundBrush = new LinearGradientBrush(new Point(0, 0), new Point(0, Height), Color.White, tintColor))
+			Color backgroundColor = Theme.InputBackgroundColor;
+			using(Brush backgroundBrush = new LinearGradientBrush(new Point(0, 0), new Point(0, Height), backgroundColor, tintColor))
 			{
 				e.Graphics.FillRectangle(backgroundBrush, ClientRectangle);
 			}
 
-			e.Graphics.DrawRectangle(Pens.Black, new Rectangle(0, 0, Width - 1, Height - 1));
+			using (Pen borderPen = new Pen(Theme.ForegroundColor))
+			{
+				e.Graphics.DrawRectangle(borderPen, new Rectangle(0, 0, Width - 1, Height - 1));
+			}
 
-			e.Graphics.InterpolationMode = InterpolationMode.HighQualityBicubic;
-			e.Graphics.DrawImage(_logoBitmap!, _logoBounds);
+			Image? logoImage = _logoImage.ResolveImage();
+			if (logoImage != null)
+			{
+				e.Graphics.InterpolationMode = InterpolationMode.HighQualityBicubic;
+				e.Graphics.DrawImage(logoImage, _logoBounds);
+			}
 		}
 
 		protected override void OnPaint(PaintEventArgs e)
@@ -246,20 +247,26 @@ namespace UnrealGameSync
 			// Draw the close button
 			if (_hoverOverClose)
 			{
-				e.Graphics.FillRectangle(SystemBrushes.ActiveCaption, _closeButtonBounds);
+				using(Brush hoverBrush = new SolidBrush(Theme.ActiveCaption))
+				{
+					e.Graphics.FillRectangle(hoverBrush, _closeButtonBounds);
+				}
 			}
-			else if (_mouseDownOverClose)
+			else if(_mouseDownOverClose)
 			{
-				e.Graphics.FillRectangle(SystemBrushes.InactiveCaption, _closeButtonBounds);
+				using(Brush downBrush = new SolidBrush(Theme.InactiveCaption))
+				{
+					e.Graphics.FillRectangle(downBrush, _closeButtonBounds);
+				}
 			}
-			if (_hoverOverClose || _mouseDownOverClose)
+			if(_hoverOverClose || _mouseDownOverClose)
 			{
-				using (Pen borderPen = new Pen(SystemColors.InactiveCaptionText, 1.0f))
+				using(Pen borderPen = new Pen(Theme.InactiveCaptionText, 1.0f))
 				{
 					e.Graphics.DrawRectangle(borderPen, new Rectangle(_closeButtonBounds.X, _closeButtonBounds.Y, _closeButtonBounds.Width - 1, _closeButtonBounds.Height - 1));
 				}
 			}
-			using (Pen crossPen = new Pen(Brushes.Black, 2.0f))
+			using(Pen crossPen = new Pen(Theme.ForegroundColor, 2.0f))
 			{
 				float offset = (_closeButtonBounds.Width - 1) / 4.0f;
 				e.Graphics.DrawLine(crossPen, _closeButtonBounds.Left + offset, _closeButtonBounds.Top + offset, _closeButtonBounds.Right - 1 - offset, _closeButtonBounds.Bottom - 1 - offset);
@@ -267,10 +274,10 @@ namespace UnrealGameSync
 			}
 
 			// Draw the caption
-			TextRenderer.DrawText(e.Graphics, _captionText, _captionFont, _captionBounds, SystemColors.HotTrack, TextFormatFlags.Left);
+			TextRenderer.DrawText(e.Graphics, _captionText, _captionFont, _captionBounds, Theme.LinkColor, TextFormatFlags.Left);
 
 			// Draw the message
-			TextRenderer.DrawText(e.Graphics, _messageText, Font, _messageBounds, SystemColors.ControlText, TextFormatFlags.Left);
+			TextRenderer.DrawText(e.Graphics, _messageText, Font, _messageBounds, Theme.TextColor, TextFormatFlags.Left);
 		}
 	}
 }
diff --git a/Engine/Source/Programs/UnrealGameSync/UnrealGameSync/Forms/OpenProjectWindow.cs b/Engine/Source/Programs/UnrealGameSync/UnrealGameSync/Forms/OpenProjectWindow.cs
--- a/Engine/Source/Programs/UnrealGameSync/UnrealGameSync/Forms/OpenProjectWindow.cs
+++ b/Engine/Source/Programs/UnrealGameSync/UnrealGameSync/Forms/OpenProjectWindow.cs
@@ -16,7 +16,7 @@ using Microsoft.Extensions.Logging;
 
 namespace UnrealGameSync
 {
-	partial class OpenProjectWindow : Form
+	partial class OpenProjectWindow : ThemedForm
 	{
 		string? _serverAndPortOverride;
 		string? _userNameOverride;
@@ -108,7 +108,10 @@ namespace UnrealGameSync
 
 		private void UpdateEnabledControls()
 		{
-			Color workspaceTextColor = WorkspaceRadioBtn.Checked ? SystemColors.ControlText : SystemColors.GrayText;
+			Color activeTextColor = Theme.TextColor;
+			Color inactiveTextColor = Theme.DimmedText;
+			
+			Color workspaceTextColor = WorkspaceRadioBtn.Checked? activeTextColor : inactiveTextColor;
 			WorkspaceNameLabel.ForeColor = workspaceTextColor;
 			WorkspaceNameTextBox.ForeColor = workspaceTextColor;
 			WorkspaceNameNewBtn.ForeColor = workspaceTextColor;
@@ -117,7 +120,7 @@ namespace UnrealGameSync
 			WorkspacePathTextBox.ForeColor = workspaceTextColor;
 			WorkspacePathBrowseBtn.ForeColor = workspaceTextColor;
 
-			Color localFileTextColor = LocalFileRadioBtn.Checked ? SystemColors.ControlText : SystemColors.GrayText;
+			Color localFileTextColor = LocalFileRadioBtn.Checked? activeTextColor : inactiveTextColor;
 			LocalFileLabel.ForeColor = localFileTextColor;
 			LocalFileTextBox.ForeColor = localFileTextColor;
 			LocalFileBrowseBtn.ForeColor = localFileTextColor;
@@ -155,6 +158,14 @@ namespace UnrealGameSync
 			}
 		}
 
+		protected override void OnHandleCreated(EventArgs e)
+		{
+			base.OnHandleCreated(e);
+			
+			// Override the automatic theming colors for the enabled state
+			UpdateEnabledControls();
+		}
+
 		private void UpdateServerLabel()
 		{
 			ServerLabel.Text = GetServerLabelText(_defaultSettings, _serverAndPortOverride, _userNameOverride);
diff --git a/Engine/Source/Programs/UnrealGameSync/UnrealGameSync/Forms/PasswordWindow.cs b/Engine/Source/Programs/UnrealGameSync/UnrealGameSync/Forms/PasswordWindow.cs
--- a/Engine/Source/Programs/UnrealGameSync/UnrealGameSync/Forms/PasswordWindow.cs
+++ b/Engine/Source/Programs/UnrealGameSync/UnrealGameSync/Forms/PasswordWindow.cs
@@ -5,7 +5,7 @@ using System.Windows.Forms;
 
 namespace UnrealGameSync
 {
-	public partial class PasswordWindow : Form
+	public partial class PasswordWindow : ThemedForm
 	{
 		public string Password { get; private set; }
 
diff --git a/Engine/Source/Programs/UnrealGameSync/UnrealGameSync/Forms/PerforceSyncSettingsWindow.cs b/Engine/Source/Programs/UnrealGameSync/UnrealGameSync/Forms/PerforceSyncSettingsWindow.cs
--- a/Engine/Source/Programs/UnrealGameSync/UnrealGameSync/Forms/PerforceSyncSettingsWindow.cs
+++ b/Engine/Source/Programs/UnrealGameSync/UnrealGameSync/Forms/PerforceSyncSettingsWindow.cs
@@ -6,7 +6,7 @@ using Microsoft.Extensions.Logging;
 
 namespace UnrealGameSync
 {
-	partial class PerforceSyncSettingsWindow : Form
+	partial class PerforceSyncSettingsWindow : ThemedForm
 	{
 		readonly UserSettings _settings;
 		readonly ILogger _logger;
diff --git a/Engine/Source/Programs/UnrealGameSync/UnrealGameSync/Forms/PresetsWindow.cs b/Engine/Source/Programs/UnrealGameSync/UnrealGameSync/Forms/PresetsWindow.cs
--- a/Engine/Source/Programs/UnrealGameSync/UnrealGameSync/Forms/PresetsWindow.cs
+++ b/Engine/Source/Programs/UnrealGameSync/UnrealGameSync/Forms/PresetsWindow.cs
@@ -6,7 +6,7 @@ using System.Windows.Forms;
 
 namespace UnrealGameSync.Forms
 {
-	public partial class PresetsWindow : Form
+	public partial class PresetsWindow : ThemedForm
 	{
 		private readonly IDictionary<Guid, WorkspaceSyncCategory> _uniqueIdToCategory;
 		private readonly IDictionary<string, Preset> _roles;
diff --git a/Engine/Source/Programs/UnrealGameSync/UnrealGameSync/Forms/ProgramsRunningWindow.cs b/Engine/Source/Programs/UnrealGameSync/UnrealGameSync/Forms/ProgramsRunningWindow.cs
--- a/Engine/Source/Programs/UnrealGameSync/UnrealGameSync/Forms/ProgramsRunningWindow.cs
+++ b/Engine/Source/Programs/UnrealGameSync/UnrealGameSync/Forms/ProgramsRunningWindow.cs
@@ -9,7 +9,7 @@ using EpicGames.Core;
 
 namespace UnrealGameSync
 {
-	public partial class ProgramsRunningWindow : Form
+	public partial class ProgramsRunningWindow : ThemedForm
 	{
 		readonly object _syncObject = new object();
 		FileReference[] _programs;
diff --git a/Engine/Source/Programs/UnrealGameSync/UnrealGameSync/Forms/ScheduleWindow.cs b/Engine/Source/Programs/UnrealGameSync/UnrealGameSync/Forms/ScheduleWindow.cs
--- a/Engine/Source/Programs/UnrealGameSync/UnrealGameSync/Forms/ScheduleWindow.cs
+++ b/Engine/Source/Programs/UnrealGameSync/UnrealGameSync/Forms/ScheduleWindow.cs
@@ -9,7 +9,7 @@ using System.Windows.Forms;
 
 namespace UnrealGameSync
 {
-	partial class ScheduleWindow : Form
+	partial class ScheduleWindow : ThemedForm
 	{
 		readonly Dictionary<UserSelectedProjectSettings, List<LatestChangeType>> _projectToLatestChangeTypes;
 
diff --git a/Engine/Source/Programs/UnrealGameSync/UnrealGameSync/Forms/SdkInfoWindow.cs b/Engine/Source/Programs/UnrealGameSync/UnrealGameSync/Forms/SdkInfoWindow.cs
--- a/Engine/Source/Programs/UnrealGameSync/UnrealGameSync/Forms/SdkInfoWindow.cs
+++ b/Engine/Source/Programs/UnrealGameSync/UnrealGameSync/Forms/SdkInfoWindow.cs
@@ -14,7 +14,7 @@ using System.Windows.Forms;
 
 namespace UnrealGameSync
 {
-	public partial class SdkInfoWindow : Form
+	public partial class SdkInfoWindow : ThemedForm
 	{
 		abstract class SdkAction
 		{
diff --git a/Engine/Source/Programs/UnrealGameSync/UnrealGameSync/Forms/SelectProjectFromWorkspaceWindow.cs b/Engine/Source/Programs/UnrealGameSync/UnrealGameSync/Forms/SelectProjectFromWorkspaceWindow.cs
--- a/Engine/Source/Programs/UnrealGameSync/UnrealGameSync/Forms/SelectProjectFromWorkspaceWindow.cs
+++ b/Engine/Source/Programs/UnrealGameSync/UnrealGameSync/Forms/SelectProjectFromWorkspaceWindow.cs
@@ -19,7 +19,7 @@ using Microsoft.Extensions.Logging;
 
 namespace UnrealGameSync
 {
-	partial class SelectProjectFromWorkspaceWindow : Form
+	partial class SelectProjectFromWorkspaceWindow : ThemedForm
 	{
 		static class EnumerateWorkspaceProjectsTask
 		{
diff --git a/Engine/Source/Programs/UnrealGameSync/UnrealGameSync/Forms/SelectStreamWindow.cs b/Engine/Source/Programs/UnrealGameSync/UnrealGameSync/Forms/SelectStreamWindow.cs
--- a/Engine/Source/Programs/UnrealGameSync/UnrealGameSync/Forms/SelectStreamWindow.cs
+++ b/Engine/Source/Programs/UnrealGameSync/UnrealGameSync/Forms/SelectStreamWindow.cs
@@ -17,7 +17,7 @@ using Microsoft.Extensions.Logging;
 
 namespace UnrealGameSync
 {
-	partial class SelectStreamWindow : Form
+	partial class SelectStreamWindow : ThemedForm
 	{
 		static class EnumerateStreamsTask
 		{
diff --git a/Engine/Source/Programs/UnrealGameSync/UnrealGameSync/Forms/SelectUserWindow.cs b/Engine/Source/Programs/UnrealGameSync/UnrealGameSync/Forms/SelectUserWindow.cs
--- a/Engine/Source/Programs/UnrealGameSync/UnrealGameSync/Forms/SelectUserWindow.cs
+++ b/Engine/Source/Programs/UnrealGameSync/UnrealGameSync/Forms/SelectUserWindow.cs
@@ -14,7 +14,7 @@ using Microsoft.Extensions.Logging;
 
 namespace UnrealGameSync
 {
-	partial class SelectUserWindow : Form
+	partial class SelectUserWindow : ThemedForm
 	{
 		static class EnumerateUsersTask
 		{
diff --git a/Engine/Source/Programs/UnrealGameSync/UnrealGameSync/Forms/SelectWorkspaceWindow.cs b/Engine/Source/Programs/UnrealGameSync/UnrealGameSync/Forms/SelectWorkspaceWindow.cs
--- a/Engine/Source/Programs/UnrealGameSync/UnrealGameSync/Forms/SelectWorkspaceWindow.cs
+++ b/Engine/Source/Programs/UnrealGameSync/UnrealGameSync/Forms/SelectWorkspaceWindow.cs
@@ -15,7 +15,7 @@ using Microsoft.Extensions.Logging;
 
 namespace UnrealGameSync
 {
-	partial class SelectWorkspaceWindow : Form
+	partial class SelectWorkspaceWindow : ThemedForm
 	{
 		class EnumerateWorkspaces
 		{
diff --git a/Engine/Source/Programs/UnrealGameSync/UnrealGameSync/Forms/SwitchStreamWarningWindow.cs b/Engine/Source/Programs/UnrealGameSync/UnrealGameSync/Forms/SwitchStreamWarningWindow.cs
--- a/Engine/Source/Programs/UnrealGameSync/UnrealGameSync/Forms/SwitchStreamWarningWindow.cs
+++ b/Engine/Source/Programs/UnrealGameSync/UnrealGameSync/Forms/SwitchStreamWarningWindow.cs
@@ -6,7 +6,7 @@ using System.Windows.Forms;
 
 namespace UnrealGameSync.Forms
 {
-	public partial class SwitchStreamWarningWindow : Form
+	public partial class SwitchStreamWarningWindow : ThemedForm
 	{
 		public SwitchStreamWarningWindow(List<string> openedFiles)
 		{
diff --git a/Engine/Source/Programs/UnrealGameSync/UnrealGameSync/Forms/SyncFilter.cs b/Engine/Source/Programs/UnrealGameSync/UnrealGameSync/Forms/SyncFilter.cs
--- a/Engine/Source/Programs/UnrealGameSync/UnrealGameSync/Forms/SyncFilter.cs
+++ b/Engine/Source/Programs/UnrealGameSync/UnrealGameSync/Forms/SyncFilter.cs
@@ -8,7 +8,7 @@ using UnrealGameSync.Forms;
 
 namespace UnrealGameSync
 {
-	partial class SyncFilter : Form
+	partial class SyncFilter : ThemedForm
 	{
 		class CheckListItem
 		{
diff --git a/Engine/Source/Programs/UnrealGameSync/UnrealGameSync/Forms/SyncFilterSyntax.cs b/Engine/Source/Programs/UnrealGameSync/UnrealGameSync/Forms/SyncFilterSyntax.cs
--- a/Engine/Source/Programs/UnrealGameSync/UnrealGameSync/Forms/SyncFilterSyntax.cs
+++ b/Engine/Source/Programs/UnrealGameSync/UnrealGameSync/Forms/SyncFilterSyntax.cs
@@ -4,7 +4,7 @@ using System.Windows.Forms;
 
 namespace UnrealGameSync
 {
-	public partial class SyncFilterSyntax : Form
+	public partial class SyncFilterSyntax : ThemedForm
 	{
 		public SyncFilterSyntax()
 		{
diff --git a/Engine/Source/Programs/UnrealGameSync/UnrealGameSync/Forms/ThemedForm.cs b/Engine/Source/Programs/UnrealGameSync/UnrealGameSync/Forms/ThemedForm.cs
new file mode 100644
--- /dev/null
+++ b/Engine/Source/Programs/UnrealGameSync/UnrealGameSync/Forms/ThemedForm.cs
@@ -0,0 +1,20 @@
+// Copyright Epic Games, Inc. All Rights Reserved.
+
+using System;
+using System.Windows.Forms;
+
+namespace UnrealGameSync
+{
+	public class ThemedForm : Form
+	{
+		protected override void OnHandleCreated(EventArgs e)
+		{
+			base.OnHandleCreated(e);
+
+			ThemeManager?.ApplyThemeRecursively(this);
+		}
+		
+		protected static ApplicationThemeManager? ThemeManager => ProgramApplicationContext.GetThemeManager();
+		protected static ApplicationTheme Theme => ProgramApplicationContext.GetTheme();
+	}
+}
diff --git a/Engine/Source/Programs/UnrealGameSync/UnrealGameSync/Forms/UpdateSettingsWindow.cs b/Engine/Source/Programs/UnrealGameSync/UnrealGameSync/Forms/UpdateSettingsWindow.cs
--- a/Engine/Source/Programs/UnrealGameSync/UnrealGameSync/Forms/UpdateSettingsWindow.cs
+++ b/Engine/Source/Programs/UnrealGameSync/UnrealGameSync/Forms/UpdateSettingsWindow.cs
@@ -11,7 +11,7 @@ using Microsoft.Extensions.Logging;
 
 namespace UnrealGameSync
 {
-	partial class UpdateSettingsWindow : Form
+	partial class UpdateSettingsWindow : ThemedForm
 	{
 		[DllImport("user32.dll")]
 		private static extern IntPtr SendMessage(IntPtr hWnd, int msg, int wParam, [MarshalAs(UnmanagedType.LPWStr)] string lParam);
diff --git a/Engine/Source/Programs/UnrealGameSync/UnrealGameSync/Forms/VariablesWindow.cs b/Engine/Source/Programs/UnrealGameSync/UnrealGameSync/Forms/VariablesWindow.cs
--- a/Engine/Source/Programs/UnrealGameSync/UnrealGameSync/Forms/VariablesWindow.cs
+++ b/Engine/Source/Programs/UnrealGameSync/UnrealGameSync/Forms/VariablesWindow.cs
@@ -10,7 +10,7 @@ using System.Windows.Forms;
 
 namespace UnrealGameSync
 {
-	public partial class VariablesWindow : Form
+	public partial class VariablesWindow : ThemedForm
 	{
 		static readonly HashSet<string> s_legacyVariables = new HashSet<string>()
 		{
diff --git a/Engine/Source/Programs/UnrealGameSync/UnrealGameSync/ProgramApplicationContext.cs b/Engine/Source/Programs/UnrealGameSync/UnrealGameSync/ProgramApplicationContext.cs
--- a/Engine/Source/Programs/UnrealGameSync/UnrealGameSync/ProgramApplicationContext.cs
+++ b/Engine/Source/Programs/UnrealGameSync/UnrealGameSync/ProgramApplicationContext.cs
@@ -22,6 +22,7 @@ namespace UnrealGameSync
 {
 	class ProgramApplicationContext : ApplicationContext
 	{
+		static ProgramApplicationContext? s_instance { get; set; }
 		readonly SynchronizationContext _mainThreadSynchronizationContext;
 
 		readonly IPerforceSettings _defaultPerforceSettings;
@@ -38,6 +39,7 @@ namespace UnrealGameSync
 		readonly IServiceProvider _serviceProvider;
 		readonly ILogger _logger;
 		readonly UserSettings _settings;
+		readonly ApplicationThemeManager _themeManager;
 		ActivationListener _activationListener;
 
 		Container _components = new Container();
@@ -64,6 +66,7 @@ namespace UnrealGameSync
 
 		public ProgramApplicationContext(IPerforceSettings defaultPerforceSettings, UpdateMonitor updateMonitor, string? apiUrl, DirectoryReference dataFolder, EventWaitHandle activateEvent, bool restoreState, string? updateSpawn, string? projectFileName, bool preview, IServiceProvider serviceProvider, string? uri)
 		{
+			s_instance = this;
 			_defaultPerforceSettings = defaultPerforceSettings;
 			_updateMonitor = updateMonitor;
 			_apiUrl = apiUrl;
@@ -90,6 +93,9 @@ namespace UnrealGameSync
 
 			// Capture the main thread's synchronization context for callbacks
 			_mainThreadSynchronizationContext = SynchronizationContext.Current!;
+			
+			// Create the theme manager before the settings so that the user's platform theme preference is detected first
+			_themeManager = new ApplicationThemeManager();
 
 			// Read the user's settings
 			_settings = UserSettings.Create(dataFolder, serviceProvider.GetRequiredService<ILogger<UserSettings>>());
@@ -221,6 +227,16 @@ namespace UnrealGameSync
 			_startupWindow.FormClosed += (s, e) => OnStartupComplete(startupTasks);
 		}
 
+		public static ApplicationThemeManager? GetThemeManager()
+		{
+			return s_instance?._themeManager;
+		}
+
+		public static ApplicationTheme GetTheme()
+		{
+			return s_instance?._settings.Theme ?? ApplicationTheme.ConstructDefaultTheme();
+		}
+
 		static bool IsValidProject(UserSelectedProjectSettings projectSettings, ILogger logger)
 		{
 			try
@@ -398,6 +414,8 @@ namespace UnrealGameSync
 				_synchronizationContext.Dispose();
 				_synchronizationContext = null;
 			}
+			
+			_themeManager.Dispose();
 		}
 
 		private void NotifyIcon_MouseDown(object? sender, MouseEventArgs e)
diff --git a/Engine/Source/Programs/UnrealGameSync/UnrealGameSyncShared/ApplicationTheme.cs b/Engine/Source/Programs/UnrealGameSync/UnrealGameSyncShared/ApplicationTheme.cs
new file mode 100644
--- /dev/null
+++ b/Engine/Source/Programs/UnrealGameSync/UnrealGameSyncShared/ApplicationTheme.cs
@@ -0,0 +1,361 @@
+// Copyright Epic Games, Inc. All Rights Reserved.
+
+using System;
+using System.Collections.Generic;
+using System.Collections.ObjectModel;
+using System.Drawing;
+using System.Reflection;
+
+namespace UnrealGameSync
+{
+	[AttributeUsage(System.AttributeTargets.Property)]
+	sealed class ThemeAttribute : System.Attribute
+	{
+
+	}
+
+	public class ApplicationTheme
+	{
+		public const string LightThemeId = "Light";
+		public const string DarkThemeId = "Dark";
+
+		public static string DefaultThemeId { get; set; } = LightThemeId;
+
+		static readonly ReadOnlyCollection<Func<ApplicationTheme>> s_defaultThemeConstructors = new ReadOnlyCollection<Func<ApplicationTheme>>(new List<Func<ApplicationTheme>>() { CreateLightTheme, CreateDarkTheme });
+
+		public string ThemeId { get; init; }
+
+		[Theme]
+		public string ThemeName { get; init; } = "Unnamed";
+
+		[Theme]
+		public bool UseDarkMode { get; init; } = false;
+		[Theme]
+		public bool UseFlatControls { get; init; } = false;
+
+		[Theme]
+		public Color ForegroundColor { get; init; } = SystemColors.ControlDark;
+		[Theme]
+		public Color BackgroundColor { get; init; } = SystemColors.Control;
+
+		[Theme]
+		public Color ActiveCaption { get; init; } = SystemColors.ActiveCaption;
+		[Theme]
+		public Color ActiveCaptionText { get; init; } = SystemColors.ActiveCaptionText;
+		[Theme]
+		public Color InactiveCaption { get; init; } = SystemColors.InactiveCaption;
+		[Theme]
+		public Color InactiveCaptionText { get; init; } = SystemColors.InactiveCaptionText;
+
+		[Theme]
+		public Color MenuHighlightBGColor { get; init; } = Color.FromArgb(179, 215, 243);
+		[Theme]
+		public Color MenuHighlightBorderColor { get; init; } = SystemColors.Highlight;
+
+		[Theme]
+		public Color TextColor { get; init; } = SystemColors.ControlText;
+		[Theme]
+		public Color DimmedText { get; init; } = SystemColors.GrayText;
+		[Theme]
+		public Color LinkColor { get; init; } = SystemColors.HotTrack;
+		[Theme]
+		public Color LinkColorHover { get; init; } = Color.FromArgb(SystemColors.HotTrack.B, SystemColors.HotTrack.G, SystemColors.HotTrack.R);
+		[Theme]
+		public Color LinkColorActive { get; init; } = Color.FromArgb(SystemColors.HotTrack.B / 2, SystemColors.HotTrack.G / 2, SystemColors.HotTrack.R);
+		[Theme]
+		public Color SuccessColor { get; init; } = Color.FromArgb(28, 180, 64);
+		[Theme]
+		public Color WarningColor { get; init; } = Color.FromArgb(128, 128, 0);
+		[Theme]
+		public Color ErrorColor { get; init; } = Color.FromArgb(189, 54, 47);
+		[Theme]
+		public Color StandardNotificationBGTint { get; init; } = Color.FromArgb(220, 220, 240);
+		[Theme]
+		public Color SuccessNotificationBGTint { get; init; } = Color.FromArgb(248, 254, 246);
+		[Theme]
+		public Color WarningNotificationBGTint { get; init; } = Color.FromArgb(240, 240, 220);
+		[Theme]
+		public Color ErrorNotificationBGTint { get; init; } = Color.FromArgb(240, 220, 220);
+
+		[Theme]
+		public Color InputBackgroundColor { get; init; } = SystemColors.Window;
+		[Theme]
+		public Color InputBorderColor { get; init; } = Color.Empty;
+
+		[Theme]
+		public Color CodeBadgeColor { get; init; } = Color.FromArgb(116, 185, 255);
+		[Theme]
+		public Color ContentBadgeColor { get; init; } = Color.FromArgb(162, 155, 255);
+
+		[Theme]
+		public Color SyncBadgeColor { get; init; } = Color.FromArgb(140, 180, 230);
+		[Theme]
+		public Color SyncHoverBadgeColor { get; init; } = Color.FromArgb(112, 146, 190);
+
+		[Theme]
+		public bool ShouldInvertColorBrightness { get; init; } = false;
+
+		public static string GetThemeConfigSectionName(string themeId) => "Theme." + themeId;
+		static ApplicationTheme CreateLightTheme()
+		{
+			return new ApplicationTheme(LightThemeId)
+			{
+				ThemeName = (DefaultThemeId == LightThemeId) ? "Light (Default)" : "Light",
+				// Inline property defaults are expected to be for Light Theme
+			};
+		}
+
+		static ApplicationTheme CreateDarkTheme()
+		{
+			return new ApplicationTheme(DarkThemeId)
+			{
+				ThemeName = (DefaultThemeId == DarkThemeId) ? "Dark (Default)" : "Dark",
+
+				UseDarkMode = true,
+				UseFlatControls = true,
+
+				ForegroundColor = SystemColors.Control,
+				BackgroundColor = Color.FromArgb(30, 30, 30),
+				
+				ActiveCaption = Color.FromArgb(46, 73, 102),
+				ActiveCaptionText = Color.FromArgb(255, 255, 255),
+				InactiveCaption = Color.FromArgb(37, 51, 65),
+				InactiveCaptionText = Color.FromArgb(255, 255, 255),
+				
+				MenuHighlightBGColor = Color.FromArgb(12, 48, 75),
+				MenuHighlightBorderColor = Color.FromArgb(41, 159, 255),
+
+				TextColor = SystemColors.ControlLight,
+				DimmedText = SystemColors.ControlDarkDark,
+				LinkColor = Color.FromArgb(84, 170, 255),
+				LinkColorHover = Color.FromArgb(255, 170, 84),
+				LinkColorActive = Color.FromArgb(127, 85, 84),
+				SuccessColor = Color.FromArgb(74, 227, 110),
+				WarningColor = Color.FromArgb(255, 255, 128),
+				ErrorColor = Color.FromArgb(208, 74, 67),
+				
+				StandardNotificationBGTint = Color.FromArgb(15, 15, 36),
+				SuccessNotificationBGTint = Color.FromArgb(6, 18, 2),
+				WarningNotificationBGTint = Color.FromArgb(36, 36, 15),
+				ErrorNotificationBGTint = Color.FromArgb(36, 15, 15),
+
+				InputBackgroundColor = Color.FromArgb(45, 45, 45),
+				InputBorderColor = Color.Black,
+
+				CodeBadgeColor = Color.FromArgb(0, 69, 138),
+				ContentBadgeColor = Color.FromArgb(83, 54, 150),
+				
+				SyncBadgeColor = Color.FromArgb(25, 65, 113),
+				SyncHoverBadgeColor = Color.FromArgb(65, 99, 144),
+				
+				ShouldInvertColorBrightness = true
+			};
+		}
+
+		public static Dictionary<string, string> GetThemeList(ConfigFile configFile)
+		{
+			Dictionary<string, string> themeList = new Dictionary<string, string>();
+
+			ConfigSection? themesSection = configFile.FindSection("Themes");
+			string[] themeIds = themesSection?.GetValues("Themes", Array.Empty<string>()) ?? Array.Empty<string>();
+
+			foreach (string themeId in themeIds)
+			{
+				ConfigSection? themeSection = configFile.FindSection(GetThemeConfigSectionName(themeId));
+				if (themeSection != null)
+				{
+					themeList.Add(themeId, themeSection.GetValue("ThemeName", "Unnamed"));
+				}
+			}
+
+			return themeList;
+		}
+
+		public static void AddDefaultThemesToConfig(ConfigFile configFile)
+		{
+			ConfigSection themesSection = configFile.FindOrAddSection("Themes");
+			string[] oldThemeIds = themesSection.GetValues("Themes", Array.Empty<string>());
+			List<string> themeIds = new List<string>(oldThemeIds);
+
+			foreach (Func<ApplicationTheme> func in s_defaultThemeConstructors)
+			{
+				ApplicationTheme theme = func();
+				// Default themes are not editable, we add them to the ini as reference for creating new themes
+				theme.AddThemeToConfig(configFile);
+
+				if (!themeIds.Contains(theme.ThemeId))
+				{
+					themeIds.Add(theme.ThemeId);
+				}
+			}
+
+			themesSection.SetValues("Themes", themeIds.ToArray());
+		}
+
+		public static bool IsColorLight(Color color)
+		{
+			return ((0.2126 * color.R) + (0.7152 * color.G) + (0.0722 * color.B)) > 127;
+		}
+
+		public static Color InvertColorBrightness(Color color)
+		{
+			float hue = color.GetHue();
+			float saturation = color.GetSaturation();
+			float brightness = 1.0f - color.GetBrightness();
+			
+			float c = (1 - Math.Abs((2 * brightness) - 1)) * saturation;
+			float x = c * (1 - Math.Abs(((hue / 60) % 2) - 1));
+			float m = brightness - (c / 2);
+
+			switch ((int)hue / 60)
+			{
+				case 0:
+					return Color.FromArgb(color.A, (int)((c + m) * 255), (int)((x + m) * 255), (int)(m * 255));
+				case 1:
+					return Color.FromArgb(color.A, (int)((x + m) * 255), (int)((c + m) * 255), (int)(m * 255));
+				case 2:
+					return Color.FromArgb(color.A, (int)(m * 255), (int)((c + m) * 255), (int)((x + m) * 255));
+				case 3:
+					return Color.FromArgb(color.A, (int)(m * 255), (int)((x + m) * 255), (int)((c + m) * 255));
+				case 4:
+					return Color.FromArgb(color.A, (int)((x + m) * 255), (int)(m * 255), (int)((c + m) * 255));
+				case 5:
+					return Color.FromArgb(color.A, (int)((c + m) * 255), (int)(m * 255), (int)((x + m) * 255));
+			}
+			
+			return color;
+		}
+
+		public static ApplicationTheme? LoadThemeFromConfig(string themeId, ConfigFile configFile)
+		{
+			ConfigSection? themeSection = configFile.FindSection(GetThemeConfigSectionName(themeId));
+			if (themeSection != null)
+			{
+				return new ApplicationTheme(themeId, themeSection);
+			}
+
+			return null;
+		}
+		
+
+		public static ApplicationTheme ConstructDefaultTheme()
+		{
+			return (DefaultThemeId == LightThemeId) ? CreateLightTheme() : CreateDarkTheme();
+		}
+
+		ApplicationTheme(string inThemeId, ConfigSection? themeSection = null)
+		{
+			ThemeId = inThemeId;
+
+			if (themeSection != null)
+			{
+				foreach (PropertyInfo propertyInfo in GetType().GetProperties())
+				{
+					if (propertyInfo.GetCustomAttribute<ThemeAttribute>() != null)
+					{
+						if (!themeSection.Pairs.ContainsKey(propertyInfo.Name))
+						{
+							continue;
+						}
+
+						if (propertyInfo.PropertyType == typeof(int))
+						{
+							int? value = themeSection.GetOptionalIntValue(propertyInfo.Name, null);
+							if (value != null)
+							{
+								propertyInfo.SetValue(this, value);
+							}
+						}
+						else if (propertyInfo.PropertyType == typeof(long) && propertyInfo.GetValue(this) is long defaultLongValue)
+						{
+							long value = themeSection.GetValue(propertyInfo.Name, defaultLongValue);
+							propertyInfo.SetValue(this, value);
+						}
+						else if (propertyInfo.PropertyType == typeof(bool) && propertyInfo.GetValue(this) is bool defaultBoolValue)
+						{
+							bool value = themeSection.GetValue(propertyInfo.Name, defaultBoolValue);
+							propertyInfo.SetValue(this, value);
+						}
+						else if (propertyInfo.PropertyType == typeof(string))
+						{
+							string? value = themeSection.GetValue(propertyInfo.Name);
+							if (value != null)
+							{
+								propertyInfo.SetValue(this, value);
+							}
+						}
+						else if (propertyInfo.PropertyType == typeof(Color))
+						{
+							string? value = themeSection.GetValue(propertyInfo.Name);
+							if (value != null)
+							{
+								string[] channels = value.Split(',');
+
+								Color color = Color.FromArgb(
+									Int32.Parse(channels[3].Substring(2)), // alpha is stored last
+									Int32.Parse(channels[0].Substring(2)),
+									Int32.Parse(channels[1].Substring(2)),
+									Int32.Parse(channels[2].Substring(2))
+								);
+								
+								propertyInfo.SetValue(this, color);
+							}
+						}
+						else
+						{
+							// unhandled
+						}
+					}
+				}
+			}
+		}
+
+		public void AddThemeToConfig(ConfigFile configFile)
+		{
+			ConfigSection themeSection = configFile.FindOrAddSection(GetThemeConfigSectionName(ThemeId));
+			themeSection.Clear();
+
+			foreach (PropertyInfo propertyInfo in GetType().GetProperties())
+			{
+				if (propertyInfo.GetCustomAttribute<ThemeAttribute>() != null)
+				{
+					object? propertyValue = propertyInfo.GetValue(this);
+					if (propertyValue is int valueInt)
+					{
+						themeSection.SetValue(propertyInfo.Name, valueInt);
+					}
+					else if (propertyValue is long valueLong)
+					{
+						themeSection.SetValue(propertyInfo.Name, valueLong);
+					}
+					else if (propertyValue is bool valueBool)
+					{
+						themeSection.SetValue(propertyInfo.Name, valueBool);
+					}
+					else if (propertyValue is string valueString)
+					{
+						themeSection.SetValue(propertyInfo.Name, valueString);
+					}
+					else if (propertyValue is Color valueColor)
+					{
+						themeSection.SetValue(propertyInfo.Name, $"R={valueColor.R},G={valueColor.G},B={valueColor.B},A={valueColor.A}");
+					}
+					else
+					{
+						// unhandled
+					}
+				}
+			}
+		}
+
+		public Color? CheckInvertColor(Color? color)
+		{
+			if (color != null && ShouldInvertColorBrightness)
+			{
+				return InvertColorBrightness(color.Value);
+			}
+
+			return color;
+		}
+	}
+}
diff --git a/Engine/Source/Programs/UnrealGameSync/UnrealGameSyncShared/UserSettings.cs b/Engine/Source/Programs/UnrealGameSync/UnrealGameSyncShared/UserSettings.cs
--- a/Engine/Source/Programs/UnrealGameSync/UnrealGameSyncShared/UserSettings.cs
+++ b/Engine/Source/Programs/UnrealGameSync/UnrealGameSyncShared/UserSettings.cs
@@ -557,6 +557,8 @@ namespace UnrealGameSync
 		public BuildConfig CompiledEditorBuildConfig { get; set; } // NB: This assumes not using precompiled editor. See CurrentBuildConfig.
 		public TabLabels TabLabels { get; set; }
 		public long NextLauncherVersionCheck { get; set; }
+		public ApplicationTheme? Theme { get; set; }
+		public Dictionary<string, string> ThemeList { get; init; }
 
 		// Precompiled binaries
 		public List<ArchiveSettings> Archives { get; init; } = new List<ArchiveSettings>();
@@ -644,6 +646,8 @@ namespace UnrealGameSync
 				configFile.TryLoad(legacyFileName, logger);
 			}
 
+			ApplicationTheme.AddDefaultThemesToConfig(configFile);
+
 			GlobalSettings? coreSettingsData = null;
 			if (FileReference.Exists(coreFileName))
 			{
@@ -699,6 +703,10 @@ namespace UnrealGameSync
 			ShowLocalTimes = _configFile.GetValue("General.ShowLocalTimes", false);
 			KeepInTray = _configFile.GetValue("General.KeepInTray", true);
 
+			string themeName = _configFile.GetValue("General.Theme", ApplicationTheme.DefaultThemeId);
+			Theme = ApplicationTheme.LoadThemeFromConfig(themeName, _configFile);
+			ThemeList = ApplicationTheme.GetThemeList(_configFile);
+
 			EnabledTools.Clear();
 			EnabledTools.UnionWith(_configFile.GetGuidValues("General.EnabledTools", Array.Empty<Guid>()));
 			if (_configFile.GetValue("General.EnableP4VExtensions", false))
@@ -921,6 +929,11 @@ namespace UnrealGameSync
 			}
 		}
 
+		public ApplicationTheme? LoadTheme(string themeId)
+		{
+			return ApplicationTheme.LoadThemeFromConfig(themeId, _configFile);
+		}
+
 		protected override void ImportWorkspaceState(DirectoryReference rootDir, string clientName, string branchPath, WorkspaceState currentWorkspace)
 		{
 			// Read the workspace settings
@@ -1088,6 +1101,7 @@ namespace UnrealGameSync
 			generalSection.SetValue("SyncTypeID", SyncTypeId);
 			generalSection.SetValue("ClientSubmitOptions", ClientSubmitOptions.ToString());
 			generalSection.SetValue("NextLauncherVersionCheck", NextLauncherVersionCheck);
+			generalSection.SetValue("Theme", Theme?.ThemeId ?? ApplicationTheme.DefaultThemeId);
 
 			// Build configuration
 			generalSection.SetValue("BuildConfig", CompiledEditorBuildConfig.ToString());
