Fixed a create window handle exception which occurred when adding custom filters in the Sync Filter view.

diff --git a/Engine/Source/Programs/UnrealGameSync/UnrealGameSync/ApplicationThemeManager.cs b/Engine/Source/Programs/UnrealGameSync/UnrealGameSync/ApplicationThemeManager.cs
--- a/Engine/Source/Programs/UnrealGameSync/UnrealGameSync/ApplicationThemeManager.cs
+++ b/Engine/Source/Programs/UnrealGameSync/UnrealGameSync/ApplicationThemeManager.cs
@@ -164,20 +164,48 @@ namespace UnrealGameSync
 				ApplyThemeRecursively(form);
 			}
 		}
-		public void ApplyThemeRecursively(Control control)
+
+		private void ControlHandleCreated(object? sender, EventArgs args)
 		{
-			control.HandleCreated += (_, __) =>
+			if (sender is Control control)
 			{
 				ApplyThemeRecursively(control);
-			};
-			control.EnabledChanged +=  (_, __) => ApplyThemeRecursively(control);
-			control.ControlAdded += (_, args) =>
+			}
+		}
+
+		private void ControlEnabledChanged(object? sender, EventArgs args)
+		{
+			if (sender is Control control)
 			{
-				if (args.Control != null)
-				{
-					ApplyThemeRecursively(args.Control);
-				}
-			};
+				ApplyThemeRecursively(control);
+			}
+		}
+
+		private void ControlAdded(object? sender, ControlEventArgs args)
+		{
+			if (args.Control != null)
+			{
+				ApplyThemeRecursively(args.Control);
+			}
+		}
+
+		private void ControlHandleDestroyed(object? sender, EventArgs args)
+		{
+			if (sender is Control control)
+			{
+				control.HandleCreated -= ControlHandleCreated;
+				control.EnabledChanged -= ControlEnabledChanged;
+				control.ControlAdded -= ControlAdded;
+				control.HandleDestroyed -= ControlHandleDestroyed;
+			}
+		}
+
+		public void ApplyThemeRecursively(Control control)
+		{
+			control.HandleCreated += ControlHandleCreated;
+			control.EnabledChanged += ControlEnabledChanged;
+			control.ControlAdded += ControlAdded;
+			control.HandleDestroyed += ControlHandleDestroyed;
 
 			if (control is ICustomThemeHandler themeHandler)
 			{
