-- Fix for Fast Sync --

We need to sync only the leaves from the synctree as we have a partial view of the depot files and cannot trusfully sync any folder.

diff --git a/Engine/Source/Programs/UnrealGameSync/UnrealGameSyncShared/WorkspaceUpdate.cs b/Engine/Source/Programs/UnrealGameSync/UnrealGameSyncShared/WorkspaceUpdate.cs
--- a/Engine/Source/Programs/UnrealGameSync/UnrealGameSyncShared/WorkspaceUpdate.cs
+++ b/Engine/Source/Programs/UnrealGameSync/UnrealGameSyncShared/WorkspaceUpdate.cs
@@ -544,22 +544,13 @@ namespace UnrealGameSync
 
 			private void GetOptimizedSyncCommandsWithoutExcludedFiles(string prefix, int changeNumber, SyncBatchBuilder builder)
 			{
-				if (CanUseWildcard
-				 && IncludedFiles.Any()
-				 && TotalSize < builder.MaxSizePerList)
+				foreach (KeyValuePair<string, long> file in IncludedFiles)
 				{
-					builder.Add($"{prefix}/...@{changeNumber}", TotalSize);
+					builder.Add($"{prefix}/{file.Key}@{changeNumber}", file.Value);
 				}
-				else
+				foreach (KeyValuePair<string, SyncTree> pair in NameToSubTree)
 				{
-					foreach (KeyValuePair<string, long> file in IncludedFiles)
-					{
-						builder.Add($"{prefix}/{file.Key}@{changeNumber}", file.Value);
-					}
-					foreach (KeyValuePair<string, SyncTree> pair in NameToSubTree)
-					{
-						pair.Value.GetOptimizedSyncCommandsWithoutExcludedFiles($"{prefix}/{PerforceUtils.EscapePath(pair.Key)}", changeNumber, builder);
-					}
+					pair.Value.GetOptimizedSyncCommandsWithoutExcludedFiles($"{prefix}/{PerforceUtils.EscapePath(pair.Key)}", changeNumber, builder);
 				}
 			}
 		}
