Fixed a styling issue in DarkTheme with DataGridView.

diff --git a/Engine/Source/Programs/UnrealGameSync/UnrealGameSync/ApplicationThemeManager.cs b/Engine/Source/Programs/UnrealGameSync/UnrealGameSync/ApplicationThemeManager.cs
--- a/Engine/Source/Programs/UnrealGameSync/UnrealGameSync/ApplicationThemeManager.cs
+++ b/Engine/Source/Programs/UnrealGameSync/UnrealGameSync/ApplicationThemeManager.cs
@@ -318,6 +318,11 @@ namespace UnrealGameSync
 				case Panel panel:
 					panel.BorderStyle = BorderStyle.None;
 					break;
+				case DataGridView dataGridView:
+					dataGridView.BackgroundColor = Theme.InputBackgroundColor;
+					dataGridView.DefaultCellStyle.BackColor = Theme.InputBackgroundColor;
+					dataGridView.ForeColor = control.Enabled ? Theme.TextColor : Theme.DimmedText;
+					break;
 			}
 		}
 
