Fixes for combine view and workspace filter related to presets:
- categories values were not locked in the workspace view
- views from presets where not shown in the global and worskpace views

diff --git a/Engine/Source/Programs/UnrealGameSync/UnrealGameSync/Controls/SyncFilterControl.cs b/Engine/Source/Programs/UnrealGameSync/UnrealGameSync/Controls/SyncFilterControl.cs
--- a/Engine/Source/Programs/UnrealGameSync/UnrealGameSync/Controls/SyncFilterControl.cs
+++ b/Engine/Source/Programs/UnrealGameSync/UnrealGameSync/Controls/SyncFilterControl.cs
@@ -1,14 +1,34 @@
 // Copyright Epic Games, Inc. All Rights Reserved.
 
+using System;
 using System.Collections.Generic;
+using System.Collections.ObjectModel;
+using System.ComponentModel;
 using System.Data;
+using System.Drawing;
 using System.Linq;
 using System.Windows.Forms;
+using static System.Windows.Forms.VisualStyles.VisualStyleElement;
 
 namespace UnrealGameSync.Controls
 {
 	public partial class SyncFilterControl : UserControl
 	{
+		private class ViewItem
+		{
+			public bool Locked { get; set; } = false;
+			public string Value { get; set; } = String.Empty;
+
+			public ViewItem(bool locked, string value)
+			{
+				Locked = locked;
+				Value = value;
+			}
+		}
+
+		private BindingList<ViewItem> _views = new BindingList<ViewItem>();
+
+
 		static readonly HashSet<string> s_stripLines = new HashSet<string>()
 		{
 			"; Rules are specified one per line, and may use any standard Perforce wildcards:",
@@ -23,6 +43,49 @@ namespace UnrealGameSync.Controls
 		public SyncFilterControl()
 		{
 			InitializeComponent();
+			ViewDataGrid.RowHeadersVisible = false;
+			ViewDataGrid.ColumnHeadersVisible = false;
+			ViewDataGrid.AutoGenerateColumns = true;
+			ViewDataGrid.EditMode = DataGridViewEditMode.EditOnEnter;
+			ViewDataGrid.AllowUserToAddRows = true;
+			ViewDataGrid.DefaultValuesNeeded += ViewDataGrid_DefaultValuesNeeded;
+			ViewDataGrid.CellBeginEdit += ViewDataGrid_CellBeginEdit;
+			ViewDataGrid.UserDeletingRow += ViewDataGrid_UserDeletingRow;
+
+			_views.AddingNew += new AddingNewEventHandler(Views_AddingNew);
+
+			ViewDataGrid.DataSource = _views;
+		}
+
+		void Views_AddingNew(object?sender, AddingNewEventArgs e)
+		{
+			e.NewObject = new ViewItem(false, "change me");
+		}
+
+		private void ViewDataGrid_CellBeginEdit(object? sender, DataGridViewCellCancelEventArgs e)
+		{
+			if (e.RowIndex >= 0 && e.RowIndex < ViewDataGrid.Rows.Count)
+			{
+				DataGridViewRow? row = ViewDataGrid.Rows[e.RowIndex];
+				if (row.DataBoundItem is ViewItem item)
+				{
+					e.Cancel = item.Locked;
+				}
+			}
+		}
+
+		private void ViewDataGrid_UserDeletingRow(object? sender, DataGridViewRowCancelEventArgs e)
+		{
+			if (e.Row != null && e.Row.DataBoundItem is ViewItem item)
+			{
+				e.Cancel = item.Locked;
+			}
+		}
+
+		private void ViewDataGrid_DefaultValuesNeeded(object? sender, System.Windows.Forms.DataGridViewRowEventArgs e)
+		{
+			e.Row.Cells[nameof(ViewItem.Locked)].Value = false;
+			e.Row.Cells[nameof(ViewItem.Value)].Value = String.Empty;
 		}
 
 		private void SyntaxButton_LinkClicked(object sender, LinkLabelLinkClickedEventArgs e)
@@ -31,15 +94,66 @@ namespace UnrealGameSync.Controls
 			dialog.ShowDialog(ParentForm);
 		}
 
-		public void SetView(string[] view)
+		public void SetView(string preset, string[] presetViews, string[] views)
 		{
-			ViewTextBox.Lines = view.Where(x => !s_stripLines.Contains(x.Trim())).SkipWhile(x => x.Trim().Length == 0 || x.Trim() == ";").ToArray();
-			ViewTextBox.Select(ViewTextBox.Text.Length, 0);
+			if (!String.IsNullOrWhiteSpace(preset))
+			{
+				foreach (string presetView in presetViews)
+				{
+					_views.Add(new ViewItem(true,$"{presetView} [Preset: {preset} - Locked]"));
+				}
+			}
+
+			foreach (string view in views)
+			{
+				_views.Add(new ViewItem(false, view));
+			}
+
+			foreach (DataGridViewColumn column in ViewDataGrid.Columns)
+			{
+				if (column.Name.Equals(nameof(ViewItem.Locked)))
+				{
+					column.Visible = false;
+				}
+
+				if (column.Name.Equals(nameof(ViewItem.Value)))
+				{
+					column.AutoSizeMode = DataGridViewAutoSizeColumnMode.Fill;
+				}
+			}
+
+			foreach (DataGridViewRow? row in ViewDataGrid.Rows)
+			{
+				if (row == null) continue;
+
+				if (row.DataBoundItem is ViewItem item)
+				{
+					if (item.Locked)
+					{
+						row.ReadOnly = true;
+						foreach (DataGridViewCell cell in row.Cells)
+						{
+							cell.ReadOnly = true;
+						}
+					}
+				}
+			}
 		}
 
 		public string[] GetView()
 		{
-			return ViewTextBox.Lines;
+			List<string> views = new List<string>();
+			foreach (ViewItem item in _views)
+			{
+				if (item.Locked)
+				{
+					continue;
+				}
+
+				views.Add(item.Value);
+			}
+
+			return views.ToArray();
 		}
 	}
 }
diff --git a/Engine/Source/Programs/UnrealGameSync/UnrealGameSync/Forms/SyncFilter.cs b/Engine/Source/Programs/UnrealGameSync/UnrealGameSync/Forms/SyncFilter.cs
--- a/Engine/Source/Programs/UnrealGameSync/UnrealGameSync/Forms/SyncFilter.cs
+++ b/Engine/Source/Programs/UnrealGameSync/UnrealGameSync/Forms/SyncFilter.cs
@@ -66,19 +66,22 @@ namespace UnrealGameSync
 			Preset? role = new Preset();
 
 			_roles.TryGetValue(_roleName, out role);
-			
+			List<string> roleViews = role != null ? role.Views.ToList() : new List<string>();
+
 			Dictionary<Guid, bool> syncCategories = WorkspaceSyncCategory.GetDefault(_uniqueIdToCategory.Values);
 
 			WorkspaceSyncCategory.ApplyDelta(syncCategories, GlobalFilter.GetCategories());
-			GlobalControl.SetView(GlobalFilter.View.ToArray());
+			GlobalControl.SetView(_roleName, roleViews.ToArray(), GlobalFilter.View.ToArray());
 			SetExcludedCategories(GlobalControl.CategoriesCheckList, _uniqueIdToCategory, role, syncCategories);
 			GlobalControl.SyncAllProjects.Checked = GlobalFilter.AllProjects ?? false;
 			GlobalControl.IncludeAllProjectsInSolution.Checked = GlobalFilter.AllProjectsInSln ?? false;
 			GlobalControl.GenerateUprojectSpecificSolution.Checked = GlobalFilter.UprojectSpecificSln ?? false;
 
 			WorkspaceSyncCategory.ApplyDelta(syncCategories, WorkspaceFilter.GetCategories());
-			WorkspaceControl.SetView(WorkspaceFilter.View.ToArray());
+			WorkspaceControl.SetView(_roleName, roleViews.ToArray(), WorkspaceFilter.View.ToArray());
 			SetExcludedCategories(WorkspaceControl.CategoriesCheckList, _uniqueIdToCategory, role, syncCategories);
+
+			WorkspaceControl.CategoriesCheckList.ItemCheck += WorkspaceControl_CategoriesCheckList_ItemCheck;
 			WorkspaceControl.SyncAllProjects.Checked = WorkspaceFilter.AllProjects ?? GlobalFilter.AllProjects ?? false;
 			WorkspaceControl.IncludeAllProjectsInSolution.Checked = WorkspaceFilter.AllProjectsInSln ?? GlobalFilter.AllProjectsInSln ?? false;
 			WorkspaceControl.GenerateUprojectSpecificSolution.Checked = WorkspaceFilter.UprojectSpecificSln ?? GlobalFilter.UprojectSpecificSln ?? false;
@@ -104,6 +107,19 @@ namespace UnrealGameSync
 			WorkspaceControl.CategoriesCheckList.SetItemCheckState(e.Index, e.NewValue);
 		}
 
+		private void WorkspaceControl_CategoriesCheckList_ItemCheck(object? sender, ItemCheckEventArgs e)
+		{
+			// do not allow changing the value of locked categories
+			if (sender is System.Windows.Forms.CheckedListBox listBox)
+			{
+				if (listBox.Items[e.Index] is CheckListItem item && item.Locked)
+				{
+					e.NewValue = e.CurrentValue;
+					return;
+				}
+			}
+		}
+
 		private void GlobalControl_SyncAllProjects_CheckStateChanged(object? sender, EventArgs e)
 		{
 			WorkspaceControl.SyncAllProjects.Checked = GlobalControl.SyncAllProjects.Checked;
