UGS: Fix compiler warnings

diff --git a/Engine/Source/Programs/UnrealGameSync/UnrealGameSync/ApplicationThemeManager.cs b/Engine/Source/Programs/UnrealGameSync/UnrealGameSync/ApplicationThemeManager.cs
--- a/Engine/Source/Programs/UnrealGameSync/UnrealGameSync/ApplicationThemeManager.cs
+++ b/Engine/Source/Programs/UnrealGameSync/UnrealGameSync/ApplicationThemeManager.cs
@@ -11,7 +11,7 @@ namespace UnrealGameSync
 {
 	public sealed class InvertibleImage : IDisposable
 	{
-		bool _disposeImage;
+		readonly bool _disposeImage;
 		Image? _standardImage;
 		Image? _invertedImage;
 
@@ -85,7 +85,9 @@ namespace UnrealGameSync
 
 	public sealed class ApplicationThemeManager : IDisposable
 	{
+#pragma warning disable CA1822 // Make property static
 		public ApplicationTheme Theme => ProgramApplicationContext.GetTheme();
+#pragma warning restore CA1822
 		
 		readonly InvertibleImage _dropListIcon = new InvertibleImage(global::UnrealGameSync.Properties.Resources.DropList, false);
 
diff --git a/Engine/Source/Programs/UnrealGameSync/UnrealGameSync/Controls/CustomListViewControl.cs b/Engine/Source/Programs/UnrealGameSync/UnrealGameSync/Controls/CustomListViewControl.cs
--- a/Engine/Source/Programs/UnrealGameSync/UnrealGameSync/Controls/CustomListViewControl.cs
+++ b/Engine/Source/Programs/UnrealGameSync/UnrealGameSync/Controls/CustomListViewControl.cs
@@ -2,7 +2,6 @@
 
 using System;
 using System.Drawing;
-using System.Runtime.InteropServices;
 using System.Windows.Forms;
 using System.Windows.Forms.VisualStyles;
 
diff --git a/Engine/Source/Programs/UnrealGameSync/UnrealGameSync/Controls/LogSplitContainer.cs b/Engine/Source/Programs/UnrealGameSync/UnrealGameSync/Controls/LogSplitContainer.cs
--- a/Engine/Source/Programs/UnrealGameSync/UnrealGameSync/Controls/LogSplitContainer.cs
+++ b/Engine/Source/Programs/UnrealGameSync/UnrealGameSync/Controls/LogSplitContainer.cs
@@ -26,7 +26,7 @@ namespace UnrealGameSync
 		Color _activeCaption = SystemColors.ActiveCaption;
 		Color _inactiveCaption = SystemColors.InactiveCaption;
 		Color _inactiveCaptionText = SystemColors.InactiveCaptionText;
-		InvertibleImage _logIcon = new InvertibleImage(Properties.Resources.Log, false);
+		readonly InvertibleImage _logIcon = new InvertibleImage(Properties.Resources.Log, false);
 
 		public event Action<bool>? OnVisibilityChanged;
 
diff --git a/Engine/Source/Programs/UnrealGameSync/UnrealGameSync/Controls/SyncFilterControl.cs b/Engine/Source/Programs/UnrealGameSync/UnrealGameSync/Controls/SyncFilterControl.cs
--- a/Engine/Source/Programs/UnrealGameSync/UnrealGameSync/Controls/SyncFilterControl.cs
+++ b/Engine/Source/Programs/UnrealGameSync/UnrealGameSync/Controls/SyncFilterControl.cs
@@ -2,13 +2,8 @@
 
 using System;
 using System.Collections.Generic;
-using System.Collections.ObjectModel;
 using System.ComponentModel;
-using System.Data;
-using System.Drawing;
-using System.Linq;
 using System.Windows.Forms;
-using static System.Windows.Forms.VisualStyles.VisualStyleElement;
 
 namespace UnrealGameSync.Controls
 {
@@ -26,19 +21,18 @@ namespace UnrealGameSync.Controls
 			}
 		}
 
-		private BindingList<ViewItem> _views = new BindingList<ViewItem>();
+		private readonly BindingList<ViewItem> _views = new BindingList<ViewItem>();
 
-
-		static readonly HashSet<string> s_stripLines = new HashSet<string>()
-		{
-			"; Rules are specified one per line, and may use any standard Perforce wildcards:",
-			";    ?    Matches one character.",
-			";    *    Matches any sequence of characters, except for a directory separator.",
-			";    ...  Matches any sequence of characters, including directory separators.",
-			"; Patterns may match any file fragment (eg. *.pdb), or may be rooted to the branch (eg. /Engine/Binaries/.../*.pdb).",
-			"; To exclude files which match a pattern, prefix the rule with the '-' character (eg. -/Engine/Documentation/...)",
-			"; Global rules are applied to the files being synced first, followed by any workspace-specific patterns.",
-		};
+		// static readonly HashSet<string> s_stripLines = new HashSet<string>()
+		// {
+		// 	"; Rules are specified one per line, and may use any standard Perforce wildcards:",
+		// 	";    ?    Matches one character.",
+		// 	";    *    Matches any sequence of characters, except for a directory separator.",
+		// 	";    ...  Matches any sequence of characters, including directory separators.",
+		// 	"; Patterns may match any file fragment (eg. *.pdb), or may be rooted to the branch (eg. /Engine/Binaries/.../*.pdb).",
+		// 	"; To exclude files which match a pattern, prefix the rule with the '-' character (eg. -/Engine/Documentation/...)",
+		// 	"; Global rules are applied to the files being synced first, followed by any workspace-specific patterns.",
+		// };
 
 		public SyncFilterControl()
 		{
@@ -111,12 +105,12 @@ namespace UnrealGameSync.Controls
 
 			foreach (DataGridViewColumn column in ViewDataGrid.Columns)
 			{
-				if (column.Name.Equals(nameof(ViewItem.Locked)))
+				if (column.Name.Equals(nameof(ViewItem.Locked), StringComparison.Ordinal))
 				{
 					column.Visible = false;
 				}
 
-				if (column.Name.Equals(nameof(ViewItem.Value)))
+				if (column.Name.Equals(nameof(ViewItem.Value), StringComparison.Ordinal))
 				{
 					column.AutoSizeMode = DataGridViewAutoSizeColumnMode.Fill;
 				}
@@ -124,7 +118,10 @@ namespace UnrealGameSync.Controls
 
 			foreach (DataGridViewRow? row in ViewDataGrid.Rows)
 			{
-				if (row == null) continue;
+				if (row == null)
+				{
+					continue;
+				}
 
 				if (row.DataBoundItem is ViewItem item)
 				{
diff --git a/Engine/Source/Programs/UnrealGameSync/UnrealGameSync/Forms/FailedToDeleteWindow.cs b/Engine/Source/Programs/UnrealGameSync/UnrealGameSync/Forms/FailedToDeleteWindow.cs
--- a/Engine/Source/Programs/UnrealGameSync/UnrealGameSync/Forms/FailedToDeleteWindow.cs
+++ b/Engine/Source/Programs/UnrealGameSync/UnrealGameSync/Forms/FailedToDeleteWindow.cs
@@ -1,7 +1,5 @@
 // Copyright Epic Games, Inc. All Rights Reserved.
 
-using System.Windows.Forms;
-
 namespace UnrealGameSync
 {
 	public partial class FailedToDeleteWindow : ThemedForm
diff --git a/Engine/Source/Programs/UnrealGameSync/UnrealGameSync/Forms/LauncherUpdateWindow.cs b/Engine/Source/Programs/UnrealGameSync/UnrealGameSync/Forms/LauncherUpdateWindow.cs
--- a/Engine/Source/Programs/UnrealGameSync/UnrealGameSync/Forms/LauncherUpdateWindow.cs
+++ b/Engine/Source/Programs/UnrealGameSync/UnrealGameSync/Forms/LauncherUpdateWindow.cs
@@ -1,7 +1,5 @@
 // Copyright Epic Games, Inc. All Rights Reserved.
 
-using System.Windows.Forms;
-
 namespace UnrealGameSync.Forms
 {
 	public partial class LauncherUpdateWindow : ThemedForm
diff --git a/Engine/Source/Programs/UnrealGameSync/UnrealGameSync/Forms/LeaveCommentWindow.cs b/Engine/Source/Programs/UnrealGameSync/UnrealGameSync/Forms/LeaveCommentWindow.cs
--- a/Engine/Source/Programs/UnrealGameSync/UnrealGameSync/Forms/LeaveCommentWindow.cs
+++ b/Engine/Source/Programs/UnrealGameSync/UnrealGameSync/Forms/LeaveCommentWindow.cs
@@ -1,7 +1,5 @@
 // Copyright Epic Games, Inc. All Rights Reserved.
 
-using System.Windows.Forms;
-
 namespace UnrealGameSync
 {
 	public partial class LeaveCommentWindow : ThemedForm
diff --git a/Engine/Source/Programs/UnrealGameSync/UnrealGameSync/Forms/LogWindow.cs b/Engine/Source/Programs/UnrealGameSync/UnrealGameSync/Forms/LogWindow.cs
--- a/Engine/Source/Programs/UnrealGameSync/UnrealGameSync/Forms/LogWindow.cs
+++ b/Engine/Source/Programs/UnrealGameSync/UnrealGameSync/Forms/LogWindow.cs
@@ -1,7 +1,5 @@
 // Copyright Epic Games, Inc. All Rights Reserved.
 
-using System.Windows.Forms;
-
 namespace UnrealGameSync
 {
 	public partial class LogWindow : ThemedForm
diff --git a/Engine/Source/Programs/UnrealGameSync/UnrealGameSync/Forms/NotificationWindow.cs b/Engine/Source/Programs/UnrealGameSync/UnrealGameSync/Forms/NotificationWindow.cs
--- a/Engine/Source/Programs/UnrealGameSync/UnrealGameSync/Forms/NotificationWindow.cs
+++ b/Engine/Source/Programs/UnrealGameSync/UnrealGameSync/Forms/NotificationWindow.cs
@@ -20,7 +20,7 @@ namespace UnrealGameSync
 
 		// Logo
 		Rectangle _logoBounds;
-		InvertibleImage _logoImage;
+		readonly InvertibleImage _logoImage;
 
 		// Caption
 		Rectangle _captionBounds;
diff --git a/Engine/Source/Programs/UnrealGameSync/UnrealGameSync/Forms/PasswordWindow.cs b/Engine/Source/Programs/UnrealGameSync/UnrealGameSync/Forms/PasswordWindow.cs
--- a/Engine/Source/Programs/UnrealGameSync/UnrealGameSync/Forms/PasswordWindow.cs
+++ b/Engine/Source/Programs/UnrealGameSync/UnrealGameSync/Forms/PasswordWindow.cs
@@ -1,7 +1,6 @@
 ﻿// Copyright Epic Games, Inc. All Rights Reserved.
 
 using System;
-using System.Windows.Forms;
 
 namespace UnrealGameSync
 {
diff --git a/Engine/Source/Programs/UnrealGameSync/UnrealGameSync/Forms/PerforceSyncSettingsWindow.cs b/Engine/Source/Programs/UnrealGameSync/UnrealGameSync/Forms/PerforceSyncSettingsWindow.cs
--- a/Engine/Source/Programs/UnrealGameSync/UnrealGameSync/Forms/PerforceSyncSettingsWindow.cs
+++ b/Engine/Source/Programs/UnrealGameSync/UnrealGameSync/Forms/PerforceSyncSettingsWindow.cs
@@ -1,7 +1,6 @@
 // Copyright Epic Games, Inc. All Rights Reserved.
 
 using System;
-using System.Windows.Forms;
 using Microsoft.Extensions.Logging;
 
 namespace UnrealGameSync
diff --git a/Engine/Source/Programs/UnrealGameSync/UnrealGameSync/Forms/SyncFilterSyntax.cs b/Engine/Source/Programs/UnrealGameSync/UnrealGameSync/Forms/SyncFilterSyntax.cs
--- a/Engine/Source/Programs/UnrealGameSync/UnrealGameSync/Forms/SyncFilterSyntax.cs
+++ b/Engine/Source/Programs/UnrealGameSync/UnrealGameSync/Forms/SyncFilterSyntax.cs
@@ -1,7 +1,5 @@
 // Copyright Epic Games, Inc. All Rights Reserved.
 
-using System.Windows.Forms;
-
 namespace UnrealGameSync
 {
 	public partial class SyncFilterSyntax : ThemedForm
