Performance improvements for UGS dark theme.

diff --git a/Engine/Source/Programs/UnrealGameSync/UnrealGameSync/ApplicationThemeManager.cs b/Engine/Source/Programs/UnrealGameSync/UnrealGameSync/ApplicationThemeManager.cs
--- a/Engine/Source/Programs/UnrealGameSync/UnrealGameSync/ApplicationThemeManager.cs
+++ b/Engine/Source/Programs/UnrealGameSync/UnrealGameSync/ApplicationThemeManager.cs
@@ -1,7 +1,9 @@
 // Copyright Epic Games, Inc. All Rights Reserved.
 
 using System;
+using System.Collections.Generic;
 using System.Drawing;
+using System.Linq;
 using System.Runtime.InteropServices;
 using System.Windows.Forms;
 using Microsoft.Win32;
@@ -194,20 +196,64 @@ namespace UnrealGameSync
 		private void ControlHandleDestroyed(object? sender, EventArgs args)
 		{
 			if (sender is Control control)
+			{
+				UnitialiseControl(control);
+			}
+		}
+
+		private void ToolStripDropDownClosed(object? sender, EventArgs args)
+		{
+			if (sender is ToolStripDropDown menu)
+			{
+				UninitialiseToolStripDropDown(menu);
+			}
+		}
+
+		private void UninitialiseToolStripDropDown(ToolStripDropDown dropDown)
+		{
+			dropDown.Closed -= ToolStripDropDownClosed;
+			UnitialiseControl(dropDown);
+
+			foreach (ToolStripDropDownItem item in dropDown.Items.OfType<ToolStripDropDownItem>())
+			{
+				UninitialiseToolStripDropDown(item.DropDown);
+			}
+		}
+
+		private void UnitialiseControl(Control control)
+		{
+			if (_registeredControls.Remove(control))
 			{
 				control.HandleCreated -= ControlHandleCreated;
 				control.EnabledChanged -= ControlEnabledChanged;
 				control.ControlAdded -= ControlAdded;
 				control.HandleDestroyed -= ControlHandleDestroyed;
+
+				foreach (Control childControl in control.Controls)
+				{
+					UnitialiseControl(childControl);
+				}
 			}
 		}
 
+		private readonly HashSet<Control> _registeredControls = new HashSet<Control>();
+
 		public void ApplyThemeRecursively(Control control)
 		{
-			control.HandleCreated += ControlHandleCreated;
-			control.EnabledChanged += ControlEnabledChanged;
-			control.ControlAdded += ControlAdded;
-			control.HandleDestroyed += ControlHandleDestroyed;
+			if ( _registeredControls.Contains(control) == false )
+			{
+				_registeredControls.Add(control);
+
+				control.HandleCreated += ControlHandleCreated;
+				control.ControlAdded += ControlAdded;
+				control.HandleDestroyed += ControlHandleDestroyed;
+				control.EnabledChanged += ControlEnabledChanged;
+
+				if (control is ToolStripDropDown menu)
+				{
+					menu.Closed += ToolStripDropDownClosed;
+				}
+			}
 
 			if (control is ICustomThemeHandler themeHandler)
 			{
