Fixed a bug in the Sync Filter view where deleting custom view entries would result in an error when clicking the okay button.

diff --git a/Engine/Source/Programs/UnrealGameSync/UnrealGameSync/Controls/SyncFilterControl.cs b/Engine/Source/Programs/UnrealGameSync/UnrealGameSync/Controls/SyncFilterControl.cs
--- a/Engine/Source/Programs/UnrealGameSync/UnrealGameSync/Controls/SyncFilterControl.cs
+++ b/Engine/Source/Programs/UnrealGameSync/UnrealGameSync/Controls/SyncFilterControl.cs
@@ -147,7 +147,10 @@ namespace UnrealGameSync.Controls
 					continue;
 				}
 
-				views.Add(item.Value);
+				if ( String.IsNullOrEmpty( item.Value ) == false )
+				{
+					views.Add(item.Value);
+				}
 			}
 
 			return views.ToArray();
