Fixing the foreground colour of tree views with the new dark theme styling.

diff --git a/Engine/Source/Programs/UnrealGameSync/UnrealGameSync/ApplicationThemeManager.cs b/Engine/Source/Programs/UnrealGameSync/UnrealGameSync/ApplicationThemeManager.cs
--- a/Engine/Source/Programs/UnrealGameSync/UnrealGameSync/ApplicationThemeManager.cs
+++ b/Engine/Source/Programs/UnrealGameSync/UnrealGameSync/ApplicationThemeManager.cs
@@ -366,6 +366,9 @@ namespace UnrealGameSync
 				case Panel panel:
 					panel.BorderStyle = BorderStyle.None;
 					break;
+				case TreeView:
+					control.ForeColor = control.Enabled ? Theme.TextColor : Theme.DimmedText;
+					break;
 				case DataGridView dataGridView:
 					dataGridView.BackgroundColor = Theme.InputBackgroundColor;
 					dataGridView.DefaultCellStyle.BackColor = Theme.InputBackgroundColor;
