Fixed a bug where the SDK Info text colour in light theme would be faded.

diff --git a/Engine/Source/Programs/UnrealGameSync/UnrealGameSync/Forms/SdkInfoWindow.cs b/Engine/Source/Programs/UnrealGameSync/UnrealGameSync/Forms/SdkInfoWindow.cs
--- a/Engine/Source/Programs/UnrealGameSync/UnrealGameSync/Forms/SdkInfoWindow.cs
+++ b/Engine/Source/Programs/UnrealGameSync/UnrealGameSync/Forms/SdkInfoWindow.cs
@@ -244,7 +244,7 @@ namespace UnrealGameSync
 			}
 			else if (e.ColumnIndex != columnHeader3.Index)
 			{
-				TextRenderer.DrawText(e.Graphics, e.SubItem.Text, SdkListView.Font, e.Bounds, SdkListView.ForeColor, TextFormatFlags.EndEllipsis | TextFormatFlags.SingleLine | TextFormatFlags.Left | TextFormatFlags.VerticalCenter | TextFormatFlags.NoPrefix);
+				TextRenderer.DrawText(e.Graphics, e.SubItem.Text, SdkListView.Font, e.Bounds, Theme.TextColor, TextFormatFlags.EndEllipsis | TextFormatFlags.SingleLine | TextFormatFlags.Left | TextFormatFlags.VerticalCenter | TextFormatFlags.NoPrefix);
 			}
 			else
 			{
