Fixed a bug where the column header text in the Customize Commands dialog was faded in Light theme.

diff --git a/Engine/Source/Programs/UnrealGameSync/UnrealGameSync/Forms/ModifyBuildStepsWindow.cs b/Engine/Source/Programs/UnrealGameSync/UnrealGameSync/Forms/ModifyBuildStepsWindow.cs
--- a/Engine/Source/Programs/UnrealGameSync/UnrealGameSync/Forms/ModifyBuildStepsWindow.cs
+++ b/Engine/Source/Programs/UnrealGameSync/UnrealGameSync/Forms/ModifyBuildStepsWindow.cs
@@ -135,7 +135,7 @@ namespace UnrealGameSync
 			}
 			else
 			{
-				TextRenderer.DrawText(e.Graphics, e.Header.Text, BuildStepList.Font, e.Bounds, ForeColor, TextFormatFlags.Left | TextFormatFlags.VerticalCenter);
+				TextRenderer.DrawText(e.Graphics, e.Header.Text, BuildStepList.Font, e.Bounds, Theme.TextColor, TextFormatFlags.Left | TextFormatFlags.VerticalCenter);
 			}
 		}
 
