Added option to UGS to toggle whether to use the perforce native dll.

diff --git a/Engine/Source/Programs/UnrealGameSync/UnrealGameSync/Forms/PerforceSyncSettingsWindow.cs b/Engine/Source/Programs/UnrealGameSync/UnrealGameSync/Forms/PerforceSyncSettingsWindow.cs
--- a/Engine/Source/Programs/UnrealGameSync/UnrealGameSync/Forms/PerforceSyncSettingsWindow.cs
+++ b/Engine/Source/Programs/UnrealGameSync/UnrealGameSync/Forms/PerforceSyncSettingsWindow.cs
@@ -28,6 +28,7 @@ namespace UnrealGameSync
 			numericUpDownSyncErrorRetryDelay.Value = syncOptions.SyncErrorRetryDelay ?? PerforceSyncOptions.DefaultSyncErrorRetryDelay;
 
 			FastSyncCheckBox.Checked = syncOptions.UseFastSync ?? false;
+			UseNativeLibraryCheckBox.Checked = syncOptions.UseNativeClient ?? PerforceSyncOptions.DefaultUseNativeLibrary;
 		}
 
 		private void OkButton_Click(object sender, EventArgs e)
@@ -37,6 +38,7 @@ namespace UnrealGameSync
 			_settings.SyncOptions.NumSyncErrorRetries = GetValueIfNotDefault((int)numericUpDownRetriesOnSyncError.Value, PerforceSyncOptions.DefaultNumSyncErrorRetries);
 			_settings.SyncOptions.SyncErrorRetryDelay = GetValueIfNotDefault((int)numericUpDownSyncErrorRetryDelay.Value, PerforceSyncOptions.DefaultSyncErrorRetryDelay);
 			_settings.SyncOptions.UseFastSync = FastSyncCheckBox.Checked;
+			_settings.SyncOptions.UseNativeClient = UseNativeLibraryCheckBox.Checked;
 			_settings.Save(_logger);
 
 			DialogResult = System.Windows.Forms.DialogResult.OK;
@@ -61,6 +63,7 @@ namespace UnrealGameSync
 			numericUpDownRetriesOnSyncError.Value = PerforceSyncOptions.DefaultNumSyncErrorRetries;
 			numericUpDownSyncErrorRetryDelay.Value = PerforceSyncOptions.DefaultSyncErrorRetryDelay;
 			FastSyncCheckBox.Checked = PerforceSyncOptions.DefaultUseFastSync;
+			UseNativeLibraryCheckBox.Checked = PerforceSyncOptions.DefaultUseNativeLibrary;
 		}
 	}
 }
diff --git a/Engine/Source/Programs/UnrealGameSync/UnrealGameSync/ProgramApplicationContext.cs b/Engine/Source/Programs/UnrealGameSync/UnrealGameSync/ProgramApplicationContext.cs
--- a/Engine/Source/Programs/UnrealGameSync/UnrealGameSync/ProgramApplicationContext.cs
+++ b/Engine/Source/Programs/UnrealGameSync/UnrealGameSync/ProgramApplicationContext.cs
@@ -130,6 +130,8 @@ namespace UnrealGameSync
 				_settings.Save(_logger);
 			}
 
+			((PerforceSettings)_defaultPerforceSettings).PreferNativeClient = _settings.SyncOptions.UseNativeClient ?? PerforceSyncOptions.DefaultUseNativeLibrary;
+
 			// Register the update listener
 			updateMonitor.OnUpdateAvailable += OnUpdateAvailableCallback;
 
diff --git a/Engine/Source/Programs/UnrealGameSync/UnrealGameSyncShared/WorkspaceUpdate.cs b/Engine/Source/Programs/UnrealGameSync/UnrealGameSyncShared/WorkspaceUpdate.cs
--- a/Engine/Source/Programs/UnrealGameSync/UnrealGameSyncShared/WorkspaceUpdate.cs
+++ b/Engine/Source/Programs/UnrealGameSync/UnrealGameSyncShared/WorkspaceUpdate.cs
@@ -70,6 +70,7 @@ namespace UnrealGameSync
 		public const int DefaultSyncErrorRetryDelay = 0;
 
 		public const bool DefaultUseFastSync = false;
+		public const bool DefaultUseNativeLibrary = true;
 
 		public int? NumThreads { get; set; }
 
@@ -81,6 +82,8 @@ namespace UnrealGameSync
 
 		public bool? UseFastSync { get; set; }
 
+		public bool? UseNativeClient { get; set; }
+
 		public PerforceSyncOptions Clone()
 		{
 			return (PerforceSyncOptions)MemberwiseClone();
