Disabled automatic setting of the UGS theme.

diff --git a/Engine/Source/Programs/UnrealGameSync/UnrealGameSync/ApplicationThemeManager.cs b/Engine/Source/Programs/UnrealGameSync/UnrealGameSync/ApplicationThemeManager.cs
--- a/Engine/Source/Programs/UnrealGameSync/UnrealGameSync/ApplicationThemeManager.cs
+++ b/Engine/Source/Programs/UnrealGameSync/UnrealGameSync/ApplicationThemeManager.cs
@@ -95,7 +95,7 @@ namespace UnrealGameSync
 
 		public ApplicationThemeManager()
 		{
-			UpdateDefaultTheme();
+			//UpdateDefaultTheme();
 		}
 		
 		public void Dispose()
