Updated the Light Theme badge text to be White.

diff --git a/Engine/Source/Programs/UnrealGameSync/UnrealGameSync/Controls/WorkspaceControl.cs b/Engine/Source/Programs/UnrealGameSync/UnrealGameSync/Controls/WorkspaceControl.cs
--- a/Engine/Source/Programs/UnrealGameSync/UnrealGameSync/Controls/WorkspaceControl.cs
+++ b/Engine/Source/Programs/UnrealGameSync/UnrealGameSync/Controls/WorkspaceControl.cs
@@ -3273,7 +3273,7 @@ namespace UnrealGameSync
 					}
 				}
 
-				Color textColor = ApplicationTheme.IsColorLight(badgeColor) ? Color.Black : Color.White;
+				Color textColor = Color.White;
 				TextRenderer.DrawText(graphics, badgeText, _badgeFont, badgeRect, textColor, TextFormatFlags.HorizontalCenter | TextFormatFlags.VerticalCenter | TextFormatFlags.SingleLine | TextFormatFlags.NoPrefix | TextFormatFlags.PreserveGraphicsClipping);
 			}
 		}
