Added a quick change to remove the requirement for a restart of UGS after toggling the UseNativeClient perforce option.

diff --git a/Engine/Source/Programs/UnrealGameSync/UnrealGameSync/Controls/WorkspaceControl.cs b/Engine/Source/Programs/UnrealGameSync/UnrealGameSync/Controls/WorkspaceControl.cs
--- a/Engine/Source/Programs/UnrealGameSync/UnrealGameSync/Controls/WorkspaceControl.cs
+++ b/Engine/Source/Programs/UnrealGameSync/UnrealGameSync/Controls/WorkspaceControl.cs
@@ -3032,6 +3032,13 @@ namespace UnrealGameSync
 		{
 			UpdateBuildSteps();
 			UpdateStatusPanel();
+			UpdatePerforceSettings();
+		}
+
+		// Note this should be changed to be done on a global level as it's unecessary to do this per workspace tab
+		private void UpdatePerforceSettings()
+		{
+			((PerforceSettings)_perforceSettings).PreferNativeClient = _settings.SyncOptions.UseNativeClient ?? true;
 		}
 
 		[return: NotNullIfNotNull("text")]
