Fixing up the badge text colour for dark theme.

diff --git a/Engine/Source/Programs/UnrealGameSync/UnrealGameSync/Controls/WorkspaceControl.cs b/Engine/Source/Programs/UnrealGameSync/UnrealGameSync/Controls/WorkspaceControl.cs
--- a/Engine/Source/Programs/UnrealGameSync/UnrealGameSync/Controls/WorkspaceControl.cs
+++ b/Engine/Source/Programs/UnrealGameSync/UnrealGameSync/Controls/WorkspaceControl.cs
@@ -3281,6 +3281,13 @@ namespace UnrealGameSync
 				}
 
 				Color textColor = Color.White;
+
+				ApplicationThemeManager? manager = ProgramApplicationContext.GetThemeManager();
+				if ( manager != null && manager.Theme.UseDarkMode )
+				{
+					textColor = ApplicationTheme.IsColorLight(badgeColor) ? Color.Black : Color.White;
+				}
+				
 				TextRenderer.DrawText(graphics, badgeText, _badgeFont, badgeRect, textColor, TextFormatFlags.HorizontalCenter | TextFormatFlags.VerticalCenter | TextFormatFlags.SingleLine | TextFormatFlags.NoPrefix | TextFormatFlags.PreserveGraphicsClipping);
 			}
 		}
