Fixed a bug where disabled editable combo boxes would appear with highlighted text.

diff --git a/Engine/Source/Programs/UnrealGameSync/UnrealGameSync/ApplicationThemeManager.cs b/Engine/Source/Programs/UnrealGameSync/UnrealGameSync/ApplicationThemeManager.cs
--- a/Engine/Source/Programs/UnrealGameSync/UnrealGameSync/ApplicationThemeManager.cs
+++ b/Engine/Source/Programs/UnrealGameSync/UnrealGameSync/ApplicationThemeManager.cs
@@ -337,9 +337,15 @@ namespace UnrealGameSync
 					control.ForeColor = control.Enabled ? Theme.TextColor : Theme.DimmedText;
 					textBox.BorderStyle = Theme.UseFlatControls ? BorderStyle.FixedSingle : BorderStyle.Fixed3D;
 					break;
-				case ComboBox:
+				case ComboBox comboBox:
 					control.BackColor = Theme.InputBackgroundColor;
 					control.ForeColor = control.Enabled ? Theme.TextColor : Theme.DimmedText;
+
+					if (control.Focused == false)
+					{
+						comboBox.SelectionLength = 0;
+					}
+
 					break;
 				case ReadOnlyCheckbox readOnlyCheckbox:
 					readOnlyCheckbox.ForeColor = readOnlyCheckbox.ReadOnly ? Theme.DimmedText : Theme.TextColor;
