Fixed a bug where if no project tabs have been created in UGS an exception is thrown on mouse move.

diff --git a/Engine/Source/Programs/UnrealGameSync/UnrealGameSync/Controls/TabControl.cs b/Engine/Source/Programs/UnrealGameSync/UnrealGameSync/Controls/TabControl.cs
--- a/Engine/Source/Programs/UnrealGameSync/UnrealGameSync/Controls/TabControl.cs
+++ b/Engine/Source/Programs/UnrealGameSync/UnrealGameSync/Controls/TabControl.cs
@@ -469,12 +469,16 @@ namespace UnrealGameSync
 				}
 			}
 
-			bool isNowHoveringCloseBtn = _hoverTabIdx == _selectedTabIdx
-			                      && ((e.Location.X) > (_tabs[_hoverTabIdx]._minX + _tabs[_hoverTabIdx]._width - _tabs[_hoverTabIdx]._closeButtonWidth));
-			if (_isHoveringCloseBtn != isNowHoveringCloseBtn)
+
+			if (_selectedTabIdx != -1)
 			{
-				_isHoveringCloseBtn = isNowHoveringCloseBtn;
-				Invalidate();
+				bool isNowHoveringCloseBtn = _hoverTabIdx == _selectedTabIdx
+								  && ((e.Location.X) > (_tabs[_hoverTabIdx]._minX + _tabs[_hoverTabIdx]._width - _tabs[_hoverTabIdx]._closeButtonWidth));
+				if (_isHoveringCloseBtn != isNowHoveringCloseBtn)
+				{
+					_isHoveringCloseBtn = isNowHoveringCloseBtn;
+					Invalidate();
+				}
 			}
 		}
 
