Fixed a bug where the item text in the UGS Editor Commandline Arguments window was difficult to read in light theme.

diff --git a/Engine/Source/Programs/UnrealGameSync/UnrealGameSync/Forms/ArgumentsWindow.cs b/Engine/Source/Programs/UnrealGameSync/UnrealGameSync/Forms/ArgumentsWindow.cs
--- a/Engine/Source/Programs/UnrealGameSync/UnrealGameSync/Forms/ArgumentsWindow.cs
+++ b/Engine/Source/Programs/UnrealGameSync/UnrealGameSync/Forms/ArgumentsWindow.cs
@@ -170,7 +170,7 @@ namespace UnrealGameSync
 				CheckBoxState state = e.Item.Checked ? CheckBoxState.CheckedNormal : CheckBoxState.UncheckedNormal;
 				Size checkSize = CheckBoxRenderer.GetGlyphSize(e.Graphics, state);
 				CheckBoxRenderer.DrawCheckBox(e.Graphics, new Point(e.Bounds.Left + 4, e.Bounds.Top + (e.Bounds.Height - checkSize.Height) / 2), state);
-				DrawItemLabel(e.Graphics, ForeColor, e.Item);
+				DrawItemLabel(e.Graphics, Theme.TextColor, e.Item);
 			}
 			else
 			{
