PR #12713: UGS: Faster archive download from Horde

diff --git a/Engine/Source/Programs/UnrealGameSync/UnrealGameSyncShared/ArchiveInfo.cs b/Engine/Source/Programs/UnrealGameSync/UnrealGameSyncShared/ArchiveInfo.cs
--- a/Engine/Source/Programs/UnrealGameSync/UnrealGameSyncShared/ArchiveInfo.cs
+++ b/Engine/Source/Programs/UnrealGameSync/UnrealGameSyncShared/ArchiveInfo.cs
@@ -10,6 +10,8 @@ using System.Threading.Tasks;
 using EpicGames.Core;
 using EpicGames.Horde;
 using EpicGames.Horde.Artifacts;
+using EpicGames.Horde.Storage;
+using EpicGames.Horde.Storage.Nodes;
 using EpicGames.Perforce;
 using Microsoft.Extensions.Logging;
 
@@ -277,34 +279,42 @@ namespace UnrealGameSync
 
 			public async Task<bool> DownloadAsync(IPerforceConnection perforce, DirectoryReference localRootPath, FileReference manifestFileName, ILogger logger, ProgressValue progress, CancellationToken cancellationToken)
 			{
-				DirectoryReference configDir = UserSettings.GetConfigDir(localRootPath);
-				UserSettings.CreateConfigDir(configDir);
-
-				FileReference tempZipFileName = FileReference.Combine(configDir, "archive.zip");
-				try
+				IArtifact? artifact = await _hordeClient.Artifacts.GetAsync(_artifactId, cancellationToken);
+				if (artifact == null)
 				{
-					HordeHttpClient hordeHttpClient = _hordeClient.CreateHttpClient();
+					return false;
+				}
 
-					using (FileStream stream = FileReference.Open(tempZipFileName, FileMode.Create, FileAccess.Write, FileShare.None))
-					{
-						await using Stream sourceStream = await hordeHttpClient.GetArtifactZipAsync(_artifactId, cancellationToken);
-						await sourceStream.CopyToAsync(stream, cancellationToken);
-					}
+				ExtractOptions options = new() { Progress = new ExtractStatsLogger(logger) };
+				await artifact.Content.ExtractAsync(localRootPath.ToDirectoryInfo(), options, logger, cancellationToken);
 
-					if (tempZipFileName.ToFileInfo().Length == 0)
-					{
-						return false;
-					}
-					ArchiveUtils.ExtractFiles(tempZipFileName, localRootPath, manifestFileName, progress, logger);
-				}
-				finally
+				FileReference.Delete(manifestFileName);
+
+				ArchiveManifest manifest = new();
+				await GatherFilesForManifestAsync(String.Empty, artifact.Content, manifest, DateTime.UtcNow, cancellationToken);
+
+				await using (FileStream manifestStream = FileReference.Open(manifestFileName, FileMode.Create, FileAccess.Write))
 				{
-					FileReference.SetAttributes(tempZipFileName, FileAttributes.Normal);
-					FileReference.Delete(tempZipFileName);
+					manifest.Write(manifestStream);
 				}
 
 				return true;
 			}
+
+			private static async Task GatherFilesForManifestAsync(string directoryPath, IBlobRef<DirectoryNode> directoryRef, ArchiveManifest manifest, DateTime timeStamp, CancellationToken cancellationToken)
+			{
+				DirectoryNode directoryNode = await directoryRef.ReadBlobAsync(cancellationToken);
+
+				foreach (FileEntry fileEntry in directoryNode.Files)
+				{
+					manifest.Files.Add(new ArchiveManifestFile(Path.Combine(directoryPath, fileEntry.Name), fileEntry.Length, fileEntry.ModTime));
+				}
+
+				foreach (DirectoryEntry directoryEntry in directoryNode.Directories)
+				{
+					await GatherFilesForManifestAsync(Path.Combine(directoryPath, directoryEntry.Name), directoryEntry.Handle, manifest, timeStamp, cancellationToken);
+				}
+			}
 		}
 
 		public static async Task<List<HordeArchiveChannel>> GetChannelsAsync(IHordeClient hordeClient, string projectIdentifier, CancellationToken cancellationToken)
