Clear the currently synced CL when switching the stream a workspace targets.
This correctly updates the state of UGS to recognise the synced CL has been invalidated and also fixes a bug with Fast Sync where it misses files in the sync.

diff --git a/Engine/Source/Programs/UnrealGameSync/UnrealGameSync/Controls/WorkspaceControl.cs b/Engine/Source/Programs/UnrealGameSync/UnrealGameSync/Controls/WorkspaceControl.cs
--- a/Engine/Source/Programs/UnrealGameSync/UnrealGameSync/Controls/WorkspaceControl.cs
+++ b/Engine/Source/Programs/UnrealGameSync/UnrealGameSync/Controls/WorkspaceControl.cs
@@ -848,6 +848,15 @@ namespace UnrealGameSync
 		private void StreamChanged()
 		{
 			_owner.StreamChanged(this);
+
+			// Reset the current CL we have synced and the current Code CL we have tracked 
+			// since these no longer are accurate now we are targetting a new stream
+			_workspace?.ModifyState(x =>
+			{
+				x.CurrentChangeNumber = -1;
+				x.CurrentCodeChangeNumber = -1;
+			});
+
 			/*
 			StatusPanel.SuspendDisplay();
 
@@ -4467,12 +4476,6 @@ namespace UnrealGameSync
 						StatusPanel.SuspendLayout();
 						StreamChanged();
 
-						// Reset the last code change we found when switching streams
-						_workspace.ModifyState(x =>
-						{
-							x.CurrentCodeChangeNumber = -1;
-						});
-
 						StatusPanel.ResumeLayout();
 					}
 				}
