Add the custom and user build steps to the list of commands to be run by the build command un ugs cli

diff --git a/Engine/Source/Programs/UnrealGameSync/UnrealGameSyncCmd/Program.cs b/Engine/Source/Programs/UnrealGameSync/UnrealGameSyncCmd/Program.cs
--- a/Engine/Source/Programs/UnrealGameSync/UnrealGameSyncCmd/Program.cs
+++ b/Engine/Source/Programs/UnrealGameSync/UnrealGameSyncCmd/Program.cs
@@ -1394,6 +1394,8 @@ namespace UnrealGameSyncCmd
 
 				ProjectInfo projectInfo = new ProjectInfo(settings.RootDir, state.Current);
 
+				UserProjectSettings projectSettings = context.UserSettings.FindOrAddProjectSettings(projectInfo, settings, logger);
+
 				if (options.ListOnly)
 				{
 					ConfigFile projectConfig = await ConfigUtils.ReadProjectConfigFileAsync(perforceClient, projectInfo, logger, CancellationToken.None);
@@ -1402,6 +1404,9 @@ namespace UnrealGameSyncCmd
 
 					Dictionary<Guid, ConfigObject> buildStepObjects = ConfigUtils.GetDefaultBuildStepObjects(projectInfo, editorTarget.GetFileNameWithoutAnyExtensions(), EditorConfig, projectConfig, false);
 
+					BuildStep.MergeBuildStepObjects(buildStepObjects, projectConfig.GetValues("Build.Step", Array.Empty<string>()).Select(x => new ConfigObject(x)));
+					BuildStep.MergeBuildStepObjects(buildStepObjects, projectSettings.BuildSteps);
+
 					logger.LogInformation("Available build steps:");
 					logger.LogInformation("");
 					logger.LogInformation("  Id                                   | Description                              | Type       | Enabled");
@@ -1424,7 +1429,12 @@ namespace UnrealGameSyncCmd
 					steps = new HashSet<Guid> { id };
 				}
 
-				WorkspaceUpdateContext updateContext = new WorkspaceUpdateContext(state.Current.CurrentChangeNumber, WorkspaceUpdateOptions.Build, BuildConfig.Development, null, new List<ConfigObject>(), steps);
+				WorkspaceUpdateContext updateContext = new WorkspaceUpdateContext(
+					state.Current.CurrentChangeNumber,
+					WorkspaceUpdateOptions.Build,
+					BuildConfig.Development,
+					null,
+					projectSettings.BuildSteps, steps);
 
 				WorkspaceUpdate update = new WorkspaceUpdate(updateContext);
 				(WorkspaceUpdateResult result, string message) = await update.ExecuteAsync(perforceClient.Settings, projectInfo, state, context.Logger, CancellationToken.None);
