Add the initialized project to the list of Open and Recent projects in the settings file

diff --git a/Engine/Source/Programs/UnrealGameSync/UnrealGameSyncCmd/Program.cs b/Engine/Source/Programs/UnrealGameSync/UnrealGameSyncCmd/Program.cs
--- a/Engine/Source/Programs/UnrealGameSync/UnrealGameSyncCmd/Program.cs
+++ b/Engine/Source/Programs/UnrealGameSync/UnrealGameSyncCmd/Program.cs
@@ -139,15 +139,19 @@ namespace UnrealGameSyncCmd
 			public CommandLineArguments Arguments { get; }
 			public ILogger Logger { get; }
 			public ILoggerFactory LoggerFactory { get; }
+
+			public UserSettings? GlobalSettings { get; }
+
 			public GlobalSettingsFile UserSettings { get; }
 
 			public IHordeClient? HordeClient { get; }
 
-			public CommandContext(CommandLineArguments arguments, ILogger logger, ILoggerFactory loggerFactory, GlobalSettingsFile userSettings, IHordeClient? hordeClient)
+			public CommandContext(CommandLineArguments arguments, ILogger logger, ILoggerFactory loggerFactory, GlobalSettingsFile userSettings, UserSettings? globalSettings, IHordeClient? hordeClient)
 			{
 				Arguments = arguments;
 				Logger = logger;
 				LoggerFactory = loggerFactory;
+				GlobalSettings = globalSettings;
 				UserSettings = userSettings;
 				HordeClient = hordeClient;
 			}
@@ -254,10 +258,12 @@ namespace UnrealGameSyncCmd
 				ServiceProvider serviceProvider = services.BuildServiceProvider();
 				IHordeClient? hordeClient = serviceProvider.GetService<IHordeClient>();
 
+				UserSettings? globalSettings = null;
 				GlobalSettingsFile settings;
 				if (RuntimeInformation.IsOSPlatform(OSPlatform.Windows))
 				{
-					settings = UserSettings.Create(globalConfigFolder, logger);
+					globalSettings = UserSettings.Create(globalConfigFolder, logger);
+					settings = globalSettings as GlobalSettingsFile;
 				}
 				else
 				{
@@ -328,7 +334,7 @@ namespace UnrealGameSyncCmd
 				}
 
 				Command instance = (Command)Activator.CreateInstance(command.Type)!;
-				await instance.ExecuteAsync(new CommandContext(args, logger, loggerFactory, settings, hordeClient));
+				await instance.ExecuteAsync(new CommandContext(args, logger, loggerFactory, settings, globalSettings, hordeClient));
 				return 0;
 			}
 			catch (UserErrorException ex)
@@ -632,15 +638,15 @@ namespace UnrealGameSyncCmd
 				// Create the perforce connection
 				if (initName != null)
 				{
-					await InitNewClientAsync(perforce, initName, hostName, options, logger);
+					await InitNewClientAsync(perforce, initName, hostName, options, context.GlobalSettings, logger);
 				}
 				else
 				{
-					await InitExistingClientAsync(perforce, hostName, options, logger);
+					await InitExistingClientAsync(perforce, hostName, options, context.GlobalSettings, logger);
 				}
 			}
 
-			static async Task InitNewClientAsync(IPerforceConnection perforce, string streamName, string hostName, InitCommandOptions options, ILogger logger)
+			static async Task InitNewClientAsync(IPerforceConnection perforce, string streamName, string hostName, InitCommandOptions options, UserSettings? globalSettings, ILogger logger)
 			{
 				logger.LogInformation("Checking stream...");
 
@@ -694,6 +700,20 @@ namespace UnrealGameSyncCmd
 				options.ApplyTo(settings);
 				settings.Save(logger);
 
+				if (globalSettings != null)
+				{
+					UserSelectedProjectSettings selectedProjectSettings =
+						new UserSelectedProjectSettings(
+							null, 
+							null, 
+							UserSelectedProjectType.Client, 
+							$"//{client.Name}{projectPath}",
+							$"{client.Root}{projectPath}".Replace("/", @"\"));
+					globalSettings.OpenProjects.Add(selectedProjectSettings);
+					globalSettings.RecentProjects.Add(selectedProjectSettings);
+					globalSettings.Save(logger);
+				}
+
 				logger.LogInformation("Initialized {ClientName} with root at {RootDir}", clientName, clientDir);
 			}
 
@@ -709,7 +729,7 @@ namespace UnrealGameSyncCmd
 				}
 			}
 
-			static async Task InitExistingClientAsync(IPerforceConnection perforce, string hostName, InitCommandOptions options, ILogger logger)
+			static async Task InitExistingClientAsync(IPerforceConnection perforce, string hostName, InitCommandOptions options, UserSettings? globalSettings, ILogger logger)
 			{
 				DirectoryReference currentDir = DirectoryReference.GetCurrentDirectory();
 
@@ -753,6 +773,20 @@ namespace UnrealGameSyncCmd
 				options.ApplyTo(settings);
 				settings.Save(logger);
 
+				if (globalSettings != null)
+				{
+					UserSelectedProjectSettings selectedProjectSettings =
+						new UserSelectedProjectSettings(
+							null,
+							null,
+							UserSelectedProjectType.Client,
+							$"//{client.Name}{projectPath}",
+							$"{client.Root}{projectPath}".Replace("/", @"\"));
+					globalSettings.OpenProjects.Add(selectedProjectSettings);
+					globalSettings.RecentProjects.Add(selectedProjectSettings);
+					globalSettings.Save(logger);
+				}
+
 				logger.LogInformation("Initialized workspace at {RootDir} for {ClientProject}", clientDir, settings.ClientProjectPath);
 			}
 
