UnrealGameSyncCmd: Load UseFastSync setting from Global.json

#rnx

diff --git a/Engine/Source/Programs/UnrealGameSync/UnrealGameSyncShared/WorkspaceUpdate.cs b/Engine/Source/Programs/UnrealGameSync/UnrealGameSyncShared/WorkspaceUpdate.cs
--- a/Engine/Source/Programs/UnrealGameSync/UnrealGameSyncShared/WorkspaceUpdate.cs
+++ b/Engine/Source/Programs/UnrealGameSync/UnrealGameSyncShared/WorkspaceUpdate.cs
@@ -146,6 +146,10 @@ namespace UnrealGameSync
 			{
 				options |= WorkspaceUpdateOptions.UprojectSpecificSolution;
 			}
+			if (globalSettings.Perforce.UseFastSync is true)
+			{
+				options |= WorkspaceUpdateOptions.FastSync;
+			}
 			return options;
 		}
 	}
