Opting out of a Sync Category changes the filter, but Fast Sync has been enabled

switching fast sync off when the filters have changed or the refilter flag is set.

diff --git a/Engine/Source/Programs/UnrealGameSync/UnrealGameSyncShared/WorkspaceUpdate.cs b/Engine/Source/Programs/UnrealGameSync/UnrealGameSyncShared/WorkspaceUpdate.cs
--- a/Engine/Source/Programs/UnrealGameSync/UnrealGameSyncShared/WorkspaceUpdate.cs
+++ b/Engine/Source/Programs/UnrealGameSync/UnrealGameSyncShared/WorkspaceUpdate.cs
@@ -688,6 +688,10 @@ namespace UnrealGameSync
 					// If the hash differs, enumerate everything in the workspace to find what needs to be removed
 					if (nextSyncFilterHash != state.CurrentSyncFilterHash || (Context.Options & WorkspaceUpdateOptions.Refilter) != 0)
 					{
+						// disable fast sync to make sure that all the files are considered
+						logger.LogInformation("Disabling fast sync as filters are being reconsidered.");
+						Context.Options &= ~WorkspaceUpdateOptions.FastSync;
+
 						using (TelemetryStopwatch filterStopwatch = new TelemetryStopwatch("Workspace_Sync_FilterChanged", project.TelemetryProjectIdentifier))
 						{
 							logger.LogInformation("Filter has changed ({PrevHash} -> {NextHash}); finding files in workspace that need to be removed.", (String.IsNullOrEmpty(state.CurrentSyncFilterHash)) ? "None" : state.CurrentSyncFilterHash, nextSyncFilterHash);
