UGS - Added ability to download cloud artifacts via .uartifact files

diff --git a/Engine/Source/Programs/UnrealGameSync/UnrealGameSync/ArtifactDownload.cs b/Engine/Source/Programs/UnrealGameSync/UnrealGameSync/ArtifactDownload.cs
--- a/Engine/Source/Programs/UnrealGameSync/UnrealGameSync/ArtifactDownload.cs
+++ b/Engine/Source/Programs/UnrealGameSync/UnrealGameSync/ArtifactDownload.cs
@@ -1,6 +1,8 @@
 // Copyright Epic Games, Inc. All Rights Reserved.
 
 using System;
+using System.Collections;
+using System.Collections.Generic;
 using System.Diagnostics;
 using System.IO;
 using System.Linq;
@@ -11,6 +13,8 @@ using EpicGames.Core;
 using EpicGames.Horde;
 using EpicGames.Horde.Storage;
 using EpicGames.Horde.Storage.Nodes;
+using EpicGames.OIDC;
+using Microsoft.Extensions.Configuration;
 using Microsoft.Extensions.DependencyInjection;
 using Microsoft.Extensions.Logging;
 using Microsoft.Win32;
@@ -78,6 +82,120 @@ namespace UnrealGameSync
 			}
 		}
 
+		static async Task DoSyncZenAsync(Uri baseUri, RefName refName, DirectoryReference outputDir, IProgress<string> progress, CancellationToken cancellationToken)
+		{
+			DirectoryReference stateFolder = DirectoryReference.Combine(outputDir, ".ugs");
+			DirectoryReference.CreateDirectory(stateFolder);
+
+			using (ILoggerProvider loggerProvider = Logging.CreateLoggerProvider(FileReference.Combine(stateFolder, "sync.log")))
+			{
+				ILogger logger = loggerProvider.CreateLogger(nameof(ArtifactDownload));
+
+				DirectoryReference dataFolder = DirectoryReference.Combine(DirectoryReference.GetSpecialFolder(Environment.SpecialFolder.LocalApplicationData)!, "UnrealGameSync");
+				// Read the user's settings
+				UserSettings settings = UserSettings.Create(dataFolder, logger);
+
+				DirectoryInfo? engineDir = null;
+				DirectoryInfo? gameDir = null;
+				if (settings.LastProject != null && settings.LastProject.LocalPath != null)
+				{
+					gameDir = new DirectoryInfo(Path.GetDirectoryName(settings.LastProject.LocalPath)!);
+
+					string dir = settings.LastProject.LocalPath;
+					// search for an "Engine" directory in the locations upwards from the uproject file
+					// TODO: This does not work for external projects that use an installed engine build
+					while (!String.IsNullOrEmpty(dir))
+					{
+						string? possibleDir = Path.GetDirectoryName(dir);
+						if (possibleDir == null)
+						{
+							break;
+						}
+						string possiblyEngineDir = Path.Combine(possibleDir, "Engine");
+						if (Directory.Exists(possiblyEngineDir))
+						{
+							engineDir = new DirectoryInfo(possiblyEngineDir);
+							break;
+						}
+
+						// did not find an engine folder, continue the search upwards
+						dir = possibleDir;
+					}
+				}
+				else
+				{
+					logger.LogError("Unable to find a previous workspace, thus unable to locate your engine directory and can not download build");
+					return;
+				}
+
+				if (engineDir == null)
+				{
+					logger.LogError("Unable to locate a Engine folder from project '{ProjectPath}'. Unable to download build", settings.LastProject.LocalPath);
+					return;
+				}
+				string host = $"{baseUri.Scheme}://{baseUri.Host}";
+				string path = baseUri.AbsolutePath;
+				// remove api path if specified
+				path = path.Replace("api/v2/builds/", "", StringComparison.OrdinalIgnoreCase);
+				string[] components = path.Split('/', StringSplitOptions.RemoveEmptyEntries);
+				if (components.Length != 3)
+				{
+					logger.LogError("Failed to base uri path {Uri}, did not match expected format of namespace/bucket/buildId", path);
+					return;
+				}
+
+				string namespaceId = components[0];
+				string bucketId = components[1];
+				string buildId = components[2];
+
+				using ITokenStore tokenStore = TokenStoreFactory.CreateTokenStore();
+				IConfiguration config = ProviderConfigurationFactory.ReadConfiguration(engineDir, gameDir);
+
+				OidcTokenManager tokenManager = OidcTokenManager.CreateTokenManager(config, tokenStore, null);
+
+				// TODO: temporary hard code our auth provider until we have moved this configuration away from the engine syn
+				string providerName = "EpicGames-Okta";
+				OidcTokenInfo tokenInfo = await tokenManager.GetAccessToken(providerName, cancellationToken);
+				if (tokenInfo.AccessToken == null)
+				{
+					logger.LogError("Failed to get a access token");
+					return;
+				}
+
+				string accessToken = tokenInfo.AccessToken;
+				string accessTokenEnvVar = "UE-CloudDataCacheAccessToken";
+				// pass the access token via environment variable to avoid showing it on the cli
+				string cmdline = $"builds download --url {host} --namespace {namespaceId} --bucket {bucketId} --access-token-env {accessTokenEnvVar} --plain-progress {outputDir} {buildId}";
+
+				progress.Report("Connecting to server...");
+
+				FileInfo zenExe = new FileInfo(Path.Combine(engineDir.FullName, "Binaries/Win64/zen.exe"));
+
+				// we need to propagate the environment while also adding the access token we want to override
+				Dictionary<string, string> processEnv = new Dictionary<string, string>();
+				foreach (DictionaryEntry entry in Environment.GetEnvironmentVariables())
+				{
+					if (entry is { Key: string k, Value: string v })
+					{
+						processEnv[k] = v;
+					}
+				}
+
+				processEnv[accessTokenEnvVar] = accessToken;
+
+				int exitCode = await Utility.ExecuteProcessAsync(zenExe.FullName, zenExe.DirectoryName, cmdline, s =>
+				{
+					logger.LogInformation(s);
+
+					progress.Report(s);
+				}, processEnv, cancellationToken);
+				if (exitCode != 0)
+				{
+					logger.LogError("Error running {App} {Args} shutdown with exitcode {ExitCode}", zenExe, cmdline, exitCode);
+				}
+			}
+		}
+
 		public static void RegisterFileAssociations(string application)
 		{
 			string progId = "UnrealGameSync.Artifact";
@@ -148,7 +266,21 @@ namespace UnrealGameSync
 
 			try
 			{
-				DownloadProgressWindow.Execute((p, ctx) => DoSyncAsync(descriptor.BaseUrl, descriptor.RefName, outputDir, p, ctx), CancellationToken.None);
+				DownloadProgressWindow.Execute((p, ctx) =>
+				{
+					if (descriptor.Backend.Type == ArtifactBackend.BackendTypeEnum.Horde)
+					{
+						return DoSyncAsync(descriptor.BaseUrl, descriptor.RefName, outputDir, p, ctx);
+					}
+					else if (descriptor.Backend.Type == ArtifactBackend.BackendTypeEnum.Zen)
+					{
+						return DoSyncZenAsync(descriptor.BaseUrl, descriptor.RefName, outputDir, p, ctx);
+					}
+					else
+					{
+						throw new NotImplementedException($"Unknown backend type {descriptor.Backend.Type}");
+					}
+				}, CancellationToken.None);
 			}
 			catch (OperationCanceledException)
 			{
diff --git a/Engine/Source/Programs/UnrealGameSync/UnrealGameSyncShared/Utility.cs b/Engine/Source/Programs/UnrealGameSync/UnrealGameSyncShared/Utility.cs
--- a/Engine/Source/Programs/UnrealGameSync/UnrealGameSyncShared/Utility.cs
+++ b/Engine/Source/Programs/UnrealGameSync/UnrealGameSyncShared/Utility.cs
@@ -488,6 +488,24 @@ namespace UnrealGameSync
 			}
 		}
 
+		public static async Task<int> ExecuteProcessAsync(string fileName, string? workingDir, string commandLine, Action<string> outputLine, IReadOnlyDictionary<string, string>? env = null, CancellationToken cancellationToken = default)
+		{
+			using (ManagedProcessGroup newProcessGroup = new ManagedProcessGroup())
+			using (ManagedProcess newProcess = new ManagedProcess(newProcessGroup, fileName, commandLine, workingDir, env, null, ProcessPriorityClass.Normal))
+			{
+				for (; ; )
+				{
+					string? line = await newProcess.ReadLineAsync(cancellationToken);
+					if (line == null)
+					{
+						await newProcess.WaitForExitAsync(cancellationToken);
+						return newProcess.ExitCode;
+					}
+					outputLine(line);
+				}
+			}
+		}
+
 		public static async Task<int> ExecuteProcessAsync(string fileName, string? workingDir, string commandLine, Action<string> outputLine, CancellationToken cancellationToken)
 		{
 			using (ManagedProcessGroup newProcessGroup = new ManagedProcessGroup())
