UGS - Quote local download path passed to zen cli when downloading artifacts

diff --git a/Engine/Source/Programs/UnrealGameSync/UnrealGameSync/ArtifactDownload.cs b/Engine/Source/Programs/UnrealGameSync/UnrealGameSync/ArtifactDownload.cs
--- a/Engine/Source/Programs/UnrealGameSync/UnrealGameSync/ArtifactDownload.cs
+++ b/Engine/Source/Programs/UnrealGameSync/UnrealGameSync/ArtifactDownload.cs
@@ -165,7 +165,7 @@ namespace UnrealGameSync
 				string accessToken = tokenInfo.AccessToken;
 				string accessTokenEnvVar = "UE-CloudDataCacheAccessToken";
 				// pass the access token via environment variable to avoid showing it on the cli
-				string cmdline = $"builds download --url {host} --namespace {namespaceId} --bucket {bucketId} --access-token-env {accessTokenEnvVar} --plain-progress {outputDir} {buildId}";
+				string cmdline = $"builds download --url {host} --namespace {namespaceId} --bucket {bucketId} --access-token-env {accessTokenEnvVar} --plain-progress \"{outputDir}\" {buildId}";
 
 				progress.Report("Connecting to server...");
 
