batch reading the existing log file to avoid exception

diff --git a/Engine/Source/Programs/UnrealGameSync/UnrealGameSync/Controls/LogControl.cs b/Engine/Source/Programs/UnrealGameSync/UnrealGameSync/Controls/LogControl.cs
--- a/Engine/Source/Programs/UnrealGameSync/UnrealGameSync/Controls/LogControl.cs
+++ b/Engine/Source/Programs/UnrealGameSync/UnrealGameSync/Controls/LogControl.cs
@@ -204,11 +204,25 @@ namespace UnrealGameSync
 
 			_logFileStream.Seek(0, SeekOrigin.Begin);
 
+
+			int batchSize = 1024;
 			using StreamReader reader = new StreamReader(_logFileStream, leaveOpen: true);
-			string text = reader.ReadToEnd().TrimEnd('\r', '\n');
-			if (text.Length > 0)
+
+			List<string> lines = new List<string>(batchSize);
+			while (!reader.EndOfStream)
+			{
+				lines.Add(reader.ReadLine()!.TrimEnd('\r', '\n'));
+
+				if (lines.Count >= batchSize)
+				{
+					AddLinesInternal(lines.Select(x => x + "\n").ToList());
+					lines.Clear();
+				}
+			}
+
+			if (lines.Any())
 			{
-				AddLinesInternal(text.Split('\n').Select(x => x + "\n").ToList());
+				AddLinesInternal(lines.Select(x => x + "\n").ToList());
 			}
 
 			Invalidate();
