UGS: Fix compiler warnings

diff --git a/Engine/Source/Programs/UnrealGameSync/UnrealGameSync/ApplicationThemeManager.cs b/Engine/Source/Programs/UnrealGameSync/UnrealGameSync/ApplicationThemeManager.cs
--- a/Engine/Source/Programs/UnrealGameSync/UnrealGameSync/ApplicationThemeManager.cs
+++ b/Engine/Source/Programs/UnrealGameSync/UnrealGameSync/ApplicationThemeManager.cs
@@ -109,6 +109,7 @@ namespace UnrealGameSync
 		[DllImport("uxtheme.dll", CharSet = CharSet.Unicode)]
 		private static extern int SetWindowTheme(IntPtr hWnd, string? pszSubAppName, string? pszSubIdList);
 		
+		[System.Diagnostics.CodeAnalysis.SuppressMessage("CodeQuality", "IDE0051:Remove unused private members", Justification = "Not currently used")]
 		private static void UpdateDefaultTheme()
 		{
 			// Set the default theme based on the user's system settings
diff --git a/Engine/Source/Programs/UnrealGameSync/UnrealGameSync/ArtifactDownload.cs b/Engine/Source/Programs/UnrealGameSync/UnrealGameSync/ArtifactDownload.cs
--- a/Engine/Source/Programs/UnrealGameSync/UnrealGameSync/ArtifactDownload.cs
+++ b/Engine/Source/Programs/UnrealGameSync/UnrealGameSync/ArtifactDownload.cs
@@ -82,7 +82,7 @@ namespace UnrealGameSync
 			}
 		}
 
-		static async Task DoSyncZenAsync(Uri baseUri, RefName refName, DirectoryReference outputDir, IProgress<string> progress, CancellationToken cancellationToken)
+		static async Task DoSyncZenAsync(Uri baseUri, RefName _, DirectoryReference outputDir, IProgress<string> progress, CancellationToken cancellationToken)
 		{
 			DirectoryReference stateFolder = DirectoryReference.Combine(outputDir, ".ugs");
 			DirectoryReference.CreateDirectory(stateFolder);
@@ -151,7 +151,7 @@ namespace UnrealGameSync
 				using ITokenStore tokenStore = TokenStoreFactory.CreateTokenStore();
 				IConfiguration config = ProviderConfigurationFactory.ReadConfiguration(engineDir, gameDir);
 
-				OidcTokenManager tokenManager = OidcTokenManager.CreateTokenManager(config, tokenStore, null);
+				OidcTokenManager tokenManager = OidcTokenManager.CreateTokenManager(config, tokenStore, null!);
 
 				// TODO: temporary hard code our auth provider until we have moved this configuration away from the engine syn
 				string providerName = "EpicGames-Okta";
@@ -185,7 +185,7 @@ namespace UnrealGameSync
 
 				int exitCode = await Utility.ExecuteProcessAsync(zenExe.FullName, zenExe.DirectoryName, cmdline, s =>
 				{
-					logger.LogInformation(s);
+					logger.LogInformation("{Line}", s);
 
 					progress.Report(s);
 				}, processEnv, cancellationToken);
diff --git a/Engine/Source/Programs/UnrealGameSync/UnrealGameSyncCmd/Program.cs b/Engine/Source/Programs/UnrealGameSync/UnrealGameSyncCmd/Program.cs
--- a/Engine/Source/Programs/UnrealGameSync/UnrealGameSyncCmd/Program.cs
+++ b/Engine/Source/Programs/UnrealGameSync/UnrealGameSyncCmd/Program.cs
@@ -24,7 +24,6 @@ using Microsoft.Extensions.Configuration;
 using Microsoft.Extensions.DependencyInjection;
 using Microsoft.Extensions.Logging;
 using Serilog;
-using Serilog.Core;
 using UnrealGameSync;
 
 namespace UnrealGameSyncCmd
@@ -713,7 +712,7 @@ namespace UnrealGameSyncCmd
 							null, 
 							UserSelectedProjectType.Client, 
 							$"//{client.Name}{projectPath}",
-							$"{client.Root}{projectPath}".Replace("/", @"\"));
+							$"{client.Root}{projectPath}".Replace("/", @"\", StringComparison.Ordinal));
 					globalSettings.OpenProjects.Add(selectedProjectSettings);
 					globalSettings.RecentProjects.Add(selectedProjectSettings);
 					globalSettings.Save(logger);
@@ -786,7 +785,7 @@ namespace UnrealGameSyncCmd
 							null,
 							UserSelectedProjectType.Client,
 							$"//{client.Name}{projectPath}",
-							$"{client.Root}{projectPath}".Replace("/", @"\"));
+							$"{client.Root}{projectPath}".Replace("/", @"\", StringComparison.Ordinal));
 					globalSettings.OpenProjects.Add(selectedProjectSettings);
 					globalSettings.RecentProjects.Add(selectedProjectSettings);
 					globalSettings.Save(logger);
@@ -1950,7 +1949,7 @@ namespace UnrealGameSyncCmd
 							{
 								if (toolInfo.Id == guid)
 								{
-									logger.LogInformation("Adding tool '{id}' - '{name}' to the list of enabled tools", id, toolInfo.Name);
+									logger.LogInformation("Adding tool '{Id}' - '{Name}' to the list of enabled tools", id, toolInfo.Name);
 									newEnabledTools.Add(guid);
 									break;
 								}
@@ -1958,12 +1957,12 @@ namespace UnrealGameSyncCmd
 
 							if (!newEnabledTools.Contains(guid))
 							{
-								logger.LogError("Could not install tool '{id}'", id);
+								logger.LogError("Could not install tool '{Id}'", id);
 							}
 						}
 						else
 						{
-							logger.LogError("'{id}' is not a valid guid", id);
+							logger.LogError("'{Id}' is not a valid guid", id);
 						}
 					}
 
@@ -1975,7 +1974,7 @@ namespace UnrealGameSyncCmd
 							{
 								if (toolInfo.Id == guid)
 								{
-									logger.LogInformation("Removing tool '{id}' - '{name}' from the list of enabled tools", id, toolInfo.Name);
+									logger.LogInformation("Removing tool '{Id}' - '{Name}' from the list of enabled tools", id, toolInfo.Name);
 									newEnabledTools.Remove(guid);
 									break;
 								}
@@ -1983,12 +1982,12 @@ namespace UnrealGameSyncCmd
 
 							if (newEnabledTools.Contains(guid))
 							{
-								logger.LogError("Could not uninstall tool '{id}'", id);
+								logger.LogError("Could not uninstall tool '{Id}'", id);
 							}
 						}
 						else
 						{
-							logger.LogError("'{id}' is not a valid guid", id);
+							logger.LogError("'{Id}' is not a valid guid", id);
 						}
 					}
 
diff --git a/Engine/Source/Programs/UnrealGameSync/UnrealGameSyncShared/AsyncDisposer.cs b/Engine/Source/Programs/UnrealGameSync/UnrealGameSyncShared/AsyncDisposer.cs
--- a/Engine/Source/Programs/UnrealGameSync/UnrealGameSyncShared/AsyncDisposer.cs
+++ b/Engine/Source/Programs/UnrealGameSync/UnrealGameSyncShared/AsyncDisposer.cs
@@ -54,6 +54,8 @@ namespace UnrealGameSync
 				_logger.LogInformation("Waiting for {NumTasks} tasks to complete", tasksCopy.Count);
 				await Task.WhenAny(waitTask, Task.Delay(TimeSpan.FromSeconds(5.0)));
 			}
+			
+			GC.SuppressFinalize(this);
 		}
 	}
 }
\ No newline at end of file
diff --git a/Engine/Source/Programs/UnrealGameSync/UnrealGameSyncShared/ToolUpdateMonitor.cs b/Engine/Source/Programs/UnrealGameSync/UnrealGameSyncShared/ToolUpdateMonitor.cs
--- a/Engine/Source/Programs/UnrealGameSync/UnrealGameSyncShared/ToolUpdateMonitor.cs
+++ b/Engine/Source/Programs/UnrealGameSync/UnrealGameSyncShared/ToolUpdateMonitor.cs
@@ -40,7 +40,7 @@ namespace UnrealGameSync
 	{
 		public string? InstallCommand { get; set; }
 		public string? UninstallCommand { get; set; }
-		public List<ToolLink> StatusPanelLinks { get; set; } = new List<ToolLink>();
+		public List<ToolLink> StatusPanelLinks { get; } = new List<ToolLink>();
 		public bool SafeWhenBusy { get; set; }
 	}
 
@@ -50,7 +50,7 @@ namespace UnrealGameSync
 		public Guid Id { get; set; }
 		public string Name { get; set; }
 		public string Description { get; set; }
-		public HashSet<Guid> DependsOnToolIds { get; set; } = new HashSet<Guid>();
+		public HashSet<Guid> DependsOnToolIds { get; } = new HashSet<Guid>();
 		public ToolSettings Settings { get; set; } = new ToolSettings();
 		public string Revision { get; set; }
 
@@ -68,7 +68,7 @@ namespace UnrealGameSync
 		}
 	}
 
-	public class ToolUpdateMonitor : IDisposable
+	public sealed class ToolUpdateMonitor : IDisposable
 	{
 #pragma warning disable CA2213 // warning CA2213: 'ToolUpdateMonitor' contains field '_cancellationSource' that is of IDisposable type 'CancellationTokenSource', but it is never disposed. Change the Dispose method on 'ToolUpdateMonitor' to call Close or Dispose on this field.
 		readonly CancellationTokenSource _cancellationSource;
@@ -93,7 +93,7 @@ namespace UnrealGameSync
 		DirectoryReference ToolsDir { get; }
 		UserSettings Settings { get; }
 
-		public Action? OnChange;
+		public Action? OnChange { get; set; }
 
 		public Tuple<bool, string> LastStatus { get; private set; } = Tuple.Create(false, "Starting...");
 
diff --git a/Engine/Source/Programs/UnrealGameSync/UnrealGameSyncShared/WorkspaceUpdate.cs b/Engine/Source/Programs/UnrealGameSync/UnrealGameSyncShared/WorkspaceUpdate.cs
--- a/Engine/Source/Programs/UnrealGameSync/UnrealGameSyncShared/WorkspaceUpdate.cs
+++ b/Engine/Source/Programs/UnrealGameSync/UnrealGameSyncShared/WorkspaceUpdate.cs
@@ -6,7 +6,6 @@ using System.Collections.Generic;
 using System.Diagnostics;
 using System.IO;
 using System.Linq;
-using System.Runtime.CompilerServices;
 using System.Runtime.InteropServices;
 using System.Security.Cryptography;
 using System.Text;
