[UGS] Prevent an error when cleaning workspaces that need a large number of files to by synched. Breaks into batches of 10K files to sync.  Also fix a few syntax errors.

diff --git a/Engine/Source/Programs/Shared/EpicGames.Serialization/CbWriter.cs b/Engine/Source/Programs/Shared/EpicGames.Serialization/CbWriter.cs
--- a/Engine/Source/Programs/Shared/EpicGames.Serialization/CbWriter.cs
+++ b/Engine/Source/Programs/Shared/EpicGames.Serialization/CbWriter.cs
@@ -693,7 +693,7 @@ namespace EpicGames.Serialization
 		/// </summary>
 		public void Clear()
 		{
-			base.Reset();
+			Reset();
 
 			_freeChunks.AddRange(_chunks);
 			_chunks.Clear();
diff --git a/Engine/Source/Programs/UnrealGameSync/UnrealGameSync/Controls/WorkspaceControl.cs b/Engine/Source/Programs/UnrealGameSync/UnrealGameSync/Controls/WorkspaceControl.cs
--- a/Engine/Source/Programs/UnrealGameSync/UnrealGameSync/Controls/WorkspaceControl.cs
+++ b/Engine/Source/Programs/UnrealGameSync/UnrealGameSync/Controls/WorkspaceControl.cs
@@ -7066,7 +7066,7 @@ namespace UnrealGameSync
 			if (!String.IsNullOrWhiteSpace(_authorFilterText) && _authorFilterText != AuthorFilterPlaceholderText)
 			{
 				_authorsFilter.Clear();
-				_authorsFilter.AddAll( _authorFilterText.Split([';', '+'], StringSplitOptions.TrimEntries));
+				_authorsFilter.AddAll( _authorFilterText.Split((char[])[';', '+'], StringSplitOptions.TrimEntries));
 			}
 			else
 			{
diff --git a/Engine/Source/Programs/UnrealGameSync/UnrealGameSync/DeleteFilesTask.cs b/Engine/Source/Programs/UnrealGameSync/UnrealGameSync/DeleteFilesTask.cs
--- a/Engine/Source/Programs/UnrealGameSync/UnrealGameSync/DeleteFilesTask.cs
+++ b/Engine/Source/Programs/UnrealGameSync/UnrealGameSync/DeleteFilesTask.cs
@@ -20,18 +20,26 @@ namespace UnrealGameSync
 
 			if (filesToSync.Count > 0)
 			{
-				using IPerforceConnection perforce = await PerforceConnection.CreateAsync(perforceSettings, logger);
-
-				List<string> revisionsToSync = new List<string>();
-				foreach (FileInfo fileToSync in filesToSync)
+				// Do batches to avoid blowing out the RPC, 10K files take around 50 seconds to sync.
+				const int BatchSize = 10_000;
+				for (int startIndex = 0; startIndex < filesToSync.Count; startIndex += BatchSize)
 				{
-					revisionsToSync.Add(String.Format("{0}#have", PerforceUtils.EscapePath(fileToSync.FullName)));
-				}
+					logger.LogInformation("syncing: {StartIndex}/{Count}, percent {Percent}", startIndex, filesToSync.Count,
+						(int)Math.Round(100 * (float)startIndex / filesToSync.Count));
 
-				List<PerforceResponse<SyncRecord>> failedRecords = await perforce.TrySyncAsync(SyncOptions.Force, revisionsToSync, cancellationToken).Where(x => x.Failed).ToListAsync(cancellationToken);
-				foreach (PerforceResponse<SyncRecord> failedRecord in failedRecords)
-				{
-					failMessage.Append(failedRecord.ToString());
+					int batchCount = Math.Min(BatchSize, filesToSync.Count - startIndex);
+					List<FileInfo> batchFilesToSync = filesToSync.GetRange(startIndex, batchCount);
+					using IPerforceConnection perforce = await PerforceConnection.CreateAsync(perforceSettings, logger);
+					List<string> revisionsToSync = new List<string>();
+					foreach (FileInfo fileToSync in batchFilesToSync)
+					{
+						revisionsToSync.Add(String.Format("{0}#have", PerforceUtils.EscapePath(fileToSync.FullName)));
+					}
+					List<PerforceResponse<SyncRecord>> failedRecords = await perforce.TrySyncAsync(SyncOptions.Force, revisionsToSync, cancellationToken).Where(x => x.Failed).ToListAsync(cancellationToken);
+					foreach (PerforceResponse<SyncRecord> failedRecord in failedRecords)
+					{
+						failMessage.Append(failedRecord.ToString());
+					}
 				}
 			}
 
diff --git a/Engine/Source/Programs/UnrealGameSync/UnrealGameSyncShared/UserSettings.cs b/Engine/Source/Programs/UnrealGameSync/UnrealGameSyncShared/UserSettings.cs
--- a/Engine/Source/Programs/UnrealGameSync/UnrealGameSyncShared/UserSettings.cs
+++ b/Engine/Source/Programs/UnrealGameSync/UnrealGameSyncShared/UserSettings.cs
@@ -592,7 +592,7 @@ namespace UnrealGameSync
 		readonly Dictionary<FileReference, UserProjectSettings> _projectKeyToSettings = new Dictionary<FileReference, UserProjectSettings>();
 
 		// Perforce settings
-		public PerforceSyncOptions SyncOptions => base.Global.Perforce;
+		public PerforceSyncOptions SyncOptions => Global.Perforce;
 		public ClientSubmitOptions ClientSubmitOptions { get; set; }
 
 		private List<UserSelectedProjectSettings> ReadProjectList(string settingName, string legacySettingName)
