Changing the error to a info

diff --git a/Engine/Source/Programs/UnrealGameSync/UnrealGameSyncCmd/Program.cs b/Engine/Source/Programs/UnrealGameSync/UnrealGameSyncCmd/Program.cs
--- a/Engine/Source/Programs/UnrealGameSync/UnrealGameSyncCmd/Program.cs
+++ b/Engine/Source/Programs/UnrealGameSync/UnrealGameSyncCmd/Program.cs
@@ -1871,7 +1871,7 @@ namespace UnrealGameSyncCmd
 				ILogger logger = context.Logger;
 				if (!RuntimeInformation.IsOSPlatform(OSPlatform.Windows))
 				{
-					logger.LogError("This command is only available for Windows.");
+					logger.LogInformation("This command is only available for Windows.");
 					return;
 				}
 	
