- activate fast sync if syncing to the current CL
- trying first with fast sync then without if the preview sync fails, this can happen with huge distance between the current CL and the new one
- fixed timings alignment
- added explaination for when fastsync is disabled

diff --git a/Engine/Source/Programs/UnrealGameSync/UnrealGameSyncShared/WorkspaceUpdate.cs b/Engine/Source/Programs/UnrealGameSync/UnrealGameSyncShared/WorkspaceUpdate.cs
--- a/Engine/Source/Programs/UnrealGameSync/UnrealGameSyncShared/WorkspaceUpdate.cs
+++ b/Engine/Source/Programs/UnrealGameSync/UnrealGameSyncShared/WorkspaceUpdate.cs
@@ -457,7 +457,7 @@ namespace UnrealGameSync
 		class SyncTree
 		{
 			public bool CanUseWildcard;
-			private readonly bool CanContainExcludedFiles;
+			public bool CanContainExcludedFiles { get; set; } = false;
 			public int TotalIncludedFiles;
 			public long TotalSize;
 			public int TotalExcludedFiles;
@@ -859,98 +859,132 @@ namespace UnrealGameSync
 
 							if (Context.Options.HasFlag(WorkspaceUpdateOptions.Sync) && Context.Options.HasFlag(WorkspaceUpdateOptions.FastSync))
 							{
-								if (state.CurrentChangeNumber > 0 && Context.ChangeNumber > state.CurrentChangeNumber)
+								if (state.CurrentChangeNumber > 0 && Context.ChangeNumber >= state.CurrentChangeNumber)
 								{
 									syncFilter = $"{state.CurrentChangeNumber},{Context.ChangeNumber}";
 									logger.LogInformation("Using Fast Sync between current changelist {CurrentChangelist} and target changelist {TargetChangelist}", state.CurrentChangeNumber, Context.ChangeNumber);
 								}
 								else
 								{
-									logger.LogInformation("Fast Sync cannot be used with current changelist {CurrentChangelist} and target changelist {TargetChangelist}", state.CurrentChangeNumber, Context.ChangeNumber);
+									if (state.CurrentChangeNumber <= 0)
+									{
+										logger.LogInformation("Fast Sync cannot be used with current changelist {CurrentChangelist} and target changelist {TargetChangelist} because there is no current change list to compare to.", state.CurrentChangeNumber, Context.ChangeNumber);
+									}
+
+									if (Context.ChangeNumber < state.CurrentChangeNumber)
+									{
+										logger.LogInformation("Fast Sync cannot be used with current changelist {CurrentChangelist} and target changelist {TargetChangelist} because workspace syncing backwards.", state.CurrentChangeNumber, Context.ChangeNumber);
+									}
+
 								}
 							}
 
 							foreach (string syncPath in syncPaths)
 							{
-								List<SyncFile> syncFiles = new List<SyncFile>();
-								await foreach (PerforceResponse<SyncRecord> response in perforce.TrySyncAsync(SyncOptions.PreviewOnly, -1, 0, -1, -1, -1, -1, $"{syncPath}@{syncFilter}", cancellationToken))
+								async Task<(WorkspaceUpdateResult, string)> inlineGetSyncRecords(bool firstTry)
 								{
-									if (!response.Succeeded)
+									List<SyncFile> syncFiles = new List<SyncFile>();
+									await foreach (PerforceResponse<SyncRecord> response in perforce.TrySyncAsync(SyncOptions.PreviewOnly, -1, 0, -1, -1, -1, -1, $"{syncPath}@{syncFilter}", cancellationToken))
 									{
-										return (WorkspaceUpdateResult.FailedToSync, $"Couldn't enumerate changes matching {syncPath}.");
-									}
-									if (response.Info != null)
-									{
-										logger.LogInformation("Note: {Note}", response.Info.Data);
-										continue;
-									}
+										if (!response.Succeeded)
+										{
+											// if there is a failure retry without fast sync
+											if (firstTry 
+											&& Context.Options.HasFlag(WorkspaceUpdateOptions.Sync)
+											&& Context.Options.HasFlag(WorkspaceUpdateOptions.FastSync))
+											{
+												// change sync filter
+												syncFilter = Context.Options.HasFlag(WorkspaceUpdateOptions.Sync) ? $"{Context.ChangeNumber}" : $"={Context.ChangeNumber}";
+												Context.Options &= ~WorkspaceUpdateOptions.FastSync;
+												syncTree.CanContainExcludedFiles = false;
+												logger.LogWarning($"Couldn't enumerate changes matching with fast sync enabled: {syncPath}. Disabling fast sync and retrying.");
 
-									SyncRecord record = response.Data;
+												return await inlineGetSyncRecords(false);
+											}
 
-									// Get the full local path
-									string relativePath;
-									try
-									{
-										FileReference syncFile = new FileReference(record.Path.ToString());
-										relativePath = PerforceUtils.GetClientRelativePath(project.LocalRootPath, syncFile);
-									}
-									catch (PathTooLongException)
-									{
-										logger.LogInformation("The local path for {Path} exceeds the maximum allowed by Windows. Re-sync your workspace to a directory with a shorter name, or delete the file from the server.", record.Path);
-										return (WorkspaceUpdateResult.FailedToSync, "File exceeds maximum path length allowed by Windows.");
-									}
+											return (WorkspaceUpdateResult.FailedToSync, $"Couldn't enumerate changes matching {syncPath}.");
+										}
+										if (response.Info != null)
+										{
+											logger.LogInformation("Note: {Note}", response.Info.Data);
+											continue;
+										}
 
-									// Create the sync record
-									long syncSize = (record.Action == SyncAction.Deleted) ? 0 : record.FileSize;
-									syncFiles.Add(new SyncFile(record.DepotFile.ToString(), relativePath, syncSize));
-									counter.Increment();
-								}
+										SyncRecord record = response.Data;
 
-								// Also sync the currently open files
-								await foreach (PerforceResponse<OpenedRecord> response in perforce.TryOpenedAsync(OpenedOptions.None, FileSpecList.Any, cancellationToken))
-								{
-									if (!response.Succeeded)
-									{
-										return (WorkspaceUpdateResult.FailedToSync, $"Couldn't enumerate changes matching {syncPath}.");
-									}
+										// Get the full local path
+										string relativePath;
+										try
+										{
+											FileReference syncFile = new FileReference(record.Path.ToString());
+											relativePath = PerforceUtils.GetClientRelativePath(project.LocalRootPath, syncFile);
+										}
+										catch (PathTooLongException)
+										{
+											logger.LogInformation("The local path for {Path} exceeds the maximum allowed by Windows. Re-sync your workspace to a directory with a shorter name, or delete the file from the server.", record.Path);
+											return (WorkspaceUpdateResult.FailedToSync, "File exceeds maximum path length allowed by Windows.");
+										}
 
-									OpenedRecord record = response.Data;
-									if (!String.IsNullOrEmpty(record.DepotFile) && !String.IsNullOrEmpty(record.ClientFile))
+										// Create the sync record
+										long syncSize = (record.Action == SyncAction.Deleted) ? 0 : record.FileSize;
+										syncFiles.Add(new SyncFile(record.DepotFile.ToString(), relativePath, syncSize));
+										counter.Increment();
+									}
+									// Also sync the currently open files
+									await foreach (PerforceResponse<OpenedRecord> response in perforce.TryOpenedAsync(OpenedOptions.None, FileSpecList.Any, cancellationToken))
 									{
-										if (record.Action != FileAction.Add && record.Action != FileAction.Branch && record.Action != FileAction.MoveAdd)
+										if (!response.Succeeded)
 										{
-											string relativePath = PerforceUtils.GetClientRelativePath(record.ClientFile);
-											syncFiles.Add(new SyncFile(record.DepotFile, relativePath, 0));
+											return (WorkspaceUpdateResult.FailedToSync, $"Couldn't enumerate changes matching {syncPath}.");
+										}
+
+										OpenedRecord record = response.Data;
+										if (!String.IsNullOrEmpty(record.DepotFile) && !String.IsNullOrEmpty(record.ClientFile))
+										{
+											if (record.Action != FileAction.Add && record.Action != FileAction.Branch && record.Action != FileAction.MoveAdd)
+											{
+												string relativePath = PerforceUtils.GetClientRelativePath(record.ClientFile);
+												syncFiles.Add(new SyncFile(record.DepotFile, relativePath, 0));
+											}
 										}
 									}
-								}
 
-								// Enumerate all the files to be synced. NOTE: depotPath is escaped, whereas clientPath is not.
-								List<string> syncRelativePaths = new();
-								foreach (SyncFile syncRecord in syncFiles)
-								{
-									if (filter.Matches(syncRecord.RelativePath))
+									// Enumerate all the files to be synced. NOTE: depotPath is escaped, whereas clientPath is not.
+									List<string> syncRelativePaths = new();
+									foreach (SyncFile syncRecord in syncFiles)
 									{
-										syncTree.IncludeFile(PerforceUtils.EscapePath(syncRecord.RelativePath), syncRecord.Size, logger);
-										syncDepotPaths.Add(syncRecord.DepotFile);
-										syncRelativePaths.Add(syncRecord.RelativePath);
-										requiredFreeSpace += syncRecord.Size;
+										if (filter.Matches(syncRecord.RelativePath))
+										{
+											syncTree.IncludeFile(PerforceUtils.EscapePath(syncRecord.RelativePath), syncRecord.Size, logger);
+											syncDepotPaths.Add(syncRecord.DepotFile);
+											syncRelativePaths.Add(syncRecord.RelativePath);
+											requiredFreeSpace += syncRecord.Size;
+										}
+										else
+										{
+											syncTree.ExcludeFile(PerforceUtils.EscapePath(syncRecord.RelativePath));
+										}
 									}
-									else
+
+									Parallel.ForEach(syncRelativePaths, syncRelativePath =>
 									{
-										syncTree.ExcludeFile(PerforceUtils.EscapePath(syncRecord.RelativePath));
-									}
-								}
+										// If the file exists the required free space can be reduced as those bytes will be replaced.
+										FileInfo localFileInfo = FileReference.Combine(project.LocalRootPath, syncRelativePath).ToFileInfo();
+										if (localFileInfo.Exists)
+										{
+											Interlocked.Add(ref requiredFreeSpace, -localFileInfo.Length);
+										}
+									});
+
+									// return success
+									return (WorkspaceUpdateResult.Success, String.Empty);
+								};
 
-								Parallel.ForEach(syncRelativePaths, syncRelativePath =>
+								(WorkspaceUpdateResult, string) ret = await inlineGetSyncRecords(true);
+								if (ret.Item1 != WorkspaceUpdateResult.Success)
 								{
-									// If the file exists the required free space can be reduced as those bytes will be replaced.
-									FileInfo localFileInfo = FileReference.Combine(project.LocalRootPath, syncRelativePath).ToFileInfo();
-									if (localFileInfo.Exists)
-									{
-										Interlocked.Add(ref requiredFreeSpace, -localFileInfo.Length);
-									}
-								});
+									return ret;
+								}
 							}
 
 							try
@@ -1610,7 +1644,7 @@ namespace UnrealGameSync
 			logger.LogInformation("Total time: {TotalTime}", FormatTime(times.Sum(x => (long)(x.Item2.TotalMilliseconds / 1000))));
 			foreach (Tuple<string, TimeSpan> time in times)
 			{
-				logger.LogInformation("   {Name,-8}: {TimeSeconds}", time.Item1, FormatTime((long)(time.Item2.TotalMilliseconds / 1000)));
+				logger.LogInformation("    {TimeSeconds-16} {Name}", FormatTime((long)(time.Item2.TotalMilliseconds / 1000)), time.Item1);
 			}
 			if (numFilesSynced > 0)
 			{
