UGS - Added support for downloading pcbs from cloud storage
Added redistribution of zen binaries so that UGS works without a workspace that has them synced.

diff --git a/Engine/Source/Programs/UnrealGameSync/UnrealGameSync/ArtifactDownload.cs b/Engine/Source/Programs/UnrealGameSync/UnrealGameSync/ArtifactDownload.cs
--- a/Engine/Source/Programs/UnrealGameSync/UnrealGameSync/ArtifactDownload.cs
+++ b/Engine/Source/Programs/UnrealGameSync/UnrealGameSync/ArtifactDownload.cs
@@ -1,8 +1,6 @@
 // Copyright Epic Games, Inc. All Rights Reserved.
 
 using System;
-using System.Collections;
-using System.Collections.Generic;
 using System.Diagnostics;
 using System.IO;
 using System.Linq;
@@ -82,117 +80,27 @@ namespace UnrealGameSync
 			}
 		}
 
-		static async Task DoSyncZenAsync(Uri baseUri, RefName _, DirectoryReference outputDir, IProgress<string> progress, CancellationToken cancellationToken)
+		static async Task DoSyncZenAsync(Uri baseUri, DirectoryReference outputDir, IProgress<string> progress, CancellationToken cancellationToken)
 		{
 			DirectoryReference stateFolder = DirectoryReference.Combine(outputDir, ".ugs");
 			DirectoryReference.CreateDirectory(stateFolder);
 
 			using (ILoggerProvider loggerProvider = Logging.CreateLoggerProvider(FileReference.Combine(stateFolder, "sync.log")))
 			{
-				ILogger logger = loggerProvider.CreateLogger(nameof(ArtifactDownload));
-
-				DirectoryReference dataFolder = DirectoryReference.Combine(DirectoryReference.GetSpecialFolder(Environment.SpecialFolder.LocalApplicationData)!, "UnrealGameSync");
-				// Read the user's settings
-				UserSettings settings = UserSettings.Create(dataFolder, logger);
-
-				DirectoryInfo? engineDir = null;
-				DirectoryInfo? gameDir = null;
-				if (settings.LastProject != null && settings.LastProject.LocalPath != null)
-				{
-					gameDir = new DirectoryInfo(Path.GetDirectoryName(settings.LastProject.LocalPath)!);
-
-					string dir = settings.LastProject.LocalPath;
-					// search for an "Engine" directory in the locations upwards from the uproject file
-					// TODO: This does not work for external projects that use an installed engine build
-					while (!String.IsNullOrEmpty(dir))
-					{
-						string? possibleDir = Path.GetDirectoryName(dir);
-						if (possibleDir == null)
-						{
-							break;
-						}
-						string possiblyEngineDir = Path.Combine(possibleDir, "Engine");
-						if (Directory.Exists(possiblyEngineDir))
-						{
-							engineDir = new DirectoryInfo(possiblyEngineDir);
-							break;
-						}
-
-						// did not find an engine folder, continue the search upwards
-						dir = possibleDir;
-					}
-				}
-				else
-				{
-					logger.LogError("Unable to find a previous workspace, thus unable to locate your engine directory and can not download build");
-					return;
-				}
-
-				if (engineDir == null)
-				{
-					logger.LogError("Unable to locate a Engine folder from project '{ProjectPath}'. Unable to download build", settings.LastProject.LocalPath);
-					return;
-				}
-				string host = $"{baseUri.Scheme}://{baseUri.Host}";
-				string path = baseUri.AbsolutePath;
-				// remove api path if specified
-				path = path.Replace("api/v2/builds/", "", StringComparison.OrdinalIgnoreCase);
-				string[] components = path.Split('/', StringSplitOptions.RemoveEmptyEntries);
-				if (components.Length != 3)
+				ServiceCollection services = new ServiceCollection();
+				services.AddLogging(builder =>
 				{
-					logger.LogError("Failed to base uri path {Uri}, did not match expected format of namespace/bucket/buildId", path);
-					return;
-				}
-
-				string namespaceId = components[0];
-				string bucketId = components[1];
-				string buildId = components[2];
-
-				using ITokenStore tokenStore = TokenStoreFactory.CreateTokenStore();
-				IConfiguration config = ProviderConfigurationFactory.ReadConfiguration(engineDir, gameDir);
-
-				OidcTokenManager tokenManager = OidcTokenManager.CreateTokenManager(config, tokenStore, null!);
+					builder.AddProvider(loggerProvider);
+					builder.SetMinimumLevel(LogLevel.Trace);
+				});
+				services.AddCloudStorage();
 
-				// TODO: temporary hard code our auth provider until we have moved this configuration away from the engine syn
-				string providerName = "EpicGames-Okta";
-				OidcTokenInfo tokenInfo = await tokenManager.GetAccessToken(providerName, cancellationToken);
-				if (tokenInfo.AccessToken == null)
-				{
-					logger.LogError("Failed to get a access token");
-					return;
-				}
+				await using ServiceProvider serviceProvider = services.BuildServiceProvider();
 
-				string accessToken = tokenInfo.AccessToken;
-				string accessTokenEnvVar = "UE-CloudDataCacheAccessToken";
-				// pass the access token via environment variable to avoid showing it on the cli
-				string cmdline = $"builds download --url {host} --namespace {namespaceId} --bucket {bucketId} --access-token-env {accessTokenEnvVar} --plain-progress \"{outputDir}\" {buildId}";
+				ICloudStorage cloudStorage = serviceProvider.GetRequiredService<ICloudStorage>();
 
 				progress.Report("Connecting to server...");
-
-				FileInfo zenExe = new FileInfo(Path.Combine(engineDir.FullName, "Binaries/Win64/zen.exe"));
-
-				// we need to propagate the environment while also adding the access token we want to override
-				Dictionary<string, string> processEnv = new Dictionary<string, string>();
-				foreach (DictionaryEntry entry in Environment.GetEnvironmentVariables())
-				{
-					if (entry is { Key: string k, Value: string v })
-					{
-						processEnv[k] = v;
-					}
-				}
-
-				processEnv[accessTokenEnvVar] = accessToken;
-
-				int exitCode = await Utility.ExecuteProcessAsync(zenExe.FullName, zenExe.DirectoryName, cmdline, s =>
-				{
-					logger.LogInformation("{Line}", s);
-
-					progress.Report(s);
-				}, processEnv, cancellationToken);
-				if (exitCode != 0)
-				{
-					logger.LogError("Error running {App} {Args} shutdown with exitcode {ExitCode}", zenExe, cmdline, exitCode);
-				}
+				await cloudStorage.DownloadBuildFromUriAsync(baseUri, outputDir, progress, cancellationToken);
 			}
 		}
 
@@ -274,7 +182,7 @@ namespace UnrealGameSync
 					}
 					else if (descriptor.Backend.Type == ArtifactBackend.BackendTypeEnum.Zen)
 					{
-						return DoSyncZenAsync(descriptor.BaseUrl, descriptor.RefName, outputDir, p, ctx);
+						return DoSyncZenAsync(descriptor.BaseUrl, outputDir, p, ctx);
 					}
 					else
 					{
diff --git a/Engine/Source/Programs/UnrealGameSync/UnrealGameSync/PerforceMonitor.cs b/Engine/Source/Programs/UnrealGameSync/UnrealGameSync/PerforceMonitor.cs
--- a/Engine/Source/Programs/UnrealGameSync/UnrealGameSync/PerforceMonitor.cs
+++ b/Engine/Source/Programs/UnrealGameSync/UnrealGameSync/PerforceMonitor.cs
@@ -55,6 +55,7 @@ namespace UnrealGameSync
 #pragma warning disable CA2213 // Dispose of _hordeClient
 		readonly IHordeClient? _hordeClient;
 #pragma warning restore CA2213
+		private readonly ICloudStorage? _cloudStorage;
 
 		string[] prevCodeRules = Array.Empty<string>();
 
@@ -90,6 +91,7 @@ namespace UnrealGameSync
 			_synchronizationContext = SynchronizationContext.Current!;
 			_cancellationSource = new CancellationTokenSource();
 			_hordeClient = serviceProvider.GetService<IHordeClient>();
+			_cloudStorage = serviceProvider.GetService<ICloudStorage>();
 
 			AvailableArchiveChannels = (new List<IArchiveChannel>()).AsReadOnly();
 		}
@@ -497,7 +499,7 @@ namespace UnrealGameSync
 
 		async Task<bool> UpdateArchivesAsync(IPerforceConnection perforce, CancellationToken cancellationToken)
 		{
-			List<BaseArchiveChannel> newArchives = await BaseArchive.EnumerateChannelsAsync(perforce, _hordeClient, LatestProjectConfigFile, _selectedProjectIdentifier, cancellationToken);
+			List<BaseArchiveChannel> newArchives = await BaseArchive.EnumerateChannelsAsync(perforce, _hordeClient, _cloudStorage, LatestProjectConfigFile, _selectedProjectIdentifier, cancellationToken);
 
 			// Check if the information has changed
 			if (!Enumerable.SequenceEqual(_archives, newArchives))
diff --git a/Engine/Source/Programs/UnrealGameSync/UnrealGameSync/Program.cs b/Engine/Source/Programs/UnrealGameSync/UnrealGameSync/Program.cs
--- a/Engine/Source/Programs/UnrealGameSync/UnrealGameSync/Program.cs
+++ b/Engine/Source/Programs/UnrealGameSync/UnrealGameSync/Program.cs
@@ -207,6 +207,7 @@ namespace UnrealGameSync
 				services.AddSingleton(sp => TokenStoreFactory.CreateTokenStore());
 				services.AddSingleton<OidcTokenManager>();
 
+				services.AddCloudStorage();
 				if (launcherSettings.HordeServer != null)
 				{
 					services.AddHorde(options =>
diff --git a/Engine/Source/Programs/UnrealGameSync/UnrealGameSync/UnrealGameSync.csproj b/Engine/Source/Programs/UnrealGameSync/UnrealGameSync/UnrealGameSync.csproj
--- a/Engine/Source/Programs/UnrealGameSync/UnrealGameSync/UnrealGameSync.csproj
+++ b/Engine/Source/Programs/UnrealGameSync/UnrealGameSync/UnrealGameSync.csproj
@@ -50,6 +50,19 @@
     <ProjectReference Include="..\..\Shared\EpicGames.Perforce\EpicGames.Perforce.csproj" />
     <ProjectReference Include="..\UnrealGameSyncShared\UnrealGameSyncShared.csproj" />
   </ItemGroup>
+
+  <ItemGroup Condition="Exists('..\..\..\..\Binaries\Win64\zen.exe')">
+    <None Include="..\..\..\..\Binaries\Win64\zen.exe" TargetPath="Binaries\Win64\zen.exe" Link="Binaries\zen.exe">
+      <CopyToOutputDirectory>PreserveNewest</CopyToOutputDirectory>
+    </None>
+    <None Include="..\..\..\..\Binaries\Linux\zen" TargetPath="Binaries\Linux\zen" Link="Binaries\zen">
+      <CopyToOutputDirectory>PreserveNewest</CopyToOutputDirectory>
+    </None>
+    <None Include="..\..\..\..\Binaries\Mac\zen" TargetPath="Binaries\Mac\zen" Link="Binaries\zen">
+      <CopyToOutputDirectory>PreserveNewest</CopyToOutputDirectory>
+    </None>
+  </ItemGroup>
+
   <ItemGroup>
     <Reference Include="Microsoft.VisualStudio.Setup.Configuration.Interop">
       <HintPath>..\..\..\..\Binaries\ThirdParty\VisualStudio\Microsoft.VisualStudio.Setup.Configuration.Interop.dll</HintPath>
diff --git a/Engine/Source/Programs/UnrealGameSync/UnrealGameSyncCmd/Program.cs b/Engine/Source/Programs/UnrealGameSync/UnrealGameSyncCmd/Program.cs
--- a/Engine/Source/Programs/UnrealGameSync/UnrealGameSyncCmd/Program.cs
+++ b/Engine/Source/Programs/UnrealGameSync/UnrealGameSyncCmd/Program.cs
@@ -149,8 +149,9 @@ namespace UnrealGameSyncCmd
 			public GlobalSettingsFile UserSettings { get; }
 
 			public IHordeClient? HordeClient { get; }
+			public ICloudStorage? CloudStorage { get; }
 
-			public CommandContext(CommandLineArguments arguments, ILogger logger, ILoggerFactory loggerFactory, GlobalSettingsFile userSettings, UserSettings? globalSettings, IHordeClient? hordeClient)
+			public CommandContext(CommandLineArguments arguments, ILogger logger, ILoggerFactory loggerFactory, GlobalSettingsFile userSettings, UserSettings? globalSettings, IHordeClient? hordeClient, ICloudStorage? cloudStorage)
 			{
 				Arguments = arguments;
 				Logger = logger;
@@ -158,6 +159,7 @@ namespace UnrealGameSyncCmd
 				GlobalSettings = globalSettings;
 				UserSettings = userSettings;
 				HordeClient = hordeClient;
+				CloudStorage = cloudStorage;
 			}
 		}
 
@@ -259,8 +261,12 @@ namespace UnrealGameSyncCmd
 				{
 					services.AddHorde(options => options.ServerUrl = new Uri(launcherSettings.HordeServer));
 				}
+
+				services.AddCloudStorage();
+
 				ServiceProvider serviceProvider = services.BuildServiceProvider();
 				IHordeClient? hordeClient = serviceProvider.GetService<IHordeClient>();
+				ICloudStorage? cloudStorage = serviceProvider.GetService<ICloudStorage>();
 
 				UserSettings? globalSettings = null;
 				GlobalSettingsFile settings;
@@ -338,7 +344,7 @@ namespace UnrealGameSyncCmd
 				}
 
 				Command instance = (Command)Activator.CreateInstance(command.Type)!;
-				await instance.ExecuteAsync(new CommandContext(args, logger, loggerFactory, settings, globalSettings, hordeClient));
+				await instance.ExecuteAsync(new CommandContext(args, logger, loggerFactory, settings, globalSettings, hordeClient, cloudStorage));
 				return 0;
 			}
 			catch (UserErrorException ex)
@@ -944,9 +950,9 @@ namespace UnrealGameSyncCmd
 
 				if (syncOptions.Binaries)
 				{
-					List<BaseArchiveChannel> archives = await BaseArchive.EnumerateChannelsAsync(perforceClient, context.HordeClient, projectConfig, state.Current.ProjectIdentifier, CancellationToken.None);
+					List<BaseArchiveChannel> archives = await BaseArchive.EnumerateChannelsAsync(perforceClient, context.HordeClient, context.CloudStorage, projectConfig, state.Current.ProjectIdentifier, CancellationToken.None);
 
-					BaseArchiveChannel? editorArchiveInfo = archives.FirstOrDefault(x => x.Name == IArchiveChannel.EditorArchiveType);
+					BaseArchiveChannel? editorArchiveInfo = archives.FirstOrDefault(x => x.Type == IArchiveChannel.EditorArchiveType);
 					if (editorArchiveInfo == null)
 					{
 						throw new UserErrorException("No editor archives found for project");
diff --git a/Engine/Source/Programs/UnrealGameSync/UnrealGameSyncShared/ArchiveInfo.cs b/Engine/Source/Programs/UnrealGameSync/UnrealGameSyncShared/ArchiveInfo.cs
--- a/Engine/Source/Programs/UnrealGameSync/UnrealGameSyncShared/ArchiveInfo.cs
+++ b/Engine/Source/Programs/UnrealGameSync/UnrealGameSyncShared/ArchiveInfo.cs
@@ -13,6 +13,7 @@ using EpicGames.Horde.Artifacts;
 using EpicGames.Horde.Storage;
 using EpicGames.Horde.Storage.Nodes;
 using EpicGames.Perforce;
+using EpicGames.Serialization;
 using Microsoft.Extensions.Logging;
 
 namespace UnrealGameSync
@@ -20,6 +21,7 @@ namespace UnrealGameSync
 	public interface IArchive
 	{
 		public string Key { get; }
+		bool? RemoveOldBinaries { get; }
 
 		Task<bool> DownloadAsync(IPerforceConnection perforce, DirectoryReference localRootPath, FileReference manifestFileName, ILogger logger, ProgressValue progress, CancellationToken cancellationToken);
 	}
@@ -122,6 +124,7 @@ namespace UnrealGameSync
 		class PerforceArchive : IArchive
 		{
 			public string Key { get; }
+			public bool? RemoveOldBinaries { get; } = true;
 
 			public PerforceArchive(string key)
 				=> Key = key;
@@ -269,6 +272,7 @@ namespace UnrealGameSync
 			readonly ArtifactId _artifactId;
 
 			public string Key { get; }
+			public bool? RemoveOldBinaries { get; } = true;
 
 			public HordeArchive(IHordeClient hordeClient, ArtifactId artifactId)
 			{
@@ -362,13 +366,196 @@ namespace UnrealGameSync
 			return channels;
 		}
 	}
+	
+	public class CloudArchiveChannel : BaseArchiveChannel
+	{
+		private readonly string _archiveType;
+		private readonly string _streamName;
+
+		public CloudArchiveChannel(string name, string type, string archiveType, string streamName, bool bIgnoreRequiredBadges = false)
+			: base(name, type, bIgnoreRequiredBadges)
+		{
+			_archiveType = archiveType;
+			_streamName = streamName;
+		}
+
+		public override bool Equals(object? other)
+		{
+			CloudArchiveChannel? otherArchive = other as CloudArchiveChannel;
+			return otherArchive != null && Name == otherArchive.Name && Type == otherArchive.Type
+				&& Enumerable.SequenceEqual(ChangeNumberToArchive.Select(x => (x.Key, x.Value)), otherArchive.ChangeNumberToArchive.Select(x => (x.Key, x.Value)));
+		}
+
+		public override int GetHashCode()
+		{
+			throw new NotSupportedException();
+		}
+
+		class CloudArchive : IArchive
+		{
+			readonly ICloudStorage _cloudStorageClient;
+			private readonly string _host;
+			private readonly string _namespaceId;
+			private readonly string _bucketId;
+
+			public string Key { get; }
+			public bool? RemoveOldBinaries { get; } = false;
+
+			public CloudArchive(ICloudStorage cloudStorageClient, string host, string namespaceId, string bucketId, string buildId)
+			{
+				_cloudStorageClient = cloudStorageClient;
+				_host = host;
+				_namespaceId = namespaceId;
+				_bucketId = bucketId;
+				Key = buildId;
+			}
+
+			public async Task<bool> DownloadAsync(IPerforceConnection perforce, DirectoryReference localRootPath, FileReference manifestFileName, ILogger logger, ProgressValue progress, CancellationToken cancellationToken)
+			{
+				DirectoryReference configDir = UserSettings.GetConfigDir(localRootPath);
+				UserSettings.CreateConfigDir(configDir);
+
+				DirectoryReference zenStateFolder = DirectoryReference.Combine(configDir, ".zen");
+
+				Progress<string> progressCallback = new Progress<string>((s =>
+				{
+#pragma warning disable CA2254
+					logger.LogInformation(s);
+#pragma warning restore CA2254
+				}));
+
+				await _cloudStorageClient.DownloadBuildAsync(_host, _namespaceId, _bucketId,Key, localRootPath, zenStateDirectory: zenStateFolder, progressCallback, cancellationToken);
+
+				FileReference.Delete(manifestFileName);
+
+				ArchiveManifest manifest = new();
+				
+				// we do not output a list of files we wrote as zencli keeps this state internally and will remove any old files when it updates next
+				// and we explicitly do not want UGS to be removing files that we could use to patch the new build
+				//manifest.Files.Add();
+
+				await using (FileStream manifestStream = FileReference.Open(manifestFileName, FileMode.Create, FileAccess.Write))
+				{
+					manifest.Write(manifestStream);
+				}
+
+				return true;
+			}
+		}
+
+		public static async Task<List<CloudArchiveChannel>> GetChannelsAsync(ICloudStorage cloudStorage, ConfigFile latestProjectConfigFile, string projectIdentifier, CancellationToken cancellationToken)
+		{
+			List<CloudArchiveChannel> channels = new List<CloudArchiveChannel>();
+			try
+			{
+				string? host = null;
+				string? namespaceId = null;
+				ConfigSection? cloudStorageSection = latestProjectConfigFile.FindSection("CloudStorage");
+				if (cloudStorageSection != null)
+				{
+					namespaceId = cloudStorageSection.GetValue("Namespace", null);
+					host = cloudStorageSection.GetValue("Host", null);
+				}
+
+				if (namespaceId == null || host == null)
+				{
+					return channels;
+				}
+
+				ConfigSection? projectConfigSection = latestProjectConfigFile.FindSection(projectIdentifier);
+				if (projectConfigSection != null)
+				{
+					foreach (string archiveValue in projectConfigSection.GetValues("Archives", Array.Empty<string>()))
+					{
+						CloudArchiveChannel? channel;
+						if (CloudArchiveChannel.TryParseConfigEntryAsync(archiveValue, out channel))
+						{
+							await channel.FindArtifactsAsync(cloudStorage, projectIdentifier, host, namespaceId, cancellationToken);
+							channels.Add(channel!);
+						}
+					}
+				}
+			}
+			catch (Exception)
+			{
+			}
+			return channels;
+		}
+
+		private async Task FindArtifactsAsync(ICloudStorage cloudStorage, string projectIdentifier, string host, string namespaceId, CancellationToken cancellationToken)
+		{
+			CbWriter queryWriter = new CbWriter();
+			queryWriter.BeginObject();
+			queryWriter.BeginObject("query");
+			queryWriter.BeginObject("ugs-project");
+			queryWriter.WriteString("$eq", projectIdentifier);
+			queryWriter.EndObject(); // ugs-project
+			queryWriter.EndObject(); // query
+			queryWriter.EndObject();
+
+			CbObject query = queryWriter.ToObject();
+
+			string SanitizeBucketComponent(string s)
+			{
+				// . is reserved for bucket component separation
+				return s.Replace(".", "-", StringComparison.OrdinalIgnoreCase);
+			}
+			string projectName = Path.GetFileNameWithoutExtension(projectIdentifier);
+			string platform = "windows";
+			string bucketId = $"{SanitizeBucketComponent(projectName)}.{SanitizeBucketComponent(_archiveType)}.{SanitizeBucketComponent(_streamName)}.{SanitizeBucketComponent(platform)}";
+			bucketId = bucketId.ToLower();
+
+			IAsyncEnumerable<FoundBuildResponse> artifactResponses = cloudStorage.FindBuildAsync(host, namespaceId, bucketId, query, cancellationToken: cancellationToken);
+			await foreach (FoundBuildResponse response in artifactResponses)
+			{
+				if (Int32.TryParse(response.Commit, out int commit))
+				{
+					CloudArchive archive = new CloudArchive(cloudStorage, host, namespaceId, bucketId, response.BuildId);
+					ChangeNumberToArchive[commit] = archive;
+				}
+				// not a integer commit, unable to map to perforce
+			}
+		}
+
+		public static bool TryParseConfigEntryAsync(string text, [NotNullWhen(true)] out CloudArchiveChannel? channel)
+		{
+			ConfigObject obj = new ConfigObject(text);
+
+			string? name = obj.GetValue("Name", null);
+			if (name == null)
+			{
+				channel = null;
+				return false;
+			}
+
+			string? streamName = obj.GetValue("StreamName", null);
+			if (streamName == null)
+			{
+				channel = null;
+				return false;
+			}
+
+			string type = obj.GetValue("Type", IArchiveChannel.EditorArchiveType);
+			string archiveType = obj.GetValue("ArchiveType", "ugs-pcb");
+			bool bIgnoreRequiredBadges = obj.GetValue("bIgnoreRequiredBadges", false);
+
+			channel = new CloudArchiveChannel(name, type, archiveType, streamName, bIgnoreRequiredBadges);
+			return true;
+		}
+	}
 
 	public static class BaseArchive
 	{
-		public static async Task<List<BaseArchiveChannel>> EnumerateChannelsAsync(IPerforceConnection perforce, IHordeClient? hordeClient, ConfigFile latestProjectConfigFile, string projectIdentifier, CancellationToken cancellationToken)
+		public static async Task<List<BaseArchiveChannel>> EnumerateChannelsAsync(IPerforceConnection perforce, IHordeClient? hordeClient, ICloudStorage? cloudStorage, ConfigFile latestProjectConfigFile, string projectIdentifier, CancellationToken cancellationToken)
 		{
 			List<BaseArchiveChannel> newArchives = new List<BaseArchiveChannel>();
 			newArchives.AddRange(await PerforceArchiveChannel.GetChannelsAsync(perforce, latestProjectConfigFile, projectIdentifier, cancellationToken));
+			// prefer Cloud Storage pcbs over horde artifacts
+			if (cloudStorage != null && cloudStorage.IsEnabled(latestProjectConfigFile)) 
+			{
+				newArchives.AddRange(await CloudArchiveChannel.GetChannelsAsync(cloudStorage, latestProjectConfigFile, projectIdentifier, cancellationToken));
+			}
+
 			if (hordeClient != null)
 			{
 				newArchives.AddRange(await HordeArchiveChannel.GetChannelsAsync(hordeClient, projectIdentifier, cancellationToken));
diff --git a/Engine/Source/Programs/UnrealGameSync/UnrealGameSyncShared/CloudStorage.cs b/Engine/Source/Programs/UnrealGameSync/UnrealGameSyncShared/CloudStorage.cs
new file mode 100644
--- /dev/null
+++ b/Engine/Source/Programs/UnrealGameSync/UnrealGameSyncShared/CloudStorage.cs
@@ -0,0 +1,364 @@
+// Copyright Epic Games, Inc. All Rights Reserved.
+
+using System;
+using System.Collections;
+using System.Collections.Generic;
+using System.IO;
+using System.Runtime.CompilerServices;
+using System.Text.Json;
+using System.Threading;
+using System.Threading.Tasks;
+using EpicGames.Core;
+using EpicGames.OIDC;
+using EpicGames.Serialization;
+using Microsoft.Extensions.Configuration;
+using Microsoft.Extensions.DependencyInjection;
+using Microsoft.Extensions.Logging;
+
+namespace UnrealGameSync
+{
+	/// <summary>
+	/// A build found during search
+	/// </summary>
+	/// <param name="buildName"></param>
+	/// <param name="buildId"></param>
+	/// <param name="commit"></param>
+	public class FoundBuildResponse(string buildName, string buildId, string commit)
+	{
+		public string BuildId { get; init; } = buildId;
+		public string Commit { get; init; } = commit;
+		public string Name { get; set; } = buildName;
+	}
+
+	/// <summary>
+	/// Interface for interacting with UE Cloud Storage
+	/// </summary>
+	public interface ICloudStorage
+	{
+		/// <summary>
+		/// True if the config file enables cloud storage
+		/// </summary>
+		/// <param name="configFile">Config file to lookup cloud storage settings in</param>
+		/// <returns></returns>
+		bool IsEnabled(ConfigFile configFile);
+
+		/// <summary>
+		/// Download a build from UE Cloud Storage based on an uri
+		/// </summary>
+		/// <param name="uri">The uri</param>
+		/// <param name="outputDir">The directory to write the build to</param>
+		/// <param name="progress">Progress reporting</param>
+		/// <param name="cancellationToken">A cancellation token</param>
+		/// <returns></returns>
+		Task<bool> DownloadBuildFromUriAsync(Uri uri, DirectoryReference outputDir, IProgress<string> progress, CancellationToken cancellationToken);
+
+		/// <summary>
+		/// Download a build using all the required fields
+		/// </summary>
+		/// <param name="host">The host url</param>
+		/// <param name="namespaceId">The namespace</param>
+		/// <param name="bucketId">The bucket</param>
+		/// <param name="buildId">The id of the build</param>
+		/// <param name="outputDir">The directory to write the build to</param>
+		/// <param name="zenStateDirectory">A optional directory were zen cli should create its state. In the output dir by default</param>
+		/// <param name="progress">Progress reporting</param>
+		/// <param name="cancellationToken">A cancellation token</param>
+		/// <returns></returns>
+		Task<bool> DownloadBuildAsync(string host, string namespaceId, string bucketId, string buildId, DirectoryReference outputDir, DirectoryReference? zenStateDirectory, IProgress<string> progress, CancellationToken cancellationToken);
+
+		/// <summary>
+		/// Search for builds in UE Cloud Storage
+		/// </summary>
+		/// <param name="host">The host url</param>
+		/// <param name="namespaceId">The namespace</param>
+		/// <param name="bucketId">The bucket</param>
+		/// <param name="query">CbObject that describes the query, see UE Cloud Storage documentation</param>
+		/// <param name="cancellationToken">A cancellation token</param>
+		/// <returns></returns>
+		IAsyncEnumerable<FoundBuildResponse> FindBuildAsync(string host, string namespaceId, string bucketId, CbObject query, CancellationToken cancellationToken);
+	}
+
+	class CloudStorage(ILogger<CloudStorage> logger, ITokenStore tokenStore) : ICloudStorage, IDisposable
+	{
+		private OidcTokenManager? _tokenManager;
+		private string? _providerName;
+
+		private readonly SemaphoreSlim _createOidcTokenSemaphore = new SemaphoreSlim(1, 1);
+		public bool IsEnabled(ConfigFile configFile)
+		{
+			ConfigSection? section = configFile.FindSection("CloudStorage");
+			if (section == null)
+			{
+				return false;
+			}
+
+			return section.GetValue("Enabled", false);
+		}
+
+		public async Task CreateOidcTokenManagerAsync(Uri uri)
+		{
+			ClientAuthConfigurationV1? remoteAuthConfig = await ProviderConfigurationFactory.ReadRemoteAuthConfigurationAsync(uri, ProviderConfigurationFactory.DefaultEncryptionKey);
+			if (remoteAuthConfig == null)
+			{
+				throw new NotImplementedException("Unable to fetch remote auth configuration");
+			}
+
+			IConfiguration config = ProviderConfigurationFactory.BindOptions(remoteAuthConfig);
+
+			_tokenManager = OidcTokenManager.CreateTokenManager(config, tokenStore);
+			_providerName = remoteAuthConfig.DefaultProvider;
+		}
+
+		public async Task<bool> DownloadBuildFromUriAsync(Uri uri, DirectoryReference outputDir, IProgress<string> progress, CancellationToken cancellationToken)
+		{
+			string host = $"{uri.Scheme}://{uri.Host}";
+			string path = uri.AbsolutePath;
+			// remove api path if specified
+			path = path.Replace("api/v2/builds/", "", StringComparison.OrdinalIgnoreCase);
+			string[] components = path.Split('/', StringSplitOptions.RemoveEmptyEntries);
+			if (components.Length != 3)
+			{
+				logger.LogError("Failed to base uri path {Uri}, did not match expected format of namespace/bucket/buildId", path);
+				return true;
+			}
+
+			string namespaceId = components[0];
+			string bucketId = components[1];
+			string buildId = components[2];	
+
+			return await DownloadBuildAsync(host, namespaceId, bucketId, buildId, outputDir, zenStateDirectory: null, progress, cancellationToken);
+		}
+
+		public async Task<bool> DownloadBuildAsync(string host, string namespaceId, string bucketId, string buildId, DirectoryReference outputDir, DirectoryReference? zenStateDirectory, IProgress<string> progress, CancellationToken cancellationToken)
+		{
+			string? accessToken = await GetAccessToken(new Uri(host));
+
+			if (accessToken == null)
+			{
+				throw new Exception("Unable to download build as no access token was found");
+			}
+
+			string zenStateArg = zenStateDirectory == null ? "" : $" --zen-folder-path {zenStateDirectory} ";
+			string cmdline = $"builds download --host {host} --namespace {namespaceId} --bucket {bucketId} {zenStateArg} --plain-progress \"{outputDir}\" {buildId}";
+
+			return await RunZenCli(cmdline, accessToken, progress, cancellationToken);
+		}
+
+		public async Task<bool> RunZenCli(string cmdline, string accessToken, IProgress<string> progress, CancellationToken cancellationToken)
+		{
+			string accessTokenEnvVar = "UE-CloudDataCacheAccessToken";
+			// pass the access token via environment variable to avoid showing it on the cli
+			cmdline += $" --access-token-env {accessTokenEnvVar} ";
+
+			progress.Report("Connecting to server...");
+
+			FileInfo zenExe;
+			switch (RuntimePlatform.Current)
+			{
+				case RuntimePlatform.Type.Windows:
+					zenExe = new FileInfo("Binaries/Win64/zen.exe");
+					break;
+				case RuntimePlatform.Type.Linux:
+					zenExe = new FileInfo("Binaries/Linux/zen");
+					break;
+				case RuntimePlatform.Type.Mac:
+					zenExe = new FileInfo("Binaries/Mac/zen");
+					break;
+				default:
+					throw new ArgumentOutOfRangeException($"Unknown runtime platform {RuntimePlatform.Current}");
+			}
+
+			// we need to propagate the environment while also adding the access token we want to override
+			Dictionary<string, string> processEnv = new Dictionary<string, string>();
+			foreach (DictionaryEntry entry in Environment.GetEnvironmentVariables())
+			{
+				if (entry is { Key: string k, Value: string v })
+				{
+					processEnv[k] = v;
+				}
+			}
+
+			processEnv[accessTokenEnvVar] = accessToken;
+
+			int exitCode = await Utility.ExecuteProcessAsync(zenExe.FullName, zenExe.DirectoryName, cmdline, s =>
+			{
+				progress.Report(s);
+			}, processEnv, cancellationToken);
+			if (exitCode != 0)
+			{
+				logger.LogError("Error running {App} {Args} shutdown with exitcode {ExitCode}", zenExe, cmdline, exitCode);
+			}
+
+			return exitCode != 0;
+		}
+
+		public async IAsyncEnumerable<FoundBuildResponse> FindBuildAsync(string host, string namespaceId, string bucketId, CbObject query, [EnumeratorCancellation] CancellationToken cancellationToken)
+		{
+			string? accessToken = await GetAccessToken(new Uri(host));
+
+			if (accessToken == null)
+			{
+				throw new Exception("Unable to list builds as no access token was found");
+			}
+
+			string tempInputFile = Path.GetTempFileName();
+			string tempOutputFile = Path.GetTempFileName();
+			tempInputFile = Path.ChangeExtension(tempInputFile, "cbo");
+			tempOutputFile = Path.ChangeExtension(tempOutputFile, "cbo");
+			try
+			{
+				await File.WriteAllBytesAsync(tempInputFile, query.GetView().ToArray(), cancellationToken);
+				string cmdline = $"builds list --host {host} --namespace {namespaceId} --bucket {bucketId} --query-path {tempInputFile} --result-path {tempOutputFile}";
+
+				Progress<string> progress = new Progress<string>();
+				bool failed = await RunZenCli(cmdline, accessToken, progress, cancellationToken);
+
+				if (failed)
+				{
+					logger.LogError("Failed to list builds, zen non-zero exit code");
+					yield break;
+				}
+				byte[] b = await File.ReadAllBytesAsync(tempOutputFile, cancellationToken);
+				CbObject o = new CbObject(b);
+
+				foreach (CbField result in o["results"].AsArray())
+				{
+					CbObject metadata = result["metadata"].AsObject();
+					if (metadata.Equals(CbObject.Empty))
+					{
+						// we require metadata to be able to know which commit this build belongs to
+						continue;
+					}
+
+					string buildName = metadata["name"].AsString();
+					string buildId = result["buildId"].AsString();
+					CbField commitField = metadata["commit"];
+					if (commitField.Equals(CbField.Empty))
+					{
+						// we require metadata to be able to know which commit this build belongs to
+						continue;
+					}
+
+					string commit = "";
+					if (commitField.IsInteger())
+					{
+						commit = commitField.AsInt64().ToString();
+
+					}
+					if (commitField.IsFloat())
+					{
+						// due to jsons problematic number types we sometimes get doubles for integers, these are likely changelists and thus whole numbers so we just cast them
+						commit = ((long)commitField.AsDouble()).ToString();
+					}
+					else if (commitField.IsString())
+					{
+						commit = commitField.AsString();
+					}
+					
+					if (String.IsNullOrEmpty(commit) )
+					{
+						// if there is no commit specified we can not know which changelist this PCB belong to
+						continue;
+					}
+					yield return new FoundBuildResponse(buildName, buildId, commit);
+				}
+			}
+			finally
+			{
+				if (File.Exists(tempInputFile))
+				{
+					File.Delete(tempInputFile);
+				}
+
+				if (File.Exists(tempOutputFile))
+				{
+					File.Delete(tempOutputFile);
+				}
+			}
+		}
+
+		private async Task<string?> GetAccessToken(Uri uri)
+		{
+			if (_tokenManager == null)
+			{
+				await _createOidcTokenSemaphore.WaitAsync(); 
+				try
+				{
+					// check to see if it wasn't created while we waited for the semaphore
+#pragma warning disable CA1508
+					if (_tokenManager == null)
+#pragma warning restore CA1508
+					{
+						await CreateOidcTokenManagerAsync(uri); 
+					}
+				}
+				finally
+				{
+					_createOidcTokenSemaphore.Release();
+				}
+			}
+
+			OidcTokenInfo tokenInfo = await _tokenManager!.GetAccessToken(_providerName!, CancellationToken.None);
+			if (tokenInfo.AccessToken == null)
+			{
+				logger.LogError("Failed to get a access token");
+				return null;
+			}
+
+			return tokenInfo.AccessToken;
+		}
+
+		private void Dispose(bool disposing)
+		{
+			if (disposing)
+			{
+				_createOidcTokenSemaphore.Dispose();
+			}
+		}
+
+		public void Dispose()
+		{
+			Dispose(true);
+			GC.SuppressFinalize(this);
+		}
+
+		~CloudStorage()
+		{
+			Dispose(false);
+		}
+	}
+
+	/// <summary>
+	/// Extension methods for Cloud Storage
+	/// </summary>
+	public static class CloudStorageExtensions
+	{
+		/// <summary>
+		/// Adds Horde-related services with the default settings
+		/// </summary>
+		/// <param name="serviceCollection">Collection to register services with</param>
+		public static void AddCloudStorage(this IServiceCollection serviceCollection)
+		{
+			serviceCollection.AddLogging();
+
+			serviceCollection.AddSingleton<ITokenStore>(sp => TokenStoreFactory.CreateTokenStore());
+			serviceCollection.AddSingleton<ICloudStorage, CloudStorage>();
+		}
+
+		/// <summary>
+		/// Adds Horde-related services
+		/// </summary>
+		/// <param name="serviceCollection">Collection to register services with</param>
+		/// <param name="configureHorde">Callback to configure options</param>
+		public static void AddCloudStorage(this IServiceCollection serviceCollection, Action<CloudStorageOptions> configureHorde)
+		{
+			serviceCollection.Configure(configureHorde);
+			AddCloudStorage(serviceCollection);
+		}
+	}
+
+	public class CloudStorageOptions
+	{
+
+	}
+}
diff --git a/Engine/Source/Programs/UnrealGameSync/UnrealGameSyncShared/WorkspaceUpdate.cs b/Engine/Source/Programs/UnrealGameSync/UnrealGameSyncShared/WorkspaceUpdate.cs
--- a/Engine/Source/Programs/UnrealGameSync/UnrealGameSyncShared/WorkspaceUpdate.cs
+++ b/Engine/Source/Programs/UnrealGameSync/UnrealGameSyncShared/WorkspaceUpdate.cs
@@ -1327,9 +1327,15 @@ namespace UnrealGameSync
 					// Sync and extract (or just remove) the given archives
 					foreach ((string archiveType, IArchive? archive) in Context.ArchiveTypeToArchive)
 					{
+						// the archive type does not require cleaning we skip this
+						bool removeOldBinaries = true;
+						if (archive != null && archive.RemoveOldBinaries.HasValue)
+						{
+							removeOldBinaries = archive.RemoveOldBinaries.HasValue;
+						}
 						// Remove any existing binaries
 						FileReference manifestFileName = FileReference.Combine(manifestDirectoryName, String.Format("{0}.zipmanifest", archiveType));
-						if (FileReference.Exists(manifestFileName))
+						if (removeOldBinaries && FileReference.Exists(manifestFileName))
 						{
 							bool isNotLastSyncedEditorArchive = archiveType != IArchiveChannel.EditorArchiveType || (archive != null && archive.Key != state.LastSyncEditorArchive);
 							if (isNotLastSyncedEditorArchive)
