Fixed a bug where we fail to apply settings for the custom tools (through reflection) due to the set properties no longer existing.

diff --git a/Engine/Source/Programs/UnrealGameSync/UnrealGameSyncShared/ToolUpdateMonitor.cs b/Engine/Source/Programs/UnrealGameSync/UnrealGameSyncShared/ToolUpdateMonitor.cs
--- a/Engine/Source/Programs/UnrealGameSync/UnrealGameSyncShared/ToolUpdateMonitor.cs
+++ b/Engine/Source/Programs/UnrealGameSync/UnrealGameSyncShared/ToolUpdateMonitor.cs
@@ -40,7 +40,11 @@ namespace UnrealGameSync
 	{
 		public string? InstallCommand { get; set; }
 		public string? UninstallCommand { get; set; }
-		public List<ToolLink> StatusPanelLinks { get; } = new List<ToolLink>();
+
+		[System.Diagnostics.CodeAnalysis.SuppressMessage("Usage", "CA2227:Collection properties should be read only", 
+			Justification = "Set property accessed through reflection")]
+		public List<ToolLink> StatusPanelLinks { get; set; } = new List<ToolLink>();
+
 		public bool SafeWhenBusy { get; set; }
 	}
 
@@ -50,7 +54,11 @@ namespace UnrealGameSync
 		public Guid Id { get; set; }
 		public string Name { get; set; }
 		public string Description { get; set; }
-		public HashSet<Guid> DependsOnToolIds { get; } = new HashSet<Guid>();
+
+		[System.Diagnostics.CodeAnalysis.SuppressMessage("Usage", "CA2227:Collection properties should be read only", 
+			Justification = "Set property accessed through reflection")]
+		public HashSet<Guid> DependsOnToolIds { get; set; } = new HashSet<Guid>();
+
 		public ToolSettings Settings { get; set; } = new ToolSettings();
 		public string Revision { get; set; }
 
