Unreal Game Sync - Fixed issue finding zen executables when using a release build of UGS. Improved error messages when that happens to ease debugging in the future.

diff --git a/Engine/Source/Programs/UnrealGameSync/UnrealGameSyncShared/CloudStorage.cs b/Engine/Source/Programs/UnrealGameSync/UnrealGameSyncShared/CloudStorage.cs
--- a/Engine/Source/Programs/UnrealGameSync/UnrealGameSyncShared/CloudStorage.cs
+++ b/Engine/Source/Programs/UnrealGameSync/UnrealGameSyncShared/CloudStorage.cs
@@ -4,6 +4,7 @@ using System;
 using System.Collections;
 using System.Collections.Generic;
 using System.IO;
+using System.Reflection;
 using System.Runtime.CompilerServices;
 using System.Text.Json;
 using System.Threading;
@@ -152,22 +153,34 @@ namespace UnrealGameSync
 
 			progress.Report("Connecting to server...");
 
+			string? ugsDir = Path.GetDirectoryName(Assembly.GetEntryAssembly()?.Location);
+			if (ugsDir == null)
+			{
+				logger.LogError("Unable to determine ugs root directory so unable to find zen executable.");
+				return true;
+			}
 			FileInfo zenExe;
 			switch (RuntimePlatform.Current)
 			{
 				case RuntimePlatform.Type.Windows:
-					zenExe = new FileInfo("Binaries/Win64/zen.exe");
+					zenExe = new FileInfo(Path.Combine(ugsDir, "Binaries/Win64/zen.exe"));
 					break;
 				case RuntimePlatform.Type.Linux:
-					zenExe = new FileInfo("Binaries/Linux/zen");
+					zenExe = new FileInfo(Path.Combine(ugsDir, "Binaries/Linux/zen"));
 					break;
 				case RuntimePlatform.Type.Mac:
-					zenExe = new FileInfo("Binaries/Mac/zen");
+					zenExe = new FileInfo(Path.Combine(ugsDir, "Binaries/Mac/zen"));
 					break;
 				default:
 					throw new ArgumentOutOfRangeException($"Unknown runtime platform {RuntimePlatform.Current}");
 			}
 
+			if (!zenExe.Exists)
+			{
+				logger.LogError("Unable to locate zen executable at {Path} unable to run it.", zenExe);
+				return true;
+			}
+
 			// we need to propagate the environment while also adding the access token we want to override
 			Dictionary<string, string> processEnv = new Dictionary<string, string>();
 			foreach (DictionaryEntry entry in Environment.GetEnvironmentVariables())
