UGS: Fix compiler warnings

Please run "Analyze" configuration prior to submitting.

diff --git a/Engine/Source/Programs/UnrealGameSync/UnrealGameSync/ArtifactDownload.cs b/Engine/Source/Programs/UnrealGameSync/UnrealGameSync/ArtifactDownload.cs
--- a/Engine/Source/Programs/UnrealGameSync/UnrealGameSync/ArtifactDownload.cs
+++ b/Engine/Source/Programs/UnrealGameSync/UnrealGameSync/ArtifactDownload.cs
@@ -11,8 +11,6 @@ using EpicGames.Core;
 using EpicGames.Horde;
 using EpicGames.Horde.Storage;
 using EpicGames.Horde.Storage.Nodes;
-using EpicGames.OIDC;
-using Microsoft.Extensions.Configuration;
 using Microsoft.Extensions.DependencyInjection;
 using Microsoft.Extensions.Logging;
 using Microsoft.Win32;
diff --git a/Engine/Source/Programs/UnrealGameSync/UnrealGameSyncShared/CloudStorage.cs b/Engine/Source/Programs/UnrealGameSync/UnrealGameSyncShared/CloudStorage.cs
--- a/Engine/Source/Programs/UnrealGameSync/UnrealGameSyncShared/CloudStorage.cs
+++ b/Engine/Source/Programs/UnrealGameSync/UnrealGameSyncShared/CloudStorage.cs
@@ -6,7 +6,6 @@ using System.Collections.Generic;
 using System.IO;
 using System.Reflection;
 using System.Runtime.CompilerServices;
-using System.Text.Json;
 using System.Threading;
 using System.Threading.Tasks;
 using EpicGames.Core;
diff --git a/Engine/Source/Programs/UnrealGameSync/UnrealGameSyncShared/WorkspaceUpdate.cs b/Engine/Source/Programs/UnrealGameSync/UnrealGameSyncShared/WorkspaceUpdate.cs
--- a/Engine/Source/Programs/UnrealGameSync/UnrealGameSyncShared/WorkspaceUpdate.cs
+++ b/Engine/Source/Programs/UnrealGameSync/UnrealGameSyncShared/WorkspaceUpdate.cs
@@ -875,7 +875,6 @@ namespace UnrealGameSync
 									{
 										logger.LogInformation("Fast Sync cannot be used with current changelist {CurrentChangelist} and target changelist {TargetChangelist} because workspace syncing backwards.", state.CurrentChangeNumber, Context.ChangeNumber);
 									}
-
 								}
 							}
 
@@ -897,7 +896,7 @@ namespace UnrealGameSync
 												syncFilter = Context.Options.HasFlag(WorkspaceUpdateOptions.Sync) ? $"{Context.ChangeNumber}" : $"={Context.ChangeNumber}";
 												Context.Options &= ~WorkspaceUpdateOptions.FastSync;
 												syncTree.CanContainExcludedFiles = false;
-												logger.LogWarning($"Couldn't enumerate changes matching with fast sync enabled: {syncPath}. Disabling fast sync and retrying.");
+												logger.LogWarning("Couldn't enumerate changes matching with fast sync enabled: {SyncPath}. Disabling fast sync and retrying.", syncPath);
 
 												return await inlineGetSyncRecords(false);
 											}
@@ -978,7 +977,7 @@ namespace UnrealGameSync
 
 									// return success
 									return (WorkspaceUpdateResult.Success, String.Empty);
-								};
+								}
 
 								(WorkspaceUpdateResult, string) ret = await inlineGetSyncRecords(true);
 								if (ret.Item1 != WorkspaceUpdateResult.Success)
