UGS - Cloud Storage now again echoes the zen progress to the log file.
Also now trigger a login if the token is not valid as the sync starts.

diff --git a/Engine/Source/Programs/UnrealGameSync/UnrealGameSync/ArtifactDownload.cs b/Engine/Source/Programs/UnrealGameSync/UnrealGameSync/ArtifactDownload.cs
--- a/Engine/Source/Programs/UnrealGameSync/UnrealGameSync/ArtifactDownload.cs
+++ b/Engine/Source/Programs/UnrealGameSync/UnrealGameSync/ArtifactDownload.cs
@@ -96,9 +96,18 @@ namespace UnrealGameSync
 				await using ServiceProvider serviceProvider = services.BuildServiceProvider();
 
 				ICloudStorage cloudStorage = serviceProvider.GetRequiredService<ICloudStorage>();
-
+				ILogger logger = serviceProvider.GetRequiredService<ILogger<ArtifactDownload>>();
 				progress.Report("Connecting to server...");
-				await cloudStorage.DownloadBuildFromUriAsync(baseUri, outputDir, progress, cancellationToken);
+				bool failure = await cloudStorage.DownloadBuildFromUriAsync(baseUri, outputDir, progress, cancellationToken);
+
+				if (!failure)
+				{
+					logger.LogInformation("Download complete");
+				}
+				else
+				{
+					logger.LogError("Failed to delete from {Url} to {OutputDirectory}", baseUri, outputDir);
+				}
 			}
 		}
 
diff --git a/Engine/Source/Programs/UnrealGameSync/UnrealGameSyncShared/CloudStorage.cs b/Engine/Source/Programs/UnrealGameSync/UnrealGameSyncShared/CloudStorage.cs
--- a/Engine/Source/Programs/UnrealGameSync/UnrealGameSyncShared/CloudStorage.cs
+++ b/Engine/Source/Programs/UnrealGameSync/UnrealGameSyncShared/CloudStorage.cs
@@ -135,7 +135,8 @@ namespace UnrealGameSync
 
 			if (accessToken == null)
 			{
-				throw new Exception("Unable to download build as no access token was found");
+				logger.LogError("Unable to determine ugs root directory so unable to find zen executable.");
+				return true;
 			}
 
 			string zenStateArg = zenStateDirectory == null ? "" : $" --zen-folder-path {zenStateDirectory} ";
@@ -192,9 +193,12 @@ namespace UnrealGameSync
 
 			processEnv[accessTokenEnvVar] = accessToken;
 
+			logger.LogInformation("Running zen cli '{ZenPath} {ZenArgs}' to download build.", zenExe.FullName, cmdline);
+
 			int exitCode = await Utility.ExecuteProcessAsync(zenExe.FullName, zenExe.DirectoryName, cmdline, s =>
 			{
 				progress.Report(s);
+				logger.LogInformation("{Line}", s);
 			}, processEnv, cancellationToken);
 			if (exitCode != 0)
 			{
@@ -289,7 +293,7 @@ namespace UnrealGameSync
 			}
 		}
 
-		private async Task<string?> GetAccessToken(Uri uri)
+		private async Task<string?> GetAccessToken(Uri uri, bool allowInteractiveLogin = true)
 		{
 			if (_tokenManager == null)
 			{
@@ -310,13 +314,23 @@ namespace UnrealGameSync
 				}
 			}
 
-			OidcTokenInfo tokenInfo = await _tokenManager!.GetAccessToken(_providerName!, CancellationToken.None);
-			if (tokenInfo.AccessToken == null)
+			OidcTokenInfo tokenInfo;
+			try
 			{
-				logger.LogError("Failed to get a access token");
-				return null;
+				tokenInfo = await _tokenManager!.GetAccessToken(_providerName!, CancellationToken.None);
+			}
+			catch (NotLoggedInException)
+			{
+				if (allowInteractiveLogin)
+				{
+					tokenInfo = await _tokenManager!.LoginAsync(_providerName!, CancellationToken.None);
+				}
+				else
+				{
+					logger.LogError("Failed to get a access token");
+					return null;
+				}
 			}
-
 			return tokenInfo.AccessToken;
 		}
 
