- check that the tool required by a build step is installed before running it

diff --git a/Engine/Source/Programs/UnrealGameSync/UnrealGameSyncCmd/Program.cs b/Engine/Source/Programs/UnrealGameSync/UnrealGameSyncCmd/Program.cs
--- a/Engine/Source/Programs/UnrealGameSync/UnrealGameSyncCmd/Program.cs
+++ b/Engine/Source/Programs/UnrealGameSync/UnrealGameSyncCmd/Program.cs
@@ -1440,53 +1440,153 @@ namespace UnrealGameSyncCmd
 
 				UserProjectSettings projectSettings = context.UserSettings.FindOrAddProjectSettings(projectInfo, settings, logger);
 
-				if (options.ListOnly)
-				{
-					ConfigFile projectConfig = await ConfigUtils.ReadProjectConfigFileAsync(perforceClient, projectInfo, logger, CancellationToken.None);
-
-					FileReference editorTarget = ConfigUtils.GetEditorTargetFile(projectInfo, projectConfig);
+				ConfigFile projectConfig = await ConfigUtils.ReadProjectConfigFileAsync(perforceClient, projectInfo, logger, CancellationToken.None);
+				FileReference editorTarget = ConfigUtils.GetEditorTargetFile(projectInfo, projectConfig);
 
-					Dictionary<Guid, ConfigObject> buildStepObjects = ConfigUtils.GetDefaultBuildStepObjects(projectInfo, editorTarget.GetFileNameWithoutAnyExtensions(), EditorConfig, projectConfig, false);
+				Dictionary<Guid, ConfigObject> buildStepObjects = ConfigUtils.GetDefaultBuildStepObjects(projectInfo, editorTarget.GetFileNameWithoutAnyExtensions(), EditorConfig, projectConfig, false);
 
+				if (context.GlobalSettings != null)
+				{
 					BuildStep.MergeBuildStepObjects(buildStepObjects, projectConfig.GetValues("Build.Step", Array.Empty<string>()).Select(x => new ConfigObject(x)));
 					BuildStep.MergeBuildStepObjects(buildStepObjects, projectSettings.BuildSteps);
+				}
+
+				if (options.ListOnly)
+				{
+					List<BuildStep> buildSteps = buildStepObjects.Values.Select(x => new BuildStep(x)).OrderBy(x => x.NormalSync).ToList();
+
+					int longestDescription = buildSteps.Max(b => String.IsNullOrWhiteSpace(b.Description) ? 0 : b.Description.Length);
+					string descriptionHeader = "Description";
+					string descriptionSpace = new string(Enumerable.Repeat(' ', longestDescription + 1 - descriptionHeader.Length).ToArray());
+					string descriptionDashes = new string(Enumerable.Repeat('-', longestDescription + 2).ToArray());
+					string format = $"  {{Id,-36}} | {{Name,-{longestDescription}}} | {{Type,-10}} | {{Enabled,-8}}";
 
 					logger.LogInformation("Available build steps:");
 					logger.LogInformation("");
-					logger.LogInformation("  Id                                   | Description                              | Type       | Enabled");
-					logger.LogInformation("  -------------------------------------|------------------------------------------|------------|-----------------");
-					foreach (BuildStep buildStep in buildStepObjects.Values.Select(x => new BuildStep(x)).OrderBy(x => x.OrderIndex))
+					logger.LogInformation("  Id                                   | {descriptionHeader}{descriptionSpace}| Type       | Enabled", descriptionHeader, descriptionSpace);
+					logger.LogInformation("  -------------------------------------|{descriptionDashes}|------------|-----------------", descriptionDashes);
+					foreach (BuildStep buildStep in buildStepObjects.Values.Select(x => new BuildStep(x)).OrderBy(x => x.NormalSync))
 					{
-						logger.LogInformation("  {Id,-36} | {Name,-40} | {Type,-10} | {Enabled,-8}", buildStep.UniqueId, buildStep.Description, buildStep.Type, buildStep.NormalSync);
+						logger.LogInformation(format, buildStep.UniqueId, buildStep.Description, buildStep.Type, buildStep.NormalSync);
 					}
 					return;
 				}
 
-				HashSet<Guid>? steps = null;
+				HashSet<Guid>? steps = new HashSet<Guid>();
 				if (target != null)
 				{
-					Guid id;
-					if (!Guid.TryParse(target, out id))
+					if (!Guid.TryParse(target, out Guid id))
 					{
 						logger.LogError("Unable to parse '{Target}' as a GUID. Pass -List to show all available build steps and their identifiers.", target);
+						return;
 					}
-					steps = new HashSet<Guid> { id };
+
+					steps.Add(id);
 				}
 
-				WorkspaceUpdateContext updateContext = new WorkspaceUpdateContext(
-					state.Current.CurrentChangeNumber,
-					WorkspaceUpdateOptions.Build,
-					BuildConfig.Development,
-					null,
-					projectSettings.BuildSteps, steps);
+				// check that the tools are installed
+				List<Guid> uninstalledTools = new List<Guid>();
+				if (context.GlobalSettings != null)
+				{
+					using ToolUpdateMonitor? toolsMonitor = await GetToolsMonitor(logger, context, perforceClient);
 
-				WorkspaceUpdate update = new WorkspaceUpdate(updateContext);
-				(WorkspaceUpdateResult result, string message) = await update.ExecuteAsync(perforceClient.Settings, projectInfo, state, context.Logger, CancellationToken.None);
+					if (toolsMonitor != null)
+					{
+						foreach (Guid step in steps)
+						{
+							if (buildStepObjects.TryGetValue(step, out ConfigObject? config))
+							{
+								BuildStep buildStep = new BuildStep(config);
+								if (CanRunStep(buildStep, toolsMonitor))
+								{
+									continue;
+								}
 
-				if (result != WorkspaceUpdateResult.Success)
+								logger.LogWarning($"Build step '{buildStep.Description}' cannot be run as the tool is not installed.");
+								uninstalledTools.Add(step);
+							}
+						}
+					}
+				}
+
+				foreach (Guid uninstalledTool in uninstalledTools)
+				{
+					steps.Remove(uninstalledTool);
+				}
+
+				if (steps.Any())
+				{
+					WorkspaceUpdateContext updateContext = new WorkspaceUpdateContext(
+						state.Current.CurrentChangeNumber,
+						WorkspaceUpdateOptions.Build,
+						BuildConfig.Development,
+						null,
+						projectSettings.BuildSteps, steps);
+
+					WorkspaceUpdate update = new WorkspaceUpdate(updateContext);
+					(WorkspaceUpdateResult result, string message) = await update.ExecuteAsync(perforceClient.Settings, projectInfo, state, context.Logger, CancellationToken.None);
+
+					if (result != WorkspaceUpdateResult.Success)
+					{
+						throw new UserErrorException("{Message}", message);
+					}
+				}
+			}
+
+			private async Task<ToolUpdateMonitor?> GetToolsMonitor(ILogger logger, CommandContext context, IPerforceConnection perforceClient)
+			{
+				ToolsCommandCommandOptions options = context.Arguments.ApplyTo<ToolsCommandCommandOptions>(logger);
+				context.Arguments.CheckAllArgumentsUsed();
+
+				DirectoryReference dataFolder = DirectoryReference.Combine(DirectoryReference.GetSpecialFolder(Environment.SpecialFolder.LocalApplicationData)!, "UnrealGameSync");
+				DirectoryReference.CreateDirectory(dataFolder);
+
+				// create a temporary service provider for the tool update monitor
+				ServiceCollection services = new ServiceCollection();
+				services.AddSingleton<IAsyncDisposer, AsyncDisposer>();
+				services.AddSingleton(sp => TokenStoreFactory.CreateTokenStore());
+				services.AddSingleton<OidcTokenManager>();
+
+				LauncherSettings launcherSettings = new LauncherSettings();
+				launcherSettings.Read();
+
+				if (launcherSettings.HordeServer != null)
+				{
+					services.AddHorde(options =>
+					{
+						options.ServerUrl = new Uri(launcherSettings.HordeServer);
+						options.AllowAuthPrompt = false;
+					});
+				}
+
+				ServiceProvider serviceProvider = services.BuildServiceProvider();
+
+				ToolUpdateMonitor toolUpdateMonitor =
+					new ToolUpdateMonitor(perforceClient.Settings, dataFolder, context.GlobalSettings!, logger, serviceProvider);
+
+				// get the list of tools available
+				logger.LogInformation("Retrieving tools information, please wait.");
+				await toolUpdateMonitor.GetDataFromBackendAsync();
+
+				return toolUpdateMonitor;
+			}
+
+			private bool CanRunStep(BuildStep step, ToolUpdateMonitor toolsMonitor)
+			{
+				if (step.ToolId != Guid.Empty)
 				{
-					throw new UserErrorException("{Message}", message);
+					string? toolName = toolsMonitor.GetToolName(step.ToolId);
+					if (toolName == null)
+					{
+						return false;
+					}
+
+					if (toolsMonitor.GetToolPath(toolName) == null)
+					{
+						return false;
+					}
 				}
+				return true;
 			}
 		}
 
