PR #13089: Fix UGS not showing investigations as "resolved" when using Horde's metadata endpoint

diff --git a/Engine/Source/Programs/Horde/Plugins/Build/HordeServer.Build/Ugs/UgsMetadataCollection.cs b/Engine/Source/Programs/Horde/Plugins/Build/HordeServer.Build/Ugs/UgsMetadataCollection.cs
--- a/Engine/Source/Programs/Horde/Plugins/Build/HordeServer.Build/Ugs/UgsMetadataCollection.cs
+++ b/Engine/Source/Programs/Horde/Plugins/Build/HordeServer.Build/Ugs/UgsMetadataCollection.cs
@@ -175,7 +175,7 @@ namespace HordeServer.Ugs
 
 					long? newSyncTime = synced.HasValue ? (synced.Value ? (long?)DateTime.UtcNow.Ticks : null) : userData.SyncTime;
 					UgsUserVote? newVote = vote.HasValue ? ((vote.Value != UgsUserVote.None) ? vote : null) : userData.Vote;
-					bool? newInvestigating = investigating.HasValue ? (investigating.Value ? investigating : null) : userData.Investigating;
+					bool? newInvestigating = investigating ?? userData.Investigating;
 					bool? newStarred = starred.HasValue ? (starred.Value ? starred : null) : userData.Starred;
 
 					List<UpdateDefinition<UgsMetadataDocument>> updates = new List<UpdateDefinition<UgsMetadataDocument>>();
diff --git a/Engine/Source/Programs/UnrealGameSync/UnrealGameSync/EventMonitor.cs b/Engine/Source/Programs/UnrealGameSync/UnrealGameSync/EventMonitor.cs
--- a/Engine/Source/Programs/UnrealGameSync/UnrealGameSync/EventMonitor.cs
+++ b/Engine/Source/Programs/UnrealGameSync/UnrealGameSync/EventMonitor.cs
@@ -239,8 +239,15 @@ namespace UnrealGameSync
 		readonly IAsyncDisposer _asyncDisposer;
 		readonly LatestData _latestIds;
 		HashSet<int> _filterChangeNumbers = new HashSet<int>();
-		readonly List<EventData> _investigationEvents = new List<EventData>();
-		List<EventData>? _activeInvestigations;
+
+		// Investigation events are collected in descending CL# order
+		readonly SortedSet<EventData> _investigationEvents = new SortedSet<EventData>(
+			Comparer<EventData>.Create((first, second) =>
+			{
+				int clOrder = second.Change.CompareTo(first.Change);
+				return clOrder != 0 ? clOrder : first.Id.CompareTo(second.Id);
+			}));
+		SortedSet<EventData>? _activeInvestigations = null;
 
 		// MetadataV2
 		readonly string _metadataStream;
@@ -421,16 +428,8 @@ namespace UnrealGameSync
 			}
 			else if (evt.Type == EventType.Investigating || evt.Type == EventType.Resolved)
 			{
-				// Insert it sorted in the investigation list
-				int insertIdx = 0;
-				while (insertIdx < _investigationEvents.Count && _investigationEvents[insertIdx].Id < evt.Id)
-				{
-					insertIdx++;
-				}
-				if (insertIdx == _investigationEvents.Count || _investigationEvents[insertIdx].Id != evt.Id)
-				{
-					_investigationEvents.Insert(insertIdx, evt);
-				}
+				// Events are automatically sorted by descending CL#
+				_investigationEvents.Add(evt);
 				_activeInvestigations = null;
 			}
 			else if (evt.Type == EventType.Syncing)
@@ -1101,40 +1100,20 @@ namespace UnrealGameSync
 		{
 			if (_activeInvestigations == null)
 			{
-				// Insert investigation events into the active list, sorted by change number.
-				_activeInvestigations = new List<EventData>();
+				// Collected investigation events are sorted by descending CL#
+				// Use the comparison logic below to only keep the last such event for each user
+				_activeInvestigations = new SortedSet<EventData>(
+					Comparer<EventData>.Create((first, second) =>
+						first.UserName.Equals(second.UserName, StringComparison.OrdinalIgnoreCase)
+							? 0
+							: first.Change.CompareTo(second.Change)));
 				foreach (EventData investigationEvent in _investigationEvents)
 				{
-					if (_filterChangeNumbers.Contains(investigationEvent.Change))
-					{
-						if (investigationEvent.Type == EventType.Investigating)
-						{
-							int insertIdx = 0;
-							while (insertIdx < _activeInvestigations.Count && _activeInvestigations[insertIdx].Change > investigationEvent.Change)
-							{
-								insertIdx++;
-							}
-							_activeInvestigations.Insert(insertIdx, investigationEvent);
-						}
-						else
-						{
-							_activeInvestigations.RemoveAll(x => String.Equals(x.UserName, investigationEvent.UserName, StringComparison.OrdinalIgnoreCase) && x.Change <= investigationEvent.Change);
-						}
-					}
+					_activeInvestigations.Add(investigationEvent);
 				}
 
-				// Remove any duplicate users
-				for (int idx = 0; idx < _activeInvestigations.Count; idx++)
-				{
-					for (int otherIdx = 0; otherIdx < idx; otherIdx++)
-					{
-						if (String.Equals(_activeInvestigations[idx].UserName, _activeInvestigations[otherIdx].UserName, StringComparison.OrdinalIgnoreCase))
-						{
-							_activeInvestigations.RemoveAt(idx--);
-							break;
-						}
-					}
-				}
+				// Only keep investigation start events
+				_activeInvestigations.RemoveWhere(e => e.Type == EventType.Resolved);
 			}
 		}
 
