Write projectstore file to the correct directory when importing snapshots via UGS.

[RN] none

diff --git a/Engine/Source/Programs/UnrealGameSync/UnrealGameSyncShared/WorkspaceUpdate.cs b/Engine/Source/Programs/UnrealGameSync/UnrealGameSyncShared/WorkspaceUpdate.cs
--- a/Engine/Source/Programs/UnrealGameSync/UnrealGameSyncShared/WorkspaceUpdate.cs
+++ b/Engine/Source/Programs/UnrealGameSync/UnrealGameSyncShared/WorkspaceUpdate.cs
@@ -1411,6 +1411,9 @@ namespace UnrealGameSync
 						stopwatch.AddData(new { MachineName = System.Net.Dns.GetHostName(), DomainName = Environment.UserDomainName, ServerAndPort = perforce.Settings.ServerAndPort, ClientName = perforce.Settings.ClientName, UserName = perforce.Settings.UserName });
 						List<Task<int>> snapshotTasks = new List<Task<int>>();
 
+						DirectoryReference batchFilesDir = DirectoryReference.Combine(project.LocalRootPath, "Engine", "Build", "BatchFiles");
+						FileReference runUbt = FileReference.Combine(batchFilesDir, $"RunUAT.{ShellScriptExt}");
+
 						Action<string, string> runApp = (string app, string cmdLine) =>
 						{
 							logger.LogInformation("zen> Running Zen command {App} {CommandLine}", app, cmdLine);
@@ -1420,7 +1423,7 @@ namespace UnrealGameSync
 						try
 						{
 							BuildIndex buildIndex = await UnrealCloudDDCBuildIndex.CreateBuildIndexAsync(Context.SnapshotHost, Context.SnapshotNamespace, (string app, string cmdLine) => runApp(app, cmdLine), true);
-							DirectoryReference cookedPlatformDir = DirectoryReference.Combine(project.EngineDir, "Saved", "Cooked", "{Platform}");
+							DirectoryReference cookedPlatformDir = DirectoryReference.Combine(project.ProjectDir, "Saved", "Cooked", "{Platform}");
 
 							foreach (SnapshotSettings snapshot in Context.Snapshots)
 							{
