Write projectstore file to the correct directory when importing snapshots via UGS.

[RN] none

diff --git a/Engine/Source/Programs/UnrealGameSync/UnrealGameSyncShared/WorkspaceUpdate.cs b/Engine/Source/Programs/UnrealGameSync/UnrealGameSyncShared/WorkspaceUpdate.cs
--- a/Engine/Source/Programs/UnrealGameSync/UnrealGameSyncShared/WorkspaceUpdate.cs
+++ b/Engine/Source/Programs/UnrealGameSync/UnrealGameSyncShared/WorkspaceUpdate.cs
@@ -1420,7 +1420,7 @@ namespace UnrealGameSync
 						try
 						{
 							BuildIndex buildIndex = await UnrealCloudDDCBuildIndex.CreateBuildIndexAsync(Context.SnapshotHost, Context.SnapshotNamespace, (string app, string cmdLine) => runApp(app, cmdLine), true);
-							DirectoryReference cookedPlatformDir = DirectoryReference.Combine(project.EngineDir, "Saved", "Cooked", "{Platform}");
+							DirectoryReference cookedPlatformDir = DirectoryReference.Combine(project.ProjectDir, "Saved", "Cooked", "{Platform}");
 
 							foreach (SnapshotSettings snapshot in Context.Snapshots)
 							{
