removed the check to see if the build had custom steps preventing the common build steps to be run

diff --git a/Engine/Source/Programs/UnrealGameSync/UnrealGameSyncCmd/Program.cs b/Engine/Source/Programs/UnrealGameSync/UnrealGameSyncCmd/Program.cs
--- a/Engine/Source/Programs/UnrealGameSync/UnrealGameSyncCmd/Program.cs
+++ b/Engine/Source/Programs/UnrealGameSync/UnrealGameSyncCmd/Program.cs
@@ -1514,22 +1514,19 @@ namespace UnrealGameSyncCmd
 					steps.Remove(uninstalledTool);
 				}
 
-				if (steps.Any())
-				{
-					WorkspaceUpdateContext updateContext = new WorkspaceUpdateContext(
-						state.Current.CurrentChangeNumber,
-						WorkspaceUpdateOptions.Build,
-						BuildConfig.Development,
-						null,
-						projectSettings.BuildSteps, steps);
+				WorkspaceUpdateContext updateContext = new WorkspaceUpdateContext(
+					state.Current.CurrentChangeNumber,
+					WorkspaceUpdateOptions.Build,
+					BuildConfig.Development,
+					null,
+					projectSettings.BuildSteps, steps);
 
-					WorkspaceUpdate update = new WorkspaceUpdate(updateContext);
-					(WorkspaceUpdateResult result, string message) = await update.ExecuteAsync(perforceClient.Settings, projectInfo, state, context.Logger, CancellationToken.None);
+				WorkspaceUpdate update = new WorkspaceUpdate(updateContext);
+				(WorkspaceUpdateResult result, string message) = await update.ExecuteAsync(perforceClient.Settings, projectInfo, state, context.Logger, CancellationToken.None);
 
-					if (result != WorkspaceUpdateResult.Success)
-					{
-						throw new UserErrorException("{Message}", message);
-					}
+				if (result != WorkspaceUpdateResult.Success)
+				{
+					throw new UserErrorException("{Message}", message);
 				}
 			}
 
