[UE-210601] Add code-sign step to UGS build graph

diff --git a/Engine/Source/Programs/UnrealGameSync/BuildUGS.xml b/Engine/Source/Programs/UnrealGameSync/BuildUGS.xml
--- a/Engine/Source/Programs/UnrealGameSync/BuildUGS.xml
+++ b/Engine/Source/Programs/UnrealGameSync/BuildUGS.xml
@@ -9,6 +9,7 @@
 	<Option Name="IncludeInstaller" DefaultValue="false" Description="Whether to include the launcher installer in the regular client distribution"/>
 	<Option Name="HordeSettings" DefaultValue="" Description="Path to a JSON file containing settings for publishing to Horde"/>
 	<Option Name="HordeSettingsMac" DefaultValue="" Description="Path to a JSON file containing settings for publishing to Horde"/>
+	<Option Name="SignBinaries" DefaultValue="false" Description="Include code signing step for binaries and installer"/>
 
 	<Property Name="UgsDir" Value="$(RootDir)/Engine/Source/Programs/UnrealGameSync"/>
 	<Property Name="UgsOutputDir" Value="$(RootDir)/LocalBuilds/UnrealGameSync"/>
@@ -43,6 +44,8 @@
 			<DotNet Arguments="publish &quot;$(UgsDir)/UnrealGameSyncCmd/UnrealGameSyncCmd.csproj&quot; --output=&quot;$(UgsPublishDir)&quot; --no-self-contained -c Release -r linux-x64 $(VersionArguments)"/>
 			<DotNet Arguments="publish &quot;$(UgsDir)/UnrealGameSyncCmd/UnrealGameSyncCmd.csproj&quot; --output=&quot;$(UgsPublishDir)&quot; --self-contained -c Release -r win-x64 $(VersionArguments)"/>
 			<DotNet Arguments="publish &quot;$(UgsDir)/UnrealGameSync/UnrealGameSync.csproj&quot; --output=&quot;$(UgsPublishDir)&quot; --self-contained -c ReleaseAutoUpdate -r win-x64 $(VersionArguments)"/>
+
+			<Sign Files="$(UgsPublishDir)/..." If="$(IsBuildMachine) and $(SignBinaries)" />
 		</Node>
 		<Node Name="Build UnrealGameSync Installer" Requires="Build UnrealGameSync">
 			<Property Name="WixDir" Value="$(RootDir)/Engine/Source/ThirdParty/WiX/3.8"/>
@@ -52,6 +55,8 @@
 			<Spawn Exe="$(WixDir)/light.exe" Arguments="-out &quot;bin/Release/UnrealGameSync.msi&quot; -pdbout &quot;bin/Release/UnrealGameSync.wixpdb&quot; -cultures:null -ext ..\..\..\ThirdParty\WiX\3.8\WixUtilExtension.dll -sice:ICE69 &quot;obj\Release\Product.wixobj&quot; &quot;obj\Release\UGSLauncher.wixobj&quot;" WorkingDir="$(UgsInstallerDir)"/>
 
 			<Copy From="$(UgsInstallerDir)/bin/Release/UnrealGameSync.msi" To="$(UgsInstallerPublishDir)/UnrealGameSync.msi"/>
+			
+			<Sign Files="$(UgsInstallerPublishDir)/UnrealGameSync.msi" If="$(IsBuildMachine) and $(SignBinaries)" />
 		</Node>
 		<Node Name="Publish UnrealGameSync Installer Horde" Requires="Build UnrealGameSync Installer">
 			<Error Message="Missing -HordeSettings=... argument" If="'$(HordeSettings)' == ''"/>
