Fixed a bug where presets would be added to the global sync filters if the tab with the preset had it's sync filter view opened and the okay button pressed.

diff --git a/Engine/Source/Programs/UnrealGameSync/UnrealGameSync/Forms/SyncFilter.cs b/Engine/Source/Programs/UnrealGameSync/UnrealGameSync/Forms/SyncFilter.cs
--- a/Engine/Source/Programs/UnrealGameSync/UnrealGameSync/Forms/SyncFilter.cs
+++ b/Engine/Source/Programs/UnrealGameSync/UnrealGameSync/Forms/SyncFilter.cs
@@ -187,7 +187,7 @@ namespace UnrealGameSync
 			newGlobalFilter.AllProjectsInSln = GlobalControl.IncludeAllProjectsInSolution.Checked;
 			newGlobalFilter.UprojectSpecificSln = GlobalControl.GenerateUprojectSpecificSolution.Checked;
 
-			Dictionary<Guid, bool> globalSyncCategories = GetCategorySettings(GlobalControl.CategoriesCheckList, GlobalFilter.GetCategories());
+			Dictionary<Guid, bool> globalSyncCategories = GetCategorySettings(GlobalControl.CategoriesCheckList, GlobalFilter.GetCategories(), includeLockedItems: false);
 			newGlobalFilter.SetCategories(WorkspaceSyncCategory.GetDelta(defaultSyncCategories, globalSyncCategories));
 
 			newWorkspaceFilter = new FilterSettings();
@@ -200,13 +200,20 @@ namespace UnrealGameSync
 			newWorkspaceFilter.SetCategories(WorkspaceSyncCategory.GetDelta(globalSyncCategories, workspaceSyncCategories));
 		}
 
-		private Dictionary<Guid, bool> GetCategorySettings(CheckedListBox listBox, IEnumerable<KeyValuePair<Guid, bool>> originalSettings)
+		private Dictionary<Guid, bool> GetCategorySettings(CheckedListBox listBox, IEnumerable<KeyValuePair<Guid, bool>> originalSettings, bool includeLockedItems = true)
 		{
+			HashSet<Guid> locked = new HashSet<Guid>();
 			Dictionary<Guid, bool> result = new Dictionary<Guid, bool>();
 			for (int idx = 0; idx < listBox.Items.Count; idx++)
 			{
 				if (listBox.Items[idx] is CheckListItem item)
 				{
+					if (item.Locked && includeLockedItems == false)
+					{
+						locked.Add(item.UniqueId);
+						continue;
+					}
+
 					Guid uniqueId = item.UniqueId;
 					if (!result.ContainsKey(uniqueId))
 					{
@@ -216,7 +223,7 @@ namespace UnrealGameSync
 			}
 			foreach (KeyValuePair<Guid, bool> originalSetting in originalSettings)
 			{
-				if (!_uniqueIdToCategory.ContainsKey(originalSetting.Key))
+				if (!_uniqueIdToCategory.ContainsKey(originalSetting.Key) || locked.Contains(originalSetting.Key))
 				{
 					result[originalSetting.Key] = originalSetting.Value;
 				}
