Build and run ZenLaunch before importing snapshots via UGS.

[RN] none

diff --git a/Engine/Source/Programs/UnrealGameSync/UnrealGameSyncShared/WorkspaceUpdate.cs b/Engine/Source/Programs/UnrealGameSync/UnrealGameSyncShared/WorkspaceUpdate.cs
--- a/Engine/Source/Programs/UnrealGameSync/UnrealGameSyncShared/WorkspaceUpdate.cs
+++ b/Engine/Source/Programs/UnrealGameSync/UnrealGameSyncShared/WorkspaceUpdate.cs
@@ -1412,6 +1412,43 @@ namespace UnrealGameSync
 						stopwatch.AddData(new { MachineName = System.Net.Dns.GetHostName(), DomainName = Environment.UserDomainName, ServerAndPort = perforce.Settings.ServerAndPort, ClientName = perforce.Settings.ClientName, UserName = perforce.Settings.UserName });
 						List<Task<int>> snapshotTasks = new List<Task<int>>();
 
+						string hostPlatform;
+
+						if (RuntimeInformation.IsOSPlatform(OSPlatform.Windows))
+						{
+							hostPlatform = "Win64";
+						}
+						else if (RuntimeInformation.IsOSPlatform(OSPlatform.OSX))
+						{
+							hostPlatform = "Mac";
+						}
+						else
+						{
+							hostPlatform = "Linux";
+						}
+
+						// If we're not syncing a precompiled editor archive, make sure ZenLaunch is built and up to date.
+						// Slightly annoying that we can't just use Context.Options.HasFlag(WorkspaceUpdateOptions.SyncArchives) here,
+						// as that flag gets set sometimes even if precompiled binaries are disabled, in order to indicate that some
+						// old archives should be cleaned up.
+						//
+						// So instead, we just check if we have any editor archives that are requested to sync to a new version,
+						// indicated by a non-null value in ArchiveTypeToArchive.
+						if (!Context.ArchiveTypeToArchive.Any(x => x.Key == IArchiveChannel.EditorArchiveType && x.Value != null))
+						{
+							DirectoryReference batchFilesDir = DirectoryReference.Combine(project.EngineDir, "Build", "BatchFiles");
+							FileReference runUbt = FileReference.Combine(batchFilesDir, $"RunUBT.{ShellScriptExt}");
+
+							string ubtCommandLine = $"\"{runUbt}\" ZenLaunch {hostPlatform} development";
+							logger.LogInformation("ubt> Building ZenLaunch with command line {CommandLine}", ubtCommandLine);
+							await ExecuteShellCommandAsync(ubtCommandLine, line => ProcessOutput(line, "ubt> ", Progress, logger), cancellationToken);
+						}
+
+						FileReference zenLaunchExe = new FileReference(Path.Combine(project.EngineDir.FullName, "Binaries", hostPlatform, String.Format("ZenLaunch{0}", RuntimePlatform.ExeExtension)));
+						string zenLaunchCommandline = $"\"{zenLaunchExe}\" \"{project.LocalFileName.FullName}\" -SponsorProcessID={Environment.ProcessId}";
+						logger.LogInformation("zen> Running ZenLaunch with command line {CommandLine}", zenLaunchCommandline);
+						await ExecuteShellCommandAsync(zenLaunchCommandline, line => ProcessOutput(line, "zen> ", Progress, logger), cancellationToken);
+
 						Action<string, string> runApp = (string app, string cmdLine) =>
 						{
 							logger.LogInformation("zen> Running Zen command {App} {CommandLine}", app, cmdLine);
