Add "sync local projects" option to UGS.

When this is enabled, UGS will sync all directories which contain a uproject file. This is different from "sync all projects in stream", as it only processes projects which were previously synced. It enables workflows where a random assortment of projects are used - just sync the relevant uprojects manually, and UGS will then keep them up to date on future syncs.

diff --git a/Engine/Source/Programs/UnrealGameSync/UnrealGameSync/Controls/WorkspaceControl.cs b/Engine/Source/Programs/UnrealGameSync/UnrealGameSync/Controls/WorkspaceControl.cs
--- a/Engine/Source/Programs/UnrealGameSync/UnrealGameSync/Controls/WorkspaceControl.cs
+++ b/Engine/Source/Programs/UnrealGameSync/UnrealGameSync/Controls/WorkspaceControl.cs
@@ -6190,7 +6190,9 @@ namespace UnrealGameSync
 			}
 
 			string[] combinedSyncFilter = UserSettings.GetCombinedSyncFilter(GetSyncCategories(), _workspaceSettings.Preset, GetPresets(), _settings.Global.Filter, _workspaceSettings.Filter, _perforceMonitor.LatestPerforceConfigSection());
-			List<string> syncPaths = WorkspaceUpdate.GetSyncPaths(_workspace.Project, _workspaceSettings.Filter.AllProjects ?? _settings.Global.Filter.AllProjects ?? false, combinedSyncFilter);
+
+			WorkspaceUpdateOptions options = WorkspaceUpdateContext.GetOptionsFromConfig(_settings.Global, _workspaceSettings);
+			List<string> syncPaths = WorkspaceUpdate.GetSyncPaths(_workspace.Project, options, combinedSyncFilter);
 
 			bool performedClean = CleanWorkspaceWindow.DoClean(ParentForm!, _perforceSettings, BranchDirectoryName, _workspace.Project.ClientRootPath, syncPaths, extraSafeToDeleteFolders.Split('\n'), extraSafeToDeleteExtensions.Split('\n'), _serviceProvider.GetRequiredService<ILogger<CleanWorkspaceWindow>>());
 
diff --git a/Engine/Source/Programs/UnrealGameSync/UnrealGameSync/Forms/SyncFilter.cs b/Engine/Source/Programs/UnrealGameSync/UnrealGameSync/Forms/SyncFilter.cs
--- a/Engine/Source/Programs/UnrealGameSync/UnrealGameSync/Forms/SyncFilter.cs
+++ b/Engine/Source/Programs/UnrealGameSync/UnrealGameSync/Forms/SyncFilter.cs
@@ -74,6 +74,7 @@ namespace UnrealGameSync
 			GlobalControl.SetView(_roleName, roleViews.ToArray(), GlobalFilter.View.ToArray());
 			SetExcludedCategories(GlobalControl.CategoriesCheckList, _uniqueIdToCategory, role, syncCategories);
 			GlobalControl.SyncAllProjects.Checked = GlobalFilter.AllProjects ?? false;
+			GlobalControl.SyncLocalProjects.Checked = GlobalFilter.LocalProjects ?? false;
 			GlobalControl.IncludeAllProjectsInSolution.Checked = GlobalFilter.AllProjectsInSln ?? false;
 			GlobalControl.GenerateUprojectSpecificSolution.Checked = GlobalFilter.UprojectSpecificSln ?? false;
 
@@ -83,11 +84,13 @@ namespace UnrealGameSync
 
 			WorkspaceControl.CategoriesCheckList.ItemCheck += WorkspaceControl_CategoriesCheckList_ItemCheck;
 			WorkspaceControl.SyncAllProjects.Checked = WorkspaceFilter.AllProjects ?? GlobalFilter.AllProjects ?? false;
+			WorkspaceControl.SyncLocalProjects.Checked = WorkspaceFilter.LocalProjects ?? GlobalFilter.LocalProjects ?? false;
 			WorkspaceControl.IncludeAllProjectsInSolution.Checked = WorkspaceFilter.AllProjectsInSln ?? GlobalFilter.AllProjectsInSln ?? false;
 			WorkspaceControl.GenerateUprojectSpecificSolution.Checked = WorkspaceFilter.UprojectSpecificSln ?? GlobalFilter.UprojectSpecificSln ?? false;
 
 			GlobalControl.CategoriesCheckList.ItemCheck += GlobalControl_CategoriesCheckList_ItemCheck;
 			GlobalControl.SyncAllProjects.CheckStateChanged += GlobalControl_SyncAllProjects_CheckStateChanged;
+			GlobalControl.SyncLocalProjects.CheckStateChanged += GlobalControl_SyncLocalProjects_CheckStateChanged;
 			GlobalControl.IncludeAllProjectsInSolution.CheckStateChanged += GlobalControl_IncludeAllProjectsInSolution_CheckStateChanged;
 			GlobalControl.GenerateUprojectSpecificSolution.CheckStateChanged += GlobalControl_GenerateUprojectSpecificSolution_CheckStateChanged;
 		}
@@ -125,6 +128,11 @@ namespace UnrealGameSync
 			WorkspaceControl.SyncAllProjects.Checked = GlobalControl.SyncAllProjects.Checked;
 		}
 
+		private void GlobalControl_SyncLocalProjects_CheckStateChanged(object? sender, EventArgs e)
+		{
+			WorkspaceControl.SyncLocalProjects.Checked = GlobalControl.SyncLocalProjects.Checked;
+		}
+
 		private void GlobalControl_IncludeAllProjectsInSolution_CheckStateChanged(object? sender, EventArgs e)
 		{
 			WorkspaceControl.IncludeAllProjectsInSolution.Checked = GlobalControl.IncludeAllProjectsInSolution.Checked;
@@ -184,6 +192,7 @@ namespace UnrealGameSync
 			newGlobalFilter = new FilterSettings();
 			newGlobalFilter.View.AddRange(GlobalControl.GetView());
 			newGlobalFilter.AllProjects = GlobalControl.SyncAllProjects.Checked;
+			newGlobalFilter.LocalProjects = GlobalControl.SyncLocalProjects.Checked;
 			newGlobalFilter.AllProjectsInSln = GlobalControl.IncludeAllProjectsInSolution.Checked;
 			newGlobalFilter.UprojectSpecificSln = GlobalControl.GenerateUprojectSpecificSolution.Checked;
 
@@ -193,6 +202,7 @@ namespace UnrealGameSync
 			newWorkspaceFilter = new FilterSettings();
 			newWorkspaceFilter.View.AddRange(WorkspaceControl.GetView());
 			newWorkspaceFilter.AllProjects = (WorkspaceControl.SyncAllProjects.Checked == newGlobalFilter.AllProjects) ? (bool?)null : WorkspaceControl.SyncAllProjects.Checked;
+			newWorkspaceFilter.LocalProjects = (WorkspaceControl.SyncLocalProjects.Checked == newGlobalFilter.LocalProjects) ? (bool?)null : WorkspaceControl.SyncLocalProjects.Checked;
 			newWorkspaceFilter.AllProjectsInSln = (WorkspaceControl.IncludeAllProjectsInSolution.Checked == newGlobalFilter.AllProjectsInSln) ? (bool?)null : WorkspaceControl.IncludeAllProjectsInSolution.Checked;
 			newWorkspaceFilter.UprojectSpecificSln = (WorkspaceControl.GenerateUprojectSpecificSolution.Checked == newGlobalFilter.UprojectSpecificSln) ? (bool?)null : WorkspaceControl.GenerateUprojectSpecificSolution.Checked;
 
diff --git a/Engine/Source/Programs/UnrealGameSync/UnrealGameSyncCmd/Program.cs b/Engine/Source/Programs/UnrealGameSync/UnrealGameSyncCmd/Program.cs
--- a/Engine/Source/Programs/UnrealGameSync/UnrealGameSyncCmd/Program.cs
+++ b/Engine/Source/Programs/UnrealGameSync/UnrealGameSyncCmd/Program.cs
@@ -1237,6 +1237,10 @@ namespace UnrealGameSyncCmd
 			[CommandLine("-OnlyCurrent", Value = "false")]
 			public bool? AllProjects { get; set; } = null;
 
+			[CommandLine("-LocalProjects", Value = "true")]
+			[CommandLine("-NoExtraLocalProjects", Value = "false")]
+			public bool? LocalProjects { get; set; } = null;
+
 			[CommandLine("-GpfAllProjects", Value = "true")]
 			[CommandLine("-GpfOnlyCurrent", Value = "false")]
 			public bool? AllProjectsInSln { get; set; } = null;
@@ -1393,6 +1397,7 @@ namespace UnrealGameSyncCmd
 				}
 
 				settings.AllProjects = commandOptions.AllProjects ?? settings.AllProjects;
+				settings.LocalProjects = commandOptions.LocalProjects ?? settings.LocalProjects;
 				settings.AllProjectsInSln = commandOptions.AllProjectsInSln ?? settings.AllProjectsInSln;
 				settings.UprojectSpecificSln = commandOptions.UprojectSpecificSln ?? settings.UprojectSpecificSln;
 			}
diff --git a/Engine/Source/Programs/UnrealGameSync/UnrealGameSyncShared/GlobalSettings.cs b/Engine/Source/Programs/UnrealGameSync/UnrealGameSyncShared/GlobalSettings.cs
--- a/Engine/Source/Programs/UnrealGameSync/UnrealGameSyncShared/GlobalSettings.cs
+++ b/Engine/Source/Programs/UnrealGameSync/UnrealGameSyncShared/GlobalSettings.cs
@@ -14,6 +14,7 @@ namespace UnrealGameSync
 		public List<Guid> ExcludeCategories { get; init; } = new List<Guid>();
 		public List<string> View { get; init; } = new List<string>();
 		public bool? AllProjects { get; set; }
+		public bool? LocalProjects { get; set; }
 		public bool? AllProjectsInSln { get; set; }
 		public bool? UprojectSpecificSln { get; set; }
 
@@ -23,6 +24,7 @@ namespace UnrealGameSync
 			ExcludeCategories.Clear();
 			View.Clear();
 			AllProjects = null;
+			LocalProjects = null;
 			AllProjectsInSln = null;
 			UprojectSpecificSln = null;
 		}
diff --git a/Engine/Source/Programs/UnrealGameSync/UnrealGameSyncShared/UserSettings.cs b/Engine/Source/Programs/UnrealGameSync/UnrealGameSyncShared/UserSettings.cs
--- a/Engine/Source/Programs/UnrealGameSync/UnrealGameSyncShared/UserSettings.cs
+++ b/Engine/Source/Programs/UnrealGameSync/UnrealGameSyncShared/UserSettings.cs
@@ -677,6 +677,7 @@ namespace UnrealGameSync
 				coreSettingsData.Filter.View.AddRange(configFile.GetValues("General.SyncFilter", Array.Empty<string>()));
 				coreSettingsData.Filter.SetCategories(GetCategorySettings(configFile.FindSection("General"), "SyncIncludedCategories", "SyncExcludedCategories"));
 				coreSettingsData.Filter.AllProjects = configFile.GetValue("General.SyncAllProjects", false);
+				coreSettingsData.Filter.LocalProjects = configFile.GetValue("General.SyncLocalProjects", false);
 				coreSettingsData.Filter.AllProjectsInSln = configFile.GetValue("General.IncludeAllProjectsInSolution", false);
 				coreSettingsData.Filter.UprojectSpecificSln = configFile.GetValue("General.UprojectSpecificSolution", false);
 			}
@@ -1058,6 +1059,9 @@ namespace UnrealGameSync
 				int syncAllProjects = workspaceSection.GetValue("SyncAllProjects", -1);
 				currentWorkspace.Filter.AllProjects = (syncAllProjects == 0) ? (bool?)false : (syncAllProjects == 1) ? (bool?)true : (bool?)null;
 
+				int syncLocalProjects = workspaceSection.GetValue("SyncLocalProjects", -1);
+				currentWorkspace.Filter.LocalProjects = (syncLocalProjects == 0) ? (bool?)false : (syncLocalProjects == 1) ? (bool?)true : (bool?)null;
+
 				int includeAllProjectsInSolution = workspaceSection.GetValue("IncludeAllProjectsInSolution", -1);
 				currentWorkspace.Filter.AllProjectsInSln = (includeAllProjectsInSolution == 0) ? (bool?)false : (includeAllProjectsInSolution == 1) ? (bool?)true : (bool?)null;
 
diff --git a/Engine/Source/Programs/UnrealGameSync/UnrealGameSyncShared/WorkspaceUpdate.cs b/Engine/Source/Programs/UnrealGameSync/UnrealGameSyncShared/WorkspaceUpdate.cs
--- a/Engine/Source/Programs/UnrealGameSync/UnrealGameSyncShared/WorkspaceUpdate.cs
+++ b/Engine/Source/Programs/UnrealGameSync/UnrealGameSyncShared/WorkspaceUpdate.cs
@@ -19,6 +19,7 @@ using EpicGames.OIDC;
 using EpicGames.Perforce;
 using EpicGames.ProjectStore;
 using Google.Protobuf.WellKnownTypes;
+using Microsoft.Build.Evaluation;
 using Microsoft.Extensions.Logging;
 
 namespace UnrealGameSync
@@ -45,7 +46,8 @@ namespace UnrealGameSync
 		Refilter = 0x10000,
 		UprojectSpecificSolution = 0x20000,
 		FastSync = 0x40000,
-		ImportSnapshots = 0x80000
+		ImportSnapshots = 0x80000,
+		SyncLocalProjects = 0x100000
 	}
 
 	public enum WorkspaceUpdateResult
@@ -145,6 +147,10 @@ namespace UnrealGameSync
 			{
 				options |= WorkspaceUpdateOptions.SyncAllProjects;
 			}
+			if (workspaceSettings.Filter.LocalProjects ?? globalSettings.Filter.LocalProjects ?? false)
+			{
+				options |= WorkspaceUpdateOptions.SyncLocalProjects;
+			}
 			if (workspaceSettings.Filter.AllProjectsInSln ?? globalSettings.Filter.AllProjectsInSln ?? false)
 			{
 				options |= WorkspaceUpdateOptions.IncludeAllProjectsInSolution;
@@ -669,7 +675,7 @@ namespace UnrealGameSync
 					}
 
 					// Figure out which paths to sync
-					List<string> relativeSyncPaths = GetRelativeSyncPaths(project, (Context.Options & WorkspaceUpdateOptions.SyncAllProjects) != 0, Context.SyncFilter);
+					List<string> relativeSyncPaths = GetRelativeSyncPaths(project, Context.Options, Context.SyncFilter);
 					List<string> syncPaths = new List<string>(relativeSyncPaths.Select(x => project.ClientRootPath + x));
 
 					// Get the user's sync filter
@@ -1780,14 +1786,17 @@ namespace UnrealGameSync
 			return !records.Any(x => x.IsMapped);
 		}
 
-		public static List<string> GetSyncPaths(ProjectInfo project, bool syncAllProjects, string[] syncFilter)
+		public static List<string> GetSyncPaths(ProjectInfo project, WorkspaceUpdateOptions options, string[] syncFilter)
 		{
-			List<string> syncPaths = GetRelativeSyncPaths(project, syncAllProjects, syncFilter);
+			List<string> syncPaths = GetRelativeSyncPaths(project, options, syncFilter);
 			return syncPaths.Select(x => project.ClientRootPath + x).ToList();
 		}
 
-		public static List<string> GetRelativeSyncPaths(ProjectInfo project, bool syncAllProjects, IReadOnlyList<string>? syncFilter)
+		public static List<string> GetRelativeSyncPaths(ProjectInfo project, WorkspaceUpdateOptions options, IReadOnlyList<string>? syncFilter)
 		{
+			bool syncAllProjects = (options & WorkspaceUpdateOptions.SyncAllProjects) != 0;
+			bool bSyncLocalProjects = (options & WorkspaceUpdateOptions.SyncLocalProjects) != 0;
+
 			List<string> syncPaths = new List<string>();
 
 			// Check the client path is formatted correctly
@@ -1820,6 +1829,57 @@ namespace UnrealGameSync
 					syncPaths.Add("/Enterprise/...");
 				}
 				syncPaths.Add(project.ClientFileName.Substring(project.ClientRootPath.Length, lastSlashIdx - project.ClientRootPath.Length) + "/...");
+
+				if (bSyncLocalProjects)
+				{
+					DirectoryInfo rootDir = project.LocalRootPath.ToDirectoryInfo();
+					// We don't use EnumerateFiles because we don't want to visit subdirectories of projects, or other directories which we know cannot contain
+					// projects but may be still around in the workspace (Intermediate, Content etc.).
+					Queue<DirectoryInfo> traversalQueue = new Queue<DirectoryInfo>();
+					traversalQueue.Enqueue(rootDir);
+					bool bAddedTemplatesDir = false;
+					while (traversalQueue.Count > 0)
+					{
+						DirectoryInfo currentDir = traversalQueue.Dequeue();
+						bool bIsProjectDir = false;
+						foreach (FileInfo file in currentDir.GetFiles())
+						{
+							if (file.Name.EndsWith(".uproject", System.StringComparison.OrdinalIgnoreCase))
+							{
+								bIsProjectDir = true;
+								string relativeProjectPath = currentDir.FullName.Substring(project.LocalRootPath.FullName.Length).Replace('\\', '/');
+
+								if (relativeProjectPath.StartsWith("/Templates", System.StringComparison.OrdinalIgnoreCase))
+								{
+									// Special handling for /Templates: this directory contains lots of projects and we don't want to add 40 or so individual
+									// template directories to the sync filter. If at least one template is present, sync all of them.
+									if (!bAddedTemplatesDir)
+									{
+										syncPaths.Add("/Templates/...");
+										bAddedTemplatesDir = true;
+									}
+								}
+								else
+								{
+									syncPaths.Add(relativeProjectPath + "/...");
+								}
+								break;
+							}
+						}
+						if (bIsProjectDir)
+						{
+							// Projects cannot have sub-projects.
+							continue;
+						}
+						foreach (DirectoryInfo subDir in currentDir.GetDirectories())
+						{
+							if (subDir.Name != "Engine" && subDir.Name != "Intermediate" && subDir.Name != "Saved" && subDir.Name != "Content" && subDir.Name != "Plugins")
+							{
+								traversalQueue.Enqueue(subDir);
+							}
+						}
+					}
+				}
 			}
 
 			// Apply the sync filter to that list. We only want inclusive rules in the output list, but we can manually apply exclusions to previous entries.
