removed the exception for when the new horde server value is the same as the default, this was clearly a bug.

diff --git a/Engine/Source/Programs/UnrealGameSync/UnrealGameSync/Launcher.cs b/Engine/Source/Programs/UnrealGameSync/UnrealGameSync/Launcher.cs
--- a/Engine/Source/Programs/UnrealGameSync/UnrealGameSync/Launcher.cs
+++ b/Engine/Source/Programs/UnrealGameSync/UnrealGameSync/Launcher.cs
@@ -155,7 +155,7 @@ namespace UnrealGameSync
 				{
 					if (String.IsNullOrEmpty(launcherSettings.HordeServer))
 					{
-						throw new Exception("No Horde server is configured");
+						launcherSettings.HordeServer = DeploymentSettings.Instance.HordeUrl;
 					}
 
 					Uri hordeServerUrl = new Uri(launcherSettings.HordeServer);
