Fixed a bug where in some circumstances trying to clear the selection of a combo box would result in an exception due to the start parameter for the combo box not being correctly initialised (this is in fact a bug in winforms).

diff --git a/Engine/Source/Programs/UnrealGameSync/UnrealGameSync/ApplicationThemeManager.cs b/Engine/Source/Programs/UnrealGameSync/UnrealGameSync/ApplicationThemeManager.cs
--- a/Engine/Source/Programs/UnrealGameSync/UnrealGameSync/ApplicationThemeManager.cs
+++ b/Engine/Source/Programs/UnrealGameSync/UnrealGameSync/ApplicationThemeManager.cs
@@ -344,6 +344,7 @@ namespace UnrealGameSync
 
 					if (control.Focused == false)
 					{
+						comboBox.SelectionStart = 0;
 						comboBox.SelectionLength = 0;
 					}
 
