Fixed a bug where we were setting the SelectionLength property on Combo Boxes that have their style set to DropDownList (i.e. don't actually have editable or selectable text).

This is undefined behaviour and was raising an exception on some user's machines.

diff --git a/Engine/Source/Programs/UnrealGameSync/UnrealGameSync/ApplicationThemeManager.cs b/Engine/Source/Programs/UnrealGameSync/UnrealGameSync/ApplicationThemeManager.cs
--- a/Engine/Source/Programs/UnrealGameSync/UnrealGameSync/ApplicationThemeManager.cs
+++ b/Engine/Source/Programs/UnrealGameSync/UnrealGameSync/ApplicationThemeManager.cs
@@ -342,9 +342,8 @@ namespace UnrealGameSync
 					control.BackColor = Theme.InputBackgroundColor;
 					control.ForeColor = control.Enabled ? Theme.TextColor : Theme.DimmedText;
 
-					if (control.Focused == false)
+					if (control.Focused == false && comboBox.DropDownStyle != ComboBoxStyle.DropDownList)
 					{
-						comboBox.SelectionStart = 0;
 						comboBox.SelectionLength = 0;
 					}
 
