Quick fix to handle if the HordeUrl is null in the deployment settings.

diff --git a/Engine/Source/Programs/UnrealGameSync/UnrealGameSync/Launcher.cs b/Engine/Source/Programs/UnrealGameSync/UnrealGameSync/Launcher.cs
--- a/Engine/Source/Programs/UnrealGameSync/UnrealGameSync/Launcher.cs
+++ b/Engine/Source/Programs/UnrealGameSync/UnrealGameSync/Launcher.cs
@@ -155,6 +155,11 @@ namespace UnrealGameSync
 				{
 					if (String.IsNullOrEmpty(launcherSettings.HordeServer))
 					{
+						if (String.IsNullOrEmpty(DeploymentSettings.Instance.HordeUrl))
+						{
+							throw new Exception("No Horde server is configured");
+						}
+
 						launcherSettings.HordeServer = DeploymentSettings.Instance.HordeUrl;
 					}
 
