Handle errors gracefully in UGS when creating Import Snapshots dropdown menu items.

[RN] none

diff --git a/Engine/Source/Programs/UnrealGameSync/UnrealGameSync/Controls/WorkspaceControl.cs b/Engine/Source/Programs/UnrealGameSync/UnrealGameSync/Controls/WorkspaceControl.cs
--- a/Engine/Source/Programs/UnrealGameSync/UnrealGameSync/Controls/WorkspaceControl.cs
+++ b/Engine/Source/Programs/UnrealGameSync/UnrealGameSync/Controls/WorkspaceControl.cs
@@ -340,9 +340,17 @@ namespace UnrealGameSync
 
 				if (cloudHost != null && cloudNamespace != null && ProjectInfo.StreamName != null)
 				{
-					UpdateImportSnapshotDropdown(cloudHost, cloudNamespace, ProjectInfo.StreamName, ProjectInfo.LocalFileName.GetFileNameWithoutAnyExtensions());
-					_settings.SnapshotHost = cloudHost;
-					_settings.SnapshotNamespace = cloudNamespace;
+					try
+					{
+						UpdateImportSnapshotDropdown(cloudHost, cloudNamespace, ProjectInfo.StreamName, ProjectInfo.LocalFileName.GetFileNameWithoutAnyExtensions());
+						_settings.SnapshotHost = cloudHost;
+						_settings.SnapshotNamespace = cloudNamespace;
+					}
+					catch (Exception e)
+					{
+						OptionsContextMenu_ImportSnapshots.Enabled = false;
+						_logger.LogError(e, "Error determining available platforms for snapshot import from host {Host} and namespace {Namespace}", cloudHost, cloudNamespace);
+					}
 				}
 				else
 				{
