Use non blocking read for EditorPerProjectUserSettings.ini

diff --git a/Engine/Source/Programs/UnrealGameSync/UnrealGameSync/Controls/WorkspaceControl.cs b/Engine/Source/Programs/UnrealGameSync/UnrealGameSync/Controls/WorkspaceControl.cs
--- a/Engine/Source/Programs/UnrealGameSync/UnrealGameSync/Controls/WorkspaceControl.cs
+++ b/Engine/Source/Programs/UnrealGameSync/UnrealGameSync/Controls/WorkspaceControl.cs
@@ -7335,8 +7335,19 @@ namespace UnrealGameSync
 				return null;
 			}
 
+			using FileStream fs = new FileStream(newConfigFile.FullName, FileMode.Open,FileAccess.Read, FileShare.Read);
+			List<string> lines = new List<string>();
+
+			using (StreamReader reader = new StreamReader(fs))
+			{
+				while (!reader.EndOfStream)
+				{
+					lines.Add(reader.ReadLine());
+				}
+			}
+
 			ConfigFile config = new ConfigFile();
-			config.Load(newConfigFile);
+			config.Parse(lines.ToArray());
 
 			ConfigSection? section = config.FindSection("/Script/EditorStyle.EditorStyleSettings");
 			if (section == null)
@@ -7375,8 +7386,19 @@ namespace UnrealGameSync
 				return null;
 			}
 
+			using FileStream fs = new FileStream(fileName.FullName, FileMode.Open, FileAccess.Read, FileShare.Read);
+			List<string> lines = new List<string>();
+
+			using (StreamReader reader = new StreamReader(fs))
+			{
+				while (!reader.EndOfStream)
+				{
+					lines.Add(reader.ReadLine());
+				}
+			}
+
 			ConfigFile config = new ConfigFile();
-			config.Load(fileName);
+			config.Parse(lines.ToArray());
 
 			ConfigSection? section = config.FindSection("/Script/EditorStyle.EditorStyleSettings");
 			if (section == null)
