Added a command to dump the global & user parameters to file.

diff --git a/Engine/Source/Programs/UnrealGameSync/UnrealGameSyncCmd/Program.cs b/Engine/Source/Programs/UnrealGameSync/UnrealGameSyncCmd/Program.cs
--- a/Engine/Source/Programs/UnrealGameSync/UnrealGameSyncCmd/Program.cs
+++ b/Engine/Source/Programs/UnrealGameSync/UnrealGameSyncCmd/Program.cs
@@ -13,6 +13,7 @@ using System.Net.Http.Json;
 using System.Reflection;
 using System.Runtime.InteropServices;
 using System.Text;
+using System.Text.Json;
 using System.Text.RegularExpressions;
 using System.Threading;
 using System.Threading.Tasks;
@@ -20,9 +21,11 @@ using EpicGames.Core;
 using EpicGames.Horde;
 using EpicGames.OIDC;
 using EpicGames.Perforce;
+
 using Microsoft.Extensions.Configuration;
 using Microsoft.Extensions.DependencyInjection;
 using Microsoft.Extensions.Logging;
+
 using Serilog;
 using UnrealGameSync;
 
@@ -136,6 +139,11 @@ namespace UnrealGameSyncCmd
 				"ugs tools",
 				"Install a Custom Tool. Only available for windows."
 			),
+			new CommandInfo("settings", typeof(SettingsCommand), typeof(SettingsCommandOptions),
+				"ugs settings",
+				"Export the current ugs settings to a file"
+			),
+			
 		};
 
 		class CommandContext
@@ -2117,5 +2125,66 @@ namespace UnrealGameSyncCmd
 				}
 			}
 		}
+
+		class SettingsCommandOptions
+		{
+			[CommandLine("-Global=")]
+			public string Global { get; set; } = String.Empty;
+
+			[CommandLine("-User=")]
+			public string User { get; set; } = String.Empty;
+		}
+
+		class SettingsCommand : Command
+		{
+			public override async Task ExecuteAsync(CommandContext context)
+			{
+				ILogger logger = context.Logger;
+
+				SettingsCommandOptions options = context.Arguments.ApplyTo<SettingsCommandOptions>(logger);
+				context.Arguments.CheckAllArgumentsUsed();
+
+				if (context.GlobalSettings != null && !String.IsNullOrWhiteSpace(options.Global))
+				{
+					logger.LogInformation("Writing global settings to {0}", options.Global);
+					await WriteSettings(options.Global, context.GlobalSettings);
+				}
+
+				if (!String.IsNullOrWhiteSpace(options.User))
+				{
+					logger.LogInformation("Writing user settings to {0}", options.User);
+					await WriteSettings(options.User, context.UserSettings);
+				}
+			}
+
+			private static async Task WriteSettings(string path, object? settings)
+			{
+				if (String.IsNullOrWhiteSpace(path) || settings == null)
+				{
+					return;
+				}
+
+				string trimmed = path.Trim('\"')
+					.Replace(Path.AltDirectorySeparatorChar, Path.DirectorySeparatorChar);
+				if (!Path.IsPathFullyQualified(trimmed))
+				{
+					return;
+				}
+
+				string? directory = Path.GetDirectoryName(trimmed);
+				if (!String.IsNullOrWhiteSpace(directory))
+				{
+					Directory.CreateDirectory(directory);
+				}
+
+				JsonSerializerOptions serializerOptions = new JsonSerializerOptions()
+				{
+					WriteIndented = true
+				};
+				string json = JsonSerializer.Serialize(settings, serializerOptions);
+
+				await File.WriteAllTextAsync(trimmed, json);
+			}
+		}
 	}
 }
