Fixing warnings

diff --git a/Engine/Source/Programs/UnrealGameSync/UnrealGameSync/Controls/WorkspaceControl.cs b/Engine/Source/Programs/UnrealGameSync/UnrealGameSync/Controls/WorkspaceControl.cs
--- a/Engine/Source/Programs/UnrealGameSync/UnrealGameSync/Controls/WorkspaceControl.cs
+++ b/Engine/Source/Programs/UnrealGameSync/UnrealGameSync/Controls/WorkspaceControl.cs
@@ -7342,7 +7342,11 @@ namespace UnrealGameSync
 			{
 				while (!reader.EndOfStream)
 				{
-					lines.Add(reader.ReadLine());
+					string? line = reader.ReadLine();
+					if (!String.IsNullOrWhiteSpace(line))
+					{
+						lines.Add(line);
+					}
 				}
 			}
 
@@ -7393,7 +7397,11 @@ namespace UnrealGameSync
 			{
 				while (!reader.EndOfStream)
 				{
-					lines.Add(reader.ReadLine());
+					string? line = reader.ReadLine();
+					if (!String.IsNullOrWhiteSpace(line))
+					{
+						lines.Add(line);
+					}
 				}
 			}
 
