Unreal Game Sync - Added uri command to spawn a custom tool with args will be used to spawn Build Storage Tool to download a build.

diff --git a/Engine/Source/Programs/UnrealGameSync/UnrealGameSync/AutomationServer.cs b/Engine/Source/Programs/UnrealGameSync/UnrealGameSync/AutomationServer.cs
--- a/Engine/Source/Programs/UnrealGameSync/UnrealGameSync/AutomationServer.cs
+++ b/Engine/Source/Programs/UnrealGameSync/UnrealGameSync/AutomationServer.cs
@@ -21,6 +21,7 @@ namespace UnrealGameSync
 		OpenProject = 2,
 		ExecCommand = 3,
 		OpenIssue = 4,
+		RunCustomTool = 5,
 	}
 
 	class AutomationRequestInput
diff --git a/Engine/Source/Programs/UnrealGameSync/UnrealGameSync/Forms/MainWindow.cs b/Engine/Source/Programs/UnrealGameSync/UnrealGameSync/Forms/MainWindow.cs
--- a/Engine/Source/Programs/UnrealGameSync/UnrealGameSync/Forms/MainWindow.cs
+++ b/Engine/Source/Programs/UnrealGameSync/UnrealGameSync/Forms/MainWindow.cs
@@ -251,6 +251,11 @@ namespace UnrealGameSync
 					AutomationRequestOutput output = OpenIssue(request);
 					request.SetOutput(new AutomationRequestOutput(AutomationRequestResult.Ok));
 				}
+				else if (request.Input.Type == AutomationRequestType.RunCustomTool)
+				{
+					AutomationRequestOutput output = RunCustomTool(request);
+					request.SetOutput(new AutomationRequestOutput(AutomationRequestResult.Ok));
+				}
 				else
 				{
 					request.SetOutput(new AutomationRequestOutput(AutomationRequestResult.Invalid));
@@ -501,6 +506,43 @@ namespace UnrealGameSync
 			return new AutomationRequestOutput(AutomationRequestResult.NotFound);
 		}
 
+		AutomationRequestOutput RunCustomTool(AutomationRequest request)
+		{
+			using BinaryReader reader = new BinaryReader(new MemoryStream(request.Input.Data));
+			string toolName = reader.ReadString();
+			string command = reader.ReadString();
+
+			for (int existingTabIdx = 0; existingTabIdx < TabControl.GetTabCount(); existingTabIdx++)
+			{
+				WorkspaceControl? existingWorkspace = TabControl.GetTabData(existingTabIdx) as WorkspaceControl;
+				if (existingWorkspace != null)
+				{
+					// make sure the tool monitor has updated its state
+					ToolUpdateMonitor.GetDataFromBackendAsync().Wait();
+
+					if (ToolUpdateMonitor.GetTools().FirstOrDefault(t => String.Equals(t.Name, toolName, StringComparison.OrdinalIgnoreCase)) == null)
+					{
+						_logger.LogError("Unable to find tool {ToolName} when attempting to process uri request to it", toolName);
+						return new AutomationRequestOutput(AutomationRequestResult.Error);
+					}
+
+					DirectoryReference? toolPath = ToolUpdateMonitor.GetToolPath(toolName);
+					if (toolPath == null)
+					{
+						_logger.LogError("Unable to find path for {ToolName}", toolName);
+						return new AutomationRequestOutput(AutomationRequestResult.Error);
+					}
+
+					// this is a async method but we do not wait for it as that blocks UGS
+					_ = ToolUpdateMonitor.RunCommandAsync(toolName, command, CancellationToken.None);
+					// we have started the tool, no need to check other workspaces for this command
+					break; 
+				}
+			}
+
+			return new AutomationRequestOutput(AutomationRequestResult.NotFound);
+		}
+
 		protected override void OnHandleCreated(EventArgs e)
 		{
 			base.OnHandleCreated(e);
diff --git a/Engine/Source/Programs/UnrealGameSync/UnrealGameSync/UriHandlers/UGSHandler.cs b/Engine/Source/Programs/UnrealGameSync/UnrealGameSync/UriHandlers/UGSHandler.cs
--- a/Engine/Source/Programs/UnrealGameSync/UnrealGameSync/UriHandlers/UGSHandler.cs
+++ b/Engine/Source/Programs/UnrealGameSync/UnrealGameSync/UriHandlers/UGSHandler.cs
@@ -68,5 +68,21 @@ namespace UnrealGameSync
 				return new UriResult() { Success = true, Request = new AutomationRequest(input) };
 			}
 		}
+
+		[UriHandler]
+		public static UriResult RunCustomTool(string toolName, string args)
+		{
+			using (MemoryStream inputDataStream = new MemoryStream())
+			{
+				using (BinaryWriter writer = new BinaryWriter(inputDataStream))
+				{
+					writer.Write(toolName);
+					writer.Write(args);
+				}
+				AutomationRequestInput input = new AutomationRequestInput(AutomationRequestType.RunCustomTool, inputDataStream.GetBuffer());
+
+				return new UriResult() { Success = true, Request = new AutomationRequest(input) };
+			}
+		}
 	}
 }
\ No newline at end of file
diff --git a/Engine/Source/Programs/UnrealGameSync/UnrealGameSyncShared/ToolUpdateMonitor.cs b/Engine/Source/Programs/UnrealGameSync/UnrealGameSyncShared/ToolUpdateMonitor.cs
--- a/Engine/Source/Programs/UnrealGameSync/UnrealGameSyncShared/ToolUpdateMonitor.cs
+++ b/Engine/Source/Programs/UnrealGameSync/UnrealGameSyncShared/ToolUpdateMonitor.cs
@@ -461,7 +461,7 @@ namespace UnrealGameSync
 			return tool;
 		}
 
-		async Task RunCommandAsync(string toolName, string command, CancellationToken cancellationToken)
+		public async Task RunCommandAsync(string toolName, string command, CancellationToken cancellationToken)
 		{
 			DirectoryReference toolPath = GetToolPathInternal(toolName);
 
diff --git a/Engine/Source/Programs/UnrealGameSync/UnrealGameSyncShared/WorkspaceUpdate.cs b/Engine/Source/Programs/UnrealGameSync/UnrealGameSyncShared/WorkspaceUpdate.cs
--- a/Engine/Source/Programs/UnrealGameSync/UnrealGameSyncShared/WorkspaceUpdate.cs
+++ b/Engine/Source/Programs/UnrealGameSync/UnrealGameSyncShared/WorkspaceUpdate.cs
@@ -15,10 +15,8 @@ using System.Text.RegularExpressions;
 using System.Threading;
 using System.Threading.Tasks;
 using EpicGames.Core;
-using EpicGames.OIDC;
 using EpicGames.Perforce;
 using EpicGames.ProjectStore;
-using Google.Protobuf.WellKnownTypes;
 using Microsoft.Build.Evaluation;
 using Microsoft.Extensions.Logging;
 
