ugs settings fixes
- added global for mac/linux
- local file is the settings.json in .ugs folder

diff --git a/Engine/Source/Programs/UnrealGameSync/UnrealGameSyncCmd/Program.cs b/Engine/Source/Programs/UnrealGameSync/UnrealGameSyncCmd/Program.cs
--- a/Engine/Source/Programs/UnrealGameSync/UnrealGameSyncCmd/Program.cs
+++ b/Engine/Source/Programs/UnrealGameSync/UnrealGameSyncCmd/Program.cs
@@ -2144,16 +2144,25 @@ namespace UnrealGameSyncCmd
 				SettingsCommandOptions options = context.Arguments.ApplyTo<SettingsCommandOptions>(logger);
 				context.Arguments.CheckAllArgumentsUsed();
 
-				if (context.GlobalSettings != null && !String.IsNullOrWhiteSpace(options.Global))
+				if (!String.IsNullOrWhiteSpace(options.Global))
 				{
-					logger.LogInformation("Writing global settings to {0}", options.Global);
-					await WriteSettings(options.Global, context.GlobalSettings);
+					if (context.GlobalSettings != null)
+					{
+						logger.LogInformation("Writing global settings to {Global}", options.Global);
+						await WriteSettings(options.Global, context.GlobalSettings);
+					}
+					else if (context.UserSettings != null)
+					{
+						logger.LogInformation("Writing global settings to {Global}", options.Global);
+						await WriteSettings(options.Global, context.UserSettings);
+					}
 				}
 
 				if (!String.IsNullOrWhiteSpace(options.User))
 				{
-					logger.LogInformation("Writing user settings to {0}", options.User);
-					await WriteSettings(options.User, context.UserSettings);
+					UserWorkspaceSettings userWorkspaceSettings = ReadRequiredUserWorkspaceSettings();
+					logger.LogInformation("Writing user settings to {User}", options.User);
+					await WriteSettings(options.User, userWorkspaceSettings);
 				}
 			}
 
