UGS: Fix compiler warnings

diff --git a/Engine/Source/Programs/UnrealGameSync/UnrealGameSyncCmd/Program.cs b/Engine/Source/Programs/UnrealGameSync/UnrealGameSyncCmd/Program.cs
--- a/Engine/Source/Programs/UnrealGameSync/UnrealGameSyncCmd/Program.cs
+++ b/Engine/Source/Programs/UnrealGameSync/UnrealGameSyncCmd/Program.cs
@@ -1476,11 +1476,13 @@ namespace UnrealGameSyncCmd
 
 					logger.LogInformation("Available build steps:");
 					logger.LogInformation("");
-					logger.LogInformation("  Id                                   | {descriptionHeader}{descriptionSpace}| Type       | Enabled", descriptionHeader, descriptionSpace);
-					logger.LogInformation("  -------------------------------------|{descriptionDashes}|------------|-----------------", descriptionDashes);
+					logger.LogInformation("  Id                                   | {DescriptionHeader}{DescriptionSpace}| Type       | Enabled", descriptionHeader, descriptionSpace);
+					logger.LogInformation("  -------------------------------------|{DescriptionDashes}|------------|-----------------", descriptionDashes);
 					foreach (BuildStep buildStep in buildStepObjects.Values.Select(x => new BuildStep(x)).OrderBy(x => x.NormalSync))
 					{
+#pragma warning disable CA2254 // Template should be a static expression
 						logger.LogInformation(format, buildStep.UniqueId, buildStep.Description, buildStep.Type, buildStep.NormalSync);
+#pragma warning restore CA2254
 					}
 					return;
 				}
@@ -1515,7 +1517,7 @@ namespace UnrealGameSyncCmd
 									continue;
 								}
 
-								logger.LogWarning($"Build step '{buildStep.Description}' cannot be run as the tool is not installed.");
+								logger.LogWarning("Build step '{Description}' cannot be run as the tool is not installed.", buildStep.Description);
 								uninstalledTools.Add(step);
 							}
 						}
@@ -1543,7 +1545,7 @@ namespace UnrealGameSyncCmd
 				}
 			}
 
-			private async Task<ToolUpdateMonitor?> GetToolsMonitor(ILogger logger, CommandContext context, IPerforceConnection perforceClient)
+			private static async Task<ToolUpdateMonitor?> GetToolsMonitor(ILogger logger, CommandContext context, IPerforceConnection perforceClient)
 			{
 				ToolsCommandCommandOptions options = context.Arguments.ApplyTo<ToolsCommandCommandOptions>(logger);
 				context.Arguments.CheckAllArgumentsUsed();
@@ -1581,7 +1583,7 @@ namespace UnrealGameSyncCmd
 				return toolUpdateMonitor;
 			}
 
-			private bool CanRunStep(BuildStep step, ToolUpdateMonitor toolsMonitor)
+			private static bool CanRunStep(BuildStep step, ToolUpdateMonitor toolsMonitor)
 			{
 				if (step.ToolId != Guid.Empty)
 				{
diff --git a/Engine/Source/Programs/UnrealGameSync/UnrealGameSyncShared/WorkspaceUpdate.cs b/Engine/Source/Programs/UnrealGameSync/UnrealGameSyncShared/WorkspaceUpdate.cs
--- a/Engine/Source/Programs/UnrealGameSync/UnrealGameSyncShared/WorkspaceUpdate.cs
+++ b/Engine/Source/Programs/UnrealGameSync/UnrealGameSyncShared/WorkspaceUpdate.cs
@@ -17,7 +17,6 @@ using System.Threading.Tasks;
 using EpicGames.Core;
 using EpicGames.Perforce;
 using EpicGames.ProjectStore;
-using Microsoft.Build.Evaluation;
 using Microsoft.Extensions.Logging;
 
 namespace UnrealGameSync
