Wildcard support in project sections of UGS ini files.

diff --git a/Engine/Source/Programs/UnrealGameSync/UnrealGameSync/PerforceMonitor.cs b/Engine/Source/Programs/UnrealGameSync/UnrealGameSync/PerforceMonitor.cs
--- a/Engine/Source/Programs/UnrealGameSync/UnrealGameSync/PerforceMonitor.cs
+++ b/Engine/Source/Programs/UnrealGameSync/UnrealGameSync/PerforceMonitor.cs
@@ -519,7 +519,7 @@ namespace UnrealGameSync
 		async Task UpdateProjectConfigFileAsync(IPerforceConnection perforce, CancellationToken cancellationToken)
 		{
 			_localConfigFiles.Clear();
-			LatestProjectConfigFile = await ConfigUtils.ReadProjectConfigFileAsync(perforce, _branchClientPath, _selectedClientFileName, _cacheFolder, _localConfigFiles, _logger, cancellationToken);
+			LatestProjectConfigFile = await ConfigUtils.ReadProjectConfigFileAsync(perforce, _branchClientPath, _selectedClientFileName, _selectedProjectIdentifier, _cacheFolder, _localConfigFiles, _logger, cancellationToken);
 		}
 
 		public List<ChangesRecord> GetChanges()
diff --git a/Engine/Source/Programs/UnrealGameSync/UnrealGameSyncShared/ConfigFile.cs b/Engine/Source/Programs/UnrealGameSync/UnrealGameSyncShared/ConfigFile.cs
--- a/Engine/Source/Programs/UnrealGameSync/UnrealGameSyncShared/ConfigFile.cs
+++ b/Engine/Source/Programs/UnrealGameSync/UnrealGameSyncShared/ConfigFile.cs
@@ -5,6 +5,7 @@ using System.Collections.Generic;
 using System.Diagnostics;
 using System.Diagnostics.CodeAnalysis;
 using System.IO;
+using System.IO.Enumeration;
 using System.Linq;
 using System.Text;
 using System.Text.Json;
@@ -561,7 +562,7 @@ namespace UnrealGameSync
 			return false;
 		}
 
-		public void Parse(string[] lines)
+		public void Parse(string[] lines, string? project = null)
 		{
 			ConfigSection? currentSection = null;
 			foreach (string line in lines)
@@ -572,6 +573,16 @@ namespace UnrealGameSync
 					if (trimLine.StartsWith("[", StringComparison.Ordinal) && trimLine.EndsWith("]", StringComparison.Ordinal))
 					{
 						string sectionName = trimLine.Substring(1, trimLine.Length - 2).Trim();
+
+						if (String.IsNullOrEmpty(project) == false &&
+							sectionName.StartsWith("//", StringComparison.Ordinal) && sectionName.Contains('*', StringComparison.Ordinal))
+						{
+							if (FileSystemName.MatchesSimpleExpression(sectionName, project))
+							{
+								sectionName = project;
+							}
+						}
+
 						currentSection = FindOrAddSection(sectionName);
 					}
 					else if (currentSection != null)
diff --git a/Engine/Source/Programs/UnrealGameSync/UnrealGameSyncShared/ConfigUtils.cs b/Engine/Source/Programs/UnrealGameSync/UnrealGameSyncShared/ConfigUtils.cs
--- a/Engine/Source/Programs/UnrealGameSync/UnrealGameSyncShared/ConfigUtils.cs
+++ b/Engine/Source/Programs/UnrealGameSync/UnrealGameSyncShared/ConfigUtils.cs
@@ -79,10 +79,10 @@ namespace UnrealGameSync
 
 		public static Task<ConfigFile> ReadProjectConfigFileAsync(IPerforceConnection perforce, ProjectInfo projectInfo, List<KeyValuePair<FileReference, DateTime>> localConfigFiles, ILogger logger, CancellationToken cancellationToken)
 		{
-			return ReadProjectConfigFileAsync(perforce, projectInfo.ClientRootPath, projectInfo.ClientFileName, projectInfo.CacheFolder, localConfigFiles, logger, cancellationToken);
+			return ReadProjectConfigFileAsync(perforce, projectInfo.ClientRootPath, projectInfo.ClientFileName, projectInfo.ProjectIdentifier, projectInfo.CacheFolder, localConfigFiles, logger, cancellationToken);
 		}
 
-		public static async Task<ConfigFile> ReadProjectConfigFileAsync(IPerforceConnection perforce, string branchClientPath, string selectedClientFileName, DirectoryReference cacheFolder, List<KeyValuePair<FileReference, DateTime>> localConfigFiles, ILogger logger, CancellationToken cancellationToken)
+		public static async Task<ConfigFile> ReadProjectConfigFileAsync(IPerforceConnection perforce, string branchClientPath, string selectedClientFileName, string projectIdentifier, DirectoryReference cacheFolder, List<KeyValuePair<FileReference, DateTime>> localConfigFiles, ILogger logger, CancellationToken cancellationToken)
 		{
 			List<string> configFilePaths = Utility.GetDepotConfigPaths(branchClientPath + "/Engine", selectedClientFileName);
 
@@ -129,7 +129,7 @@ namespace UnrealGameSync
 					{
 						try
 						{
-							projectConfig.Parse(lines.ToArray());
+							projectConfig.Parse(lines.ToArray(), projectIdentifier);
 							logger.LogDebug("Read config file from {DepotFile}", fileRecord.DepotFile);
 						}
 						catch (Exception ex)
diff --git a/Engine/Source/Programs/UnrealGameSync/UnrealGameSyncShared/OpenProjectInfo.cs b/Engine/Source/Programs/UnrealGameSync/UnrealGameSyncShared/OpenProjectInfo.cs
--- a/Engine/Source/Programs/UnrealGameSync/UnrealGameSyncShared/OpenProjectInfo.cs
+++ b/Engine/Source/Programs/UnrealGameSync/UnrealGameSyncShared/OpenProjectInfo.cs
@@ -232,7 +232,7 @@ namespace UnrealGameSync
 					ConfigFile latestProjectConfigFile = await ConfigUtils.ReadProjectConfigFileAsync(perforceClient, projectInfo, localConfigFiles, logger, cancellationToken);
 
 					// Get the local config file and stream filter
-					ConfigFile workspaceProjectConfigFile = await WorkspaceUpdate.ReadProjectConfigFile(branchDirectoryName, newSelectedFileName, logger);
+					ConfigFile workspaceProjectConfigFile = await WorkspaceUpdate.ReadProjectConfigFile(branchDirectoryName, newSelectedFileName, projectInfo.ProjectIdentifier, logger);
 					IReadOnlyList<string>? workspaceProjectStreamFilter = await WorkspaceUpdate.ReadProjectStreamFilter(perforceClient, workspaceProjectConfigFile, cancellationToken);
 
 					OpenProjectInfo workspaceSettings = new OpenProjectInfo(selectedProject, perforceSettings, projectInfo, userWorkspaceSettings, workspaceStateWrapper, latestProjectConfigFile, workspaceProjectConfigFile, workspaceProjectStreamFilter, localConfigFiles, generateP4Config);
diff --git a/Engine/Source/Programs/UnrealGameSync/UnrealGameSyncShared/WorkspaceUpdate.cs b/Engine/Source/Programs/UnrealGameSync/UnrealGameSyncShared/WorkspaceUpdate.cs
--- a/Engine/Source/Programs/UnrealGameSync/UnrealGameSyncShared/WorkspaceUpdate.cs
+++ b/Engine/Source/Programs/UnrealGameSync/UnrealGameSyncShared/WorkspaceUpdate.cs
@@ -1054,7 +1054,7 @@ namespace UnrealGameSync
 					if (Context.Options.HasFlag(WorkspaceUpdateOptions.Sync) && !Context.Options.HasFlag(WorkspaceUpdateOptions.UpdateFilter))
 					{
 						// Read the new config file
-						Context.ProjectConfigFile = await ReadProjectConfigFile(project.LocalRootPath, project.LocalFileName, logger);
+						Context.ProjectConfigFile = await ReadProjectConfigFile(project.LocalRootPath, project.LocalFileName, project.ProjectIdentifier, logger);
 						Context.ProjectStreamFilter = await ReadProjectStreamFilter(perforce, Context.ProjectConfigFile, cancellationToken);
 
 						// Get the branch name
@@ -1240,7 +1240,7 @@ namespace UnrealGameSync
 					// Continue processing sync-only actions
 					if (Context.Options.HasFlag(WorkspaceUpdateOptions.Sync) && !Context.Options.HasFlag(WorkspaceUpdateOptions.UpdateFilter))
 					{
-						Context.ProjectConfigFile ??= await ReadProjectConfigFile(project.LocalRootPath, project.LocalFileName, logger);
+						Context.ProjectConfigFile ??= await ReadProjectConfigFile(project.LocalRootPath, project.LocalFileName, project.ProjectIdentifier, logger);
 
 						// Execute any project specific post-sync steps
 						string[]? postSyncSteps = Context.ProjectConfigFile.GetValues("Sync.Step", null);
@@ -1542,7 +1542,7 @@ namespace UnrealGameSync
 			// Build everything using MegaXGE
 			if (Context.Options.HasFlag(WorkspaceUpdateOptions.Build))
 			{
-				Context.ProjectConfigFile ??= await ReadProjectConfigFile(project.LocalRootPath, project.LocalFileName, logger);
+				Context.ProjectConfigFile ??= await ReadProjectConfigFile(project.LocalRootPath, project.LocalFileName, project.ProjectIdentifier, logger);
 
 				// Figure out the new editor target
 				TargetReceipt defaultEditorReceipt = ConfigUtils.CreateDefaultEditorReceipt(project, Context.ProjectConfigFile, Context.EditorConfig);
@@ -2169,7 +2169,7 @@ namespace UnrealGameSync
 			return (WorkspaceUpdateResult.Success, "Succeeded.", null);
 		}
 
-		public static async Task<ConfigFile> ReadProjectConfigFile(DirectoryReference localRootPath, FileReference selectedLocalFileName, ILogger logger)
+		public static async Task<ConfigFile> ReadProjectConfigFile(DirectoryReference localRootPath, FileReference selectedLocalFileName, string projectIdentifier, ILogger logger)
 		{
 			// Find the valid config file paths
 			DirectoryInfo engineDir = DirectoryReference.Combine(localRootPath, "Engine").ToDirectoryInfo();
@@ -2182,7 +2182,7 @@ namespace UnrealGameSync
 				try
 				{
 					string[] lines = await File.ReadAllLinesAsync(localConfigFile.FullName);
-					projectConfig.Parse(lines);
+					projectConfig.Parse(lines, projectIdentifier);
 					logger.LogDebug("Read config file from {FileName}", localConfigFile.FullName);
 				}
 				catch (Exception ex)
