Fixed a bug where wildcard project sections would incorrectly take priority over fully pathed project sections in the UGS ini file.

diff --git a/Engine/Source/Programs/UnrealGameSync/UnrealGameSyncShared/ConfigFile.cs b/Engine/Source/Programs/UnrealGameSync/UnrealGameSyncShared/ConfigFile.cs
--- a/Engine/Source/Programs/UnrealGameSync/UnrealGameSyncShared/ConfigFile.cs
+++ b/Engine/Source/Programs/UnrealGameSync/UnrealGameSyncShared/ConfigFile.cs
@@ -565,6 +565,8 @@ namespace UnrealGameSync
 		public void Parse(string[] lines, string? project = null)
 		{
 			ConfigSection? currentSection = null;
+			bool wildCardSection = false;
+
 			foreach (string line in lines)
 			{
 				string trimLine = line.Trim();
@@ -573,6 +575,7 @@ namespace UnrealGameSync
 					if (trimLine.StartsWith("[", StringComparison.Ordinal) && trimLine.EndsWith("]", StringComparison.Ordinal))
 					{
 						string sectionName = trimLine.Substring(1, trimLine.Length - 2).Trim();
+						wildCardSection = false;
 
 						if (String.IsNullOrEmpty(project) == false &&
 							sectionName.StartsWith("//", StringComparison.Ordinal) && sectionName.Contains('*', StringComparison.Ordinal))
@@ -580,6 +583,7 @@ namespace UnrealGameSync
 							if (FileSystemName.MatchesSimpleExpression(sectionName, project))
 							{
 								sectionName = project;
+								wildCardSection = true;
 							}
 						}
 
@@ -597,7 +601,12 @@ namespace UnrealGameSync
 							}
 							else
 							{
-								currentSection.SetValue(trimLine.Substring(0, equalsIdx).TrimEnd(), value);
+								string key = trimLine.Substring(0, equalsIdx).TrimEnd();
+
+								if (wildCardSection == false || currentSection.GetValue(key) == null)
+								{
+									currentSection.SetValue(key, value);
+								}
 							}
 						}
 					}
