checking the current stream and its parents' last update time to see if we need to refresh the stream
store the times in the worskpace state

diff --git a/Engine/Source/Programs/UnrealGameSync/UnrealGameSync/Controls/WorkspaceControl.cs b/Engine/Source/Programs/UnrealGameSync/UnrealGameSync/Controls/WorkspaceControl.cs
--- a/Engine/Source/Programs/UnrealGameSync/UnrealGameSync/Controls/WorkspaceControl.cs
+++ b/Engine/Source/Programs/UnrealGameSync/UnrealGameSync/Controls/WorkspaceControl.cs
@@ -426,7 +426,14 @@ namespace UnrealGameSync
 
 			FileReference projectLogBaseName = FileReference.Combine(_workspaceDataFolder, "sync.log");
 
-			_perforceMonitor = new PerforceMonitor(perforceClientSettings, openProjectInfo.ProjectInfo, openProjectInfo.LatestProjectConfigFile, openProjectInfo.ProjectInfo.CacheFolder, openProjectInfo.LocalConfigFiles, _serviceProvider);
+			_perforceMonitor = new PerforceMonitor(
+				perforceClientSettings,
+				openProjectInfo.ProjectInfo,
+				_workspace,
+				openProjectInfo.LatestProjectConfigFile,
+				openProjectInfo.ProjectInfo.CacheFolder,
+				openProjectInfo.LocalConfigFiles,
+				_serviceProvider);
 
 			_perforceMonitor.OnUpdate += UpdateBuildListCallback;
 			_perforceMonitor.OnUpdateMetadata += UpdateBuildMetadataCallback;
diff --git a/Engine/Source/Programs/UnrealGameSync/UnrealGameSync/PerforceMonitor.cs b/Engine/Source/Programs/UnrealGameSync/UnrealGameSync/PerforceMonitor.cs
--- a/Engine/Source/Programs/UnrealGameSync/UnrealGameSync/PerforceMonitor.cs
+++ b/Engine/Source/Programs/UnrealGameSync/UnrealGameSync/PerforceMonitor.cs
@@ -34,6 +34,7 @@ namespace UnrealGameSync
 		public bool ShowChangesForAllProjects { get; set; } = true;
 
 		readonly IPerforceSettings _perforceSettings;
+		readonly Workspace _workspace;
 		readonly string _branchClientPath;
 		readonly string _selectedClientFileName;
 		readonly string _selectedProjectIdentifier;
@@ -56,7 +57,6 @@ namespace UnrealGameSync
 		readonly IHordeClient? _hordeClient;
 #pragma warning restore CA2213
 		private readonly ICloudStorage? _cloudStorage;
-
 		string[] prevCodeRules = Array.Empty<string>();
 
 		readonly SynchronizationContext _synchronizationContext;
@@ -73,9 +73,10 @@ namespace UnrealGameSync
 			private set;
 		}
 
-		public PerforceMonitor(IPerforceSettings perforceSettings, ProjectInfo projectInfo, ConfigFile projectConfigFile, DirectoryReference cacheFolder, List<KeyValuePair<FileReference, DateTime>> localConfigFiles, IServiceProvider serviceProvider)
+		public PerforceMonitor(IPerforceSettings perforceSettings, ProjectInfo projectInfo, Workspace workspace, ConfigFile projectConfigFile, DirectoryReference cacheFolder, List<KeyValuePair<FileReference, DateTime>> localConfigFiles, IServiceProvider serviceProvider)
 		{
 			_perforceSettings = perforceSettings;
+			_workspace = workspace;
 			_branchClientPath = projectInfo.ClientRootPath;
 			_selectedClientFileName = projectInfo.ClientFileName;
 			_selectedProjectIdentifier = projectInfo.ProjectIdentifier;
@@ -219,6 +220,57 @@ namespace UnrealGameSync
 							_synchronizationContext.Post(_ => OnStreamChange?.Invoke(), null);
 						}
 
+						if (!String.IsNullOrWhiteSpace(streamName))
+						{
+							Dictionary<string, DateTime> currentStreamUpdate = new Dictionary<string, DateTime>(StringComparer.OrdinalIgnoreCase);
+
+							StreamRecord record = await perforce.GetStreamAsync(streamName, true, cancellationToken);
+
+							if (record != null)
+							{
+								currentStreamUpdate.Add(streamName, record.Update);
+
+								string? parentStream = record.Parent;
+								while (!String.IsNullOrWhiteSpace(parentStream) && !parentStream.Equals("none", StringComparison.OrdinalIgnoreCase))
+								{
+									StreamRecord parentRecord = await perforce.GetStreamAsync(parentStream, true, cancellationToken);
+									currentStreamUpdate.Add(parentStream, parentRecord.Update);
+									parentStream = parentRecord.Parent;
+								}
+							}
+
+
+							bool streamHasChanged = false;
+							foreach (KeyValuePair<string, DateTime> kvp in currentStreamUpdate)
+							{
+								if (_workspace.State.Streams.TryGetValue(kvp.Key, out DateTime value))
+								{
+									if (DateTime.Compare(kvp.Value, value) != 0)
+									{
+										streamHasChanged = true;
+										break;
+									}
+								}
+							}
+
+							streamHasChanged |= currentStreamUpdate.Keys.Any(k => !_workspace.State.Streams.ContainsKey(k));
+							streamHasChanged |= _workspace.State.Streams.Keys.Any(k => !currentStreamUpdate.ContainsKey(k));
+
+							if(streamHasChanged)
+							{
+								_workspace.ModifyState((state) =>
+								{
+									state.Streams.Clear();
+									foreach (KeyValuePair<string, DateTime> item in currentStreamUpdate)
+									{
+										state.Streams.Add(item.Key, item.Value);
+									}
+								});
+								_synchronizationContext.Post(_ => OnStreamChange?.Invoke(), null);
+							}
+							
+						}
+
 						// Check for any p4 changes
 						if (!await UpdateChangesAsync(perforce, cancellationToken))
 						{
diff --git a/Engine/Source/Programs/UnrealGameSync/UnrealGameSyncShared/WorkspaceState.cs b/Engine/Source/Programs/UnrealGameSync/UnrealGameSyncShared/WorkspaceState.cs
--- a/Engine/Source/Programs/UnrealGameSync/UnrealGameSyncShared/WorkspaceState.cs
+++ b/Engine/Source/Programs/UnrealGameSync/UnrealGameSyncShared/WorkspaceState.cs
@@ -26,6 +26,8 @@ namespace UnrealGameSync
 		public string BranchPath { get; set; } = String.Empty;
 		public string ProjectPath { get; set; } = String.Empty;
 		public string? StreamName { get; set; }
+
+		public IDictionary<string, DateTime> Streams { get; set; } = new Dictionary<string, DateTime>();
 		public string ProjectIdentifier { get; set; } = String.Empty; // Should be fully reset if this changes
 		public bool IsEnterpriseProject { get; set; }
 
@@ -87,6 +89,12 @@ namespace UnrealGameSync
 			BranchPath = other.BranchPath;
 			ProjectPath = other.ProjectPath;
 			StreamName = other.StreamName;
+			Streams.Clear();
+			foreach (KeyValuePair<string, DateTime> kvp in other.Streams)
+			{
+				Streams.Add(kvp.Key, kvp.Value);
+			}
+
 			ProjectIdentifier = other.ProjectIdentifier;
 			IsEnterpriseProject = other.IsEnterpriseProject;
 
@@ -147,6 +155,7 @@ namespace UnrealGameSync
 		public string BranchPath => _inner.BranchPath;
 		public string ProjectPath => _inner.ProjectPath;
 		public string? StreamName => _inner.StreamName;
+		public IDictionary<string, DateTime> Streams => _inner.Streams;
 		public string ProjectIdentifier => _inner.ProjectIdentifier;
 		public bool IsEnterpriseProject => _inner.IsEnterpriseProject;
 
