UGS: Fix issue where Update Settings would fallback to Horde when changing theme even though Perforce was previously selected

diff --git a/Engine/Source/Programs/UnrealGameSync/UnrealGameSyncShared/LauncherSettings.cs b/Engine/Source/Programs/UnrealGameSync/UnrealGameSyncShared/LauncherSettings.cs
--- a/Engine/Source/Programs/UnrealGameSync/UnrealGameSyncShared/LauncherSettings.cs
+++ b/Engine/Source/Programs/UnrealGameSync/UnrealGameSyncShared/LauncherSettings.cs
@@ -39,6 +39,8 @@ namespace UnrealGameSync
 
 		public LauncherSettings(LauncherSettings other)
 		{
+			UpdateSource = other.UpdateSource;
+			HordeServer = other.HordeServer;
 			PerforceServerAndPort = other.PerforceServerAndPort;
 			PerforceUserName = other.PerforceUserName;
 			PerforceDepotPath = other.PerforceDepotPath;
