Unreal Game Sync - Do not trigger logins to Oidc when listing builds for PCBs as this is frequently called and can cause login spam (browsers spawned and even MFA requests).

diff --git a/Engine/Source/Programs/UnrealGameSync/UnrealGameSyncShared/CloudStorage.cs b/Engine/Source/Programs/UnrealGameSync/UnrealGameSyncShared/CloudStorage.cs
--- a/Engine/Source/Programs/UnrealGameSync/UnrealGameSyncShared/CloudStorage.cs
+++ b/Engine/Source/Programs/UnrealGameSync/UnrealGameSyncShared/CloudStorage.cs
@@ -210,7 +210,7 @@ namespace UnrealGameSync
 
 		public async IAsyncEnumerable<FoundBuildResponse> FindBuildAsync(string host, string namespaceId, string bucketId, CbObject query, [EnumeratorCancellation] CancellationToken cancellationToken)
 		{
-			string? accessToken = await GetAccessToken(new Uri(host));
+			string? accessToken = await GetAccessToken(new Uri(host), allowInteractiveLogin: false);
 
 			if (accessToken == null)
 			{
