Fixed a bug where user's synced status would be invalidated if there is no stream hierarchy

diff --git a/Engine/Source/Programs/UnrealGameSync/UnrealGameSync/PerforceMonitor.cs b/Engine/Source/Programs/UnrealGameSync/UnrealGameSync/PerforceMonitor.cs
--- a/Engine/Source/Programs/UnrealGameSync/UnrealGameSync/PerforceMonitor.cs
+++ b/Engine/Source/Programs/UnrealGameSync/UnrealGameSync/PerforceMonitor.cs
@@ -253,8 +253,11 @@ namespace UnrealGameSync
 								}
 							}
 
-							streamHasChanged |= currentStreamUpdate.Keys.Any(k => !_workspace.State.Streams.ContainsKey(k));
-							streamHasChanged |= _workspace.State.Streams.Keys.Any(k => !currentStreamUpdate.ContainsKey(k));
+							if (_workspace.State.Streams.Any() )
+							{
+								streamHasChanged |= currentStreamUpdate.Keys.Any(k => !_workspace.State.Streams.ContainsKey(k));
+								streamHasChanged |= _workspace.State.Streams.Keys.Any(k => !currentStreamUpdate.ContainsKey(k));
+							}
 
 							if(streamHasChanged)
 							{
