Fixed an issue where we never set the initial workspace streams hierarchy.

diff --git a/Engine/Source/Programs/UnrealGameSync/UnrealGameSync/PerforceMonitor.cs b/Engine/Source/Programs/UnrealGameSync/UnrealGameSync/PerforceMonitor.cs
--- a/Engine/Source/Programs/UnrealGameSync/UnrealGameSync/PerforceMonitor.cs
+++ b/Engine/Source/Programs/UnrealGameSync/UnrealGameSync/PerforceMonitor.cs
@@ -253,14 +253,17 @@ namespace UnrealGameSync
 								}
 							}
 
-							if (_workspace.State.Streams.Any() )
-							{
-								streamHasChanged |= currentStreamUpdate.Keys.Any(k => !_workspace.State.Streams.ContainsKey(k));
-								streamHasChanged |= _workspace.State.Streams.Keys.Any(k => !currentStreamUpdate.ContainsKey(k));
-							}
+							streamHasChanged |= currentStreamUpdate.Keys.Any(k => !_workspace.State.Streams.ContainsKey(k));
+							streamHasChanged |= _workspace.State.Streams.Keys.Any(k => !currentStreamUpdate.ContainsKey(k));
 
 							if(streamHasChanged)
 							{
+								// Only refresh the sync state if we already have a stored out stream hierarchy
+								if (_workspace.State.Streams.Any())
+								{
+									_synchronizationContext.Post(_ => OnStreamChange?.Invoke(), null);
+								}
+
 								_workspace.ModifyState((state) =>
 								{
 									state.Streams.Clear();
@@ -269,7 +272,6 @@ namespace UnrealGameSync
 										state.Streams.Add(item.Key, item.Value);
 									}
 								});
-								_synchronizationContext.Post(_ => OnStreamChange?.Invoke(), null);
 							}
 							
 						}
