Syncing two PCBs iteratively between two projects puts UGS in a state where Sync Succeeds, but Banner does not detect PCB has been synced

the workspace state was corrupted py the perforce monitor update during the  project switch, keeping the data in a perforce monitor state avoid the issue and seems more appropriate.

diff --git a/Engine/Source/Programs/UnrealGameSync/UnrealGameSync/PerforceMonitor.cs b/Engine/Source/Programs/UnrealGameSync/UnrealGameSync/PerforceMonitor.cs
--- a/Engine/Source/Programs/UnrealGameSync/UnrealGameSync/PerforceMonitor.cs
+++ b/Engine/Source/Programs/UnrealGameSync/UnrealGameSync/PerforceMonitor.cs
@@ -4,7 +4,9 @@ using System;
 using System.Collections.Generic;
 using System.Diagnostics;
 using System.Diagnostics.CodeAnalysis;
+using System.IO;
 using System.Linq;
+using System.Text.Json;
 using System.Threading;
 using System.Threading.Tasks;
 using EpicGames.Core;
@@ -20,6 +22,33 @@ namespace UnrealGameSync
 		IReadOnlyList<IArchiveChannel> AvailableArchiveChannels { get; }
 	}
 
+	internal class PerforceMonitorCache
+	{
+		public Dictionary<string, DateTime> Streams { get; set; } = new Dictionary<string, DateTime>();
+
+		public static bool Get(FileReference path, out PerforceMonitorCache? cache)
+		{
+			if (File.Exists(path.FullName))
+			{
+				string content = File.ReadAllText(path.FullName);
+				cache = JsonSerializer.Deserialize<PerforceMonitorCache>(content);
+				return cache != null;
+			}
+
+			cache = null;
+			return false;
+		}
+
+		public static void Set(FileReference path, PerforceMonitorCache? cache)
+		{
+			if(cache != null)
+			{
+				string content = JsonSerializer.Serialize(cache);
+				File.WriteAllText(path.FullName, content);
+			}
+		}
+	}
+
 	class PerforceMonitor : IDisposable, IArchiveInfoSource
 	{
 		class PerforceChangeSorter : IComparer<ChangesRecord>
@@ -34,7 +63,6 @@ namespace UnrealGameSync
 		public bool ShowChangesForAllProjects { get; set; } = true;
 
 		readonly IPerforceSettings _perforceSettings;
-		readonly Workspace _workspace;
 		readonly string _branchClientPath;
 		readonly string _selectedClientFileName;
 		readonly string _selectedProjectIdentifier;
@@ -59,6 +87,8 @@ namespace UnrealGameSync
 		private readonly ICloudStorage? _cloudStorage;
 		string[] prevCodeRules = Array.Empty<string>();
 
+		readonly FileReference CacheFile;
+
 		readonly SynchronizationContext _synchronizationContext;
 		public event Action? OnUpdate;
 		public event Action? OnUpdateMetadata;
@@ -76,7 +106,6 @@ namespace UnrealGameSync
 		public PerforceMonitor(IPerforceSettings perforceSettings, ProjectInfo projectInfo, Workspace workspace, ConfigFile projectConfigFile, DirectoryReference cacheFolder, List<KeyValuePair<FileReference, DateTime>> localConfigFiles, IServiceProvider serviceProvider)
 		{
 			_perforceSettings = perforceSettings;
-			_workspace = workspace;
 			_branchClientPath = projectInfo.ClientRootPath;
 			_selectedClientFileName = projectInfo.ClientFileName;
 			_selectedProjectIdentifier = projectInfo.ProjectIdentifier;
@@ -94,6 +123,9 @@ namespace UnrealGameSync
 			_hordeClient = serviceProvider.GetService<IHordeClient>();
 			_cloudStorage = serviceProvider.GetService<ICloudStorage>();
 
+			DirectoryReference configDir = UserSettings.GetConfigDir(workspace.RootDir);
+			CacheFile = FileReference.Combine(configDir, "perforce_monitor_cache.json");
+
 			AvailableArchiveChannels = (new List<IArchiveChannel>()).AsReadOnly();
 		}
 
@@ -239,41 +271,47 @@ namespace UnrealGameSync
 								}
 							}
 
-
-							bool streamHasChanged = false;
-							foreach (KeyValuePair<string, DateTime> kvp in currentStreamUpdate)
+							if (PerforceMonitorCache.Get(CacheFile, out PerforceMonitorCache ? cache))
 							{
-								if (_workspace.State.Streams.TryGetValue(kvp.Key, out DateTime value))
+								bool streamHasChanged = false;
+								foreach (KeyValuePair<string, DateTime> kvp in currentStreamUpdate)
 								{
-									if (DateTime.Compare(kvp.Value, value) != 0)
+									if (cache!.Streams.TryGetValue(kvp.Key, out DateTime value))
 									{
-										streamHasChanged = true;
-										break;
+										if (DateTime.Compare(kvp.Value, value) != 0)
+										{
+											streamHasChanged = true;
+											break;
+										}
 									}
 								}
-							}
 
-							streamHasChanged |= currentStreamUpdate.Keys.Any(k => !_workspace.State.Streams.ContainsKey(k));
-							streamHasChanged |= _workspace.State.Streams.Keys.Any(k => !currentStreamUpdate.ContainsKey(k));
-
-							if(streamHasChanged)
-							{
-								// Only refresh the sync state if we already have a stored out stream hierarchy
-								if (_workspace.State.Streams.Any())
-								{
-									_synchronizationContext.Post(_ => OnStreamChange?.Invoke(), null);
-								}
+								streamHasChanged |= currentStreamUpdate.Keys.Any(k => !cache!.Streams.ContainsKey(k));
+								streamHasChanged |= cache!.Streams.Keys.Any(k => !currentStreamUpdate.ContainsKey(k));
 
-								_workspace.ModifyState((state) =>
+								if (streamHasChanged)
 								{
-									state.Streams.Clear();
-									foreach (KeyValuePair<string, DateTime> item in currentStreamUpdate)
+									// Only refresh the sync state if we already have a stored out stream hierarchy
+									if (cache!.Streams.Any())
 									{
-										state.Streams.Add(item.Key, item.Value);
+										_synchronizationContext.Post(_ => OnStreamChange?.Invoke(), null);
 									}
-								});
+
+									PerforceMonitorCache newCache = new PerforceMonitorCache()
+									{
+										Streams = currentStreamUpdate
+									};
+									PerforceMonitorCache.Set(CacheFile, newCache);
+								}
+							}
+							else
+							{
+								PerforceMonitorCache newCache = new PerforceMonitorCache()
+								{
+									Streams = currentStreamUpdate
+								};
+								PerforceMonitorCache.Set(CacheFile,newCache);
 							}
-							
 						}
 
 						// Check for any p4 changes
diff --git a/Engine/Source/Programs/UnrealGameSync/UnrealGameSync/Workspace.cs b/Engine/Source/Programs/UnrealGameSync/UnrealGameSync/Workspace.cs
--- a/Engine/Source/Programs/UnrealGameSync/UnrealGameSync/Workspace.cs
+++ b/Engine/Source/Programs/UnrealGameSync/UnrealGameSync/Workspace.cs
@@ -17,6 +17,8 @@ namespace UnrealGameSync
 		readonly WorkspaceStateWrapper _stateWrapper;
 		public ReadOnlyWorkspaceState State { get; private set; }
 		public ProjectInfo Project { get; }
+		public DirectoryReference RootDir => _stateWrapper.RootDir;
+
 		public SynchronizationContext SynchronizationContext { get; }
 		public WorkspaceLock Lock { get; }
 		readonly ILogger _logger;
diff --git a/Engine/Source/Programs/UnrealGameSync/UnrealGameSyncShared/WorkspaceState.cs b/Engine/Source/Programs/UnrealGameSync/UnrealGameSyncShared/WorkspaceState.cs
--- a/Engine/Source/Programs/UnrealGameSync/UnrealGameSyncShared/WorkspaceState.cs
+++ b/Engine/Source/Programs/UnrealGameSync/UnrealGameSyncShared/WorkspaceState.cs
@@ -27,7 +27,6 @@ namespace UnrealGameSync
 		public string ProjectPath { get; set; } = String.Empty;
 		public string? StreamName { get; set; }
 
-		public IDictionary<string, DateTime> Streams { get; set; } = new Dictionary<string, DateTime>();
 		public string ProjectIdentifier { get; set; } = String.Empty; // Should be fully reset if this changes
 		public bool IsEnterpriseProject { get; set; }
 
@@ -89,11 +88,6 @@ namespace UnrealGameSync
 			BranchPath = other.BranchPath;
 			ProjectPath = other.ProjectPath;
 			StreamName = other.StreamName;
-			Streams.Clear();
-			foreach (KeyValuePair<string, DateTime> kvp in other.Streams)
-			{
-				Streams.Add(kvp.Key, kvp.Value);
-			}
 
 			ProjectIdentifier = other.ProjectIdentifier;
 			IsEnterpriseProject = other.IsEnterpriseProject;
@@ -155,7 +149,6 @@ namespace UnrealGameSync
 		public string BranchPath => _inner.BranchPath;
 		public string ProjectPath => _inner.ProjectPath;
 		public string? StreamName => _inner.StreamName;
-		public IDictionary<string, DateTime> Streams => _inner.Streams;
 		public string ProjectIdentifier => _inner.ProjectIdentifier;
 		public bool IsEnterpriseProject => _inner.IsEnterpriseProject;
 
@@ -291,6 +284,7 @@ namespace UnrealGameSync
 						stream.SetLength(stream.Position);
 
 						Current = new ReadOnlyWorkspaceState(state);
+
 						return Current;
 					}
 				}
