Fixed a bug where we were not correctly reading the RemoveOldBinaries property for archives. This meant that archives pulled from UnrealCloudDDC would not be retained locally preventing us from patching them in place.

diff --git a/Engine/Source/Programs/UnrealGameSync/UnrealGameSyncShared/WorkspaceUpdate.cs b/Engine/Source/Programs/UnrealGameSync/UnrealGameSyncShared/WorkspaceUpdate.cs
--- a/Engine/Source/Programs/UnrealGameSync/UnrealGameSyncShared/WorkspaceUpdate.cs
+++ b/Engine/Source/Programs/UnrealGameSync/UnrealGameSyncShared/WorkspaceUpdate.cs
@@ -1344,7 +1344,7 @@ namespace UnrealGameSync
 						bool removeOldBinaries = true;
 						if (archive != null && archive.RemoveOldBinaries.HasValue)
 						{
-							removeOldBinaries = archive.RemoveOldBinaries.HasValue;
+							removeOldBinaries = archive.RemoveOldBinaries.Value;
 						}
 						// Remove any existing binaries
 						FileReference manifestFileName = FileReference.Combine(manifestDirectoryName, String.Format("{0}.zipmanifest", archiveType));
