Unreal Game Sync - Multiple bucket search for PCBs from cloud, makes it more resilient to bucket naming changes

diff --git a/Engine/Source/Programs/UnrealGameSync/UnrealGameSyncShared/ArchiveInfo.cs b/Engine/Source/Programs/UnrealGameSync/UnrealGameSyncShared/ArchiveInfo.cs
--- a/Engine/Source/Programs/UnrealGameSync/UnrealGameSyncShared/ArchiveInfo.cs
+++ b/Engine/Source/Programs/UnrealGameSync/UnrealGameSyncShared/ArchiveInfo.cs
@@ -5,6 +5,7 @@ using System.Collections.Generic;
 using System.Diagnostics.CodeAnalysis;
 using System.IO;
 using System.Linq;
+using System.Text.RegularExpressions;
 using System.Threading;
 using System.Threading.Tasks;
 using EpicGames.Core;
@@ -486,31 +487,33 @@ namespace UnrealGameSync
 		{
 			CbWriter queryWriter = new CbWriter();
 			queryWriter.BeginObject();
-			queryWriter.BeginObject("query");
 			queryWriter.BeginObject("ugs-project");
 			queryWriter.WriteString("$eq", projectIdentifier);
 			queryWriter.EndObject(); // ugs-project
-			queryWriter.EndObject(); // query
+			queryWriter.BeginObject("branch");
+			queryWriter.WriteString("$eq", SanitizeBucketComponent(_streamName));
+			queryWriter.EndObject(); // branch
+			queryWriter.BeginObject("platform");
+			queryWriter.WriteString("$eq", "windows");
+			queryWriter.EndObject(); // platform
 			queryWriter.EndObject();
 
-			CbObject query = queryWriter.ToObject();
-
 			string SanitizeBucketComponent(string s)
 			{
 				// . is reserved for bucket component separation
-				return s.Replace(".", "-", StringComparison.OrdinalIgnoreCase);
+				return s.Replace(".", "-", StringComparison.OrdinalIgnoreCase).ToLowerInvariant();
 			}
 			string projectName = Path.GetFileNameWithoutExtension(projectIdentifier);
-			string platform = "windows";
-			string bucketId = $"{SanitizeBucketComponent(projectName)}.{SanitizeBucketComponent(_archiveType)}.{SanitizeBucketComponent(_streamName)}.{SanitizeBucketComponent(platform)}";
-			bucketId = bucketId.ToLower();
+			Regex bucketRegex = new Regex($"{SanitizeBucketComponent(projectName)}.{SanitizeBucketComponent(_archiveType)}\\..*");
+
+			CbObject query = queryWriter.ToObject();
 
-			IAsyncEnumerable<FoundBuildResponse> artifactResponses = cloudStorage.FindBuildAsync(host, namespaceId, bucketId, query, cancellationToken: cancellationToken);
+			IAsyncEnumerable<FoundBuildResponse> artifactResponses = cloudStorage.FindBuildAsync(host, namespaceId, bucketRegex, query, cancellationToken: cancellationToken);
 			await foreach (FoundBuildResponse response in artifactResponses)
 			{
 				if (Int32.TryParse(response.Commit, out int commit))
 				{
-					CloudArchive archive = new CloudArchive(cloudStorage, host, namespaceId, bucketId, response.BuildId);
+					CloudArchive archive = new CloudArchive(cloudStorage, host, namespaceId, response.BucketId, response.BuildId);
 					ChangeNumberToArchive[commit] = archive;
 				}
 				// not a integer commit, unable to map to perforce
diff --git a/Engine/Source/Programs/UnrealGameSync/UnrealGameSyncShared/CloudStorage.cs b/Engine/Source/Programs/UnrealGameSync/UnrealGameSyncShared/CloudStorage.cs
--- a/Engine/Source/Programs/UnrealGameSync/UnrealGameSyncShared/CloudStorage.cs
+++ b/Engine/Source/Programs/UnrealGameSync/UnrealGameSyncShared/CloudStorage.cs
@@ -6,6 +6,7 @@ using System.Collections.Generic;
 using System.IO;
 using System.Reflection;
 using System.Runtime.CompilerServices;
+using System.Text.RegularExpressions;
 using System.Threading;
 using System.Threading.Tasks;
 using EpicGames.Core;
@@ -22,10 +23,13 @@ namespace UnrealGameSync
 	/// </summary>
 	/// <param name="buildName"></param>
 	/// <param name="buildId"></param>
+	/// <param name="bucketId"></param>
 	/// <param name="commit"></param>
-	public class FoundBuildResponse(string buildName, string buildId, string commit)
+	public class FoundBuildResponse(string buildName, string buildId, string bucketId, string commit)
 	{
 		public string BuildId { get; init; } = buildId;
+		public string BucketId { get; init; } = bucketId;
+
 		public string Commit { get; init; } = commit;
 		public string Name { get; set; } = buildName;
 	}
@@ -67,15 +71,15 @@ namespace UnrealGameSync
 		Task<bool> DownloadBuildAsync(string host, string namespaceId, string bucketId, string buildId, DirectoryReference outputDir, DirectoryReference? zenStateDirectory, IProgress<string> progress, CancellationToken cancellationToken);
 
 		/// <summary>
-		/// Search for builds in UE Cloud Storage
+		/// Search for builds in UE Cloud Storage using a regular expression to describe which buckets the build might be in
 		/// </summary>
 		/// <param name="host">The host url</param>
 		/// <param name="namespaceId">The namespace</param>
-		/// <param name="bucketId">The bucket</param>
+		/// <param name="bucketRegex">A regex that describe which buckets the build might be in</param>
 		/// <param name="query">CbObject that describes the query, see UE Cloud Storage documentation</param>
 		/// <param name="cancellationToken">A cancellation token</param>
 		/// <returns></returns>
-		IAsyncEnumerable<FoundBuildResponse> FindBuildAsync(string host, string namespaceId, string bucketId, CbObject query, CancellationToken cancellationToken);
+		IAsyncEnumerable<FoundBuildResponse> FindBuildAsync(string host, string namespaceId, Regex bucketRegex, CbObject query, CancellationToken cancellationToken);
 	}
 
 	class CloudStorage(ILogger<CloudStorage> logger, ITokenStore tokenStore) : ICloudStorage, IDisposable
@@ -208,7 +212,7 @@ namespace UnrealGameSync
 			return exitCode != 0;
 		}
 
-		public async IAsyncEnumerable<FoundBuildResponse> FindBuildAsync(string host, string namespaceId, string bucketId, CbObject query, [EnumeratorCancellation] CancellationToken cancellationToken)
+		public async IAsyncEnumerable<FoundBuildResponse> FindBuildAsync(string host, string namespaceId, Regex bucketRegex, CbObject query, [EnumeratorCancellation] CancellationToken cancellationToken)
 		{
 			string? accessToken = await GetAccessToken(new Uri(host), allowInteractiveLogin: false);
 
@@ -223,8 +227,15 @@ namespace UnrealGameSync
 			tempOutputFile = Path.ChangeExtension(tempOutputFile, "cbo");
 			try
 			{
-				await File.WriteAllBytesAsync(tempInputFile, query.GetView().ToArray(), cancellationToken);
-				string cmdline = $"builds list --host {host} --namespace {namespaceId} --bucket {bucketId} --query-path {tempInputFile} --result-path {tempOutputFile}";
+				// copy the query and append the bucket regex filter
+				CbWriter writer = new CbWriter();
+				writer.BeginObject();
+				writer.WriteObject("query", query);
+				writer.WriteString("bucketRegex", bucketRegex.ToString()!);
+				writer.EndObject();
+
+				await File.WriteAllBytesAsync(tempInputFile, writer.ToByteArray(), cancellationToken);
+				string cmdline = $"builds list --host {host} --namespace {namespaceId} --query-path {tempInputFile} --result-path {tempOutputFile}";
 
 				Progress<string> progress = new Progress<string>();
 				bool failed = await RunZenCli(cmdline, accessToken, progress, cancellationToken);
@@ -247,6 +258,7 @@ namespace UnrealGameSync
 					}
 
 					string buildName = metadata["name"].AsString();
+					string bucketId = result["bucketId"].AsString();
 					string buildId = result["buildId"].AsString();
 					CbField commitField = metadata["commit"];
 					if (commitField.Equals(CbField.Empty))
@@ -276,7 +288,7 @@ namespace UnrealGameSync
 						// if there is no commit specified we can not know which changelist this PCB belong to
 						continue;
 					}
-					yield return new FoundBuildResponse(buildName, buildId, commit);
+					yield return new FoundBuildResponse(buildName, buildId, bucketId, commit);
 				}
 			}
 			finally
