Unreal Game Sync - Quoted file path to temp files when listing cloud builds to make whitespaces are handled correctly

diff --git a/Engine/Source/Programs/UnrealGameSync/UnrealGameSyncShared/CloudStorage.cs b/Engine/Source/Programs/UnrealGameSync/UnrealGameSyncShared/CloudStorage.cs
--- a/Engine/Source/Programs/UnrealGameSync/UnrealGameSyncShared/CloudStorage.cs
+++ b/Engine/Source/Programs/UnrealGameSync/UnrealGameSyncShared/CloudStorage.cs
@@ -235,7 +235,7 @@ namespace UnrealGameSync
 				writer.EndObject();
 
 				await File.WriteAllBytesAsync(tempInputFile, writer.ToByteArray(), cancellationToken);
-				string cmdline = $"builds list --host {host} --namespace {namespaceId} --query-path {tempInputFile} --result-path {tempOutputFile}";
+				string cmdline = $"builds list --host {host} --namespace {namespaceId} --query-path \"{tempInputFile}\" --result-path \"{tempOutputFile}\"";
 
 				Progress<string> progress = new Progress<string>();
 				bool failed = await RunZenCli(cmdline, accessToken, progress, cancellationToken);
