checking the type of the stream before trigger the stream change

diff --git a/Engine/Source/Programs/UnrealGameSync/UnrealGameSync/PerforceMonitor.cs b/Engine/Source/Programs/UnrealGameSync/UnrealGameSync/PerforceMonitor.cs
--- a/Engine/Source/Programs/UnrealGameSync/UnrealGameSync/PerforceMonitor.cs
+++ b/Engine/Source/Programs/UnrealGameSync/UnrealGameSync/PerforceMonitor.cs
@@ -292,7 +292,14 @@ namespace UnrealGameSync
 								if (streamHasChanged)
 								{
 									// Only refresh the sync state if we already have a stored out stream hierarchy
-									if (cache!.Streams.Any())
+									// Added condition - to revisit later - only trigger if the current stream is a virtual stream
+									bool isVirtual = false;
+									if (record != null)
+									{
+										isVirtual = record.Type?.Equals("virtual", StringComparison.OrdinalIgnoreCase) ?? false;
+									}
+
+									if (isVirtual && cache!.Streams.Any())
 									{
 										_synchronizationContext.Post(_ => OnStreamChange?.Invoke(), null);
 									}
