Fixed an issue where UGS ini sections which used wildcards would prioritise the first wildcard entry rather than the last. Note that explicit project paths will still always take priority.

diff --git a/Engine/Source/Programs/UnrealGameSync/UnrealGameSyncShared/ConfigFile.cs b/Engine/Source/Programs/UnrealGameSync/UnrealGameSyncShared/ConfigFile.cs
--- a/Engine/Source/Programs/UnrealGameSync/UnrealGameSyncShared/ConfigFile.cs
+++ b/Engine/Source/Programs/UnrealGameSync/UnrealGameSyncShared/ConfigFile.cs
@@ -348,6 +348,8 @@ namespace UnrealGameSync
 		public string Name { get; set; }
 		public Dictionary<string, string> Pairs { get; } = new Dictionary<string, string>();
 
+		private readonly HashSet<string> _wildCardKeys = new HashSet<string>();
+
 		public ConfigSection(string inName)
 		{
 			Name = inName;
@@ -356,6 +358,7 @@ namespace UnrealGameSync
 		public void Clear()
 		{
 			Pairs.Clear();
+			_wildCardKeys.Clear();
 		}
 
 		public void SetValue(string key, int value)
@@ -373,7 +376,7 @@ namespace UnrealGameSync
 			Pairs[key] = value ? "1" : "0";
 		}
 
-		public void SetValue(string key, string value)
+		public void SetValue(string key, string value, bool fromWildCard = false)
 		{
 			if (value == null)
 			{
@@ -383,6 +386,20 @@ namespace UnrealGameSync
 			{
 				Pairs[key] = value;
 			}
+
+			if (fromWildCard)
+			{
+				_wildCardKeys.Add(key);
+			}
+			else
+			{
+				_wildCardKeys.Remove(key);
+			}
+		}
+
+		public bool HasWildCardValue(string key)
+		{
+			return _wildCardKeys.Contains(key);
 		}
 
 		public void SetValues(string key, string[] values)
@@ -603,9 +620,11 @@ namespace UnrealGameSync
 							{
 								string key = trimLine.Substring(0, equalsIdx).TrimEnd();
 
-								if (wildCardSection == false || currentSection.GetValue(key) == null)
+								if (wildCardSection == false || // Non-wildcard settings always take priority
+									currentSection.GetValue(key) == null || // If we don't have any value already set, allow a wildcard to be set
+									currentSection.HasWildCardValue(key)) // If the current value was set from a wildcard, allow it to be overwritten
 								{
-									currentSection.SetValue(key, value);
+									currentSection.SetValue(key, value, wildCardSection);
 								}
 							}
 						}
