ugs cli refactoring
- moved all commands to a file per command
- moved the static utils in a file per "subject"
- fixed all the compilation warnings

diff --git a/Engine/Source/Programs/UnrealGameSync/UnrealGameSyncCmd/Command.cs b/Engine/Source/Programs/UnrealGameSync/UnrealGameSyncCmd/Command.cs
new file mode 100644
--- /dev/null
+++ b/Engine/Source/Programs/UnrealGameSync/UnrealGameSyncCmd/Command.cs
@@ -0,0 +1,14 @@
+﻿// Copyright Epic Games, Inc. All Rights Reserved.
+
+using System.Threading.Tasks;
+using UnrealGameSync;
+
+namespace UnrealGameSyncCmd
+{
+	internal abstract class Command
+	{
+		internal static BuildConfig EditorConfig => BuildConfig.Development;
+
+		public abstract Task ExecuteAsync(CommandContext context);
+	}
+}
diff --git a/Engine/Source/Programs/UnrealGameSync/UnrealGameSyncCmd/CommandContext.cs b/Engine/Source/Programs/UnrealGameSync/UnrealGameSyncCmd/CommandContext.cs
new file mode 100644
--- /dev/null
+++ b/Engine/Source/Programs/UnrealGameSync/UnrealGameSyncCmd/CommandContext.cs
@@ -0,0 +1,34 @@
+﻿// Copyright Epic Games, Inc. All Rights Reserved.
+
+using EpicGames.Core;
+using EpicGames.Horde;
+using Microsoft.Extensions.Logging;
+using UnrealGameSync;
+
+namespace UnrealGameSyncCmd
+{
+	internal class CommandContext
+	{
+		public CommandLineArguments Arguments { get; }
+		public ILogger Logger { get; }
+		public ILoggerFactory LoggerFactory { get; }
+
+		public UserSettings? GlobalSettings { get; }
+
+		public GlobalSettingsFile UserSettings { get; }
+
+		public IHordeClient? HordeClient { get; }
+		public ICloudStorage? CloudStorage { get; }
+
+		public CommandContext(CommandLineArguments arguments, ILogger logger, ILoggerFactory loggerFactory, GlobalSettingsFile userSettings, UserSettings? globalSettings, IHordeClient? hordeClient, ICloudStorage? cloudStorage)
+		{
+			Arguments = arguments;
+			Logger = logger;
+			LoggerFactory = loggerFactory;
+			GlobalSettings = globalSettings;
+			UserSettings = userSettings;
+			HordeClient = hordeClient;
+			CloudStorage = cloudStorage;
+		}
+	}
+}
diff --git a/Engine/Source/Programs/UnrealGameSync/UnrealGameSyncCmd/CommandInfo.cs b/Engine/Source/Programs/UnrealGameSync/UnrealGameSyncCmd/CommandInfo.cs
new file mode 100644
--- /dev/null
+++ b/Engine/Source/Programs/UnrealGameSync/UnrealGameSyncCmd/CommandInfo.cs
@@ -0,0 +1,24 @@
+﻿// Copyright Epic Games, Inc. All Rights Reserved.
+
+using System;
+
+namespace UnrealGameSyncCmd
+{
+	internal class CommandInfo
+	{
+		public string Name { get; }
+		public Type Type { get; }
+		public Type? OptionsType { get; }
+		public string? Usage { get; }
+		public string? Brief { get; }
+
+		public CommandInfo(string name, Type type, Type? optionsType, string? usage, string? brief)
+		{
+			Name = name;
+			Type = type;
+			OptionsType = optionsType;
+			Usage = usage;
+			Brief = brief;
+		}
+	}
+}
diff --git a/Engine/Source/Programs/UnrealGameSync/UnrealGameSyncCmd/Commands/BuildCommand.cs b/Engine/Source/Programs/UnrealGameSync/UnrealGameSyncCmd/Commands/BuildCommand.cs
new file mode 100644
--- /dev/null
+++ b/Engine/Source/Programs/UnrealGameSync/UnrealGameSyncCmd/Commands/BuildCommand.cs
@@ -0,0 +1,194 @@
+﻿// Copyright Epic Games, Inc. All Rights Reserved.
+
+using System;
+using System.Collections.Generic;
+using System.Linq;
+using System.Threading;
+using System.Threading.Tasks;
+using EpicGames.Core;
+using EpicGames.Horde;
+using EpicGames.OIDC;
+using EpicGames.Perforce;
+using Microsoft.Extensions.DependencyInjection;
+using Microsoft.Extensions.Logging;
+using UnrealGameSync;
+using UnrealGameSyncCmd.Utils;
+
+using UserErrorException = UnrealGameSyncCmd.Exceptions.UserErrorException;
+
+namespace UnrealGameSyncCmd.Commands
+{
+	internal class BuildCommandOptions
+	{
+		[CommandLine("-List")]
+		public bool ListOnly { get; set; }
+	}
+
+	internal class BuildCommand : Command
+	{
+		public override async Task ExecuteAsync(CommandContext context)
+		{
+			ILogger logger = context.Logger;
+			context.Arguments.TryGetPositionalArgument(out string? target);
+
+			BuildCommandOptions options = new BuildCommandOptions();
+			context.Arguments.ApplyTo(options);
+			context.Arguments.CheckAllArgumentsUsed();
+
+			UserWorkspaceSettings settings = UserSettingsUtils.ReadRequiredUserWorkspaceSettings();
+			using IPerforceConnection perforceClient = await PerforceConnectionUtils.ConnectAsync(settings, context.LoggerFactory);
+			WorkspaceStateWrapper state = await WorkspaceStateUtils.ReadWorkspaceState(perforceClient, settings, context.UserSettings, logger);
+
+			ProjectInfo projectInfo = new ProjectInfo(settings.RootDir, state.Current);
+
+			UserProjectSettings projectSettings = context.UserSettings.FindOrAddProjectSettings(projectInfo, settings, logger);
+
+			ConfigFile projectConfig = await ConfigUtils.ReadProjectConfigFileAsync(perforceClient, projectInfo, logger, CancellationToken.None);
+			FileReference editorTarget = ConfigUtils.GetEditorTargetFile(projectInfo, projectConfig);
+
+			Dictionary<Guid, ConfigObject> buildStepObjects = ConfigUtils.GetDefaultBuildStepObjects(projectInfo, editorTarget.GetFileNameWithoutAnyExtensions(), EditorConfig, projectConfig, false);
+
+			if (context.GlobalSettings != null)
+			{
+				BuildStep.MergeBuildStepObjects(buildStepObjects, projectConfig.GetValues("Build.Step", Array.Empty<string>()).Select(x => new ConfigObject(x)));
+				BuildStep.MergeBuildStepObjects(buildStepObjects, projectSettings.BuildSteps);
+			}
+
+			if (options.ListOnly)
+			{
+				List<BuildStep> buildSteps = buildStepObjects.Values.Select(x => new BuildStep(x)).OrderBy(x => x.NormalSync).ToList();
+
+				int longestDescription = buildSteps.Max(b => String.IsNullOrWhiteSpace(b.Description) ? 0 : b.Description.Length);
+				string descriptionHeader = "Description";
+				string descriptionSpace = new string(Enumerable.Repeat(' ', longestDescription + 1 - descriptionHeader.Length).ToArray());
+				string descriptionDashes = new string(Enumerable.Repeat('-', longestDescription + 2).ToArray());
+				string format = $"  {{Id,-36}} | {{Name,-{longestDescription}}} | {{Type,-10}} | {{Enabled,-8}}";
+
+				logger.LogInformation("Available build steps:");
+				logger.LogInformation("");
+				logger.LogInformation("  Id                                   | {DescriptionHeader}{DescriptionSpace}| Type       | Enabled", descriptionHeader, descriptionSpace);
+				logger.LogInformation("  -------------------------------------|{DescriptionDashes}|------------|-----------------", descriptionDashes);
+				foreach (BuildStep buildStep in buildStepObjects.Values.Select(x => new BuildStep(x)).OrderBy(x => x.NormalSync))
+				{
+					#pragma warning disable CA2254 // Template should be a static expression
+					logger.LogInformation(format, buildStep.UniqueId, buildStep.Description, buildStep.Type, buildStep.NormalSync);
+					#pragma warning restore CA2254
+				}
+				return;
+			}
+
+			HashSet<Guid>? steps = new HashSet<Guid>();
+			if (target != null)
+			{
+				if (!Guid.TryParse(target, out Guid id))
+				{
+					logger.LogError("Unable to parse '{Target}' as a GUID. Pass -List to show all available build steps and their identifiers.", target);
+					return;
+				}
+
+				steps.Add(id);
+			}
+
+			// check that the tools are installed
+			List<Guid> uninstalledTools = new List<Guid>();
+			if (context.GlobalSettings != null)
+			{
+				using ToolUpdateMonitor? toolsMonitor = await GetToolsMonitor(logger, context, perforceClient);
+
+				if (toolsMonitor != null)
+				{
+					foreach (Guid step in steps)
+					{
+						if (buildStepObjects.TryGetValue(step, out ConfigObject? config))
+						{
+							BuildStep buildStep = new BuildStep(config);
+							if (CanRunStep(buildStep, toolsMonitor))
+							{
+								continue;
+							}
+
+							logger.LogWarning("Build step '{Description}' cannot be run as the tool is not installed.", buildStep.Description);
+							uninstalledTools.Add(step);
+						}
+					}
+				}
+			}
+
+			foreach (Guid uninstalledTool in uninstalledTools)
+			{
+				steps.Remove(uninstalledTool);
+			}
+
+			WorkspaceUpdateContext updateContext = new WorkspaceUpdateContext(
+				state.Current.CurrentChangeNumber,
+				WorkspaceUpdateOptions.Build,
+				BuildConfig.Development,
+				null,
+				projectSettings.BuildSteps, steps);
+
+			WorkspaceUpdate update = new WorkspaceUpdate(updateContext);
+			(WorkspaceUpdateResult result, string message) = await update.ExecuteAsync(perforceClient.Settings, projectInfo, state, context.Logger, CancellationToken.None);
+
+			if (result != WorkspaceUpdateResult.Success)
+			{
+				throw new UserErrorException("{Message}", message);
+			}
+		}
+
+		private static async Task<ToolUpdateMonitor?> GetToolsMonitor(ILogger logger, CommandContext context, IPerforceConnection perforceClient)
+		{
+			ToolsCommandCommandOptions options = context.Arguments.ApplyTo<ToolsCommandCommandOptions>(logger);
+			context.Arguments.CheckAllArgumentsUsed();
+
+			DirectoryReference dataFolder = DirectoryReference.Combine(DirectoryReference.GetSpecialFolder(Environment.SpecialFolder.LocalApplicationData)!, "UnrealGameSync");
+			DirectoryReference.CreateDirectory(dataFolder);
+
+			// create a temporary service provider for the tool update monitor
+			ServiceCollection services = new ServiceCollection();
+			services.AddSingleton<IAsyncDisposer, AsyncDisposer>();
+			services.AddSingleton(sp => TokenStoreFactory.CreateTokenStore());
+			services.AddSingleton<OidcTokenManager>();
+
+			LauncherSettings launcherSettings = new LauncherSettings();
+			launcherSettings.Read();
+
+			if (launcherSettings.HordeServer != null)
+			{
+				services.AddHorde(options =>
+				{
+					options.ServerUrl = new Uri(launcherSettings.HordeServer);
+					options.AllowAuthPrompt = false;
+				});
+			}
+
+			ServiceProvider serviceProvider = services.BuildServiceProvider();
+
+			ToolUpdateMonitor toolUpdateMonitor =
+				new ToolUpdateMonitor(perforceClient.Settings, dataFolder, context.GlobalSettings!, logger, serviceProvider);
+
+			// get the list of tools available
+			logger.LogInformation("Retrieving tools information, please wait.");
+			await toolUpdateMonitor.GetDataFromBackendAsync();
+
+			return toolUpdateMonitor;
+		}
+
+		private static bool CanRunStep(BuildStep step, ToolUpdateMonitor toolsMonitor)
+		{
+			if (step.ToolId != Guid.Empty)
+			{
+				string? toolName = toolsMonitor.GetToolName(step.ToolId);
+				if (toolName == null)
+				{
+					return false;
+				}
+
+				if (toolsMonitor.GetToolPath(toolName) == null)
+				{
+					return false;
+				}
+			}
+			return true;
+		}
+	}
+}
diff --git a/Engine/Source/Programs/UnrealGameSync/UnrealGameSyncCmd/Commands/ChangesCommand.cs b/Engine/Source/Programs/UnrealGameSync/UnrealGameSyncCmd/Commands/ChangesCommand.cs
new file mode 100644
--- /dev/null
+++ b/Engine/Source/Programs/UnrealGameSync/UnrealGameSyncCmd/Commands/ChangesCommand.cs
@@ -0,0 +1,90 @@
+﻿// Copyright Epic Games, Inc. All Rights Reserved.
+
+using System;
+using System.Collections.Generic;
+using System.Linq;
+using System.Threading.Tasks;
+using EpicGames.Core;
+using EpicGames.Perforce;
+using Microsoft.Extensions.Logging;
+using UnrealGameSync;
+using UnrealGameSyncCmd.Utils;
+
+namespace UnrealGameSyncCmd.Commands
+{
+	internal class ChangesCommandOptions
+	{
+		[CommandLine("-Count=")]
+		public int Count { get; set; } = 10;
+
+		[CommandLine("-Lines=")]
+		public int LineCount { get; set; } = 3;
+	}
+
+	internal class ChangesCommand : Command
+	{
+		public override async Task ExecuteAsync(CommandContext context)
+		{
+			ILogger logger = context.Logger;
+
+			ChangesCommandOptions options = new ChangesCommandOptions();
+			context.Arguments.ApplyTo(options);
+			context.Arguments.CheckAllArgumentsUsed(context.Logger);
+
+			UserWorkspaceSettings settings = UserSettingsUtils.ReadRequiredUserWorkspaceSettings();
+			using IPerforceConnection perforceClient = await PerforceConnectionUtils.ConnectAsync(settings, context.LoggerFactory);
+
+			List<ChangesRecord> changes = await perforceClient.GetChangesAsync(EpicGames.Perforce.ChangesOptions.None, options.Count, ChangeStatus.Submitted, $"//{settings.ClientName}/...");
+			foreach (IEnumerable<ChangesRecord> changesBatch in changes.Batch(10))
+			{
+				List<DescribeRecord> describeRecords = await perforceClient.DescribeAsync(changesBatch.Select(x => x.Number).ToArray());
+
+				logger.LogInformation("  Change    Type     Author          Description");
+				foreach (DescribeRecord describeRecord in describeRecords)
+				{
+					PerforceChangeDetails details = new PerforceChangeDetails(describeRecord);
+
+					string type;
+					if (details.ContainsCode)
+					{
+						if (details.ContainsContent)
+						{
+							type = "Both";
+						}
+						else
+						{
+							type = "Code";
+						}
+					}
+					else
+					{
+						if (details.ContainsContent)
+						{
+							type = "Content";
+						}
+						else
+						{
+							type = "None";
+						}
+					}
+
+					string author = StringUtils.Truncate(describeRecord.User, 15);
+
+					List<string> lines = StringUtils.WordWrap(details.Description, Math.Max(ConsoleUtils.WindowWidth - 40, 10)).ToList();
+					if (lines.Count == 0)
+					{
+						lines.Add(String.Empty);
+					}
+
+					int lineCount = Math.Min(options.LineCount, lines.Count);
+
+					logger.LogInformation("  {Change,-9} {Type,-8} {Author,-15} {Description}", describeRecord.Number, type, author, lines[0]);
+					for (int lineIndex = 1; lineIndex < lineCount; lineIndex++)
+					{
+						logger.LogInformation("                                     {Description}", lines[lineIndex]);
+					}
+				}
+			}
+		}
+	}
+}
diff --git a/Engine/Source/Programs/UnrealGameSync/UnrealGameSyncCmd/Commands/ClientsCommand.cs b/Engine/Source/Programs/UnrealGameSync/UnrealGameSyncCmd/Commands/ClientsCommand.cs
new file mode 100644
--- /dev/null
+++ b/Engine/Source/Programs/UnrealGameSync/UnrealGameSyncCmd/Commands/ClientsCommand.cs
@@ -0,0 +1,41 @@
+﻿// Copyright Epic Games, Inc. All Rights Reserved.
+
+using System;
+using System.Collections.Generic;
+using System.Threading.Tasks;
+using Microsoft.Extensions.Logging;
+
+using EpicGames.Perforce;
+
+using UnrealGameSyncCmd.Options;
+using UnrealGameSyncCmd.Utils;
+
+namespace UnrealGameSyncCmd.Commands
+{
+	internal class ClientsCommandOptions : ServerOptions
+	{
+	}
+
+	internal class ClientsCommand : Command
+	{
+		public override async Task ExecuteAsync(CommandContext context)
+		{
+			ILogger logger = context.Logger;
+
+			ClientsCommandOptions options = context.Arguments.ApplyTo<ClientsCommandOptions>(logger);
+			context.Arguments.CheckAllArgumentsUsed();
+
+			using IPerforceConnection perforceClient = await PerforceConnectionUtils.ConnectAsync(options.ServerAndPort, options.UserName, null, context.LoggerFactory);
+			InfoRecord info = await perforceClient.GetInfoAsync(InfoOptions.ShortOutput);
+
+			List<ClientsRecord> clients = await perforceClient.GetClientsAsync(EpicGames.Perforce.ClientsOptions.None, perforceClient.Settings.UserName);
+			foreach (ClientsRecord client in clients)
+			{
+				if (String.Equals(info.ClientHost, client.Host, StringComparison.OrdinalIgnoreCase))
+				{
+					logger.LogInformation("{Client,-50} {Root}", client.Name, client.Root);
+				}
+			}
+		}
+	}
+}
diff --git a/Engine/Source/Programs/UnrealGameSync/UnrealGameSyncCmd/Commands/ConfigCommand.cs b/Engine/Source/Programs/UnrealGameSync/UnrealGameSyncCmd/Commands/ConfigCommand.cs
new file mode 100644
--- /dev/null
+++ b/Engine/Source/Programs/UnrealGameSync/UnrealGameSyncCmd/Commands/ConfigCommand.cs
@@ -0,0 +1,48 @@
+﻿// Copyright Epic Games, Inc. All Rights Reserved.
+
+using System.Diagnostics;
+using System.Linq;
+using System.Threading.Tasks;
+
+using Microsoft.Extensions.Logging;
+
+using UnrealGameSync;
+using UnrealGameSyncCmd.Options;
+using UnrealGameSyncCmd.Utils;
+
+namespace UnrealGameSyncCmd.Commands
+{
+	internal class ConfigCommand : Command
+	{
+		public override async Task ExecuteAsync(CommandContext context)
+		{
+			ILogger logger = context.Logger;
+
+			UserWorkspaceSettings settings = UserSettingsUtils.ReadRequiredUserWorkspaceSettings();
+			if (!context.Arguments.GetUnusedArguments().Any())
+			{
+				ProcessStartInfo startInfo = new ProcessStartInfo();
+				startInfo.FileName = settings.ConfigFile.FullName;
+				startInfo.UseShellExecute = true;
+				using (Process? editor = Process.Start(startInfo))
+				{
+					if (editor != null)
+					{
+						await editor.WaitForExitAsync();
+					}
+				}
+			}
+			else
+			{
+				ConfigCommandOptions options = new ConfigCommandOptions();
+				context.Arguments.ApplyTo(options);
+				context.Arguments.CheckAllArgumentsUsed(context.Logger);
+
+				options.ApplyTo(settings);
+				settings.Save(logger);
+
+				logger.LogInformation("Updated {ConfigFile}", settings.ConfigFile);
+			}
+		}
+	}
+}
diff --git a/Engine/Source/Programs/UnrealGameSync/UnrealGameSyncCmd/Commands/FilterCommand.cs b/Engine/Source/Programs/UnrealGameSync/UnrealGameSyncCmd/Commands/FilterCommand.cs
new file mode 100644
--- /dev/null
+++ b/Engine/Source/Programs/UnrealGameSync/UnrealGameSyncCmd/Commands/FilterCommand.cs
@@ -0,0 +1,227 @@
+﻿// Copyright Epic Games, Inc. All Rights Reserved.
+
+using System;
+using System.Collections.Generic;
+using System.Linq;
+using System.Threading;
+using System.Threading.Tasks;
+
+using Microsoft.Extensions.Logging;
+
+using EpicGames.Core;
+using EpicGames.Perforce;
+
+using UnrealGameSync;
+using UnrealGameSyncCmd.Utils;
+
+using UserErrorException = UnrealGameSyncCmd.Exceptions.UserErrorException;
+
+namespace UnrealGameSyncCmd.Commands
+{
+	internal class FilterCommandOptions
+	{
+		[CommandLine("-Reset")]
+		public bool Reset { get; set; } = false;
+
+		[CommandLine("-Include=")]
+		public List<string> Include { get; set; } = new List<string>();
+
+		[CommandLine("-Exclude=")]
+		public List<string> Exclude { get; set; } = new List<string>();
+
+		[CommandLine("-View=", ListSeparator = ';')]
+		public List<string>? View { get; set; }
+
+		[CommandLine("-AddView=", ListSeparator = ';')]
+		public List<string> AddView { get; set; } = new List<string>();
+
+		[CommandLine("-RemoveView=", ListSeparator = ';')]
+		public List<string> RemoveView { get; set; } = new List<string>();
+
+		[CommandLine("-AllProjects", Value = "true")]
+		[CommandLine("-OnlyCurrent", Value = "false")]
+		public bool? AllProjects { get; set; } = null;
+
+		[CommandLine("-LocalProjects", Value = "true")]
+		[CommandLine("-NoExtraLocalProjects", Value = "false")]
+		public bool? LocalProjects { get; set; } = null;
+
+		[CommandLine("-GpfAllProjects", Value = "true")]
+		[CommandLine("-GpfOnlyCurrent", Value = "false")]
+		public bool? AllProjectsInSln { get; set; } = null;
+
+		[CommandLine("-GpfMinimalSln", Value = "true")]
+		[CommandLine("-GpfFullSln", Value = "false")]
+		public bool? UprojectSpecificSln { get; set; } = null;
+
+		[CommandLine("-Global")]
+		public bool Global { get; set; }
+	}
+
+	internal class FilterCommand : Command
+	{
+		public override async Task ExecuteAsync(CommandContext context)
+		{
+			ILogger logger = context.Logger;
+
+			UserWorkspaceSettings workspaceSettings = UserSettingsUtils.ReadRequiredUserWorkspaceSettings();
+			using IPerforceConnection perforceClient = await PerforceConnectionUtils.ConnectAsync(workspaceSettings, context.LoggerFactory);
+			WorkspaceStateWrapper workspaceState = await WorkspaceStateUtils.ReadWorkspaceState(perforceClient, workspaceSettings, context.UserSettings, logger);
+			ProjectInfo projectInfo = new ProjectInfo(workspaceSettings.RootDir, workspaceState.Current);
+
+			ConfigFile projectConfig = await ConfigUtils.ReadProjectConfigFileAsync(perforceClient, projectInfo, logger, CancellationToken.None);
+			Dictionary<Guid, WorkspaceSyncCategory> syncCategories = ConfigUtils.GetSyncCategories(projectConfig);
+
+			FilterSettings globalFilter = context.UserSettings.Global.Filter;
+			FilterSettings workspaceFilter = workspaceSettings.Filter;
+
+			IDictionary<string, Preset> roles = ConfigUtils.GetPresets(projectConfig, workspaceState.Current.ProjectIdentifier);
+
+			roles.TryGetValue(workspaceSettings.Preset, out Preset? role);
+
+			FilterCommandOptions options = context.Arguments.ApplyTo<FilterCommandOptions>(logger);
+			context.Arguments.CheckAllArgumentsUsed(context.Logger);
+
+			if (options.Global)
+			{
+				ApplyCommandOptions(context.UserSettings.Global.Filter, options, syncCategories.Values, logger);
+				context.UserSettings.Save(logger);
+			}
+			else
+			{
+				ApplyCommandOptions(workspaceSettings.Filter, options, syncCategories.Values, logger);
+				workspaceSettings.Save(logger);
+			}
+
+			Dictionary<Guid, bool> globalCategories = globalFilter.GetCategories();
+			Dictionary<Guid, bool> workspaceCategories = workspaceFilter.GetCategories();
+
+			logger.LogInformation("Categories:");
+			foreach (WorkspaceSyncCategory syncCategory in syncCategories.Values)
+			{
+				bool enabled;
+
+				string scope = "(Default)";
+				if (globalCategories.TryGetValue(syncCategory.UniqueId, out enabled))
+				{
+					scope = "(Global)";
+				}
+				else if (workspaceCategories.TryGetValue(syncCategory.UniqueId, out enabled))
+				{
+					scope = "(Workspace)";
+				}
+				else
+				{
+					enabled = syncCategory.Enable;
+				}
+
+				if (role != null && role.Categories.TryGetValue(syncCategory.UniqueId, out RoleCategory? roleCategory))
+				{
+					scope = $"(Preset: {role.Name})";
+					enabled = roleCategory.Enabled;
+				}
+
+				logger.LogInformation("  {Id,30} {Enabled,3} {Scope,-9} {Name}", syncCategory.UniqueId, enabled ? "Yes" : "No", scope, syncCategory.Name);
+			}
+
+			if (globalFilter.View.Count > 0)
+			{
+				logger.LogInformation("");
+				logger.LogInformation("Global View:");
+				foreach (string line in globalFilter.View)
+				{
+					logger.LogInformation("  {Line}", line);
+				}
+			}
+			if (workspaceFilter.View.Count > 0)
+			{
+				logger.LogInformation("");
+				logger.LogInformation("Workspace View:");
+				foreach (string line in workspaceFilter.View)
+				{
+					logger.LogInformation("  {Line}", line);
+				}
+			}
+			if (role != null && role.Views.Count > 0)
+			{
+				logger.LogInformation("");
+				logger.LogInformation("Preset View:");
+				foreach (string line in role.Views)
+				{
+					logger.LogInformation("  {Line}", line);
+				}
+			}
+
+			string[] filter = SyncFilterUtils.ReadSyncFilter(workspaceSettings, context.UserSettings, projectConfig, workspaceState.Current.ProjectIdentifier);
+
+			logger.LogInformation("");
+			logger.LogInformation("Combined view:");
+			foreach (string filterLine in filter)
+			{
+				logger.LogInformation("  {FilterLine}", filterLine);
+			}
+		}
+
+		static void ApplyCommandOptions(FilterSettings settings, FilterCommandOptions commandOptions, IEnumerable<WorkspaceSyncCategory> syncCategories, ILogger logger)
+		{
+			if (commandOptions.Reset)
+			{
+				logger.LogInformation("Resetting settings...");
+				settings.Reset();
+			}
+
+			HashSet<Guid> includeCategories = new HashSet<Guid>(commandOptions.Include.Select(x => GetCategoryId(x, syncCategories)));
+			HashSet<Guid> excludeCategories = new HashSet<Guid>(commandOptions.Exclude.Select(x => GetCategoryId(x, syncCategories)));
+
+			Guid id = includeCategories.FirstOrDefault(x => excludeCategories.Contains(x));
+			if (id != Guid.Empty)
+			{
+				throw new UserErrorException("Category {Id} cannot be both included and excluded", id);
+			}
+
+			includeCategories.ExceptWith(settings.IncludeCategories);
+			settings.IncludeCategories.AddRange(includeCategories);
+
+			excludeCategories.ExceptWith(settings.ExcludeCategories);
+			settings.ExcludeCategories.AddRange(excludeCategories);
+
+			if (commandOptions.View != null)
+			{
+				settings.View.Clear();
+				settings.View.AddRange(commandOptions.View);
+			}
+			if (commandOptions.RemoveView.Count > 0)
+			{
+				HashSet<string> viewRemove = new HashSet<string>(commandOptions.RemoveView, StringComparer.OrdinalIgnoreCase);
+				settings.View.RemoveAll(x => viewRemove.Contains(x));
+			}
+			if (commandOptions.AddView.Count > 0)
+			{
+				HashSet<string> viewLines = new HashSet<string>(settings.View, StringComparer.OrdinalIgnoreCase);
+				settings.View.AddRange(commandOptions.AddView.Where(x => !viewLines.Contains(x)));
+			}
+
+			settings.AllProjects = commandOptions.AllProjects ?? settings.AllProjects;
+			settings.LocalProjects = commandOptions.LocalProjects ?? settings.LocalProjects;
+			settings.AllProjectsInSln = commandOptions.AllProjectsInSln ?? settings.AllProjectsInSln;
+			settings.UprojectSpecificSln = commandOptions.UprojectSpecificSln ?? settings.UprojectSpecificSln;
+		}
+
+		static Guid GetCategoryId(string text, IEnumerable<WorkspaceSyncCategory> syncCategories)
+		{
+			Guid id;
+			if (Guid.TryParse(text, out id))
+			{
+				return id;
+			}
+
+			WorkspaceSyncCategory? category = syncCategories.FirstOrDefault(x => x.Name.Equals(text, StringComparison.OrdinalIgnoreCase));
+			if (category != null)
+			{
+				return category.UniqueId;
+			}
+
+			throw new UserErrorException("Unable to find category '{Category}'", text);
+		}
+	}
+}
diff --git a/Engine/Source/Programs/UnrealGameSync/UnrealGameSyncCmd/Commands/InitCommand.cs b/Engine/Source/Programs/UnrealGameSync/UnrealGameSyncCmd/Commands/InitCommand.cs
new file mode 100644
--- /dev/null
+++ b/Engine/Source/Programs/UnrealGameSync/UnrealGameSyncCmd/Commands/InitCommand.cs
@@ -0,0 +1,226 @@
+﻿// Copyright Epic Games, Inc. All Rights Reserved.
+
+using System;
+using System.Collections.Generic;
+using System.Linq;
+using System.Net;
+using System.Text.RegularExpressions;
+using System.Threading.Tasks;
+
+using EpicGames.Core;
+using EpicGames.Perforce;
+using Microsoft.Extensions.Logging;
+
+using UnrealGameSync;
+using UnrealGameSyncCmd.Options;
+using UnrealGameSyncCmd.Utils;
+
+using UserErrorException = UnrealGameSyncCmd.Exceptions.UserErrorException;
+
+namespace UnrealGameSyncCmd.Commands
+{
+	internal class InitCommand : Command
+	{
+		public override async Task ExecuteAsync(CommandContext context)
+		{
+			ILogger logger = context.Logger;
+
+			// Get the positional argument indicating the file to look for
+			string? initName;
+			context.Arguments.TryGetPositionalArgument(out initName);
+
+			// Get the config settings from the command line
+			InitCommandOptions options = new InitCommandOptions();
+			context.Arguments.ApplyTo(options);
+			context.Arguments.CheckAllArgumentsUsed();
+
+			// Get the host name
+			using IPerforceConnection perforce = await PerforceConnectionUtils.ConnectAsync(options.ServerAndPort, options.UserName, null, context.LoggerFactory);
+			InfoRecord perforceInfo = await perforce.GetInfoAsync(InfoOptions.ShortOutput);
+			string hostName = perforceInfo.ClientHost ?? Dns.GetHostName();
+
+			// Create the perforce connection
+			if (initName != null)
+			{
+				await InitNewClientAsync(perforce, initName, hostName, options, context.GlobalSettings, logger);
+			}
+			else
+			{
+				await InitExistingClientAsync(perforce, hostName, options, context.GlobalSettings, logger);
+			}
+		}
+
+		static async Task InitNewClientAsync(IPerforceConnection perforce, string streamName, string hostName, InitCommandOptions options, UserSettings? globalSettings, ILogger logger)
+		{
+			logger.LogInformation("Checking stream...");
+
+			// Get the given stream
+			PerforceResponse<StreamRecord> streamResponse = await perforce.TryGetStreamAsync(streamName, true);
+			if (!streamResponse.Succeeded)
+			{
+				throw new UserErrorException($"Unable to find stream '{streamName}'");
+			}
+			StreamRecord stream = streamResponse.Data;
+
+			// Get the new directory for the client
+			DirectoryReference clientDir = DirectoryReference.FromString(options.ClientRoot) ?? DirectoryReference.Combine(DirectoryReference.GetCurrentDirectory(), stream.Stream.Replace('/', '+'));
+			DirectoryReference.CreateDirectory(clientDir);
+
+			// Make up a new client name 
+			string clientName = options.ClientName ?? Regex.Replace($"{perforce.Settings.UserName}_{hostName}_{stream.Stream.Trim('/')}", "[^0-9a-zA-Z_.-]", "+");
+
+			if (!options.IgnoreExistingClients)
+			{
+				// Check there are no existing clients under the current path
+				List<ClientsRecord> clients = await FindExistingClients(perforce, hostName, clientDir);
+				if (clients.Count > 0)
+				{
+					if (clients.Count == 1 && clientName.Equals(clients[0].Name, StringComparison.OrdinalIgnoreCase) && clientDir == TryParseRoot(clients[0].Root))
+					{
+						logger.LogInformation("Reusing existing client for {ClientDir} ({ClientName})", clientDir, options.ClientName);
+					}
+					else
+					{
+						throw new UserErrorException("Current directory is already within a Perforce workspace ({ClientName})", clients[0].Name);
+					}
+				}
+			}
+
+			// Create the new client
+			ClientRecord client = new ClientRecord(clientName, perforce.Settings.UserName, clientDir.FullName);
+			client.Host = hostName;
+			client.Stream = stream.Stream;
+			client.Options = ClientOptions.Rmdir;
+			await perforce.CreateClientAsync(client);
+
+			// Branch root is currently hard-coded at the root
+			string branchPath = options.BranchPath ?? String.Empty;
+			string projectPath = await ProjectUtils.FindProjectPathAsync(perforce, clientName, branchPath, options.ProjectName);
+
+			// Create the settings object
+			UserWorkspaceSettings settings = new UserWorkspaceSettings();
+			settings.RootDir = clientDir;
+			settings.Init(perforce.Settings.ServerAndPort, perforce.Settings.UserName, clientName, branchPath, projectPath);
+			options.ApplyTo(settings);
+			settings.Save(logger);
+
+			if (globalSettings != null)
+			{
+				UserSelectedProjectSettings selectedProjectSettings =
+					new UserSelectedProjectSettings(
+						null,
+						null,
+						UserSelectedProjectType.Client,
+						$"//{client.Name}{projectPath}",
+						$"{client.Root}{projectPath}".Replace("/", @"\", StringComparison.Ordinal));
+				globalSettings.OpenProjects.Add(selectedProjectSettings);
+				globalSettings.RecentProjects.Add(selectedProjectSettings);
+				globalSettings.Save(logger);
+			}
+
+			logger.LogInformation("Initialized {ClientName} with root at {RootDir}", clientName, clientDir);
+		}
+
+		static DirectoryReference? TryParseRoot(string root)
+		{
+			try
+			{
+				return new DirectoryReference(root);
+			}
+			catch
+			{
+				return null;
+			}
+		}
+
+		static async Task InitExistingClientAsync(IPerforceConnection perforce, string hostName, InitCommandOptions options, UserSettings? globalSettings, ILogger logger)
+		{
+			DirectoryReference currentDir = DirectoryReference.GetCurrentDirectory();
+
+			// Make sure the client name is set
+			string? clientName = options.ClientName;
+			if (clientName == null)
+			{
+				List<ClientsRecord> clients = await FindExistingClients(perforce, hostName, currentDir);
+				if (clients.Count == 0)
+				{
+					throw new UserErrorException("Unable to find client for {HostName} under {ClientDir}", hostName, currentDir);
+				}
+				if (clients.Count > 1)
+				{
+					throw new UserErrorException("Multiple clients found for {HostName} under {ClientDir}: {ClientList}", hostName, currentDir, String.Join(", ", clients.Select(x => x.Name)));
+				}
+
+				clientName = clients[0].Name;
+				logger.LogInformation("Found client {ClientName}", clientName);
+			}
+
+			// Get the client info
+			ClientRecord client = await perforce.GetClientAsync(clientName);
+			DirectoryReference clientDir = new DirectoryReference(client.Root);
+
+			// If a project path was specified in local syntax, try to convert it to client-relative syntax
+			string? projectName = options.ProjectName;
+			if (options.ProjectName != null && options.ProjectName.Contains('.', StringComparison.Ordinal))
+			{
+				options.ProjectName = FileReference.Combine(currentDir, options.ProjectName).MakeRelativeTo(clientDir).Replace('\\', '/');
+			}
+
+			// Branch root is currently hard-coded at the root
+			string branchPath = options.BranchPath ?? String.Empty;
+			string projectPath = await ProjectUtils.FindProjectPathAsync(perforce, clientName, branchPath, projectName);
+
+			// Create the settings object
+			UserWorkspaceSettings settings = new UserWorkspaceSettings();
+			settings.RootDir = clientDir;
+			settings.Init(perforce.Settings.ServerAndPort, perforce.Settings.UserName, clientName, branchPath, projectPath);
+			options.ApplyTo(settings);
+			settings.Save(logger);
+
+			if (globalSettings != null)
+			{
+				UserSelectedProjectSettings selectedProjectSettings =
+					new UserSelectedProjectSettings(
+						null,
+						null,
+						UserSelectedProjectType.Client,
+						$"//{client.Name}{projectPath}",
+						$"{client.Root}{projectPath}".Replace("/", @"\", StringComparison.Ordinal));
+				globalSettings.OpenProjects.Add(selectedProjectSettings);
+				globalSettings.RecentProjects.Add(selectedProjectSettings);
+				globalSettings.Save(logger);
+			}
+
+			logger.LogInformation("Initialized workspace at {RootDir} for {ClientProject}", clientDir, settings.ClientProjectPath);
+		}
+
+		static async Task<List<ClientsRecord>> FindExistingClients(IPerforceConnection perforce, string hostName, DirectoryReference clientDir)
+		{
+			List<ClientsRecord> matchingClients = new List<ClientsRecord>();
+
+			List<ClientsRecord> clients = await perforce.GetClientsAsync(ClientsOptions.None, perforce.Settings.UserName);
+			foreach (ClientsRecord client in clients)
+			{
+				if (!String.IsNullOrEmpty(client.Root) && !String.IsNullOrEmpty(client.Host) && String.Equals(hostName, client.Host, StringComparison.OrdinalIgnoreCase))
+				{
+					DirectoryReference? rootDir;
+					try
+					{
+						rootDir = new DirectoryReference(client.Root);
+					}
+					catch
+					{
+						rootDir = null;
+					}
+
+					if (rootDir != null && clientDir.IsUnderDirectory(rootDir))
+					{
+						matchingClients.Add(client);
+					}
+				}
+			}
+
+			return matchingClients;
+		}
+	}
+}
diff --git a/Engine/Source/Programs/UnrealGameSync/UnrealGameSyncCmd/Commands/InstallCommand.cs b/Engine/Source/Programs/UnrealGameSync/UnrealGameSyncCmd/Commands/InstallCommand.cs
new file mode 100644
--- /dev/null
+++ b/Engine/Source/Programs/UnrealGameSync/UnrealGameSyncCmd/Commands/InstallCommand.cs
@@ -0,0 +1,15 @@
+﻿// Copyright Epic Games, Inc. All Rights Reserved.
+
+using System.Threading.Tasks;
+using UnrealGameSyncCmd.Utils;
+
+namespace UnrealGameSyncCmd.Commands
+{
+	internal class InstallCommand : Command
+	{
+		public override async Task ExecuteAsync(CommandContext context)
+		{
+			await InstallUtils.UpdateInstallAsync(true, context.Logger);
+		}
+	}
+}
diff --git a/Engine/Source/Programs/UnrealGameSync/UnrealGameSyncCmd/Commands/LoginCommand.cs b/Engine/Source/Programs/UnrealGameSync/UnrealGameSyncCmd/Commands/LoginCommand.cs
new file mode 100644
--- /dev/null
+++ b/Engine/Source/Programs/UnrealGameSync/UnrealGameSyncCmd/Commands/LoginCommand.cs
@@ -0,0 +1,45 @@
+﻿// Copyright Epic Games, Inc. All Rights Reserved.
+
+using System.Collections.Generic;
+using System.IO;
+using System.Threading.Tasks;
+
+using Microsoft.Extensions.Configuration;
+using Microsoft.Extensions.Logging;
+
+using EpicGames.Core;
+using EpicGames.OIDC;
+using UnrealGameSync;
+using UnrealGameSyncCmd.Utils;
+
+using UserErrorException = UnrealGameSyncCmd.Exceptions.UserErrorException;
+
+namespace UnrealGameSyncCmd.Commands
+{
+	internal class LoginCommand : Command
+	{
+		public override async Task ExecuteAsync(CommandContext context)
+		{
+			ILogger logger = context.Logger;
+
+			// Get the positional argument indicating the file to look for
+			if (!context.Arguments.TryGetPositionalArgument(out string? providerIdentifier))
+			{
+				throw new UserErrorException("Missing provider identifier to login to.");
+			}
+			context.Arguments.CheckAllArgumentsUsed();
+
+			UserWorkspaceSettings settings = UserSettingsUtils.ReadRequiredUserWorkspaceSettings();
+
+			// Find the valid config file paths
+			DirectoryInfo engineDir = DirectoryReference.Combine(settings.RootDir, "Engine").ToDirectoryInfo();
+			DirectoryInfo gameDir = new DirectoryInfo(settings.ProjectPath);
+			using ITokenStore tokenStore = TokenStoreFactory.CreateTokenStore();
+			IConfiguration providerConfiguration = ProviderConfigurationFactory.ReadConfiguration(engineDir, gameDir);
+			OidcTokenManager oidcTokenManager = OidcTokenManager.CreateTokenManager(providerConfiguration, tokenStore, new List<string>() { providerIdentifier });
+			OidcTokenInfo result = await oidcTokenManager.LoginAsync(providerIdentifier);
+
+			logger.LogInformation("Logged in to provider {ProviderIdentifier}", providerIdentifier);
+		}
+	}
+}
diff --git a/Engine/Source/Programs/UnrealGameSync/UnrealGameSyncCmd/Commands/RunCommand.cs b/Engine/Source/Programs/UnrealGameSync/UnrealGameSyncCmd/Commands/RunCommand.cs
new file mode 100644
--- /dev/null
+++ b/Engine/Source/Programs/UnrealGameSync/UnrealGameSyncCmd/Commands/RunCommand.cs
@@ -0,0 +1,70 @@
+﻿// Copyright Epic Games, Inc. All Rights Reserved.
+
+using System;
+using System.Collections.Generic;
+using System.IO;
+using System.Linq;
+using System.Threading;
+using System.Threading.Tasks;
+using EpicGames.Core;
+using EpicGames.Perforce;
+using Microsoft.Extensions.Logging;
+using UnrealGameSync;
+using UnrealGameSyncCmd.Utils;
+
+using UserErrorException = UnrealGameSyncCmd.Exceptions.UserErrorException;
+
+namespace UnrealGameSyncCmd.Commands
+{
+	internal class RunCommand : Command
+	{
+		public override async Task ExecuteAsync(CommandContext context)
+		{
+			ILogger logger = context.Logger;
+
+			UserWorkspaceSettings settings = UserSettingsUtils.ReadRequiredUserWorkspaceSettings();
+			using IPerforceConnection perforceClient = await PerforceConnectionUtils.ConnectAsync(settings, context.LoggerFactory);
+			WorkspaceStateWrapper state = await WorkspaceStateUtils.ReadWorkspaceState(perforceClient, settings, context.UserSettings, logger);
+
+			ProjectInfo projectInfo = new ProjectInfo(settings.RootDir, state.Current);
+			ConfigFile projectConfig = await ConfigUtils.ReadProjectConfigFileAsync(perforceClient, projectInfo, logger, CancellationToken.None);
+
+			FileReference receiptFile = ConfigUtils.GetEditorReceiptFile(projectInfo, projectConfig, EditorConfig);
+			logger.LogDebug("Receipt file: {Receipt}", receiptFile);
+
+			if (!ConfigUtils.TryReadEditorReceipt(projectInfo, receiptFile, out TargetReceipt? receipt) || String.IsNullOrEmpty(receipt.Launch))
+			{
+				throw new UserErrorException("The editor needs to be built before you can run it. (Missing {ReceiptFile}).", receiptFile);
+			}
+			if (!File.Exists(receipt.Launch))
+			{
+				throw new UserErrorException("The editor needs to be built before you can run it. (Missing {LaunchFile}).", receipt.Launch);
+			}
+
+			List<string> launchArguments = new List<string>();
+			if (settings.LocalProjectPath.HasExtension(ProjectUtils.UProjectExtension))
+			{
+				launchArguments.Add($"\"{settings.LocalProjectPath}\"");
+			}
+			if (EditorConfig == BuildConfig.Debug || EditorConfig == BuildConfig.DebugGame)
+			{
+				launchArguments.Append(" -debug");
+			}
+			for (int idx = 0; idx < context.Arguments.Count; idx++)
+			{
+				if (!context.Arguments.HasBeenUsed(idx))
+				{
+					launchArguments.Add(context.Arguments[idx]);
+				}
+			}
+
+			string commandLine = CommandLineArguments.Join(launchArguments);
+			logger.LogInformation("Spawning: {LaunchFile} {CommandLine}", CommandLineArguments.Quote(receipt.Launch), commandLine);
+
+			if (!Utility.SpawnProcess(receipt.Launch, commandLine))
+			{
+				logger.LogError("Unable to spawn {App} {Args}", receipt.Launch, launchArguments.ToString());
+			}
+		}
+	}
+}
diff --git a/Engine/Source/Programs/UnrealGameSync/UnrealGameSyncCmd/Commands/SettingsCommand.cs b/Engine/Source/Programs/UnrealGameSync/UnrealGameSyncCmd/Commands/SettingsCommand.cs
new file mode 100644
--- /dev/null
+++ b/Engine/Source/Programs/UnrealGameSync/UnrealGameSyncCmd/Commands/SettingsCommand.cs
@@ -0,0 +1,83 @@
+﻿// Copyright Epic Games, Inc. All Rights Reserved.
+
+using System;
+using System.IO;
+using System.Text.Json;
+using System.Threading.Tasks;
+using EpicGames.Core;
+using Microsoft.Extensions.Logging;
+using UnrealGameSync;
+using UnrealGameSyncCmd.Utils;
+
+namespace UnrealGameSyncCmd.Commands
+{
+	internal class SettingsCommandOptions
+	{
+		[CommandLine("-Global=")]
+		public string Global { get; set; } = String.Empty;
+
+		[CommandLine("-User=")]
+		public string User { get; set; } = String.Empty;
+	}
+
+	internal class SettingsCommand : Command
+	{
+		public override async Task ExecuteAsync(CommandContext context)
+		{
+			ILogger logger = context.Logger;
+
+			SettingsCommandOptions options = context.Arguments.ApplyTo<SettingsCommandOptions>(logger);
+			context.Arguments.CheckAllArgumentsUsed();
+
+			if (!String.IsNullOrWhiteSpace(options.Global))
+			{
+				if (context.GlobalSettings != null)
+				{
+					logger.LogInformation("Writing global settings to {Global}", options.Global);
+					await WriteSettings(options.Global, context.GlobalSettings);
+				}
+				else if (context.UserSettings != null)
+				{
+					logger.LogInformation("Writing global settings to {Global}", options.Global);
+					await WriteSettings(options.Global, context.UserSettings);
+				}
+			}
+
+			if (!String.IsNullOrWhiteSpace(options.User))
+			{
+				UserWorkspaceSettings userWorkspaceSettings = UserSettingsUtils.ReadRequiredUserWorkspaceSettings();
+				logger.LogInformation("Writing user settings to {User}", options.User);
+				await WriteSettings(options.User, userWorkspaceSettings);
+			}
+		}
+
+		private static async Task WriteSettings(string path, object? settings)
+		{
+			if (String.IsNullOrWhiteSpace(path) || settings == null)
+			{
+				return;
+			}
+
+			string trimmed = path.Trim('\"')
+				.Replace(Path.AltDirectorySeparatorChar, Path.DirectorySeparatorChar);
+			if (!Path.IsPathFullyQualified(trimmed))
+			{
+				return;
+			}
+
+			string? directory = Path.GetDirectoryName(trimmed);
+			if (!String.IsNullOrWhiteSpace(directory))
+			{
+				Directory.CreateDirectory(directory);
+			}
+
+			JsonSerializerOptions serializerOptions = new JsonSerializerOptions()
+			{
+				WriteIndented = true
+			};
+			string json = JsonSerializer.Serialize(settings, serializerOptions);
+
+			await File.WriteAllTextAsync(trimmed, json);
+		}
+	}
+}
diff --git a/Engine/Source/Programs/UnrealGameSync/UnrealGameSyncCmd/Commands/StatusCommand.cs b/Engine/Source/Programs/UnrealGameSync/UnrealGameSyncCmd/Commands/StatusCommand.cs
new file mode 100644
--- /dev/null
+++ b/Engine/Source/Programs/UnrealGameSync/UnrealGameSyncCmd/Commands/StatusCommand.cs
@@ -0,0 +1,52 @@
+﻿// Copyright Epic Games, Inc. All Rights Reserved.
+
+using System.Threading;
+using System.Threading.Tasks;
+
+using Microsoft.Extensions.Logging;
+
+using EpicGames.Perforce;
+
+using UnrealGameSync;
+using UnrealGameSyncCmd.Utils;
+
+namespace UnrealGameSyncCmd.Commands
+{
+	internal class StatusCommand : Command
+	{
+		public override async Task ExecuteAsync(CommandContext context)
+		{
+			ILogger logger = context.Logger;
+			bool update = context.Arguments.HasOption("-Update");
+			context.Arguments.CheckAllArgumentsUsed();
+
+			UserWorkspaceSettings settings = UserSettingsUtils.ReadRequiredUserWorkspaceSettings();
+			logger.LogInformation("User: {UserName}", settings.UserName);
+			logger.LogInformation("Server: {ServerAndPort}", settings.ServerAndPort);
+			logger.LogInformation("Project: {ClientProjectPath}", settings.ClientProjectPath);
+
+			using IPerforceConnection perforceClient = await PerforceConnectionUtils.ConnectAsync(settings, context.LoggerFactory);
+
+			WorkspaceStateWrapper state = await WorkspaceStateUtils.ReadWorkspaceState(perforceClient, settings, context.UserSettings, logger);
+			if (update)
+			{
+				ProjectInfo newProjectInfo = await ProjectInfo.CreateAsync(perforceClient, settings, CancellationToken.None);
+				state.Modify(x => x.UpdateCachedProjectInfo(newProjectInfo, settings.LastModifiedTimeUtc));
+			}
+
+			string streamOrBranchName = state.Current.StreamName ?? settings.BranchPath.TrimStart('/');
+			if (state.Current.LastSyncResultMessage == null)
+			{
+				logger.LogInformation("Not currently synced to {Stream}", streamOrBranchName);
+			}
+			else if (state.Current.LastSyncResult == WorkspaceUpdateResult.Success)
+			{
+				logger.LogInformation("Synced to {Stream} CL {Change}", streamOrBranchName, state.Current.LastSyncChangeNumber);
+			}
+			else
+			{
+				logger.LogWarning("Last sync to {Stream} CL {Change} failed: {Result}", streamOrBranchName, state.Current.LastSyncChangeNumber, state.Current.LastSyncResultMessage);
+			}
+		}
+	}
+}
diff --git a/Engine/Source/Programs/UnrealGameSync/UnrealGameSyncCmd/Commands/SwitchCommand.cs b/Engine/Source/Programs/UnrealGameSync/UnrealGameSyncCmd/Commands/SwitchCommand.cs
new file mode 100644
--- /dev/null
+++ b/Engine/Source/Programs/UnrealGameSync/UnrealGameSyncCmd/Commands/SwitchCommand.cs
@@ -0,0 +1,78 @@
+﻿// Copyright Epic Games, Inc. All Rights Reserved.
+
+using System;
+using System.Linq;
+using System.Threading.Tasks;
+
+using EpicGames.Core;
+using EpicGames.Perforce;
+using Microsoft.Extensions.Logging;
+using UnrealGameSync;
+using UnrealGameSyncCmd.Utils;
+
+using UserErrorException = UnrealGameSyncCmd.Exceptions.UserErrorException;
+
+namespace UnrealGameSyncCmd.Commands
+{
+	internal class SwitchCommandOptions
+	{
+		[CommandLine("-Force")]
+		public bool Force { get; set; }
+	}
+
+	internal class SwitchCommand : Command
+	{
+		public override async Task ExecuteAsync(CommandContext context)
+		{
+			// Get the positional argument indicating the file to look for
+			string? targetName;
+			if (!context.Arguments.TryGetPositionalArgument(out targetName))
+			{
+				throw new UserErrorException("Missing stream or project name to switch to.");
+			}
+
+			// Finish argument parsing
+			SwitchCommandOptions options = new SwitchCommandOptions();
+			context.Arguments.ApplyTo(options);
+			context.Arguments.CheckAllArgumentsUsed();
+
+			if (targetName.StartsWith("//", StringComparison.Ordinal))
+			{
+				options.Force = true;
+			}
+
+			// Get a connection to the client for this workspace
+			UserWorkspaceSettings settings = UserSettingsUtils.ReadRequiredUserWorkspaceSettings();
+			using IPerforceConnection perforceClient = await PerforceConnectionUtils.ConnectAsync(settings, context.LoggerFactory);
+
+			// Check whether we're switching stream or project
+			if (targetName.StartsWith("//", StringComparison.Ordinal))
+			{
+				await SwitchStreamAsync(perforceClient, targetName, options.Force, context.Logger);
+			}
+			else
+			{
+				await SwitchProjectAsync(perforceClient, settings, targetName, context.Logger);
+			}
+		}
+
+		public static async Task SwitchStreamAsync(IPerforceConnection perforceClient, string streamName, bool force, ILogger logger)
+		{
+			if (!force && await perforceClient.OpenedAsync(OpenedOptions.None, FileSpecList.Any).AnyAsync())
+			{
+				throw new UserErrorException("Client {ClientName} has files opened. Use -Force to switch anyway.", perforceClient.Settings.ClientName!);
+			}
+
+			await perforceClient.SwitchClientToStreamAsync(streamName, SwitchClientOptions.IgnoreOpenFiles);
+
+			logger.LogInformation("Switched to stream {StreamName}", streamName);
+		}
+
+		public static async Task SwitchProjectAsync(IPerforceConnection perforceClient, UserWorkspaceSettings settings, string projectName, ILogger logger)
+		{
+			settings.ProjectPath = await ProjectUtils.FindProjectPathAsync(perforceClient, settings.ClientName, settings.BranchPath, projectName);
+			settings.Save(logger);
+			logger.LogInformation("Switched to project {ProjectPath}", settings.ClientProjectPath);
+		}
+	}
+}
diff --git a/Engine/Source/Programs/UnrealGameSync/UnrealGameSyncCmd/Commands/SyncCommand.cs b/Engine/Source/Programs/UnrealGameSync/UnrealGameSyncCmd/Commands/SyncCommand.cs
new file mode 100644
--- /dev/null
+++ b/Engine/Source/Programs/UnrealGameSync/UnrealGameSyncCmd/Commands/SyncCommand.cs
@@ -0,0 +1,213 @@
+﻿// Copyright Epic Games, Inc. All Rights Reserved.
+
+using System;
+using System.Collections.Generic;
+using System.Linq;
+using System.Runtime.InteropServices;
+using System.Threading;
+using System.Threading.Tasks;
+
+using EpicGames.Core;
+using EpicGames.Perforce;
+using Microsoft.Extensions.Logging;
+using UnrealGameSync;
+using UnrealGameSyncCmd.Utils;
+
+using UserErrorException = UnrealGameSyncCmd.Exceptions.UserErrorException;
+
+namespace UnrealGameSyncCmd.Commands
+{
+	internal class SyncCommandOptions
+	{
+		[CommandLine("-Clean")]
+		public bool Clean { get; set; }
+
+		[CommandLine("-Build")]
+		public bool Build { get; set; }
+
+		[CommandLine("-Binaries")]
+		public bool Binaries { get; set; }
+
+		[CommandLine("-NoGPF", Value = "false")]
+		[CommandLine("-NoProjectFiles", Value = "false")]
+		public bool ProjectFiles { get; set; } = true;
+
+		[CommandLine("-Clobber")]
+		public bool Clobber { get; set; }
+
+		[CommandLine("-Refilter")]
+		public bool Refilter { get; set; }
+
+		[CommandLine("-Only")]
+		public bool SingleChange { get; set; }
+	}
+
+	internal class SyncCommand : Command
+	{
+		static async Task<bool> IsCodeChangeAsync(IPerforceConnection perforce, int change)
+		{
+			DescribeRecord describeRecord = await perforce.DescribeAsync(change);
+			return IsCodeChange(describeRecord);
+		}
+
+		static bool IsCodeChange(DescribeRecord describeRecord)
+		{
+			foreach (DescribeFileRecord file in describeRecord.Files)
+			{
+				if (PerforceUtils.CodeExtensions.Any(extension => file.DepotFile.EndsWith(extension, StringComparison.OrdinalIgnoreCase)))
+				{
+					return true;
+				}
+			}
+			return false;
+		}
+
+		public override async Task ExecuteAsync(CommandContext context)
+		{
+			ILogger logger = context.Logger;
+			context.Arguments.TryGetPositionalArgument(out string? changeString);
+
+			SyncCommandOptions syncOptions = new SyncCommandOptions();
+			context.Arguments.ApplyTo(syncOptions);
+
+			context.Arguments.CheckAllArgumentsUsed();
+
+			UserWorkspaceSettings settings = UserSettingsUtils.ReadRequiredUserWorkspaceSettings();
+			using IPerforceConnection perforceClient = await PerforceConnectionUtils.ConnectAsync(settings, context.LoggerFactory);
+			WorkspaceStateWrapper state = await WorkspaceStateUtils.ReadWorkspaceState(perforceClient, settings, context.UserSettings, logger);
+
+			changeString ??= "latest";
+
+			ProjectInfo projectInfo = new ProjectInfo(settings.RootDir, state.Current);
+			UserProjectSettings projectSettings = context.UserSettings.FindOrAddProjectSettings(projectInfo, settings, logger);
+
+			ConfigFile projectConfig = await ConfigUtils.ReadProjectConfigFileAsync(perforceClient, projectInfo, logger, CancellationToken.None);
+
+			bool syncLatest = String.Equals(changeString, "latest", StringComparison.OrdinalIgnoreCase);
+
+			int change;
+			if (!Int32.TryParse(changeString, out change))
+			{
+				if (syncLatest)
+				{
+					List<ChangesRecord> changes = await perforceClient.GetChangesAsync(ChangesOptions.None, 1, ChangeStatus.Submitted, $"//{settings.ClientName}/...");
+					change = changes[0].Number;
+				}
+				else
+				{
+					throw new UserErrorException("Unknown change type for sync '{Change}'", changeString);
+				}
+			}
+
+			WorkspaceUpdateOptions options = syncOptions.SingleChange ? WorkspaceUpdateOptions.SyncSingleChange : WorkspaceUpdateOptions.Sync;
+			if (syncOptions.Clean)
+			{
+				options |= WorkspaceUpdateOptions.Clean;
+			}
+			if (syncOptions.Build)
+			{
+				options |= WorkspaceUpdateOptions.Build;
+			}
+			if (syncOptions.ProjectFiles)
+			{
+				options |= WorkspaceUpdateOptions.GenerateProjectFiles;
+			}
+			if (syncOptions.Clobber)
+			{
+				options |= WorkspaceUpdateOptions.Clobber;
+			}
+			if (syncOptions.Refilter)
+			{
+				options |= WorkspaceUpdateOptions.Refilter;
+			}
+			options |= WorkspaceUpdateContext.GetOptionsFromConfig(context.UserSettings.Global, settings);
+			options |= WorkspaceUpdateOptions.RemoveFilteredFiles;
+
+			string[] syncFilter = SyncFilterUtils.ReadSyncFilter(settings, context.UserSettings, projectConfig, state.Current.ProjectIdentifier);
+
+			using WorkspaceLock? workspaceLock = CreateWorkspaceLock(settings.RootDir);
+			if (workspaceLock != null && !await workspaceLock.TryAcquireAsync())
+			{
+				logger.LogError("Another process is already syncing this workspace.");
+				return;
+			}
+
+			WorkspaceUpdateContext updateContext = new WorkspaceUpdateContext(change, options, BuildConfig.Development, syncFilter, projectSettings.BuildSteps, null);
+			updateContext.PerforceSyncOptions = context.UserSettings.Global.Perforce;
+
+			if (syncOptions.Binaries)
+			{
+				List<BaseArchiveChannel> archives = await BaseArchive.EnumerateChannelsAsync(perforceClient, context.HordeClient, context.CloudStorage, projectConfig, state.Current.ProjectIdentifier, CancellationToken.None);
+
+				BaseArchiveChannel? editorArchiveInfo = archives.FirstOrDefault(x => x.Type == IArchiveChannel.EditorArchiveType);
+				if (editorArchiveInfo == null)
+				{
+					throw new UserErrorException("No editor archives found for project");
+				}
+
+				KeyValuePair<int, IArchive> revision = editorArchiveInfo.ChangeNumberToArchive.LastOrDefault(x => x.Key <= change);
+				if (revision.Key == 0)
+				{
+					throw new UserErrorException($"No editor archives found for CL {change}");
+				}
+
+				if (revision.Key < change)
+				{
+					int lastChange = revision.Key;
+
+					List<ChangesRecord> changeRecords = await perforceClient.GetChangesAsync(ChangesOptions.None, 1, ChangeStatus.Submitted, $"//{settings.ClientName}/...@{revision.Key + 1},{change}");
+					foreach (ChangesRecord changeRecord in changeRecords.OrderBy(x => x.Number))
+					{
+						if (await IsCodeChangeAsync(perforceClient, changeRecord.Number))
+						{
+							if (syncLatest)
+							{
+								updateContext.ChangeNumber = lastChange;
+							}
+							else
+							{
+								throw new UserErrorException($"No editor binaries found for CL {change} (last archive at CL {revision.Key}, but CL {changeRecord.Number} is a code change)");
+							}
+							break;
+						}
+						change = changeRecord.Number;
+					}
+				}
+
+				updateContext.Options |= WorkspaceUpdateOptions.SyncArchives;
+				updateContext.ArchiveTypeToArchive[IArchiveChannel.EditorArchiveType] = revision.Value;
+			}
+
+			WorkspaceUpdate update = new WorkspaceUpdate(updateContext);
+			(WorkspaceUpdateResult result, string message) = await update.ExecuteAsync(perforceClient.Settings, projectInfo, state, context.Logger, CancellationToken.None);
+			if (result == WorkspaceUpdateResult.FilesToClobber)
+			{
+				logger.LogWarning("The following files are modified in your workspace:");
+				foreach (string file in updateContext.ClobberFiles.Keys.OrderBy(x => x))
+				{
+					logger.LogWarning("  {File}", file);
+				}
+				logger.LogWarning("Use -Clobber to overwrite");
+			}
+
+			state.Modify(x => x.SetLastSyncState(result, updateContext, message));
+
+			if (result != WorkspaceUpdateResult.Success)
+			{
+				throw new UserErrorException("{Message} (Result: {Result})", message, result);
+			}
+		}
+
+		static WorkspaceLock? CreateWorkspaceLock(DirectoryReference rootDir)
+		{
+			if (RuntimeInformation.IsOSPlatform(OSPlatform.Windows))
+			{
+				return new WorkspaceLock(rootDir);
+			}
+			else
+			{
+				return null;
+			}
+		}
+	}
+}
diff --git a/Engine/Source/Programs/UnrealGameSync/UnrealGameSyncCmd/Commands/ToolsCommand.cs b/Engine/Source/Programs/UnrealGameSync/UnrealGameSyncCmd/Commands/ToolsCommand.cs
new file mode 100644
--- /dev/null
+++ b/Engine/Source/Programs/UnrealGameSync/UnrealGameSyncCmd/Commands/ToolsCommand.cs
@@ -0,0 +1,181 @@
+﻿// Copyright Epic Games, Inc. All Rights Reserved.
+
+using System;
+using System.Collections.Generic;
+using System.Linq;
+using System.Runtime.InteropServices;
+using System.Threading.Tasks;
+
+using EpicGames.Core;
+using EpicGames.Horde;
+using EpicGames.OIDC;
+using EpicGames.Perforce;
+using Microsoft.Extensions.DependencyInjection;
+using Microsoft.Extensions.Logging;
+
+using UnrealGameSync;
+using UnrealGameSyncCmd.Utils;
+
+namespace UnrealGameSyncCmd.Commands
+{
+	internal class ToolsCommandCommandOptions
+	{
+		[CommandLine("-Install=", ListSeparator = ',')]
+		public List<string> Install { get; set; } = new List<string>();
+
+		[CommandLine("-Uninstall=", ListSeparator = ',')]
+		public List<string> Uninstall { get; set; } = new List<string>();
+
+		[CommandLine("-List")]
+		public bool List { get; set; }
+
+		[CommandLine("-Update")]
+		public bool Update { get; set; }
+	}
+
+	internal class ToolsCommand : Command
+	{
+		public override async Task ExecuteAsync(CommandContext context)
+		{
+			ILogger logger = context.Logger;
+			if (!RuntimeInformation.IsOSPlatform(OSPlatform.Windows))
+			{
+				logger.LogInformation("This command is only available for Windows.");
+				return;
+			}
+
+			context.Arguments.TryGetPositionalArgument(out string? toolId);
+
+			ToolsCommandCommandOptions options = context.Arguments.ApplyTo<ToolsCommandCommandOptions>(logger);
+			context.Arguments.CheckAllArgumentsUsed();
+
+			UserWorkspaceSettings settings = UserSettingsUtils.ReadRequiredUserWorkspaceSettings();
+			using IPerforceConnection perforceClient = await PerforceConnectionUtils.ConnectAsync(settings, context.LoggerFactory);
+
+			if (context.GlobalSettings == null)
+			{
+				logger.LogError("Could not retrieve user settings.");
+				return;
+			}
+
+			DirectoryReference dataFolder = DirectoryReference.Combine(DirectoryReference.GetSpecialFolder(Environment.SpecialFolder.LocalApplicationData)!, "UnrealGameSync");
+			DirectoryReference.CreateDirectory(dataFolder);
+
+			// create a temporary service provider for the tool update monitor
+			ServiceCollection services = new ServiceCollection();
+			services.AddSingleton<IAsyncDisposer, AsyncDisposer>();
+			services.AddSingleton(sp => TokenStoreFactory.CreateTokenStore());
+			services.AddSingleton<OidcTokenManager>();
+
+			LauncherSettings launcherSettings = new LauncherSettings();
+			launcherSettings.Read();
+
+			if (launcherSettings.HordeServer != null)
+			{
+				services.AddHorde(options =>
+				{
+					options.ServerUrl = new Uri(launcherSettings.HordeServer);
+					options.AllowAuthPrompt = false;
+				});
+			}
+
+			ServiceProvider serviceProvider = services.BuildServiceProvider();
+
+			using ToolUpdateMonitor toolUpdateMonitor =
+				new ToolUpdateMonitor(perforceClient.Settings, dataFolder, context.GlobalSettings, logger, serviceProvider);
+
+			// get the list of tools available
+			logger.LogInformation("Retrieving tools information, please wait.");
+			await toolUpdateMonitor.GetDataFromBackendAsync();
+
+			// only list the tools
+			if (options.List)
+			{
+				logger.LogInformation("Available Tools:");
+				logger.LogInformation("");
+				logger.LogInformation("  Id                                   | Name                                 | Enabled         | Description                                      ");
+				logger.LogInformation("  -------------------------------------| -------------------------------------| ----------------| -------------------------------------------------");
+
+				IReadOnlyList<ToolInfo> enabled = toolUpdateMonitor.GetEnabledTools();
+
+				foreach (ToolInfo toolInfo in toolUpdateMonitor.GetTools().OrderBy(t => enabled.All(e => t.Id != e.Id)))
+				{
+					logger.LogInformation("  {Id,-36} | {Name,-36} | {Enabled,-16}| {Description,-48} ", toolInfo.Id, toolInfo.Name, enabled.Any(t => t.Id == toolInfo.Id), toolInfo.Description);
+				}
+			}
+
+			// update the list of enabled tools and update (this will update all the tools)
+			if (options.Install.Any() || options.Uninstall.Any())
+			{
+				IReadOnlyList<ToolInfo> enabled = toolUpdateMonitor.GetEnabledTools();
+				HashSet<Guid> newEnabledTools = new HashSet<Guid>(enabled.Select(t => t.Id));
+
+				foreach (string id in options.Install.Select(id => id.Trim()))
+				{
+					if (Guid.TryParse(id, out Guid guid))
+					{
+						foreach (ToolInfo toolInfo in toolUpdateMonitor.GetTools())
+						{
+							if (toolInfo.Id == guid)
+							{
+								logger.LogInformation("Adding tool '{Id}' - '{Name}' to the list of enabled tools", id, toolInfo.Name);
+								newEnabledTools.Add(guid);
+								break;
+							}
+						}
+
+						if (!newEnabledTools.Contains(guid))
+						{
+							logger.LogError("Could not install tool '{Id}'", id);
+						}
+					}
+					else
+					{
+						logger.LogError("'{Id}' is not a valid guid", id);
+					}
+				}
+
+				foreach (string id in options.Uninstall.Select(id => id.Trim()))
+				{
+					if (Guid.TryParse(id, out Guid guid))
+					{
+						foreach (ToolInfo toolInfo in toolUpdateMonitor.GetTools())
+						{
+							if (toolInfo.Id == guid)
+							{
+								logger.LogInformation("Removing tool '{Id}' - '{Name}' from the list of enabled tools", id, toolInfo.Name);
+								newEnabledTools.Remove(guid);
+								break;
+							}
+						}
+
+						if (newEnabledTools.Contains(guid))
+						{
+							logger.LogError("Could not uninstall tool '{Id}'", id);
+						}
+					}
+					else
+					{
+						logger.LogError("'{Id}' is not a valid guid", id);
+					}
+				}
+
+				if (!newEnabledTools.SequenceEqual(context.GlobalSettings.EnabledTools))
+				{
+					context.GlobalSettings.EnabledTools.Clear();
+					context.GlobalSettings.EnabledTools.UnionWith(newEnabledTools);
+					context.GlobalSettings.Save(logger);
+
+					logger.LogInformation("Updating tools");
+					await toolUpdateMonitor.UpdateToolsAsync();
+				}
+			}
+
+			// update the all the enabled tools
+			if (options.Update)
+			{
+				await toolUpdateMonitor.UpdateToolsAsync();
+			}
+		}
+	}
+}
diff --git a/Engine/Source/Programs/UnrealGameSync/UnrealGameSyncCmd/Commands/UninstallCommand.cs b/Engine/Source/Programs/UnrealGameSync/UnrealGameSyncCmd/Commands/UninstallCommand.cs
new file mode 100644
--- /dev/null
+++ b/Engine/Source/Programs/UnrealGameSync/UnrealGameSyncCmd/Commands/UninstallCommand.cs
@@ -0,0 +1,15 @@
+﻿// Copyright Epic Games, Inc. All Rights Reserved.
+
+using System.Threading.Tasks;
+using UnrealGameSyncCmd.Utils;
+
+namespace UnrealGameSyncCmd.Commands
+{
+	internal class UninstallCommand : Command
+	{
+		public override async Task ExecuteAsync(CommandContext context)
+		{
+			await InstallUtils.UpdateInstallAsync(false, context.Logger);
+		}
+	}
+}
diff --git a/Engine/Source/Programs/UnrealGameSync/UnrealGameSyncCmd/Commands/UpgradeCommand.cs b/Engine/Source/Programs/UnrealGameSync/UnrealGameSyncCmd/Commands/UpgradeCommand.cs
new file mode 100644
--- /dev/null
+++ b/Engine/Source/Programs/UnrealGameSync/UnrealGameSyncCmd/Commands/UpgradeCommand.cs
@@ -0,0 +1,93 @@
+﻿// Copyright Epic Games, Inc. All Rights Reserved.
+
+using System;
+using System.IO;
+using System.IO.Compression;
+using System.Net.Http;
+using System.Reflection;
+using System.Threading;
+using System.Threading.Tasks;
+using Microsoft.Extensions.Logging;
+
+using EpicGames.Core;
+
+using UnrealGameSync;
+using UnrealGameSyncCmd.Utils;
+
+namespace UnrealGameSyncCmd.Commands
+{
+	internal class UpgradeCommandOptions
+	{
+		[CommandLine("-Check")]
+		public bool Check { get; set; }
+
+		[CommandLine("-Force")]
+		public bool Force { get; set; }
+	}
+
+	internal class UpgradeCommand : Command
+	{
+		public override async Task ExecuteAsync(CommandContext context)
+		{
+			ILogger logger = context.Logger;
+
+			UpgradeCommandOptions options = new UpgradeCommandOptions();
+			context.Arguments.ApplyTo(options);
+			string? targetDirStr = context.Arguments.GetStringOrDefault("-TargetDir=", null);
+			context.Arguments.CheckAllArgumentsUsed(logger);
+
+			string currentVersion = VersionUtils.GetVersion();
+
+			string? latestVersion = await VersionUtils.GetLatestVersionAsync(logger, CancellationToken.None);
+			if (latestVersion == null)
+			{
+				return;
+			}
+			if (latestVersion.Equals(currentVersion, StringComparison.OrdinalIgnoreCase) && !options.Force)
+			{
+				logger.LogInformation("You are running the latest version ({Version})", currentVersion);
+				return;
+			}
+			if (options.Check)
+			{
+				logger.LogWarning("A newer version of UGS is available ({NewVersion})", latestVersion);
+				return;
+			}
+
+			using (HttpClient httpClient = new HttpClient())
+			{
+				Uri baseUrl = new Uri(DeploymentSettings.Instance.HordeUrl ?? String.Empty);
+
+				DirectoryReference currentDir = new FileReference(Assembly.GetExecutingAssembly().Location).Directory;
+
+				DirectoryReference targetDir = (targetDirStr == null) ? currentDir : DirectoryReference.Combine(currentDir, targetDirStr);
+				DirectoryReference.CreateDirectory(targetDir);
+
+				FileReference tempFile = FileReference.Combine(targetDir, "update.zip");
+				using (Stream requestStream = await httpClient.GetStreamAsync(new Uri(baseUrl, $"api/v1/tools/{VersionUtils.GetUpgradeToolName()}?action=download")))
+				{
+					using (Stream tempFileStream = FileReference.Open(tempFile, FileMode.Create, FileAccess.Write, FileShare.None))
+					{
+						await requestStream.CopyToAsync(tempFileStream);
+					}
+				}
+
+				using (FileStream stream = FileReference.Open(tempFile, FileMode.Open, FileAccess.Read, FileShare.Read))
+				{
+					using (ZipArchive archive = new ZipArchive(stream, ZipArchiveMode.Read, true))
+					{
+						foreach (ZipArchiveEntry entry in archive.Entries)
+						{
+							FileReference targetFile = FileReference.Combine(targetDir, entry.Name);
+							if (!targetFile.IsUnderDirectory(targetDir))
+							{
+								throw new InvalidDataException("Attempt to extract file outside source directory");
+							}
+							entry.ExtractToFile_CrossPlatform(targetFile.FullName, true);
+						}
+					}
+				}
+			}
+		}
+	}
+}
diff --git a/Engine/Source/Programs/UnrealGameSync/UnrealGameSyncCmd/Commands/VersionCommand.cs b/Engine/Source/Programs/UnrealGameSync/UnrealGameSyncCmd/Commands/VersionCommand.cs
new file mode 100644
--- /dev/null
+++ b/Engine/Source/Programs/UnrealGameSync/UnrealGameSyncCmd/Commands/VersionCommand.cs
@@ -0,0 +1,21 @@
+﻿// Copyright Epic Games, Inc. All Rights Reserved.
+
+using System.Threading.Tasks;
+using Microsoft.Extensions.Logging;
+using UnrealGameSyncCmd.Utils;
+
+namespace UnrealGameSyncCmd.Commands
+{
+	internal class VersionCommand : Command
+	{
+		public override Task ExecuteAsync(CommandContext context)
+		{
+			ILogger logger = context.Logger;
+
+			string version = VersionUtils.GetVersion();
+			logger.LogInformation("UnrealGameSync {Version}", version);
+
+			return Task.CompletedTask;
+		}
+	}
+}
diff --git a/Engine/Source/Programs/UnrealGameSync/UnrealGameSyncCmd/CommandsFactory.cs b/Engine/Source/Programs/UnrealGameSync/UnrealGameSyncCmd/CommandsFactory.cs
new file mode 100644
--- /dev/null
+++ b/Engine/Source/Programs/UnrealGameSync/UnrealGameSyncCmd/CommandsFactory.cs
@@ -0,0 +1,111 @@
+﻿// Copyright Epic Games, Inc. All Rights Reserved.
+
+using System;
+using System.Collections.Generic;
+using System.Linq;
+using System.Runtime.InteropServices;
+
+using UnrealGameSyncCmd.Commands;
+using UnrealGameSyncCmd.Options;
+
+namespace UnrealGameSyncCmd
+{
+	internal static class CommandsFactory
+	{
+		internal static CommandInfo[] GetCommands()
+		{
+			List<CommandInfo> commandInfos = PortableCommands();
+
+			if (RuntimeInformation.IsOSPlatform(OSPlatform.Windows))
+			{
+				commandInfos.AddRange(GetWinCommand());
+			}
+
+			return commandInfos
+				.OrderBy(x => x.Name, StringComparer.OrdinalIgnoreCase)
+				.ToArray();
+		}
+
+		private static List<CommandInfo> PortableCommands()
+		{
+			return new List<CommandInfo>
+			{
+				new CommandInfo("init", typeof(InitCommand), typeof(InitCommandOptions),
+				"ugs init [stream-path]",
+				"Create a client for the given stream, or initializes an existing client for use by UGS."
+			),
+				new CommandInfo("switch", typeof(SwitchCommand), typeof(SwitchCommandOptions),
+				"ugs switch [project name|project path|stream]",
+				"Changes the active project to the one in the workspace with the given name, or switches to a new stream."
+			),
+				new CommandInfo("changes", typeof(ChangesCommand), typeof(ChangesCommandOptions),
+				"ugs changes",
+				"List recently submitted changes to the current branch."
+			),
+				new CommandInfo("config", typeof(ConfigCommand), typeof(ConfigCommandOptions),
+				"ugs config",
+				"Updates the configuration for the current workspace."
+			),
+				new CommandInfo("filter", typeof(FilterCommand), typeof(FilterCommandOptions),
+				"ugs filter",
+				"Displays or updates the workspace or global sync filter"
+			),
+				new CommandInfo("sync", typeof(SyncCommand), typeof(SyncCommandOptions),
+				"ugs sync [change|'latest']",
+				"Syncs the current workspace to the given changelist, optionally removing all local state."
+			),
+				new CommandInfo("clients", typeof(ClientsCommand), typeof(ClientsCommandOptions),
+				"ugs clients",
+				"Lists all clients suitable for use on the current machine."
+			),
+				new CommandInfo("run", typeof(RunCommand), null,
+				"ugs run",
+				"Runs the editor for the current branch."
+			),
+				new CommandInfo("build", typeof(BuildCommand), typeof(BuildCommandOptions),
+				"ugs build [id]",
+				"Runs the default build steps for the current project, or a particular step referenced by id."
+			),
+				new CommandInfo("status", typeof(StatusCommand), null,
+				"ugs status [-update]",
+				"Shows the status of the currently synced branch."
+			),
+				new CommandInfo("login", typeof(LoginCommand), null,
+				"ugs login",
+				"Starts a interactive login flow against the configured Identity Provider"
+			),
+				new CommandInfo("version", typeof(VersionCommand), null,
+				"ugs version",
+				"Prints the current application version"
+			),
+				new CommandInfo("install", typeof(InstallCommand), null,
+				null,
+				null
+			),
+				new CommandInfo("uninstall", typeof(UninstallCommand), null,
+				null,
+				null
+			),
+				new CommandInfo("upgrade", typeof(UpgradeCommand), typeof(UpgradeCommandOptions),
+				"ugs upgrade",
+				"Upgrades the current installation with the latest build of UGS."
+			),
+				new CommandInfo("settings", typeof(SettingsCommand), typeof(SettingsCommandOptions),
+				"ugs settings",
+				"Export the current ugs settings to a file"
+			),
+			};
+		}
+
+		private static List<CommandInfo> GetWinCommand()
+		{
+			return new List<CommandInfo>()
+			{
+				new CommandInfo("tools", typeof(ToolsCommand), typeof(ToolsCommandCommandOptions),
+				"ugs tools",
+				"Install a Custom Tool. Only available for windows."
+				)
+			};
+		}
+	}
+}
diff --git a/Engine/Source/Programs/UnrealGameSync/UnrealGameSyncCmd/Exceptions/UserErrorException.cs b/Engine/Source/Programs/UnrealGameSync/UnrealGameSyncCmd/Exceptions/UserErrorException.cs
new file mode 100644
--- /dev/null
+++ b/Engine/Source/Programs/UnrealGameSync/UnrealGameSyncCmd/Exceptions/UserErrorException.cs
@@ -0,0 +1,26 @@
+﻿// Copyright Epic Games, Inc. All Rights Reserved.
+
+using System;
+using EpicGames.Core;
+using Microsoft.Extensions.Logging;
+
+namespace UnrealGameSyncCmd.Exceptions
+{
+	public sealed class UserErrorException : Exception
+	{
+		public LogEvent Event { get; }
+		public int Code { get; }
+
+		public UserErrorException(LogEvent evt)
+			: base(evt.ToString())
+		{
+			Event = evt;
+			Code = 1;
+		}
+
+		public UserErrorException(string message, params object[] args)
+			: this(LogEvent.Create(LogLevel.Error, message, args))
+		{
+		}
+	}
+}
diff --git a/Engine/Source/Programs/UnrealGameSync/UnrealGameSyncCmd/Options/ConfigCommandOptions.cs b/Engine/Source/Programs/UnrealGameSync/UnrealGameSyncCmd/Options/ConfigCommandOptions.cs
new file mode 100644
--- /dev/null
+++ b/Engine/Source/Programs/UnrealGameSync/UnrealGameSyncCmd/Options/ConfigCommandOptions.cs
@@ -0,0 +1,21 @@
+﻿// Copyright Epic Games, Inc. All Rights Reserved.
+
+using UnrealGameSync;
+
+namespace UnrealGameSyncCmd.Options
+{
+	internal class ConfigCommandOptions : ServerOptions
+	{
+		public void ApplyTo(UserWorkspaceSettings settings)
+		{
+			if (ServerAndPort != null)
+			{
+				settings.ServerAndPort = (ServerAndPort.Length == 0) ? null : ServerAndPort;
+			}
+			if (UserName != null)
+			{
+				settings.UserName = (UserName.Length == 0) ? null : UserName;
+			}
+		}
+	}
+}
diff --git a/Engine/Source/Programs/UnrealGameSync/UnrealGameSyncCmd/Options/InitCommandOptions.cs b/Engine/Source/Programs/UnrealGameSync/UnrealGameSyncCmd/Options/InitCommandOptions.cs
new file mode 100644
--- /dev/null
+++ b/Engine/Source/Programs/UnrealGameSync/UnrealGameSyncCmd/Options/InitCommandOptions.cs
@@ -0,0 +1,24 @@
+﻿// Copyright Epic Games, Inc. All Rights Reserved.
+
+using EpicGames.Core;
+
+namespace UnrealGameSyncCmd.Options
+{
+	internal class InitCommandOptions : ConfigCommandOptions
+	{
+		[CommandLine("-Client=")]
+		public string? ClientName { get; set; }
+
+		[CommandLine("-Branch=")]
+		public string? BranchPath { get; set; }
+
+		[CommandLine("-Project=")]
+		public string? ProjectName { get; set; }
+
+		[CommandLine("-ClientRoot=")]
+		public string? ClientRoot { get; set; }
+
+		[CommandLine("-IgnoreExistingClients")]
+		public bool IgnoreExistingClients { get; set; }
+	}
+}
diff --git a/Engine/Source/Programs/UnrealGameSync/UnrealGameSyncCmd/Options/ServerOptions.cs b/Engine/Source/Programs/UnrealGameSync/UnrealGameSyncCmd/Options/ServerOptions.cs
new file mode 100644
--- /dev/null
+++ b/Engine/Source/Programs/UnrealGameSync/UnrealGameSyncCmd/Options/ServerOptions.cs
@@ -0,0 +1,15 @@
+﻿// Copyright Epic Games, Inc. All Rights Reserved.
+
+using EpicGames.Core;
+
+namespace UnrealGameSyncCmd.Options
+{
+	internal class ServerOptions
+	{
+		[CommandLine("-Server=")]
+		public string? ServerAndPort { get; set; }
+
+		[CommandLine("-User=")]
+		public string? UserName { get; set; }
+	}
+}
diff --git a/Engine/Source/Programs/UnrealGameSync/UnrealGameSyncCmd/Program.cs b/Engine/Source/Programs/UnrealGameSync/UnrealGameSyncCmd/Program.cs
--- a/Engine/Source/Programs/UnrealGameSync/UnrealGameSyncCmd/Program.cs
+++ b/Engine/Source/Programs/UnrealGameSync/UnrealGameSyncCmd/Program.cs
@@ -4,220 +4,34 @@ using System;
 using System.Collections.Generic;
 using System.Diagnostics;
 using System.Globalization;
-using System.IO;
-using System.IO.Compression;
 using System.Linq;
-using System.Net;
-using System.Net.Http;
-using System.Net.Http.Json;
 using System.Reflection;
 using System.Runtime.InteropServices;
 using System.Text;
-using System.Text.Json;
-using System.Text.RegularExpressions;
 using System.Threading;
 using System.Threading.Tasks;
+
 using EpicGames.Core;
 using EpicGames.Horde;
-using EpicGames.OIDC;
 using EpicGames.Perforce;
 
-using Microsoft.Extensions.Configuration;
 using Microsoft.Extensions.DependencyInjection;
 using Microsoft.Extensions.Logging;
 
 using Serilog;
 using UnrealGameSync;
+using UnrealGameSyncCmd.Commands;
+using UnrealGameSyncCmd.Utils;
+
+using UserErrorException = UnrealGameSyncCmd.Exceptions.UserErrorException;
 
 namespace UnrealGameSyncCmd
 {
 	using ILogger = Microsoft.Extensions.Logging.ILogger;
 
-	public sealed class UserErrorException : Exception
-	{
-		public LogEvent Event { get; }
-		public int Code { get; }
-
-		public UserErrorException(LogEvent evt)
-			: base(evt.ToString())
-		{
-			Event = evt;
-			Code = 1;
-		}
-
-		public UserErrorException(string message, params object[] args)
-			: this(LogEvent.Create(LogLevel.Error, message, args))
-		{
-		}
-	}
-
 	public static class Program
 	{
-		static BuildConfig EditorConfig => BuildConfig.Development;
-
-		class CommandInfo
-		{
-			public string Name { get; }
-			public Type Type { get; }
-			public Type? OptionsType { get; }
-			public string? Usage { get; }
-			public string? Brief { get; }
-
-			public CommandInfo(string name, Type type, Type? optionsType, string? usage, string? brief)
-			{
-				Name = name;
-				Type = type;
-				OptionsType = optionsType;
-				Usage = usage;
-				Brief = brief;
-			}
-		}
-
-		static readonly CommandInfo[] _commands =
-		{
-			new CommandInfo("init", typeof(InitCommand), typeof(InitCommandOptions),
-				"ugs init [stream-path]",
-				"Create a client for the given stream, or initializes an existing client for use by UGS."
-			),
-			new CommandInfo("switch", typeof(SwitchCommand), typeof(SwitchCommandOptions),
-				"ugs switch [project name|project path|stream]",
-				"Changes the active project to the one in the workspace with the given name, or switches to a new stream."
-			),
-			new CommandInfo("changes", typeof(ChangesCommand), typeof(ChangesCommandOptions),
-				"ugs changes",
-				"List recently submitted changes to the current branch."
-			),
-			new CommandInfo("config", typeof(ConfigCommand), typeof(ConfigCommandOptions),
-				"ugs config",
-				"Updates the configuration for the current workspace."
-			),
-			new CommandInfo("filter", typeof(FilterCommand), typeof(FilterCommandOptions),
-				"ugs filter",
-				"Displays or updates the workspace or global sync filter"
-			),
-			new CommandInfo("sync", typeof(SyncCommand), typeof(SyncCommandOptions),
-				"ugs sync [change|'latest']",
-				"Syncs the current workspace to the given changelist, optionally removing all local state."
-			),
-			new CommandInfo("clients", typeof(ClientsCommand), typeof(ClientsCommandOptions),
-				"ugs clients",
-				"Lists all clients suitable for use on the current machine."
-			),
-			new CommandInfo("run", typeof(RunCommand), null,
-				"ugs run",
-				"Runs the editor for the current branch."
-			),
-			new CommandInfo("build", typeof(BuildCommand), typeof(BuildCommandOptions),
-				"ugs build [id]",
-				"Runs the default build steps for the current project, or a particular step referenced by id."
-			),
-			new CommandInfo("status", typeof(StatusCommand), null,
-				"ugs status [-update]",
-				"Shows the status of the currently synced branch."
-			),
-			new CommandInfo("login", typeof(LoginCommand), null,
-				"ugs login",
-				"Starts a interactive login flow against the configured Identity Provider"
-			),
-			new CommandInfo("version", typeof(VersionCommand), null,
-				"ugs version",
-				"Prints the current application version"
-			),
-			new CommandInfo("install", typeof(InstallCommand), null,
-				null,
-				null
-			),
-			new CommandInfo("uninstall", typeof(UninstallCommand), null,
-				null,
-				null
-			),
-			new CommandInfo("upgrade", typeof(UpgradeCommand), typeof(UpgradeCommandOptions),
-				"ugs upgrade",
-				"Upgrades the current installation with the latest build of UGS."
-			),
-			new CommandInfo("tools", typeof(ToolsCommand), typeof(ToolsCommandCommandOptions),
-				"ugs tools",
-				"Install a Custom Tool. Only available for windows."
-			),
-			new CommandInfo("settings", typeof(SettingsCommand), typeof(SettingsCommandOptions),
-				"ugs settings",
-				"Export the current ugs settings to a file"
-			),
-			
-		};
-
-		class CommandContext
-		{
-			public CommandLineArguments Arguments { get; }
-			public ILogger Logger { get; }
-			public ILoggerFactory LoggerFactory { get; }
-
-			public UserSettings? GlobalSettings { get; }
-
-			public GlobalSettingsFile UserSettings { get; }
-
-			public IHordeClient? HordeClient { get; }
-			public ICloudStorage? CloudStorage { get; }
-
-			public CommandContext(CommandLineArguments arguments, ILogger logger, ILoggerFactory loggerFactory, GlobalSettingsFile userSettings, UserSettings? globalSettings, IHordeClient? hordeClient, ICloudStorage? cloudStorage)
-			{
-				Arguments = arguments;
-				Logger = logger;
-				LoggerFactory = loggerFactory;
-				GlobalSettings = globalSettings;
-				UserSettings = userSettings;
-				HordeClient = hordeClient;
-				CloudStorage = cloudStorage;
-			}
-		}
-
-		class ServerOptions
-		{
-			[CommandLine("-Server=")]
-			public string? ServerAndPort { get; set; }
-
-			[CommandLine("-User=")]
-			public string? UserName { get; set; }
-		}
-
-		class ConfigCommandOptions : ServerOptions
-		{
-			public void ApplyTo(UserWorkspaceSettings settings)
-			{
-				if (ServerAndPort != null)
-				{
-					settings.ServerAndPort = (ServerAndPort.Length == 0) ? null : ServerAndPort;
-				}
-				if (UserName != null)
-				{
-					settings.UserName = (UserName.Length == 0) ? null : UserName;
-				}
-			}
-		}
-
-		class InitCommandOptions : ConfigCommandOptions
-		{
-			[CommandLine("-Client=")]
-			public string? ClientName { get; set; }
-
-			[CommandLine("-Branch=")]
-			public string? BranchPath { get; set; }
-
-			[CommandLine("-Project=")]
-			public string? ProjectName { get; set; }
-
-			[CommandLine("-ClientRoot=")]
-			public string? ClientRoot { get; set; }
-
-			[CommandLine("-IgnoreExistingClients")]
-			public bool IgnoreExistingClients { get; set; }
-		}
-
-		class UpdateState
-		{
-			public string? LatestVersion { get; set; }
-			public DateTime LastVersionCheck { get; set; }
-		}
+		static readonly CommandInfo[] _commands = CommandsFactory.GetCommands();
 
 		public static async Task<int> Main(string[] rawArgs)
 		{
@@ -309,14 +123,14 @@ namespace UnrealGameSyncCmd
 				// On Windows this is distributed with the GUI client, so we don't need to check for upgrades
 				if (command.Type != typeof(InstallCommand) && command.Type != typeof(UpgradeCommand) && !RuntimeInformation.IsOSPlatform(OSPlatform.Windows))
 				{
-					string version = GetVersion();
+					string version = VersionUtils.GetVersion();
 
 					DateTime utcNow = DateTime.UtcNow;
 					if (settings.Global.LastVersionCheck < utcNow - TimeSpan.FromDays(1.0) || IsUpgradeAvailable(settings, version))
 					{
 						using (CancellationTokenSource cancellationSource = new CancellationTokenSource(TimeSpan.FromSeconds(10.0)))
 						{
-							Task<string?> latestVersionTask = GetLatestVersionAsync(null, cancellationSource.Token);
+							Task<string?> latestVersionTask = VersionUtils.GetLatestVersionAsync(null, cancellationSource.Token);
 
 							Task delay = Task.Delay(TimeSpan.FromSeconds(2.0));
 							await Task.WhenAny(latestVersionTask, delay);
@@ -429,1773 +243,5 @@ namespace UnrealGameSyncCmd
 			}
 			return result.ToString();
 		}
-
-		static string GetVersion()
-		{
-			AssemblyInformationalVersionAttribute? version = Assembly.GetExecutingAssembly().GetCustomAttribute<AssemblyInformationalVersionAttribute>();
-			return version?.InformationalVersion ?? "Unknown";
-		}
-
-		class DeploymentInfo
-		{
-			public string Version { get; set; } = String.Empty;
-		}
-
-		static string? GetUpgradeToolName()
-		{
-			if (RuntimeInformation.IsOSPlatform(OSPlatform.OSX))
-			{
-				if (RuntimeInformation.OSArchitecture == Architecture.Arm64)
-				{
-					return "ugs-mac-arm64";
-				}
-				else
-				{
-					return "ugs-mac";
-				}
-			}
-			else if (RuntimeInformation.IsOSPlatform(OSPlatform.Linux))
-			{
-				return "ugs-linux";
-			}
-			else
-			{
-				return null;
-			}
-		}
-
-		static async Task<string?> GetLatestVersionAsync(ILogger? logger, CancellationToken cancellationToken)
-		{
-			string? hordeUrl = DeploymentSettings.Instance.HordeUrl;
-			if (hordeUrl == null)
-			{
-				logger?.LogError("Horde URL is not set in deployment config file. Cannot upgrade.");
-				return null;
-			}
-
-			string? toolName = GetUpgradeToolName();
-			if (toolName == null)
-			{
-				logger?.LogError("Command-line upgrades are not supported on this platform.");
-				return null;
-			}
-
-			using (HttpClient httpClient = new HttpClient())
-			{
-				Uri baseUrl = new Uri(hordeUrl);
-
-				DeploymentInfo? deploymentInfo;
-				try
-				{
-					deploymentInfo = await httpClient.GetFromJsonAsync<DeploymentInfo>(new Uri(baseUrl, $"api/v1/tools/{toolName}/deployments"), cancellationToken);
-				}
-				catch (Exception ex)
-				{
-					logger?.LogError(ex, "Failed to query for deployment info: {Message}", ex.Message);
-					return null;
-				}
-				if (deploymentInfo == null)
-				{
-					logger?.LogError("Failed to query for deployment info.");
-					return null;
-				}
-				return deploymentInfo.Version;
-			}
-		}
-
-		public static UserWorkspaceSettings? ReadOptionalUserWorkspaceSettings()
-		{
-			DirectoryReference? dir = DirectoryReference.GetCurrentDirectory();
-			for (; dir != null; dir = dir.ParentDirectory)
-			{
-				try
-				{
-					UserWorkspaceSettings? settings;
-					if (UserWorkspaceSettings.TryLoad(dir, out settings))
-					{
-						return settings;
-					}
-				}
-				catch
-				{
-					// Guard against directories we can't access, eg. [REDACTED_USER_HOME]
-				}
-			}
-			return null;
-		}
-
-		public static UserWorkspaceSettings ReadRequiredUserWorkspaceSettings()
-		{
-			UserWorkspaceSettings? settings = ReadOptionalUserWorkspaceSettings();
-			if (settings == null)
-			{
-				throw new UserErrorException("Unable to find UGS workspace in current directory.");
-			}
-			return settings;
-		}
-
-		public static async Task<WorkspaceStateWrapper> ReadWorkspaceState(IPerforceConnection perforceClient, UserWorkspaceSettings settings, GlobalSettingsFile userSettings, ILogger logger)
-		{
-			WorkspaceStateWrapper state = userSettings.FindOrAddWorkspaceState(settings);
-			if (state.Current.SettingsTimeUtc != settings.LastModifiedTimeUtc)
-			{
-				logger.LogDebug("Updating state due to modified settings timestamp");
-				ProjectInfo info = await ProjectInfo.CreateAsync(perforceClient, settings, CancellationToken.None);
-				state.Modify(x => x.UpdateCachedProjectInfo(info, settings.LastModifiedTimeUtc));
-			}
-			return state;
-		}
-
-		public static Task<IPerforceConnection> ConnectAsync(string? serverAndPort, string? userName, string? clientName, ILoggerFactory loggerFactory)
-		{
-			PerforceSettings settings = new PerforceSettings(PerforceSettings.Default);
-			settings.ClientName = clientName;
-			settings.PreferNativeClient = true;
-			if (!String.IsNullOrEmpty(serverAndPort))
-			{
-				settings.ServerAndPort = serverAndPort;
-			}
-			if (!String.IsNullOrEmpty(userName))
-			{
-				settings.UserName = userName;
-			}
-
-			return PerforceConnection.CreateAsync(settings, loggerFactory.CreateLogger("Perforce"));
-		}
-
-		public static Task<IPerforceConnection> ConnectAsync(UserWorkspaceSettings settings, ILoggerFactory loggerFactory)
-		{
-			return ConnectAsync(settings.ServerAndPort, settings.UserName, settings.ClientName, loggerFactory);
-		}
-
-		static string[] ReadSyncFilter(UserWorkspaceSettings workspaceSettings, GlobalSettingsFile userSettings, ConfigFile projectConfig, string projectIdentifier)
-		{
-			Dictionary<Guid, WorkspaceSyncCategory> syncCategories = ConfigUtils.GetSyncCategories(projectConfig);
-
-			// check if any category is from the role
-			IDictionary<string, Preset> roles = ConfigUtils.GetPresets(projectConfig, projectIdentifier);
-			if (roles.TryGetValue(workspaceSettings.Preset, out Preset? role))
-			{
-				foreach (RoleCategory roleCategory in role.Categories.Values)
-				{
-					if (syncCategories.TryGetValue(roleCategory.Id, out WorkspaceSyncCategory? category))
-					{
-						category.Enable = roleCategory.Enabled;
-					}
-				}
-			}
-			
-			ConfigSection? perforceSection = projectConfig.FindSection("Perforce");
-
-			string[] combinedSyncFilter = GlobalSettingsFile.GetCombinedSyncFilter(syncCategories, workspaceSettings.Preset, roles, userSettings.Global.Filter, workspaceSettings.Filter, perforceSection);
-
-			return combinedSyncFilter;
-		}
-
-		static async Task<string> FindProjectPathAsync(IPerforceConnection perforce, string clientName, string branchPath, string? projectName)
-		{
-			using IPerforceConnection perforceClient = await PerforceConnection.CreateAsync(new PerforceSettings(perforce.Settings) { ClientName = clientName }, perforce.Logger);
-
-			// Find or validate the selected project
-			string searchPath;
-			if (projectName == null)
-			{
-				searchPath = $"//{clientName}{branchPath}/*.uprojectdirs";
-			}
-			else if (projectName.Contains('.', StringComparison.Ordinal))
-			{
-				searchPath = $"//{clientName}{branchPath}/{projectName.TrimStart('/')}";
-			}
-			else
-			{
-				searchPath = $"//{clientName}{branchPath}/.../{projectName}.uproject";
-			}
-
-			List<FStatRecord> projectFileRecords = await perforceClient.FStatAsync(FStatOptions.ClientFileInPerforceSyntax, searchPath).ToListAsync();
-			projectFileRecords.RemoveAll(x => x.HeadAction == FileAction.Delete || x.HeadAction == FileAction.MoveDelete);
-			projectFileRecords.RemoveAll(x => !x.IsMapped);
-
-			List<string> paths = projectFileRecords.Select(x => PerforceUtils.GetClientRelativePath(x.ClientFile!)).Distinct(StringComparer.Ordinal).ToList();
-			if (paths.Count == 0)
-			{
-				throw new UserErrorException("No project file found matching {SearchPath}", searchPath);
-			}
-			if (paths.Count > 1)
-			{
-				throw new UserErrorException("Multiple projects found matching {SearchPath}: {Paths}", searchPath, String.Join(", ", paths));
-			}
-
-			return "/" + paths[0];
-		}
-
-		abstract class Command
-		{
-			public abstract Task ExecuteAsync(CommandContext context);
-		}
-
-		class InitCommand : Command
-		{
-			public override async Task ExecuteAsync(CommandContext context)
-			{
-				ILogger logger = context.Logger;
-
-				// Get the positional argument indicating the file to look for
-				string? initName;
-				context.Arguments.TryGetPositionalArgument(out initName);
-
-				// Get the config settings from the command line
-				InitCommandOptions options = new InitCommandOptions();
-				context.Arguments.ApplyTo(options);
-				context.Arguments.CheckAllArgumentsUsed();
-
-				// Get the host name
-				using IPerforceConnection perforce = await ConnectAsync(options.ServerAndPort, options.UserName, null, context.LoggerFactory);
-				InfoRecord perforceInfo = await perforce.GetInfoAsync(InfoOptions.ShortOutput);
-				string hostName = perforceInfo.ClientHost ?? Dns.GetHostName();
-
-				// Create the perforce connection
-				if (initName != null)
-				{
-					await InitNewClientAsync(perforce, initName, hostName, options, context.GlobalSettings, logger);
-				}
-				else
-				{
-					await InitExistingClientAsync(perforce, hostName, options, context.GlobalSettings, logger);
-				}
-			}
-
-			static async Task InitNewClientAsync(IPerforceConnection perforce, string streamName, string hostName, InitCommandOptions options, UserSettings? globalSettings, ILogger logger)
-			{
-				logger.LogInformation("Checking stream...");
-
-				// Get the given stream
-				PerforceResponse<StreamRecord> streamResponse = await perforce.TryGetStreamAsync(streamName, true);
-				if (!streamResponse.Succeeded)
-				{
-					throw new UserErrorException($"Unable to find stream '{streamName}'");
-				}
-				StreamRecord stream = streamResponse.Data;
-
-				// Get the new directory for the client
-				DirectoryReference clientDir = DirectoryReference.FromString(options.ClientRoot) ?? DirectoryReference.Combine(DirectoryReference.GetCurrentDirectory(), stream.Stream.Replace('/', '+'));
-				DirectoryReference.CreateDirectory(clientDir);
-
-				// Make up a new client name 
-				string clientName = options.ClientName ?? Regex.Replace($"{perforce.Settings.UserName}_{hostName}_{stream.Stream.Trim('/')}", "[^0-9a-zA-Z_.-]", "+");
-
-				if (!options.IgnoreExistingClients)
-				{
-					// Check there are no existing clients under the current path
-					List<ClientsRecord> clients = await FindExistingClients(perforce, hostName, clientDir);
-					if (clients.Count > 0)
-					{
-						if (clients.Count == 1 && clientName.Equals(clients[0].Name, StringComparison.OrdinalIgnoreCase) && clientDir == TryParseRoot(clients[0].Root))
-						{
-							logger.LogInformation("Reusing existing client for {ClientDir} ({ClientName})", clientDir, options.ClientName);
-						}
-						else
-						{
-							throw new UserErrorException("Current directory is already within a Perforce workspace ({ClientName})", clients[0].Name);
-						}
-					}
-				}
-
-				// Create the new client
-				ClientRecord client = new ClientRecord(clientName, perforce.Settings.UserName, clientDir.FullName);
-				client.Host = hostName;
-				client.Stream = stream.Stream;
-				client.Options = ClientOptions.Rmdir;
-				await perforce.CreateClientAsync(client);
-
-				// Branch root is currently hard-coded at the root
-				string branchPath = options.BranchPath ?? String.Empty;
-				string projectPath = await FindProjectPathAsync(perforce, clientName, branchPath, options.ProjectName);
-
-				// Create the settings object
-				UserWorkspaceSettings settings = new UserWorkspaceSettings();
-				settings.RootDir = clientDir;
-				settings.Init(perforce.Settings.ServerAndPort, perforce.Settings.UserName, clientName, branchPath, projectPath);
-				options.ApplyTo(settings);
-				settings.Save(logger);
-
-				if (globalSettings != null)
-				{
-					UserSelectedProjectSettings selectedProjectSettings =
-						new UserSelectedProjectSettings(
-							null, 
-							null, 
-							UserSelectedProjectType.Client, 
-							$"//{client.Name}{projectPath}",
-							$"{client.Root}{projectPath}".Replace("/", @"\", StringComparison.Ordinal));
-					globalSettings.OpenProjects.Add(selectedProjectSettings);
-					globalSettings.RecentProjects.Add(selectedProjectSettings);
-					globalSettings.Save(logger);
-				}
-
-				logger.LogInformation("Initialized {ClientName} with root at {RootDir}", clientName, clientDir);
-			}
-
-			static DirectoryReference? TryParseRoot(string root)
-			{
-				try
-				{
-					return new DirectoryReference(root);
-				}
-				catch
-				{
-					return null;
-				}
-			}
-
-			static async Task InitExistingClientAsync(IPerforceConnection perforce, string hostName, InitCommandOptions options, UserSettings? globalSettings, ILogger logger)
-			{
-				DirectoryReference currentDir = DirectoryReference.GetCurrentDirectory();
-
-				// Make sure the client name is set
-				string? clientName = options.ClientName;
-				if (clientName == null)
-				{
-					List<ClientsRecord> clients = await FindExistingClients(perforce, hostName, currentDir);
-					if (clients.Count == 0)
-					{
-						throw new UserErrorException("Unable to find client for {HostName} under {ClientDir}", hostName, currentDir);
-					}
-					if (clients.Count > 1)
-					{
-						throw new UserErrorException("Multiple clients found for {HostName} under {ClientDir}: {ClientList}", hostName, currentDir, String.Join(", ", clients.Select(x => x.Name)));
-					}
-
-					clientName = clients[0].Name;
-					logger.LogInformation("Found client {ClientName}", clientName);
-				}
-
-				// Get the client info
-				ClientRecord client = await perforce.GetClientAsync(clientName);
-				DirectoryReference clientDir = new DirectoryReference(client.Root);
-
-				// If a project path was specified in local syntax, try to convert it to client-relative syntax
-				string? projectName = options.ProjectName;
-				if (options.ProjectName != null && options.ProjectName.Contains('.', StringComparison.Ordinal))
-				{
-					options.ProjectName = FileReference.Combine(currentDir, options.ProjectName).MakeRelativeTo(clientDir).Replace('\\', '/');
-				}
-
-				// Branch root is currently hard-coded at the root
-				string branchPath = options.BranchPath ?? String.Empty;
-				string projectPath = await FindProjectPathAsync(perforce, clientName, branchPath, projectName);
-
-				// Create the settings object
-				UserWorkspaceSettings settings = new UserWorkspaceSettings();
-				settings.RootDir = clientDir;
-				settings.Init(perforce.Settings.ServerAndPort, perforce.Settings.UserName, clientName, branchPath, projectPath);
-				options.ApplyTo(settings);
-				settings.Save(logger);
-
-				if (globalSettings != null)
-				{
-					UserSelectedProjectSettings selectedProjectSettings =
-						new UserSelectedProjectSettings(
-							null,
-							null,
-							UserSelectedProjectType.Client,
-							$"//{client.Name}{projectPath}",
-							$"{client.Root}{projectPath}".Replace("/", @"\", StringComparison.Ordinal));
-					globalSettings.OpenProjects.Add(selectedProjectSettings);
-					globalSettings.RecentProjects.Add(selectedProjectSettings);
-					globalSettings.Save(logger);
-				}
-
-				logger.LogInformation("Initialized workspace at {RootDir} for {ClientProject}", clientDir, settings.ClientProjectPath);
-			}
-
-			static async Task<List<ClientsRecord>> FindExistingClients(IPerforceConnection perforce, string hostName, DirectoryReference clientDir)
-			{
-				List<ClientsRecord> matchingClients = new List<ClientsRecord>();
-
-				List<ClientsRecord> clients = await perforce.GetClientsAsync(ClientsOptions.None, perforce.Settings.UserName);
-				foreach (ClientsRecord client in clients)
-				{
-					if (!String.IsNullOrEmpty(client.Root) && !String.IsNullOrEmpty(client.Host) && String.Equals(hostName, client.Host, StringComparison.OrdinalIgnoreCase))
-					{
-						DirectoryReference? rootDir;
-						try
-						{
-							rootDir = new DirectoryReference(client.Root);
-						}
-						catch
-						{
-							rootDir = null;
-						}
-
-						if (rootDir != null && clientDir.IsUnderDirectory(rootDir))
-						{
-							matchingClients.Add(client);
-						}
-					}
-				}
-
-				return matchingClients;
-			}
-		}
-
-		class SyncCommandOptions
-		{
-			[CommandLine("-Clean")]
-			public bool Clean { get; set; }
-
-			[CommandLine("-Build")]
-			public bool Build { get; set; }
-
-			[CommandLine("-Binaries")]
-			public bool Binaries { get; set; }
-
-			[CommandLine("-NoGPF", Value = "false")]
-			[CommandLine("-NoProjectFiles", Value = "false")]
-			public bool ProjectFiles { get; set; } = true;
-
-			[CommandLine("-Clobber")]
-			public bool Clobber { get; set; }
-
-			[CommandLine("-Refilter")]
-			public bool Refilter { get; set; }
-
-			[CommandLine("-Only")]
-			public bool SingleChange { get; set; }
-		}
-
-		class SyncCommand : Command
-		{
-			static async Task<bool> IsCodeChangeAsync(IPerforceConnection perforce, int change)
-			{
-				DescribeRecord describeRecord = await perforce.DescribeAsync(change);
-				return IsCodeChange(describeRecord);
-			}
-
-			static bool IsCodeChange(DescribeRecord describeRecord)
-			{
-				foreach (DescribeFileRecord file in describeRecord.Files)
-				{
-					if (PerforceUtils.CodeExtensions.Any(extension => file.DepotFile.EndsWith(extension, StringComparison.OrdinalIgnoreCase)))
-					{
-						return true;
-					}
-				}
-				return false;
-			}
-
-			public override async Task ExecuteAsync(CommandContext context)
-			{
-				ILogger logger = context.Logger;
-				context.Arguments.TryGetPositionalArgument(out string? changeString);
-
-				SyncCommandOptions syncOptions = new SyncCommandOptions();
-				context.Arguments.ApplyTo(syncOptions);
-
-				context.Arguments.CheckAllArgumentsUsed();
-
-				UserWorkspaceSettings settings = ReadRequiredUserWorkspaceSettings();
-				using IPerforceConnection perforceClient = await ConnectAsync(settings, context.LoggerFactory);
-				WorkspaceStateWrapper state = await ReadWorkspaceState(perforceClient, settings, context.UserSettings, logger);
-
-				changeString ??= "latest";
-
-				ProjectInfo projectInfo = new ProjectInfo(settings.RootDir, state.Current);
-				UserProjectSettings projectSettings = context.UserSettings.FindOrAddProjectSettings(projectInfo, settings, logger);
-
-				ConfigFile projectConfig = await ConfigUtils.ReadProjectConfigFileAsync(perforceClient, projectInfo, logger, CancellationToken.None);
-
-				bool syncLatest = String.Equals(changeString, "latest", StringComparison.OrdinalIgnoreCase);
-
-				int change;
-				if (!Int32.TryParse(changeString, out change))
-				{
-					if (syncLatest)
-					{
-						List<ChangesRecord> changes = await perforceClient.GetChangesAsync(ChangesOptions.None, 1, ChangeStatus.Submitted, $"//{settings.ClientName}/...");
-						change = changes[0].Number;
-					}
-					else
-					{
-						throw new UserErrorException("Unknown change type for sync '{Change}'", changeString);
-					}
-				}
-
-				WorkspaceUpdateOptions options = syncOptions.SingleChange ? WorkspaceUpdateOptions.SyncSingleChange : WorkspaceUpdateOptions.Sync;
-				if (syncOptions.Clean)
-				{
-					options |= WorkspaceUpdateOptions.Clean;
-				}
-				if (syncOptions.Build)
-				{
-					options |= WorkspaceUpdateOptions.Build;
-				}
-				if (syncOptions.ProjectFiles)
-				{
-					options |= WorkspaceUpdateOptions.GenerateProjectFiles;
-				}
-				if (syncOptions.Clobber)
-				{
-					options |= WorkspaceUpdateOptions.Clobber;
-				}
-				if (syncOptions.Refilter)
-				{
-					options |= WorkspaceUpdateOptions.Refilter;
-				}
-				options |= WorkspaceUpdateContext.GetOptionsFromConfig(context.UserSettings.Global, settings);
-				options |= WorkspaceUpdateOptions.RemoveFilteredFiles;
-
-				string[] syncFilter = ReadSyncFilter(settings, context.UserSettings, projectConfig, state.Current.ProjectIdentifier);
-
-				using WorkspaceLock? workspaceLock = CreateWorkspaceLock(settings.RootDir);
-				if (workspaceLock != null && !await workspaceLock.TryAcquireAsync())
-				{
-					logger.LogError("Another process is already syncing this workspace.");
-					return;
-				}
-
-				WorkspaceUpdateContext updateContext = new WorkspaceUpdateContext(change, options, BuildConfig.Development, syncFilter, projectSettings.BuildSteps, null);
-				updateContext.PerforceSyncOptions = context.UserSettings.Global.Perforce;
-
-				if (syncOptions.Binaries)
-				{
-					List<BaseArchiveChannel> archives = await BaseArchive.EnumerateChannelsAsync(perforceClient, context.HordeClient, context.CloudStorage, projectConfig, state.Current.ProjectIdentifier, CancellationToken.None);
-
-					BaseArchiveChannel? editorArchiveInfo = archives.FirstOrDefault(x => x.Type == IArchiveChannel.EditorArchiveType);
-					if (editorArchiveInfo == null)
-					{
-						throw new UserErrorException("No editor archives found for project");
-					}
-
-					KeyValuePair<int, IArchive> revision = editorArchiveInfo.ChangeNumberToArchive.LastOrDefault(x => x.Key <= change);
-					if (revision.Key == 0)
-					{
-						throw new UserErrorException($"No editor archives found for CL {change}");
-					}
-
-					if (revision.Key < change)
-					{
-						int lastChange = revision.Key;
-
-						List<ChangesRecord> changeRecords = await perforceClient.GetChangesAsync(ChangesOptions.None, 1, ChangeStatus.Submitted, $"//{settings.ClientName}/...@{revision.Key + 1},{change}");
-						foreach (ChangesRecord changeRecord in changeRecords.OrderBy(x => x.Number))
-						{
-							if (await IsCodeChangeAsync(perforceClient, changeRecord.Number))
-							{
-								if (syncLatest)
-								{
-									updateContext.ChangeNumber = lastChange;
-								}
-								else
-								{
-									throw new UserErrorException($"No editor binaries found for CL {change} (last archive at CL {revision.Key}, but CL {changeRecord.Number} is a code change)");
-								}
-								break;
-							}
-							change = changeRecord.Number;
-						}
-					}
-
-					updateContext.Options |= WorkspaceUpdateOptions.SyncArchives;
-					updateContext.ArchiveTypeToArchive[IArchiveChannel.EditorArchiveType] = revision.Value;
-				}
-
-				WorkspaceUpdate update = new WorkspaceUpdate(updateContext);
-				(WorkspaceUpdateResult result, string message) = await update.ExecuteAsync(perforceClient.Settings, projectInfo, state, context.Logger, CancellationToken.None);
-				if (result == WorkspaceUpdateResult.FilesToClobber)
-				{
-					logger.LogWarning("The following files are modified in your workspace:");
-					foreach (string file in updateContext.ClobberFiles.Keys.OrderBy(x => x))
-					{
-						logger.LogWarning("  {File}", file);
-					}
-					logger.LogWarning("Use -Clobber to overwrite");
-				}
-
-				state.Modify(x => x.SetLastSyncState(result, updateContext, message));
-
-				if (result != WorkspaceUpdateResult.Success)
-				{
-					throw new UserErrorException("{Message} (Result: {Result})", message, result);
-				}
-			}
-
-			static WorkspaceLock? CreateWorkspaceLock(DirectoryReference rootDir)
-			{
-				if (RuntimeInformation.IsOSPlatform(OSPlatform.Windows))
-				{
-					return new WorkspaceLock(rootDir);
-				}
-				else
-				{
-					return null;
-				}
-			}
-		}
-
-		class ClientsCommandOptions : ServerOptions
-		{
-		}
-
-		class ClientsCommand : Command
-		{
-			public override async Task ExecuteAsync(CommandContext context)
-			{
-				ILogger logger = context.Logger;
-
-				ClientsCommandOptions options = context.Arguments.ApplyTo<ClientsCommandOptions>(logger);
-				context.Arguments.CheckAllArgumentsUsed();
-
-				using IPerforceConnection perforceClient = await ConnectAsync(options.ServerAndPort, options.UserName, null, context.LoggerFactory);
-				InfoRecord info = await perforceClient.GetInfoAsync(InfoOptions.ShortOutput);
-
-				List<ClientsRecord> clients = await perforceClient.GetClientsAsync(EpicGames.Perforce.ClientsOptions.None, perforceClient.Settings.UserName);
-				foreach (ClientsRecord client in clients)
-				{
-					if (String.Equals(info.ClientHost, client.Host, StringComparison.OrdinalIgnoreCase))
-					{
-						logger.LogInformation("{Client,-50} {Root}", client.Name, client.Root);
-					}
-				}
-			}
-		}
-
-		class RunCommand : Command
-		{
-			public override async Task ExecuteAsync(CommandContext context)
-			{
-				ILogger logger = context.Logger;
-
-				UserWorkspaceSettings settings = ReadRequiredUserWorkspaceSettings();
-				using IPerforceConnection perforceClient = await ConnectAsync(settings, context.LoggerFactory);
-				WorkspaceStateWrapper state = await ReadWorkspaceState(perforceClient, settings, context.UserSettings, logger);
-
-				ProjectInfo projectInfo = new ProjectInfo(settings.RootDir, state.Current);
-				ConfigFile projectConfig = await ConfigUtils.ReadProjectConfigFileAsync(perforceClient, projectInfo, logger, CancellationToken.None);
-
-				FileReference receiptFile = ConfigUtils.GetEditorReceiptFile(projectInfo, projectConfig, EditorConfig);
-				logger.LogDebug("Receipt file: {Receipt}", receiptFile);
-
-				if (!ConfigUtils.TryReadEditorReceipt(projectInfo, receiptFile, out TargetReceipt? receipt) || String.IsNullOrEmpty(receipt.Launch))
-				{
-					throw new UserErrorException("The editor needs to be built before you can run it. (Missing {ReceiptFile}).", receiptFile);
-				}
-				if (!File.Exists(receipt.Launch))
-				{
-					throw new UserErrorException("The editor needs to be built before you can run it. (Missing {LaunchFile}).", receipt.Launch);
-				}
-
-				List<string> launchArguments = new List<string>();
-				if (settings.LocalProjectPath.HasExtension(".uproject"))
-				{
-					launchArguments.Add($"\"{settings.LocalProjectPath}\"");
-				}
-				if (EditorConfig == BuildConfig.Debug || EditorConfig == BuildConfig.DebugGame)
-				{
-					launchArguments.Append(" -debug");
-				}
-				for (int idx = 0; idx < context.Arguments.Count; idx++)
-				{
-					if (!context.Arguments.HasBeenUsed(idx))
-					{
-						launchArguments.Add(context.Arguments[idx]);
-					}
-				}
-
-				string commandLine = CommandLineArguments.Join(launchArguments);
-				logger.LogInformation("Spawning: {LaunchFile} {CommandLine}", CommandLineArguments.Quote(receipt.Launch), commandLine);
-
-				if (!Utility.SpawnProcess(receipt.Launch, commandLine))
-				{
-					logger.LogError("Unable to spawn {App} {Args}", receipt.Launch, launchArguments.ToString());
-				}
-			}
-		}
-
-		class ChangesCommandOptions
-		{
-			[CommandLine("-Count=")]
-			public int Count { get; set; } = 10;
-
-			[CommandLine("-Lines=")]
-			public int LineCount { get; set; } = 3;
-		}
-
-		class ChangesCommand : Command
-		{
-			public override async Task ExecuteAsync(CommandContext context)
-			{
-				ILogger logger = context.Logger;
-
-				ChangesCommandOptions options = new ChangesCommandOptions();
-				context.Arguments.ApplyTo(options);
-				context.Arguments.CheckAllArgumentsUsed(context.Logger);
-
-				UserWorkspaceSettings settings = ReadRequiredUserWorkspaceSettings();
-				using IPerforceConnection perforceClient = await ConnectAsync(settings, context.LoggerFactory);
-
-				List<ChangesRecord> changes = await perforceClient.GetChangesAsync(EpicGames.Perforce.ChangesOptions.None, options.Count, ChangeStatus.Submitted, $"//{settings.ClientName}/...");
-				foreach (IEnumerable<ChangesRecord> changesBatch in changes.Batch(10))
-				{
-					List<DescribeRecord> describeRecords = await perforceClient.DescribeAsync(changesBatch.Select(x => x.Number).ToArray());
-
-					logger.LogInformation("  Change    Type     Author          Description");
-					foreach (DescribeRecord describeRecord in describeRecords)
-					{
-						PerforceChangeDetails details = new PerforceChangeDetails(describeRecord);
-
-						string type;
-						if (details.ContainsCode)
-						{
-							if (details.ContainsContent)
-							{
-								type = "Both";
-							}
-							else
-							{
-								type = "Code";
-							}
-						}
-						else
-						{
-							if (details.ContainsContent)
-							{
-								type = "Content";
-							}
-							else
-							{
-								type = "None";
-							}
-						}
-
-						string author = StringUtils.Truncate(describeRecord.User, 15);
-
-						List<string> lines = StringUtils.WordWrap(details.Description, Math.Max(ConsoleUtils.WindowWidth - 40, 10)).ToList();
-						if (lines.Count == 0)
-						{
-							lines.Add(String.Empty);
-						}
-
-						int lineCount = Math.Min(options.LineCount, lines.Count);
-
-						logger.LogInformation("  {Change,-9} {Type,-8} {Author,-15} {Description}", describeRecord.Number, type, author, lines[0]);
-						for (int lineIndex = 1; lineIndex < lineCount; lineIndex++)
-						{
-							logger.LogInformation("                                     {Description}", lines[lineIndex]);
-						}
-					}
-				}
-			}
-		}
-
-		class ConfigCommand : Command
-		{
-			public override async Task ExecuteAsync(CommandContext context)
-			{
-				ILogger logger = context.Logger;
-
-				UserWorkspaceSettings settings = ReadRequiredUserWorkspaceSettings();
-				if (!context.Arguments.GetUnusedArguments().Any())
-				{
-					ProcessStartInfo startInfo = new ProcessStartInfo();
-					startInfo.FileName = settings.ConfigFile.FullName;
-					startInfo.UseShellExecute = true;
-					using (Process? editor = Process.Start(startInfo))
-					{
-						if (editor != null)
-						{
-							await editor.WaitForExitAsync();
-						}
-					}
-				}
-				else
-				{
-					ConfigCommandOptions options = new ConfigCommandOptions();
-					context.Arguments.ApplyTo(options);
-					context.Arguments.CheckAllArgumentsUsed(context.Logger);
-
-					options.ApplyTo(settings);
-					settings.Save(logger);
-
-					logger.LogInformation("Updated {ConfigFile}", settings.ConfigFile);
-				}
-			}
-		}
-
-		class FilterCommandOptions
-		{
-			[CommandLine("-Reset")]
-			public bool Reset { get; set; } = false;
-
-			[CommandLine("-Include=")]
-			public List<string> Include { get; set; } = new List<string>();
-
-			[CommandLine("-Exclude=")]
-			public List<string> Exclude { get; set; } = new List<string>();
-
-			[CommandLine("-View=", ListSeparator = ';')]
-			public List<string>? View { get; set; }
-
-			[CommandLine("-AddView=", ListSeparator = ';')]
-			public List<string> AddView { get; set; } = new List<string>();
-
-			[CommandLine("-RemoveView=", ListSeparator = ';')]
-			public List<string> RemoveView { get; set; } = new List<string>();
-
-			[CommandLine("-AllProjects", Value = "true")]
-			[CommandLine("-OnlyCurrent", Value = "false")]
-			public bool? AllProjects { get; set; } = null;
-
-			[CommandLine("-LocalProjects", Value = "true")]
-			[CommandLine("-NoExtraLocalProjects", Value = "false")]
-			public bool? LocalProjects { get; set; } = null;
-
-			[CommandLine("-GpfAllProjects", Value = "true")]
-			[CommandLine("-GpfOnlyCurrent", Value = "false")]
-			public bool? AllProjectsInSln { get; set; } = null;
-
-			[CommandLine("-GpfMinimalSln", Value = "true")]
-			[CommandLine("-GpfFullSln", Value = "false")]
-			public bool? UprojectSpecificSln { get; set; } = null;
-
-			[CommandLine("-Global")]
-			public bool Global { get; set; }
-		}
-
-		class FilterCommand : Command
-		{
-			public override async Task ExecuteAsync(CommandContext context)
-			{
-				ILogger logger = context.Logger;
-
-				UserWorkspaceSettings workspaceSettings = ReadRequiredUserWorkspaceSettings();
-				using IPerforceConnection perforceClient = await ConnectAsync(workspaceSettings, context.LoggerFactory);
-				WorkspaceStateWrapper workspaceState = await ReadWorkspaceState(perforceClient, workspaceSettings, context.UserSettings, logger);
-				ProjectInfo projectInfo = new ProjectInfo(workspaceSettings.RootDir, workspaceState.Current);
-
-				ConfigFile projectConfig = await ConfigUtils.ReadProjectConfigFileAsync(perforceClient, projectInfo, logger, CancellationToken.None);
-				Dictionary<Guid, WorkspaceSyncCategory> syncCategories = ConfigUtils.GetSyncCategories(projectConfig);
-
-				FilterSettings globalFilter = context.UserSettings.Global.Filter;
-				FilterSettings workspaceFilter = workspaceSettings.Filter;
-
-				IDictionary<string, Preset> roles = ConfigUtils.GetPresets(projectConfig, workspaceState.Current.ProjectIdentifier);
-
-				roles.TryGetValue(workspaceSettings.Preset, out Preset? role);
-				
-				FilterCommandOptions options = context.Arguments.ApplyTo<FilterCommandOptions>(logger);
-				context.Arguments.CheckAllArgumentsUsed(context.Logger);
-
-				if (options.Global)
-				{
-					ApplyCommandOptions(context.UserSettings.Global.Filter, options, syncCategories.Values, logger);
-					context.UserSettings.Save(logger);
-				}
-				else
-				{
-					ApplyCommandOptions(workspaceSettings.Filter, options, syncCategories.Values, logger);
-					workspaceSettings.Save(logger);
-				}
-
-				Dictionary<Guid, bool> globalCategories = globalFilter.GetCategories();
-				Dictionary<Guid, bool> workspaceCategories = workspaceFilter.GetCategories();
-
-				logger.LogInformation("Categories:");
-				foreach (WorkspaceSyncCategory syncCategory in syncCategories.Values)
-				{
-					bool enabled;
-
-					string scope = "(Default)";
-					if (globalCategories.TryGetValue(syncCategory.UniqueId, out enabled))
-					{
-						scope = "(Global)";
-					}
-					else if (workspaceCategories.TryGetValue(syncCategory.UniqueId, out enabled))
-					{
-						scope = "(Workspace)";
-					}
-					else
-					{
-						enabled = syncCategory.Enable;
-					}
-
-					if (role != null && role.Categories.TryGetValue(syncCategory.UniqueId, out RoleCategory? roleCategory))
-					{
-						scope = $"(Preset: {role.Name})";
-						enabled = roleCategory.Enabled;
-					}
-					
-					logger.LogInformation("  {Id,30} {Enabled,3} {Scope,-9} {Name}", syncCategory.UniqueId, enabled ? "Yes" : "No", scope, syncCategory.Name);
-				}
-
-				if (globalFilter.View.Count > 0)
-				{
-					logger.LogInformation("");
-					logger.LogInformation("Global View:");
-					foreach (string line in globalFilter.View)
-					{
-						logger.LogInformation("  {Line}", line);
-					}
-				}
-				if (workspaceFilter.View.Count > 0)
-				{
-					logger.LogInformation("");
-					logger.LogInformation("Workspace View:");
-					foreach (string line in workspaceFilter.View)
-					{
-						logger.LogInformation("  {Line}", line);
-					}
-				}
-				if (role != null && role.Views.Count > 0)
-				{
-					logger.LogInformation("");
-					logger.LogInformation("Preset View:");
-					foreach (string line in role.Views)
-					{
-						logger.LogInformation("  {Line}", line);
-					}
-				}
-
-				string[] filter = ReadSyncFilter(workspaceSettings, context.UserSettings, projectConfig, workspaceState.Current.ProjectIdentifier);
-
-				logger.LogInformation("");
-				logger.LogInformation("Combined view:");
-				foreach (string filterLine in filter)
-				{
-					logger.LogInformation("  {FilterLine}", filterLine);
-				}
-			}
-
-			static void ApplyCommandOptions(FilterSettings settings, FilterCommandOptions commandOptions, IEnumerable<WorkspaceSyncCategory> syncCategories, ILogger logger)
-			{
-				if (commandOptions.Reset)
-				{
-					logger.LogInformation("Resetting settings...");
-					settings.Reset();
-				}
-
-				HashSet<Guid> includeCategories = new HashSet<Guid>(commandOptions.Include.Select(x => GetCategoryId(x, syncCategories)));
-				HashSet<Guid> excludeCategories = new HashSet<Guid>(commandOptions.Exclude.Select(x => GetCategoryId(x, syncCategories)));
-
-				Guid id = includeCategories.FirstOrDefault(x => excludeCategories.Contains(x));
-				if (id != Guid.Empty)
-				{
-					throw new UserErrorException("Category {Id} cannot be both included and excluded", id);
-				}
-
-				includeCategories.ExceptWith(settings.IncludeCategories);
-				settings.IncludeCategories.AddRange(includeCategories);
-
-				excludeCategories.ExceptWith(settings.ExcludeCategories);
-				settings.ExcludeCategories.AddRange(excludeCategories);
-
-				if (commandOptions.View != null)
-				{
-					settings.View.Clear();
-					settings.View.AddRange(commandOptions.View);
-				}
-				if (commandOptions.RemoveView.Count > 0)
-				{
-					HashSet<string> viewRemove = new HashSet<string>(commandOptions.RemoveView, StringComparer.OrdinalIgnoreCase);
-					settings.View.RemoveAll(x => viewRemove.Contains(x));
-				}
-				if (commandOptions.AddView.Count > 0)
-				{
-					HashSet<string> viewLines = new HashSet<string>(settings.View, StringComparer.OrdinalIgnoreCase);
-					settings.View.AddRange(commandOptions.AddView.Where(x => !viewLines.Contains(x)));
-				}
-
-				settings.AllProjects = commandOptions.AllProjects ?? settings.AllProjects;
-				settings.LocalProjects = commandOptions.LocalProjects ?? settings.LocalProjects;
-				settings.AllProjectsInSln = commandOptions.AllProjectsInSln ?? settings.AllProjectsInSln;
-				settings.UprojectSpecificSln = commandOptions.UprojectSpecificSln ?? settings.UprojectSpecificSln;
-			}
-
-			static Guid GetCategoryId(string text, IEnumerable<WorkspaceSyncCategory> syncCategories)
-			{
-				Guid id;
-				if (Guid.TryParse(text, out id))
-				{
-					return id;
-				}
-
-				WorkspaceSyncCategory? category = syncCategories.FirstOrDefault(x => x.Name.Equals(text, StringComparison.OrdinalIgnoreCase));
-				if (category != null)
-				{
-					return category.UniqueId;
-				}
-
-				throw new UserErrorException("Unable to find category '{Category}'", text);
-			}
-		}
-
-		class BuildCommandOptions
-		{
-			[CommandLine("-List")]
-			public bool ListOnly { get; set; }
-		}
-
-		class BuildCommand : Command
-		{
-			public override async Task ExecuteAsync(CommandContext context)
-			{
-				ILogger logger = context.Logger;
-				context.Arguments.TryGetPositionalArgument(out string? target);
-
-				BuildCommandOptions options = new BuildCommandOptions();
-				context.Arguments.ApplyTo(options);
-				context.Arguments.CheckAllArgumentsUsed();
-
-				UserWorkspaceSettings settings = ReadRequiredUserWorkspaceSettings();
-				using IPerforceConnection perforceClient = await ConnectAsync(settings, context.LoggerFactory);
-				WorkspaceStateWrapper state = await ReadWorkspaceState(perforceClient, settings, context.UserSettings, logger);
-
-				ProjectInfo projectInfo = new ProjectInfo(settings.RootDir, state.Current);
-
-				UserProjectSettings projectSettings = context.UserSettings.FindOrAddProjectSettings(projectInfo, settings, logger);
-
-				ConfigFile projectConfig = await ConfigUtils.ReadProjectConfigFileAsync(perforceClient, projectInfo, logger, CancellationToken.None);
-				FileReference editorTarget = ConfigUtils.GetEditorTargetFile(projectInfo, projectConfig);
-
-				Dictionary<Guid, ConfigObject> buildStepObjects = ConfigUtils.GetDefaultBuildStepObjects(projectInfo, editorTarget.GetFileNameWithoutAnyExtensions(), EditorConfig, projectConfig, false);
-
-				if (context.GlobalSettings != null)
-				{
-					BuildStep.MergeBuildStepObjects(buildStepObjects, projectConfig.GetValues("Build.Step", Array.Empty<string>()).Select(x => new ConfigObject(x)));
-					BuildStep.MergeBuildStepObjects(buildStepObjects, projectSettings.BuildSteps);
-				}
-
-				if (options.ListOnly)
-				{
-					List<BuildStep> buildSteps = buildStepObjects.Values.Select(x => new BuildStep(x)).OrderBy(x => x.NormalSync).ToList();
-
-					int longestDescription = buildSteps.Max(b => String.IsNullOrWhiteSpace(b.Description) ? 0 : b.Description.Length);
-					string descriptionHeader = "Description";
-					string descriptionSpace = new string(Enumerable.Repeat(' ', longestDescription + 1 - descriptionHeader.Length).ToArray());
-					string descriptionDashes = new string(Enumerable.Repeat('-', longestDescription + 2).ToArray());
-					string format = $"  {{Id,-36}} | {{Name,-{longestDescription}}} | {{Type,-10}} | {{Enabled,-8}}";
-
-					logger.LogInformation("Available build steps:");
-					logger.LogInformation("");
-					logger.LogInformation("  Id                                   | {DescriptionHeader}{DescriptionSpace}| Type       | Enabled", descriptionHeader, descriptionSpace);
-					logger.LogInformation("  -------------------------------------|{DescriptionDashes}|------------|-----------------", descriptionDashes);
-					foreach (BuildStep buildStep in buildStepObjects.Values.Select(x => new BuildStep(x)).OrderBy(x => x.NormalSync))
-					{
-#pragma warning disable CA2254 // Template should be a static expression
-						logger.LogInformation(format, buildStep.UniqueId, buildStep.Description, buildStep.Type, buildStep.NormalSync);
-#pragma warning restore CA2254
-					}
-					return;
-				}
-
-				HashSet<Guid>? steps = new HashSet<Guid>();
-				if (target != null)
-				{
-					if (!Guid.TryParse(target, out Guid id))
-					{
-						logger.LogError("Unable to parse '{Target}' as a GUID. Pass -List to show all available build steps and their identifiers.", target);
-						return;
-					}
-
-					steps.Add(id);
-				}
-
-				// check that the tools are installed
-				List<Guid> uninstalledTools = new List<Guid>();
-				if (context.GlobalSettings != null)
-				{
-					using ToolUpdateMonitor? toolsMonitor = await GetToolsMonitor(logger, context, perforceClient);
-
-					if (toolsMonitor != null)
-					{
-						foreach (Guid step in steps)
-						{
-							if (buildStepObjects.TryGetValue(step, out ConfigObject? config))
-							{
-								BuildStep buildStep = new BuildStep(config);
-								if (CanRunStep(buildStep, toolsMonitor))
-								{
-									continue;
-								}
-
-								logger.LogWarning("Build step '{Description}' cannot be run as the tool is not installed.", buildStep.Description);
-								uninstalledTools.Add(step);
-							}
-						}
-					}
-				}
-
-				foreach (Guid uninstalledTool in uninstalledTools)
-				{
-					steps.Remove(uninstalledTool);
-				}
-
-				WorkspaceUpdateContext updateContext = new WorkspaceUpdateContext(
-					state.Current.CurrentChangeNumber,
-					WorkspaceUpdateOptions.Build,
-					BuildConfig.Development,
-					null,
-					projectSettings.BuildSteps, steps);
-
-				WorkspaceUpdate update = new WorkspaceUpdate(updateContext);
-				(WorkspaceUpdateResult result, string message) = await update.ExecuteAsync(perforceClient.Settings, projectInfo, state, context.Logger, CancellationToken.None);
-
-				if (result != WorkspaceUpdateResult.Success)
-				{
-					throw new UserErrorException("{Message}", message);
-				}
-			}
-
-			private static async Task<ToolUpdateMonitor?> GetToolsMonitor(ILogger logger, CommandContext context, IPerforceConnection perforceClient)
-			{
-				ToolsCommandCommandOptions options = context.Arguments.ApplyTo<ToolsCommandCommandOptions>(logger);
-				context.Arguments.CheckAllArgumentsUsed();
-
-				DirectoryReference dataFolder = DirectoryReference.Combine(DirectoryReference.GetSpecialFolder(Environment.SpecialFolder.LocalApplicationData)!, "UnrealGameSync");
-				DirectoryReference.CreateDirectory(dataFolder);
-
-				// create a temporary service provider for the tool update monitor
-				ServiceCollection services = new ServiceCollection();
-				services.AddSingleton<IAsyncDisposer, AsyncDisposer>();
-				services.AddSingleton(sp => TokenStoreFactory.CreateTokenStore());
-				services.AddSingleton<OidcTokenManager>();
-
-				LauncherSettings launcherSettings = new LauncherSettings();
-				launcherSettings.Read();
-
-				if (launcherSettings.HordeServer != null)
-				{
-					services.AddHorde(options =>
-					{
-						options.ServerUrl = new Uri(launcherSettings.HordeServer);
-						options.AllowAuthPrompt = false;
-					});
-				}
-
-				ServiceProvider serviceProvider = services.BuildServiceProvider();
-
-				ToolUpdateMonitor toolUpdateMonitor =
-					new ToolUpdateMonitor(perforceClient.Settings, dataFolder, context.GlobalSettings!, logger, serviceProvider);
-
-				// get the list of tools available
-				logger.LogInformation("Retrieving tools information, please wait.");
-				await toolUpdateMonitor.GetDataFromBackendAsync();
-
-				return toolUpdateMonitor;
-			}
-
-			private static bool CanRunStep(BuildStep step, ToolUpdateMonitor toolsMonitor)
-			{
-				if (step.ToolId != Guid.Empty)
-				{
-					string? toolName = toolsMonitor.GetToolName(step.ToolId);
-					if (toolName == null)
-					{
-						return false;
-					}
-
-					if (toolsMonitor.GetToolPath(toolName) == null)
-					{
-						return false;
-					}
-				}
-				return true;
-			}
-		}
-
-		class StatusCommand : Command
-		{
-			public override async Task ExecuteAsync(CommandContext context)
-			{
-				ILogger logger = context.Logger;
-				bool update = context.Arguments.HasOption("-Update");
-				context.Arguments.CheckAllArgumentsUsed();
-
-				UserWorkspaceSettings settings = ReadRequiredUserWorkspaceSettings();
-				logger.LogInformation("User: {UserName}", settings.UserName);
-				logger.LogInformation("Server: {ServerAndPort}", settings.ServerAndPort);
-				logger.LogInformation("Project: {ClientProjectPath}", settings.ClientProjectPath);
-
-				using IPerforceConnection perforceClient = await ConnectAsync(settings, context.LoggerFactory);
-
-				WorkspaceStateWrapper state = await ReadWorkspaceState(perforceClient, settings, context.UserSettings, logger);
-				if (update)
-				{
-					ProjectInfo newProjectInfo = await ProjectInfo.CreateAsync(perforceClient, settings, CancellationToken.None);
-					state.Modify(x => x.UpdateCachedProjectInfo(newProjectInfo, settings.LastModifiedTimeUtc));
-				}
-
-				string streamOrBranchName = state.Current.StreamName ?? settings.BranchPath.TrimStart('/');
-				if (state.Current.LastSyncResultMessage == null)
-				{
-					logger.LogInformation("Not currently synced to {Stream}", streamOrBranchName);
-				}
-				else if (state.Current.LastSyncResult == WorkspaceUpdateResult.Success)
-				{
-					logger.LogInformation("Synced to {Stream} CL {Change}", streamOrBranchName, state.Current.LastSyncChangeNumber);
-				}
-				else
-				{
-					logger.LogWarning("Last sync to {Stream} CL {Change} failed: {Result}", streamOrBranchName, state.Current.LastSyncChangeNumber, state.Current.LastSyncResultMessage);
-				}
-			}
-		}
-
-		class LoginCommand : Command
-		{
-			public override async Task ExecuteAsync(CommandContext context)
-			{
-				ILogger logger = context.Logger;
-
-				// Get the positional argument indicating the file to look for
-				if (!context.Arguments.TryGetPositionalArgument(out string? providerIdentifier))
-				{
-					throw new UserErrorException("Missing provider identifier to login to.");
-				}
-				context.Arguments.CheckAllArgumentsUsed();
-
-				UserWorkspaceSettings settings = ReadRequiredUserWorkspaceSettings();
-
-				// Find the valid config file paths
-				DirectoryInfo engineDir = DirectoryReference.Combine(settings.RootDir, "Engine").ToDirectoryInfo();
-				DirectoryInfo gameDir = new DirectoryInfo(settings.ProjectPath);
-				using ITokenStore tokenStore = TokenStoreFactory.CreateTokenStore();
-				IConfiguration providerConfiguration = ProviderConfigurationFactory.ReadConfiguration(engineDir, gameDir);
-				OidcTokenManager oidcTokenManager = OidcTokenManager.CreateTokenManager(providerConfiguration, tokenStore, new List<string>() { providerIdentifier });
-				OidcTokenInfo result = await oidcTokenManager.LoginAsync(providerIdentifier);
-
-				logger.LogInformation("Logged in to provider {ProviderIdentifier}", providerIdentifier);
-			}
-		}
-
-		class SwitchCommandOptions
-		{
-			[CommandLine("-Force")]
-			public bool Force { get; set; }
-		}
-
-		class SwitchCommand : Command
-		{
-			public override async Task ExecuteAsync(CommandContext context)
-			{
-				// Get the positional argument indicating the file to look for
-				string? targetName;
-				if (!context.Arguments.TryGetPositionalArgument(out targetName))
-				{
-					throw new UserErrorException("Missing stream or project name to switch to.");
-				}
-
-				// Finish argument parsing
-				SwitchCommandOptions options = new SwitchCommandOptions();
-				context.Arguments.ApplyTo(options);
-				context.Arguments.CheckAllArgumentsUsed();
-
-				if (targetName.StartsWith("//", StringComparison.Ordinal))
-				{
-					options.Force = true;
-				}
-
-				// Get a connection to the client for this workspace
-				UserWorkspaceSettings settings = ReadRequiredUserWorkspaceSettings();
-				using IPerforceConnection perforceClient = await ConnectAsync(settings, context.LoggerFactory);
-
-				// Check whether we're switching stream or project
-				if (targetName.StartsWith("//", StringComparison.Ordinal))
-				{
-					await SwitchStreamAsync(perforceClient, targetName, options.Force, context.Logger);
-				}
-				else
-				{
-					await SwitchProjectAsync(perforceClient, settings, targetName, context.Logger);
-				}
-			}
-
-			public static async Task SwitchStreamAsync(IPerforceConnection perforceClient, string streamName, bool force, ILogger logger)
-			{
-				if (!force && await perforceClient.OpenedAsync(OpenedOptions.None, FileSpecList.Any).AnyAsync())
-				{
-					throw new UserErrorException("Client {ClientName} has files opened. Use -Force to switch anyway.", perforceClient.Settings.ClientName!);
-				}
-
-				await perforceClient.SwitchClientToStreamAsync(streamName, SwitchClientOptions.IgnoreOpenFiles);
-
-				logger.LogInformation("Switched to stream {StreamName}", streamName);
-			}
-
-			public static async Task SwitchProjectAsync(IPerforceConnection perforceClient, UserWorkspaceSettings settings, string projectName, ILogger logger)
-			{
-				settings.ProjectPath = await FindProjectPathAsync(perforceClient, settings.ClientName, settings.BranchPath, projectName);
-				settings.Save(logger);
-				logger.LogInformation("Switched to project {ProjectPath}", settings.ClientProjectPath);
-			}
-		}
-
-		class VersionCommand : Command
-		{
-			public override Task ExecuteAsync(CommandContext context)
-			{
-				ILogger logger = context.Logger;
-
-				string version = GetVersion();
-				logger.LogInformation("UnrealGameSync {Version}", version);
-
-				return Task.CompletedTask;
-			}
-		}
-
-		class InstallCommand : Command
-		{
-			public override async Task ExecuteAsync(CommandContext context)
-			{
-				await UpdateInstallAsync(true, context.Logger);
-			}
-		}
-
-		class UninstallCommand : Command
-		{
-			public override async Task ExecuteAsync(CommandContext context)
-			{
-				await UpdateInstallAsync(false, context.Logger);
-			}
-		}
-
-		static async Task UpdateInstallAsync(bool install, ILogger logger)
-		{
-			DirectoryReference? installDir = GetInstallFolder();
-			if (installDir != null)
-			{
-				UpdateInstalledFiles(install, installDir, logger);
-			}
-			else
-			{
-				installDir = new FileReference(Assembly.GetExecutingAssembly().GetOriginalLocation()).Directory;
-			}
-
-			if (OperatingSystem.IsWindows())
-			{
-				const string EnvVarName = "PATH";
-
-				string? pathVar = Environment.GetEnvironmentVariable(EnvVarName, EnvironmentVariableTarget.User);
-				pathVar ??= String.Empty;
-
-				List<string> paths = new List<string>(pathVar.Split(Path.PathSeparator, StringSplitOptions.RemoveEmptyEntries));
-
-				int changes = paths.RemoveAll(x => x.Equals(installDir.FullName, StringComparison.OrdinalIgnoreCase));
-				if (install)
-				{
-					paths.Add(installDir.FullName);
-					changes++;
-				}
-				if (changes > 0)
-				{
-					pathVar = String.Join(Path.PathSeparator, paths);
-					Environment.SetEnvironmentVariable(EnvVarName, pathVar, EnvironmentVariableTarget.User);
-				}
-
-				if (install)
-				{
-					logger.LogInformation("Added {Path} to PATH environment variable", installDir);
-				}
-				else
-				{
-					logger.LogInformation("Removed {Path} from PATH environment variable", installDir);
-				}
-			}
-			else if (OperatingSystem.IsMacOS())
-			{
-				DirectoryReference? userDir = DirectoryReference.GetSpecialFolder(Environment.SpecialFolder.UserProfile);
-				if (userDir != null)
-				{
-					FileReference configFile = FileReference.Combine(userDir, ".zshrc");
-					await UpdateAliasAsync(configFile, install, installDir, logger);
-				}
-			}
-			else if (OperatingSystem.IsLinux())
-			{
-				DirectoryReference? userDir = DirectoryReference.GetSpecialFolder(Environment.SpecialFolder.UserProfile);
-				if (userDir != null)
-				{
-					FileReference configFile = FileReference.Combine(userDir, ".bashrc");
-					await UpdateAliasAsync(configFile, install, installDir, logger);
-				}
-			}
-		}
-
-		static DirectoryReference? GetInstallFolder()
-		{
-			if (OperatingSystem.IsWindows())
-			{
-				DirectoryReference? installDir = DirectoryReference.GetSpecialFolder(Environment.SpecialFolder.LocalApplicationData);
-				if (installDir != null)
-				{
-					return DirectoryReference.Combine(installDir, "Epic Games", "UgsCmd");
-				}
-			}
-			return null;
-		}
-
-		static void UpdateInstalledFiles(bool install, DirectoryReference installDir, ILogger logger)
-		{
-			FileReference assemblyFile = new FileReference(Assembly.GetExecutingAssembly().GetOriginalLocation());
-			DirectoryReference sourceDir = assemblyFile.Directory;
-
-			DirectoryReference TempDir = DirectoryReference.Combine(installDir.ParentDirectory!, "~" + installDir.GetDirectoryName());
-			if (DirectoryReference.Exists(TempDir))
-			{
-				DirectoryReference.Delete(TempDir, true);
-			}
-
-			if (DirectoryReference.Exists(installDir))
-			{
-				logger.LogInformation("Removing application files from {Dir}", installDir);
-
-				Directory.Move(installDir.FullName, TempDir.FullName);
-				DirectoryReference.Delete(TempDir, true);
-			}
-
-			if (install)
-			{
-				logger.LogInformation("Copying application files to {Dir}", installDir);
-
-				DirectoryReference.CreateDirectory(TempDir);
-				foreach (FileReference SourceFile in DirectoryReference.EnumerateFiles(sourceDir, "*", SearchOption.AllDirectories))
-				{
-					FileReference TargetFile = FileReference.Combine(TempDir, SourceFile.MakeRelativeTo(sourceDir));
-					DirectoryReference.CreateDirectory(TargetFile.Directory);
-					FileReference.Copy(SourceFile, TargetFile, true);
-				}
-				Directory.Move(TempDir.FullName, installDir.FullName);
-			}
-		}
-
-		static async Task UpdateAliasAsync(FileReference configFile, bool install, DirectoryReference installDir, ILogger logger)
-		{
-			List<string> lines = new List<string>();
-			if (FileReference.Exists(configFile))
-			{
-				lines.AddRange(await FileReference.ReadAllLinesAsync(configFile));
-				lines.RemoveAll(x => Regex.IsMatch(x, @"^\s*alias\s+ugs\s*="));
-			}
-			if (install)
-			{
-				lines.Add($"alias ugs={FileReference.Combine(installDir, "ugs")}");
-			}
-
-			await FileReference.WriteAllLinesAsync(configFile, lines);
-
-			if (install)
-			{
-				logger.LogInformation("Added 'ugs' alias to {ConfigFile}", configFile);
-			}
-			else
-			{
-				logger.LogInformation("Removed 'ugs' alias from {ConfigFile}", configFile);
-			}
-		}
-
-		class UpgradeCommandOptions
-		{
-			[CommandLine("-Check")]
-			public bool Check { get; set; }
-
-			[CommandLine("-Force")]
-			public bool Force { get; set; }
-		}
-
-		class UpgradeCommand : Command
-		{
-			public override async Task ExecuteAsync(CommandContext context)
-			{
-				ILogger logger = context.Logger;
-
-				UpgradeCommandOptions options = new UpgradeCommandOptions();
-				context.Arguments.ApplyTo(options);
-				string? targetDirStr = context.Arguments.GetStringOrDefault("-TargetDir=", null);
-				context.Arguments.CheckAllArgumentsUsed(logger);
-
-				string currentVersion = GetVersion();
-
-				string? latestVersion = await GetLatestVersionAsync(logger, CancellationToken.None);
-				if (latestVersion == null)
-				{
-					return;
-				}
-				if (latestVersion.Equals(currentVersion, StringComparison.OrdinalIgnoreCase) && !options.Force)
-				{
-					logger.LogInformation("You are running the latest version ({Version})", currentVersion);
-					return;
-				}
-				if (options.Check)
-				{
-					logger.LogWarning("A newer version of UGS is available ({NewVersion})", latestVersion);
-					return;
-				}
-
-				using (HttpClient httpClient = new HttpClient())
-				{
-					Uri baseUrl = new Uri(DeploymentSettings.Instance.HordeUrl ?? String.Empty);
-
-					DirectoryReference currentDir = new FileReference(Assembly.GetExecutingAssembly().Location).Directory;
-
-					DirectoryReference targetDir = (targetDirStr == null) ? currentDir : DirectoryReference.Combine(currentDir, targetDirStr);
-					DirectoryReference.CreateDirectory(targetDir);
-
-					FileReference tempFile = FileReference.Combine(targetDir, "update.zip");
-					using (Stream requestStream = await httpClient.GetStreamAsync(new Uri(baseUrl, $"api/v1/tools/{GetUpgradeToolName()}?action=download")))
-					{
-						using (Stream tempFileStream = FileReference.Open(tempFile, FileMode.Create, FileAccess.Write, FileShare.None))
-						{
-							await requestStream.CopyToAsync(tempFileStream);
-						}
-					}
-
-					using (FileStream stream = FileReference.Open(tempFile, FileMode.Open, FileAccess.Read, FileShare.Read))
-					{
-						using (ZipArchive archive = new ZipArchive(stream, ZipArchiveMode.Read, true))
-						{
-							foreach (ZipArchiveEntry entry in archive.Entries)
-							{
-								FileReference targetFile = FileReference.Combine(targetDir, entry.Name);
-								if (!targetFile.IsUnderDirectory(targetDir))
-								{
-									throw new InvalidDataException("Attempt to extract file outside source directory");
-								}
-								entry.ExtractToFile_CrossPlatform(targetFile.FullName, true);
-							}
-						}
-					}
-				}
-			}
-		}
-
-		class ToolsCommandCommandOptions
-		{
-			[CommandLine("-Install=", ListSeparator = ',')]
-			public List<string> Install { get; set; }  = new List<string>();
-
-			[CommandLine("-Uninstall=", ListSeparator = ',')]
-			public List<string> Uninstall { get; set; } = new List<string>();
-
-			[CommandLine("-List")]
-			public bool List { get; set; }
-
-			[CommandLine("-Update")]
-			public bool Update { get; set; }
-		}
-
-		class ToolsCommand : Command
-		{
-			public override async Task ExecuteAsync(CommandContext context)
-			{
-				ILogger logger = context.Logger;
-				if (!RuntimeInformation.IsOSPlatform(OSPlatform.Windows))
-				{
-					logger.LogInformation("This command is only available for Windows.");
-					return;
-				}
-	
-				context.Arguments.TryGetPositionalArgument(out string? toolId);
-
-				ToolsCommandCommandOptions options = context.Arguments.ApplyTo<ToolsCommandCommandOptions>(logger);
-				context.Arguments.CheckAllArgumentsUsed();
-
-				UserWorkspaceSettings settings = ReadRequiredUserWorkspaceSettings();
-				using IPerforceConnection perforceClient = await ConnectAsync(settings, context.LoggerFactory);
-
-				if (context.GlobalSettings == null)
-				{
-					logger.LogError("Could not retrieve user settings.");
-					return;
-				}
-
-				DirectoryReference dataFolder = DirectoryReference.Combine(DirectoryReference.GetSpecialFolder(Environment.SpecialFolder.LocalApplicationData)!, "UnrealGameSync");
-				DirectoryReference.CreateDirectory(dataFolder);
-
-				// create a temporary service provider for the tool update monitor
-				ServiceCollection services = new ServiceCollection();
-				services.AddSingleton<IAsyncDisposer, AsyncDisposer>();
-				services.AddSingleton(sp => TokenStoreFactory.CreateTokenStore());
-				services.AddSingleton<OidcTokenManager>();
-
-				LauncherSettings launcherSettings = new LauncherSettings();
-				launcherSettings.Read();
-
-				if (launcherSettings.HordeServer != null)
-				{
-					services.AddHorde(options =>
-					{
-						options.ServerUrl = new Uri(launcherSettings.HordeServer);
-						options.AllowAuthPrompt = false;
-					});
-				}
-
-				ServiceProvider serviceProvider = services.BuildServiceProvider();
-
-				using ToolUpdateMonitor toolUpdateMonitor =
-					new ToolUpdateMonitor(perforceClient.Settings, dataFolder, context.GlobalSettings, logger, serviceProvider);
-
-				// get the list of tools available
-				logger.LogInformation("Retrieving tools information, please wait.");
-				await toolUpdateMonitor.GetDataFromBackendAsync();
-
-				// only list the tools
-				if (options.List)
-				{
-					logger.LogInformation("Available Tools:");
-					logger.LogInformation("");
-					logger.LogInformation("  Id                                   | Name                                 | Enabled         | Description                                      ");
-					logger.LogInformation("  -------------------------------------| -------------------------------------| ----------------| -------------------------------------------------");
-
-					IReadOnlyList<ToolInfo> enabled = toolUpdateMonitor.GetEnabledTools();
-
-					foreach (ToolInfo toolInfo in toolUpdateMonitor.GetTools().OrderBy(t => enabled.All( e => t.Id != e.Id)))
-					{
-						logger.LogInformation("  {Id,-36} | {Name,-36} | {Enabled,-16}| {Description,-48} ", toolInfo.Id, toolInfo.Name, enabled.Any( t => t.Id == toolInfo.Id), toolInfo.Description);
-					}
-				}
-
-				// update the list of enabled tools and update (this will update all the tools)
-				if (options.Install.Any() || options.Uninstall.Any())
-				{
-					IReadOnlyList<ToolInfo> enabled = toolUpdateMonitor.GetEnabledTools();
-					HashSet<Guid> newEnabledTools = new HashSet<Guid>(enabled.Select(t => t.Id));
-					
-					foreach (string id in options.Install.Select(id => id.Trim()))
-					{
-						if (Guid.TryParse(id, out Guid guid))
-						{
-							foreach (ToolInfo toolInfo in toolUpdateMonitor.GetTools())
-							{
-								if (toolInfo.Id == guid)
-								{
-									logger.LogInformation("Adding tool '{Id}' - '{Name}' to the list of enabled tools", id, toolInfo.Name);
-									newEnabledTools.Add(guid);
-									break;
-								}
-							}
-
-							if (!newEnabledTools.Contains(guid))
-							{
-								logger.LogError("Could not install tool '{Id}'", id);
-							}
-						}
-						else
-						{
-							logger.LogError("'{Id}' is not a valid guid", id);
-						}
-					}
-
-					foreach (string id in options.Uninstall.Select(id => id.Trim()))
-					{
-						if (Guid.TryParse(id, out Guid guid))
-						{
-							foreach (ToolInfo toolInfo in toolUpdateMonitor.GetTools())
-							{
-								if (toolInfo.Id == guid)
-								{
-									logger.LogInformation("Removing tool '{Id}' - '{Name}' from the list of enabled tools", id, toolInfo.Name);
-									newEnabledTools.Remove(guid);
-									break;
-								}
-							}
-
-							if (newEnabledTools.Contains(guid))
-							{
-								logger.LogError("Could not uninstall tool '{Id}'", id);
-							}
-						}
-						else
-						{
-							logger.LogError("'{Id}' is not a valid guid", id);
-						}
-					}
-
-					if (!newEnabledTools.SequenceEqual(context.GlobalSettings.EnabledTools))
-					{
-						context.GlobalSettings.EnabledTools.Clear();
-						context.GlobalSettings.EnabledTools.UnionWith(newEnabledTools);
-						context.GlobalSettings.Save(logger);
-
-						logger.LogInformation("Updating tools");
-						await toolUpdateMonitor.UpdateToolsAsync();
-					}
-				}
-
-				// update the all the enabled tools
-				if (options.Update)
-				{
-					await toolUpdateMonitor.UpdateToolsAsync();
-				}
-			}
-		}
-
-		class SettingsCommandOptions
-		{
-			[CommandLine("-Global=")]
-			public string Global { get; set; } = String.Empty;
-
-			[CommandLine("-User=")]
-			public string User { get; set; } = String.Empty;
-		}
-
-		class SettingsCommand : Command
-		{
-			public override async Task ExecuteAsync(CommandContext context)
-			{
-				ILogger logger = context.Logger;
-
-				SettingsCommandOptions options = context.Arguments.ApplyTo<SettingsCommandOptions>(logger);
-				context.Arguments.CheckAllArgumentsUsed();
-
-				if (!String.IsNullOrWhiteSpace(options.Global))
-				{
-					if (context.GlobalSettings != null)
-					{
-						logger.LogInformation("Writing global settings to {Global}", options.Global);
-						await WriteSettings(options.Global, context.GlobalSettings);
-					}
-					else if (context.UserSettings != null)
-					{
-						logger.LogInformation("Writing global settings to {Global}", options.Global);
-						await WriteSettings(options.Global, context.UserSettings);
-					}
-				}
-
-				if (!String.IsNullOrWhiteSpace(options.User))
-				{
-					UserWorkspaceSettings userWorkspaceSettings = ReadRequiredUserWorkspaceSettings();
-					logger.LogInformation("Writing user settings to {User}", options.User);
-					await WriteSettings(options.User, userWorkspaceSettings);
-				}
-			}
-
-			private static async Task WriteSettings(string path, object? settings)
-			{
-				if (String.IsNullOrWhiteSpace(path) || settings == null)
-				{
-					return;
-				}
-
-				string trimmed = path.Trim('\"')
-					.Replace(Path.AltDirectorySeparatorChar, Path.DirectorySeparatorChar);
-				if (!Path.IsPathFullyQualified(trimmed))
-				{
-					return;
-				}
-
-				string? directory = Path.GetDirectoryName(trimmed);
-				if (!String.IsNullOrWhiteSpace(directory))
-				{
-					Directory.CreateDirectory(directory);
-				}
-
-				JsonSerializerOptions serializerOptions = new JsonSerializerOptions()
-				{
-					WriteIndented = true
-				};
-				string json = JsonSerializer.Serialize(settings, serializerOptions);
-
-				await File.WriteAllTextAsync(trimmed, json);
-			}
-		}
 	}
 }
diff --git a/Engine/Source/Programs/UnrealGameSync/UnrealGameSyncCmd/Utils/InstallUtils.cs b/Engine/Source/Programs/UnrealGameSync/UnrealGameSyncCmd/Utils/InstallUtils.cs
new file mode 100644
--- /dev/null
+++ b/Engine/Source/Programs/UnrealGameSync/UnrealGameSyncCmd/Utils/InstallUtils.cs
@@ -0,0 +1,152 @@
+﻿// Copyright Epic Games, Inc. All Rights Reserved.
+
+using System;
+using System.Collections.Generic;
+using System.IO;
+using System.Reflection;
+using System.Text.RegularExpressions;
+using System.Threading.Tasks;
+
+using Microsoft.Extensions.Logging;
+
+using EpicGames.Core;
+
+namespace UnrealGameSyncCmd.Utils
+{
+	internal static class InstallUtils
+	{
+		internal static async Task UpdateInstallAsync(bool install, ILogger logger)
+		{
+			DirectoryReference? installDir = GetInstallFolder();
+			if (installDir != null)
+			{
+				UpdateInstalledFiles(install, installDir, logger);
+			}
+			else
+			{
+				installDir = new FileReference(Assembly.GetExecutingAssembly().GetOriginalLocation()).Directory;
+			}
+
+			if (OperatingSystem.IsWindows())
+			{
+				const string EnvVarName = "PATH";
+
+				string? pathVar = Environment.GetEnvironmentVariable(EnvVarName, EnvironmentVariableTarget.User);
+				pathVar ??= String.Empty;
+
+				List<string> paths = new List<string>(pathVar.Split(Path.PathSeparator, StringSplitOptions.RemoveEmptyEntries));
+
+				int changes = paths.RemoveAll(x => x.Equals(installDir.FullName, StringComparison.OrdinalIgnoreCase));
+				if (install)
+				{
+					paths.Add(installDir.FullName);
+					changes++;
+				}
+				if (changes > 0)
+				{
+					pathVar = String.Join(Path.PathSeparator, paths);
+					Environment.SetEnvironmentVariable(EnvVarName, pathVar, EnvironmentVariableTarget.User);
+				}
+
+				if (install)
+				{
+					logger.LogInformation("Added {Path} to PATH environment variable", installDir);
+				}
+				else
+				{
+					logger.LogInformation("Removed {Path} from PATH environment variable", installDir);
+				}
+			}
+			else if (OperatingSystem.IsMacOS())
+			{
+				DirectoryReference? userDir = DirectoryReference.GetSpecialFolder(Environment.SpecialFolder.UserProfile);
+				if (userDir != null)
+				{
+					FileReference configFile = FileReference.Combine(userDir, ".zshrc");
+					await UpdateAliasAsync(configFile, install, installDir, logger);
+				}
+			}
+			else if (OperatingSystem.IsLinux())
+			{
+				DirectoryReference? userDir = DirectoryReference.GetSpecialFolder(Environment.SpecialFolder.UserProfile);
+				if (userDir != null)
+				{
+					FileReference configFile = FileReference.Combine(userDir, ".bashrc");
+					await UpdateAliasAsync(configFile, install, installDir, logger);
+				}
+			}
+		}
+
+		internal static DirectoryReference? GetInstallFolder()
+		{
+			if (OperatingSystem.IsWindows())
+			{
+				DirectoryReference? installDir = DirectoryReference.GetSpecialFolder(Environment.SpecialFolder.LocalApplicationData);
+				if (installDir != null)
+				{
+					return DirectoryReference.Combine(installDir, "Epic Games", "UgsCmd");
+				}
+			}
+			return null;
+		}
+
+		internal static void UpdateInstalledFiles(bool install, DirectoryReference installDir, ILogger logger)
+		{
+			FileReference assemblyFile = new FileReference(Assembly.GetExecutingAssembly().GetOriginalLocation());
+			DirectoryReference sourceDir = assemblyFile.Directory;
+
+			DirectoryReference TempDir = DirectoryReference.Combine(installDir.ParentDirectory!, "~" + installDir.GetDirectoryName());
+			if (DirectoryReference.Exists(TempDir))
+			{
+				DirectoryReference.Delete(TempDir, true);
+			}
+
+			if (DirectoryReference.Exists(installDir))
+			{
+				logger.LogInformation("Removing application files from {Dir}", installDir);
+
+				Directory.Move(installDir.FullName, TempDir.FullName);
+				DirectoryReference.Delete(TempDir, true);
+			}
+
+			if (install)
+			{
+				logger.LogInformation("Copying application files to {Dir}", installDir);
+
+				DirectoryReference.CreateDirectory(TempDir);
+				foreach (FileReference SourceFile in DirectoryReference.EnumerateFiles(sourceDir, "*", SearchOption.AllDirectories))
+				{
+					FileReference TargetFile = FileReference.Combine(TempDir, SourceFile.MakeRelativeTo(sourceDir));
+					DirectoryReference.CreateDirectory(TargetFile.Directory);
+					FileReference.Copy(SourceFile, TargetFile, true);
+				}
+				Directory.Move(TempDir.FullName, installDir.FullName);
+			}
+		}
+
+		internal static async Task UpdateAliasAsync(FileReference configFile, bool install, DirectoryReference installDir, ILogger logger)
+		{
+			List<string> lines = new List<string>();
+			if (FileReference.Exists(configFile))
+			{
+				lines.AddRange(await FileReference.ReadAllLinesAsync(configFile));
+				lines.RemoveAll(x => Regex.IsMatch(x, @"^\s*alias\s+ugs\s*="));
+			}
+			if (install)
+			{
+				lines.Add($"alias ugs={FileReference.Combine(installDir, "ugs")}");
+			}
+
+			await FileReference.WriteAllLinesAsync(configFile, lines);
+
+			if (install)
+			{
+				logger.LogInformation("Added 'ugs' alias to {ConfigFile}", configFile);
+			}
+			else
+			{
+				logger.LogInformation("Removed 'ugs' alias from {ConfigFile}", configFile);
+			}
+		}
+	}
+}
diff --git a/Engine/Source/Programs/UnrealGameSync/UnrealGameSyncCmd/Utils/PerforceConnectionUtils.cs b/Engine/Source/Programs/UnrealGameSync/UnrealGameSyncCmd/Utils/PerforceConnectionUtils.cs
new file mode 100644
--- /dev/null
+++ b/Engine/Source/Programs/UnrealGameSync/UnrealGameSyncCmd/Utils/PerforceConnectionUtils.cs
@@ -0,0 +1,36 @@
+﻿// Copyright Epic Games, Inc. All Rights Reserved.
+
+using System;
+
+using System.Threading.Tasks;
+using EpicGames.Perforce;
+using Microsoft.Extensions.Logging;
+using UnrealGameSync;
+
+namespace UnrealGameSyncCmd.Utils
+{
+	public static class PerforceConnectionUtils
+	{
+		public static Task<IPerforceConnection> ConnectAsync(string? serverAndPort, string? userName, string? clientName, ILoggerFactory loggerFactory)
+		{
+			PerforceSettings settings = new PerforceSettings(PerforceSettings.Default);
+			settings.ClientName = clientName;
+			settings.PreferNativeClient = true;
+			if (!String.IsNullOrEmpty(serverAndPort))
+			{
+				settings.ServerAndPort = serverAndPort;
+			}
+			if (!String.IsNullOrEmpty(userName))
+			{
+				settings.UserName = userName;
+			}
+
+			return PerforceConnection.CreateAsync(settings, loggerFactory.CreateLogger("Perforce"));
+		}
+
+		public static Task<IPerforceConnection> ConnectAsync(UserWorkspaceSettings settings, ILoggerFactory loggerFactory)
+		{
+			return ConnectAsync(settings.ServerAndPort, settings.UserName, settings.ClientName, loggerFactory);
+		}
+	}
+}
diff --git a/Engine/Source/Programs/UnrealGameSync/UnrealGameSyncCmd/Utils/ProjectUtils.cs b/Engine/Source/Programs/UnrealGameSync/UnrealGameSyncCmd/Utils/ProjectUtils.cs
new file mode 100644
--- /dev/null
+++ b/Engine/Source/Programs/UnrealGameSync/UnrealGameSyncCmd/Utils/ProjectUtils.cs
@@ -0,0 +1,54 @@
+﻿// Copyright Epic Games, Inc. All Rights Reserved.
+
+using System;
+using System.Collections.Generic;
+using System.Linq;
+using System.Threading.Tasks;
+
+using EpicGames.Perforce;
+using UnrealGameSyncCmd.Exceptions;
+
+namespace UnrealGameSyncCmd.Utils
+{
+	internal static class ProjectUtils
+	{
+		internal const string UProjectExtension = ".uproject";
+		internal const string UProjectDirsExtension = ".uprojectdirs";
+
+		internal static async Task<string> FindProjectPathAsync(IPerforceConnection perforce, string clientName, string branchPath, string? projectName)
+		{
+			using IPerforceConnection perforceClient = await PerforceConnection.CreateAsync(new PerforceSettings(perforce.Settings) { ClientName = clientName }, perforce.Logger);
+
+			// Find or validate the selected project
+			string searchPath;
+			if (projectName == null)
+			{
+				searchPath = $"//{clientName}{branchPath}/*{UProjectDirsExtension}";
+			}
+			else if (projectName.Contains('.', StringComparison.Ordinal))
+			{
+				searchPath = $"//{clientName}{branchPath}/{projectName.TrimStart('/')}";
+			}
+			else
+			{
+				searchPath = $"//{clientName}{branchPath}/.../{projectName}{UProjectExtension}";
+			}
+
+			List<FStatRecord> projectFileRecords = await perforceClient.FStatAsync(FStatOptions.ClientFileInPerforceSyntax, searchPath).ToListAsync();
+			projectFileRecords.RemoveAll(x => x.HeadAction == FileAction.Delete || x.HeadAction == FileAction.MoveDelete);
+			projectFileRecords.RemoveAll(x => !x.IsMapped);
+
+			List<string> paths = projectFileRecords.Select(x => PerforceUtils.GetClientRelativePath(x.ClientFile!)).Distinct(StringComparer.Ordinal).ToList();
+			if (paths.Count == 0)
+			{
+				throw new UserErrorException("No project file found matching {SearchPath}", searchPath);
+			}
+			if (paths.Count > 1)
+			{
+				throw new UserErrorException("Multiple projects found matching {SearchPath}: {Paths}", searchPath, String.Join(", ", paths));
+			}
+
+			return "/" + paths[0];
+		}
+	}
+}
diff --git a/Engine/Source/Programs/UnrealGameSync/UnrealGameSyncCmd/Utils/SyncFilterUtils.cs b/Engine/Source/Programs/UnrealGameSync/UnrealGameSyncCmd/Utils/SyncFilterUtils.cs
new file mode 100644
--- /dev/null
+++ b/Engine/Source/Programs/UnrealGameSync/UnrealGameSyncCmd/Utils/SyncFilterUtils.cs
@@ -0,0 +1,35 @@
+﻿// Copyright Epic Games, Inc. All Rights Reserved.
+
+using System;
+using System.Collections.Generic;
+using UnrealGameSync;
+
+namespace UnrealGameSyncCmd.Utils
+{
+	internal static class SyncFilterUtils
+	{
+		internal static string[] ReadSyncFilter(UserWorkspaceSettings workspaceSettings, GlobalSettingsFile userSettings, ConfigFile projectConfig, string projectIdentifier)
+		{
+			Dictionary<Guid, WorkspaceSyncCategory> syncCategories = ConfigUtils.GetSyncCategories(projectConfig);
+
+			// check if any category is from the role
+			IDictionary<string, Preset> roles = ConfigUtils.GetPresets(projectConfig, projectIdentifier);
+			if (roles.TryGetValue(workspaceSettings.Preset, out Preset? role))
+			{
+				foreach (RoleCategory roleCategory in role.Categories.Values)
+				{
+					if (syncCategories.TryGetValue(roleCategory.Id, out WorkspaceSyncCategory? category))
+					{
+						category.Enable = roleCategory.Enabled;
+					}
+				}
+			}
+
+			ConfigSection? perforceSection = projectConfig.FindSection("Perforce");
+
+			string[] combinedSyncFilter = GlobalSettingsFile.GetCombinedSyncFilter(syncCategories, workspaceSettings.Preset, roles, userSettings.Global.Filter, workspaceSettings.Filter, perforceSection);
+
+			return combinedSyncFilter;
+		}
+	}
+}
diff --git a/Engine/Source/Programs/UnrealGameSync/UnrealGameSyncCmd/Utils/UserSettingsUtils.cs b/Engine/Source/Programs/UnrealGameSync/UnrealGameSyncCmd/Utils/UserSettingsUtils.cs
new file mode 100644
--- /dev/null
+++ b/Engine/Source/Programs/UnrealGameSync/UnrealGameSyncCmd/Utils/UserSettingsUtils.cs
@@ -0,0 +1,41 @@
+﻿// Copyright Epic Games, Inc. All Rights Reserved.
+
+using EpicGames.Core;
+using UnrealGameSync;
+
+namespace UnrealGameSyncCmd.Utils
+{
+	public static class UserSettingsUtils
+	{
+		public static UserWorkspaceSettings? ReadOptionalUserWorkspaceSettings()
+		{
+			DirectoryReference? dir = DirectoryReference.GetCurrentDirectory();
+			for (; dir != null; dir = dir.ParentDirectory)
+			{
+				try
+				{
+					UserWorkspaceSettings? settings;
+					if (UserWorkspaceSettings.TryLoad(dir, out settings))
+					{
+						return settings;
+					}
+				}
+				catch
+				{
+					// Guard against directories we can't access, eg. [REDACTED_USER_HOME]
+				}
+			}
+			return null;
+		}
+
+		public static UserWorkspaceSettings ReadRequiredUserWorkspaceSettings()
+		{
+			UserWorkspaceSettings? settings = ReadOptionalUserWorkspaceSettings();
+			if (settings == null)
+			{
+				throw new UserErrorException("Unable to find UGS workspace in current directory.");
+			}
+			return settings;
+		}
+	}
+}
diff --git a/Engine/Source/Programs/UnrealGameSync/UnrealGameSyncCmd/Utils/VersionUtils.cs b/Engine/Source/Programs/UnrealGameSync/UnrealGameSyncCmd/Utils/VersionUtils.cs
new file mode 100644
--- /dev/null
+++ b/Engine/Source/Programs/UnrealGameSync/UnrealGameSyncCmd/Utils/VersionUtils.cs
@@ -0,0 +1,91 @@
+﻿// Copyright Epic Games, Inc. All Rights Reserved.
+
+using System;
+using System.Net.Http;
+using System.Net.Http.Json;
+using System.Reflection;
+using System.Runtime.InteropServices;
+using System.Threading;
+using System.Threading.Tasks;
+
+using Microsoft.Extensions.Logging;
+using UnrealGameSync;
+
+namespace UnrealGameSyncCmd.Utils
+{
+	public class DeploymentInfo
+	{
+		public string Version { get; set; } = String.Empty;
+	}
+
+	public static class VersionUtils
+	{
+		public static string GetVersion()
+		{
+			AssemblyInformationalVersionAttribute? version = Assembly.GetExecutingAssembly().GetCustomAttribute<AssemblyInformationalVersionAttribute>();
+			return version?.InformationalVersion ?? "Unknown";
+		}
+
+		public static async Task<string?> GetLatestVersionAsync(ILogger? logger, CancellationToken cancellationToken)
+		{
+			string? hordeUrl = DeploymentSettings.Instance.HordeUrl;
+			if (hordeUrl == null)
+			{
+				logger?.LogError("Horde URL is not set in deployment config file. Cannot upgrade.");
+				return null;
+			}
+
+			string? toolName = GetUpgradeToolName();
+			if (toolName == null)
+			{
+				logger?.LogError("Command-line upgrades are not supported on this platform.");
+				return null;
+			}
+
+			using (HttpClient httpClient = new HttpClient())
+			{
+				Uri baseUrl = new Uri(hordeUrl);
+
+				DeploymentInfo? deploymentInfo;
+				try
+				{
+					deploymentInfo = await httpClient.GetFromJsonAsync<DeploymentInfo>(new Uri(baseUrl, $"api/v1/tools/{toolName}/deployments"), cancellationToken);
+				}
+				catch (Exception ex)
+				{
+					logger?.LogError(ex, "Failed to query for deployment info: {Message}", ex.Message);
+					return null;
+				}
+				if (deploymentInfo == null)
+				{
+					logger?.LogError("Failed to query for deployment info.");
+					return null;
+				}
+				return deploymentInfo.Version;
+			}
+		}
+
+		public static string? GetUpgradeToolName()
+		{
+			if (RuntimeInformation.IsOSPlatform(OSPlatform.OSX))
+			{
+				if (RuntimeInformation.OSArchitecture == Architecture.Arm64)
+				{
+					return "ugs-mac-arm64";
+				}
+				else
+				{
+					return "ugs-mac";
+				}
+			}
+			else if (RuntimeInformation.IsOSPlatform(OSPlatform.Linux))
+			{
+				return "ugs-linux";
+			}
+			else
+			{
+				return null;
+			}
+		}
+	}
+}
diff --git a/Engine/Source/Programs/UnrealGameSync/UnrealGameSyncCmd/Utils/WorkspaceStateUtils.cs b/Engine/Source/Programs/UnrealGameSync/UnrealGameSyncCmd/Utils/WorkspaceStateUtils.cs
new file mode 100644
--- /dev/null
+++ b/Engine/Source/Programs/UnrealGameSync/UnrealGameSyncCmd/Utils/WorkspaceStateUtils.cs
@@ -0,0 +1,26 @@
+﻿// Copyright Epic Games, Inc. All Rights Reserved.
+
+using System.Threading;
+using System.Threading.Tasks;
+
+using EpicGames.Perforce;
+using Microsoft.Extensions.Logging;
+using UnrealGameSync;
+
+namespace UnrealGameSyncCmd.Utils
+{
+	internal static class WorkspaceStateUtils
+	{
+		public static async Task<WorkspaceStateWrapper> ReadWorkspaceState(IPerforceConnection perforceClient, UserWorkspaceSettings settings, GlobalSettingsFile userSettings, ILogger logger)
+		{
+			WorkspaceStateWrapper state = userSettings.FindOrAddWorkspaceState(settings);
+			if (state.Current.SettingsTimeUtc != settings.LastModifiedTimeUtc)
+			{
+				logger.LogDebug("Updating state due to modified settings timestamp");
+				ProjectInfo info = await ProjectInfo.CreateAsync(perforceClient, settings, CancellationToken.None);
+				state.Modify(x => x.UpdateCachedProjectInfo(info, settings.LastModifiedTimeUtc));
+			}
+			return state;
+		}
+	}
+}
