fix to avoid a "too many row" issue in UGS getting the changes for a range in PCBS

diff --git a/Engine/Source/Programs/UnrealGameSync/UnrealGameSync/PerforceMonitor.cs b/Engine/Source/Programs/UnrealGameSync/UnrealGameSync/PerforceMonitor.cs
--- a/Engine/Source/Programs/UnrealGameSync/UnrealGameSync/PerforceMonitor.cs
+++ b/Engine/Source/Programs/UnrealGameSync/UnrealGameSync/PerforceMonitor.cs
@@ -433,10 +433,10 @@ namespace UnrealGameSync
 						}
 					}
 				}
+
 				if (minZippedChangeNumber != -1 && minZippedChangeNumber < oldestChangeNumber)
 				{
-					string[] filteredPaths = depotPaths.Select(x => $"{x}@{minZippedChangeNumber},{oldestChangeNumber - 1}").ToArray();
-					List<ChangesRecord> zipChanges = await perforce.GetChangesAsync(ChangesOptions.None, -1, ChangeStatus.Submitted, filteredPaths, cancellationToken);
+					List<ChangesRecord> zipChanges = await perforce.GetChangesAsync(ChangesOptions.None, null, minZippedChangeNumber, -1, ChangeStatus.Submitted, null, depotPaths, cancellationToken);
 					newChanges.AddRange(zipChanges);
 				}
 			}
