added command to list, explore and set archive channels (win only for now)

diff --git a/Engine/Source/Programs/UnrealGameSync/UnrealGameSyncCmd/Commands/ArchiveCommand.cs b/Engine/Source/Programs/UnrealGameSync/UnrealGameSyncCmd/Commands/ArchiveCommand.cs
new file mode 100644
--- /dev/null
+++ b/Engine/Source/Programs/UnrealGameSync/UnrealGameSyncCmd/Commands/ArchiveCommand.cs
@@ -0,0 +1,200 @@
+﻿// Copyright Epic Games, Inc. All Rights Reserved.
+
+using System;
+using System.Collections.Generic;
+using System.Linq;
+using System.Threading;
+using System.Threading.Tasks;
+using EpicGames.Core;
+using EpicGames.Perforce;
+using Microsoft.Extensions.Logging;
+using UnrealGameSync;
+using UnrealGameSyncCmd.Utils;
+
+namespace UnrealGameSyncCmd.Commands
+{
+	internal class ArchiveCommandOptions
+	{
+		[CommandLine("-List")]
+		public bool List { get; set; } = false;
+
+		[CommandLine("-Changes")]
+		public bool Changes { get; set; } = false;
+
+		[CommandLine("-Set")]
+		public bool Set { get; set; } = false;
+
+		[CommandLine("-Unset")]
+		public bool Unset { get; set; } = false;
+
+		[CommandLine("-ArchiveType")]
+		public string ArchiveType { get; set; } = String.Empty;
+
+		[CommandLine("-ArchiveName")]
+		public string ArchiveName { get; set; } = String.Empty;
+	}
+
+	class ArchiveCommand : Command
+	{
+		public override async Task ExecuteAsync(CommandContext context)
+		{
+			ILogger logger = context.Logger;
+			context.Arguments.TryGetPositionalArgument(out string? _);
+
+			ArchiveCommandOptions options = new ArchiveCommandOptions();
+			context.Arguments.ApplyTo(options);
+			context.Arguments.CheckAllArgumentsUsed();
+
+			UserWorkspaceSettings settings = UserSettingsUtils.ReadRequiredUserWorkspaceSettings();
+			using IPerforceConnection perforceClient = await PerforceConnectionUtils.ConnectAsync(settings, context.LoggerFactory);
+
+			WorkspaceStateWrapper state = await WorkspaceStateUtils.ReadWorkspaceState(perforceClient, settings, context.UserSettings, logger);
+
+			ProjectInfo projectInfo = new ProjectInfo(settings.RootDir, state.Current);
+			UserProjectSettings projectSettings = context.UserSettings.FindOrAddProjectSettings(projectInfo, settings, logger);
+
+			ConfigFile projectConfig = await ConfigUtils.ReadProjectConfigFileAsync(perforceClient, projectInfo, logger, CancellationToken.None);
+
+			List<BaseArchiveChannel> archives = await BaseArchive.EnumerateChannelsAsync(perforceClient, context.HordeClient, context.CloudStorage, projectConfig, state.Current.ProjectIdentifier, CancellationToken.None);
+			List<IArchiveChannel> selected = GetSelectedArchiveChannels(context.GlobalSettings!, archives);
+
+			if (options.List)
+			{
+				foreach (BaseArchiveChannel archive in archives)
+				{
+					if (!String.IsNullOrWhiteSpace(options.ArchiveType) && !options.ArchiveType.Equals(archive.Type, StringComparison.OrdinalIgnoreCase))
+					{
+						continue;
+					}
+
+					if (!String.IsNullOrWhiteSpace(options.ArchiveName) && !options.ArchiveName.Equals(archive.Name, StringComparison.OrdinalIgnoreCase))
+					{
+						continue;
+					}
+
+					DisplayArchive(logger, selected, archive);
+				}
+			}
+
+			if (options.Changes)
+			{
+				foreach (BaseArchiveChannel archive in archives)
+				{
+					if (!String.IsNullOrWhiteSpace(options.ArchiveType) && !options.ArchiveType.Equals(archive.Type, StringComparison.OrdinalIgnoreCase))
+					{
+						continue;
+					}
+
+					if (!String.IsNullOrWhiteSpace(options.ArchiveName) && !options.ArchiveName.Equals(archive.Name, StringComparison.OrdinalIgnoreCase))
+					{
+						continue;
+					}
+
+					DisplayArchive(logger, selected, archive);
+					DisplayChanges(logger, archive);
+				}
+			}
+
+			if (options.Set)
+			{
+				if (String.IsNullOrWhiteSpace(options.ArchiveType))
+				{
+					logger.LogError("Parameter ArchiveType is required to set an archive.");
+					return;
+				}
+
+				if (String.IsNullOrWhiteSpace(options.ArchiveName))
+				{
+					logger.LogError("Parameter ArchiveName is required to set an archive.");
+					return;
+				}
+
+				List<BaseArchiveChannel> enabled = archives.Where(a => a.Type.Equals(options.ArchiveType, StringComparison.OrdinalIgnoreCase) && a.Name.Equals(options.ArchiveName, StringComparison.OrdinalIgnoreCase)).ToList();
+				if(enabled.Any())
+				{
+					ResetArchiveSettings(logger, context.GlobalSettings!);
+					foreach (BaseArchiveChannel archive in enabled)
+					{
+						ArchiveSettings? archiveSettings = context.GlobalSettings!.Archives.FirstOrDefault(x => x.Type == archive.Type);
+						if (archiveSettings != null)
+						{
+							archiveSettings.Enabled = true;
+							archiveSettings.Order.Remove(archive.Name);
+							archiveSettings.Order.Insert(0, archive.Name);
+						}
+					}
+
+					context.GlobalSettings!.Save(logger);
+				}
+			}
+			else if (options.Unset)
+			{
+				ResetArchiveSettings(logger, context.GlobalSettings!);
+			}
+		}
+
+		private static void ResetArchiveSettings(ILogger logger, UserSettings userSettings)
+		{
+			foreach (ArchiveSettings archiveSettings in userSettings.Archives)
+			{
+				archiveSettings.Enabled = false;
+			}
+
+			userSettings.Save(logger);
+		}
+
+		private static List<IArchiveChannel> GetSelectedArchiveChannels(UserSettings userSettings, IReadOnlyList<IArchiveChannel> archives)
+		{
+			Dictionary<string, KeyValuePair<IArchiveChannel, int>> archiveTypeToSelection = new Dictionary<string, KeyValuePair<IArchiveChannel, int>>();
+			foreach (IArchiveChannel archive in archives)
+			{
+				ArchiveSettings? archiveSettings = userSettings.Archives.FirstOrDefault(x => x.Type == archive.Type);
+				if (archiveSettings != null && archiveSettings.Enabled)
+				{
+					int preference = archiveSettings.Order.IndexOf(archive.Name);
+					if (preference == -1)
+					{
+						preference = archiveSettings.Order.Count;
+					}
+
+					KeyValuePair<IArchiveChannel, int> existingItem;
+					if (!archiveTypeToSelection.TryGetValue(archive.Type, out existingItem) || existingItem.Value > preference)
+					{
+						archiveTypeToSelection[archive.Type] = new KeyValuePair<IArchiveChannel, int>(archive, preference);
+					}
+				}
+			}
+			return archiveTypeToSelection.Select(x => x.Value.Key).ToList();
+		}
+
+		private static void DisplayArchive(ILogger logger, List<IArchiveChannel> selected, BaseArchiveChannel archive)
+		{
+
+			string annotations = String.Empty;
+			switch (archive)
+			{
+				case PerforceArchiveChannel perforceArchiveChannel:
+					annotations = $"perforce:{perforceArchiveChannel.DepotPath}";
+					break;
+				case HordeArchiveChannel _:
+					annotations = $"horde";
+					break;
+				case CloudArchiveChannel _:
+					annotations = $"cloud";
+					break;
+				default:
+					break;
+			}
+
+			logger.LogInformation("...channel {Selected} | {Type} | {Name} | {Annotations}", selected.Contains(archive), archive.Type, archive.Name, annotations);
+		}
+
+		private static void DisplayChanges(ILogger logger, BaseArchiveChannel archive)
+		{
+			foreach (KeyValuePair<int, IArchive> change in archive.ChangeNumberToArchive)
+			{
+				logger.LogInformation("...change {ChangeNumber} {ArchiveKey}", change.Key, change.Value.Key);
+			}
+		}
+	}
+}
diff --git a/Engine/Source/Programs/UnrealGameSync/UnrealGameSyncCmd/CommandsFactory.cs b/Engine/Source/Programs/UnrealGameSync/UnrealGameSyncCmd/CommandsFactory.cs
--- a/Engine/Source/Programs/UnrealGameSync/UnrealGameSyncCmd/CommandsFactory.cs
+++ b/Engine/Source/Programs/UnrealGameSync/UnrealGameSyncCmd/CommandsFactory.cs
@@ -104,7 +104,11 @@ namespace UnrealGameSyncCmd
 				new CommandInfo("tools", typeof(ToolsCommand), typeof(ToolsCommandCommandOptions),
 				"ugs tools",
 				"Install a Custom Tool. Only available for windows."
-				)
+				),
+				new CommandInfo("archive", typeof(ArchiveCommand), typeof(ArchiveCommandOptions),
+				"ugs archive",
+				"Execute a command related to archives."
+			),
 			};
 		}
 	}
diff --git a/Engine/Source/Programs/UnrealGameSync/UnrealGameSyncCmd/UnrealGameSyncCmd.csproj b/Engine/Source/Programs/UnrealGameSync/UnrealGameSyncCmd/UnrealGameSyncCmd.csproj
--- a/Engine/Source/Programs/UnrealGameSync/UnrealGameSyncCmd/UnrealGameSyncCmd.csproj
+++ b/Engine/Source/Programs/UnrealGameSync/UnrealGameSyncCmd/UnrealGameSyncCmd.csproj
@@ -40,4 +40,16 @@
     </None>
   </ItemGroup>
 
+  <ItemGroup Condition="Exists('..\..\..\..\Binaries\Win64\zen.exe')">
+    <None Include="..\..\..\..\Binaries\Win64\zen.exe" TargetPath="Binaries\Win64\zen.exe" Link="Binaries\zen.exe">
+      <CopyToOutputDirectory>PreserveNewest</CopyToOutputDirectory>
+    </None>
+    <None Include="..\..\..\..\Binaries\Linux\zen" TargetPath="Binaries\Linux\zen" Link="Binaries\zen">
+      <CopyToOutputDirectory>PreserveNewest</CopyToOutputDirectory>
+    </None>
+    <None Include="..\..\..\..\Binaries\Mac\zen" TargetPath="Binaries\Mac\zen" Link="Binaries\zen">
+      <CopyToOutputDirectory>PreserveNewest</CopyToOutputDirectory>
+    </None>
+  </ItemGroup>
+
 </Project>
