Adding support for InfoMessages  in UGS, displayed on the left hand side of the top banner.

diff --git a/Engine/Source/Programs/UnrealGameSync/UnrealGameSync/Controls/StatusPanel.cs b/Engine/Source/Programs/UnrealGameSync/UnrealGameSync/Controls/StatusPanel.cs
--- a/Engine/Source/Programs/UnrealGameSync/UnrealGameSync/Controls/StatusPanel.cs
+++ b/Engine/Source/Programs/UnrealGameSync/UnrealGameSync/Controls/StatusPanel.cs
@@ -69,6 +69,12 @@ namespace UnrealGameSync
 		public bool MouseDown { get; set; }
 		public Rectangle Bounds { get; set; }
 
+		public string? ToolTip
+		{
+			get;
+			set;
+		}
+
 		public Point Layout(Graphics graphics, Point location, StatusElementResources resources)
 		{
 			Size size = Measure(graphics, resources);
@@ -109,17 +115,21 @@ namespace UnrealGameSync
 		readonly Image _strip;
 		readonly Size _iconSize;
 		readonly int _index;
+		readonly bool _applyDpiScaling;
 
-		public IconStripStatusElement(Image strip, Size iconSize, int index)
+		public IconStripStatusElement(Image strip, Size iconSize, int index, bool applyDpiScaling = true)
 		{
 			_strip = strip;
 			_iconSize = iconSize;
 			_index = index;
+			_applyDpiScaling = applyDpiScaling;
 		}
 
 		public override Size Measure(Graphics graphics, StatusElementResources resources)
 		{
-			return new Size((int)(_iconSize.Width * graphics.DpiX / 96.0f), (int)(_iconSize.Height * graphics.DpiY / 96.0f));
+			float dpiScaling = _applyDpiScaling ? (graphics.DpiX / 96.0f) : 1.0f;
+
+			return new Size((int)(_iconSize.Width * dpiScaling), (int)(_iconSize.Height * dpiScaling));
 		}
 
 		public override void Draw(Graphics graphics, StatusElementResources resources)
@@ -333,52 +343,56 @@ namespace UnrealGameSync
 			_modified = true;
 		}
 
-		public void Add(StatusElement element)
+		public string? ToolTip
+		{ 
+			get;
+			set;
+		}
+
+		public void Add(StatusElement element, string? toolTip = null)
 		{
+			element.ToolTip = toolTip;
 			_elements.Add(element);
 			_modified = true;
 		}
 
-		public void AddIcon(Image inIcon)
+		public void AddIcon(Image inIcon, string? toolTip = null)
 		{
-			_elements.Add(new IconStatusElement(inIcon));
-			_modified = true;
+			Add(new IconStatusElement(inIcon), toolTip);
 		}
 
-		public void AddIcon(Image inStrip, Size inIconSize, int inIndex)
+		public void AddIcon(Image inStrip, Size inIconSize, int inIndex, bool applyDpiScaling = true, string? toolTip = null)
 		{
-			_elements.Add(new IconStripStatusElement(inStrip, inIconSize, inIndex));
-			_modified = true;
+			Add(new IconStripStatusElement(inStrip, inIconSize, inIndex, applyDpiScaling), toolTip);
 		}
 
-		public void AddText(string inText, FontStyle inStyle = FontStyle.Regular)
+		public void AddText(string inText, FontStyle inStyle = FontStyle.Regular, string? toolTip = null)
 		{
 			Color textColor = ProgramApplicationContext.GetTheme().TextColor;
-			_elements.Add(new TextStatusElement(inText, textColor, inStyle));
-			_modified = true;
+			Add(new TextStatusElement(inText, textColor, inStyle), toolTip);
 		}
 
-		public void AddText(string inText, Color inColor, FontStyle inStyle = FontStyle.Regular)
+		public void AddText(string inText, Color inColor, FontStyle inStyle = FontStyle.Regular, string? toolTip = null)
 		{
-			_elements.Add(new TextStatusElement(inText, inColor, inStyle));
-			_modified = true;
+			Add(new TextStatusElement(inText, inColor, inStyle), toolTip);
 		}
 
-		public void AddLink(string inText, FontStyle inStyle, Action inLinkAction)
+		public void AddLink(string inText, FontStyle inStyle, Action inLinkAction, string? toolTip = null)
 		{
-			_elements.Add(new LinkStatusElement(inText, inStyle, (p, r) => { inLinkAction(); }));
-			_modified = true;
+			Add(new LinkStatusElement(inText, inStyle, (p, r) => { inLinkAction(); }), toolTip);
 		}
 
-		public void AddLink(string inText, FontStyle inStyle, Action<Point, Rectangle> inLinkAction)
+		public void AddLink(string inText, FontStyle inStyle, Action<Point, Rectangle> inLinkAction, string? toolTip = null)
 		{
-			_elements.Add(new LinkStatusElement(inText, inStyle, inLinkAction));
-			_modified = true;
+			Add(new LinkStatusElement(inText, inStyle, inLinkAction), toolTip);
 		}
 
-		public void AddBadge(string inText, Color inBackgroundColor, Action<Point, Rectangle>? inClickAction)
+		public void AddBadge(string inText, Color inBackgroundColor, Action<Point, Rectangle>? inClickAction, string? toolTip = null)
 		{
-			_elements.Add(new BadgeStatusElement(inText, inBackgroundColor, inClickAction));
+			BadgeStatusElement element = new BadgeStatusElement(inText, inBackgroundColor, inClickAction);
+			element.ToolTip = toolTip;
+
+			_elements.Add(element);
 			if (_elements.Count >= 2)
 			{
 				BadgeStatusElement? prevBadge = _elements[^2] as BadgeStatusElement;
@@ -392,10 +406,9 @@ namespace UnrealGameSync
 			_modified = true;
 		}
 
-		public void AddProgressBar(float progress)
+		public void AddProgressBar(float progress, string? toolTip = null)
 		{
-			_elements.Add(new ProgressBarStatusElement(progress));
-			_modified = true;
+			Add(new ProgressBarStatusElement(progress), toolTip);
 		}
 
 		public bool HitTest(Point location, [NotNullWhen(true)] out StatusElement? outElement)
@@ -441,11 +454,14 @@ namespace UnrealGameSync
 	class StatusPanel : Panel, ICustomThemeHandler
 	{
 		const float LineSpacing = 1.35f;
+		const int MaxInfoLineWidth = 700;
+		const int InfoMargin = 15;
 
 		InvertibleImage _projectLogo = new InvertibleImage(Properties.Resources.DefaultProjectLogo, false);
 		Rectangle _projectLogoBounds;
 		StatusElementResources? _resources;
 		readonly List<StatusLine> _lines = new List<StatusLine>();
+		readonly List<StatusLine> _infoLines = new List<StatusLine>();
 		StatusLine? _caption;
 		Pen? _alertDividerPen;
 		int _alertDividerY;
@@ -457,11 +473,17 @@ namespace UnrealGameSync
 		StatusElement? _mouseDownElement;
 		int _contentWidth = 400;
 		int _suspendDisplayCount;
+		readonly System.Windows.Forms.ToolTip _toolTip = new System.Windows.Forms.ToolTip();
+		string? _currentToolTip = null;
 
 		public StatusPanel()
 		{
 			DoubleBuffered = true;
 
+			_toolTip.ShowAlways = true;
+			_toolTip.InitialDelay = 1000;
+			_toolTip.ReshowDelay = 500;
+
 			_alertDividerPen = new Pen(Color.Black);
 			_alertDividerPen.DashPattern = new float[] { 2, 2 };
 		}
@@ -497,6 +519,8 @@ namespace UnrealGameSync
 					_alertDividerPen.Dispose();
 					_alertDividerPen = null;
 				}
+
+				_toolTip.Dispose();
 				
 				_projectLogo.Dispose();
 				
@@ -526,10 +550,11 @@ namespace UnrealGameSync
 		{
 			InvalidateElements();
 			_lines.Clear();
+			_infoLines.Clear();
 			_caption = null;
 		}
 
-		public void Set(IEnumerable<StatusLine> newLines, StatusLine? newCaption, StatusLine? newAlert, Color? newTintColor)
+		public void Set(IEnumerable<StatusLine> newLines, IEnumerable<StatusLine> newInfoLines, StatusLine? newCaption, StatusLine? newAlert, Color? newTintColor)
 		{
 			_resources ??= new StatusElementResources(Font);
 			
@@ -542,6 +567,8 @@ namespace UnrealGameSync
 			InvalidateElements();
 			_lines.Clear();
 			_lines.AddRange(newLines);
+			_infoLines.Clear();
+			_infoLines.AddRange(newInfoLines);
 			_caption = newCaption;
 			_alert = newAlert;
 			LayoutElements();
@@ -561,6 +588,10 @@ namespace UnrealGameSync
 			{
 				Invalidate(line.Bounds);
 			}
+			foreach(StatusLine line in _infoLines)
+			{
+				Invalidate(line.Bounds);
+			}
 			if (_caption != null)
 			{
 				Invalidate(_caption.Bounds);
@@ -571,13 +602,26 @@ namespace UnrealGameSync
 			}
 		}
 
-		protected bool HitTest(Point location, [NotNullWhen(true)] out StatusElement? outElement)
+		protected bool HitTest(Point location, [NotNullWhen(true)] out StatusElement? outElement, out StatusLine? outLine)
 		{
 			outElement = null;
+			outLine = null;
+
 			foreach (StatusLine line in _lines)
 			{
 				if (line.HitTest(location, out outElement))
 				{
+					outLine = line;
+
+					return true;
+				}
+			}
+			foreach (StatusLine line in _infoLines)
+			{
+				if (line.HitTest(location, out outElement))
+				{
+					outLine = line;
+
 					return true;
 				}
 			}
@@ -585,6 +629,8 @@ namespace UnrealGameSync
 			{
 				if (_caption.HitTest(location, out outElement))
 				{
+					outLine = _caption;
+
 					return true;
 				}
 			}
@@ -592,9 +638,12 @@ namespace UnrealGameSync
 			{
 				if (_alert.HitTest(location, out outElement))
 				{
+					outLine = _alert;
+
 					return true;
 				}
 			}
+
 			return false;
 		}
 
@@ -633,7 +682,7 @@ namespace UnrealGameSync
 			int logoHeight = (int)(drawProjectLogo.Height * logoScale);
 
 			// Figure out where the split between content and the logo is going to be
-			int dividerX = ((Width - logoWidth - _contentWidth) / 2) + logoWidth;
+			int dividerX = Math.Max( ((Width - logoWidth - _contentWidth) / 2) + logoWidth, MaxInfoLineWidth);
 
 			// Get the logo position
 			int logoX = dividerX - logoWidth;
@@ -662,6 +711,18 @@ namespace UnrealGameSync
 				line.Layout(graphics, new Point(dividerX + 5, (int)lineY), _resources!);
 				lineY += (int)(Font.Height * LineSpacing * line.LineHeight * 0.5f);
 			}
+
+			// Measure up all the info line height
+			float totalInfoLineHeight = _infoLines.Sum(x => x.LineHeight);
+
+			// Space out all the info lines
+			float infoLineY = (bodyHeight - InfoMargin - totalInfoLineHeight * (int)(Font.Height * LineSpacing));
+			foreach (StatusLine line in _infoLines)
+			{
+				infoLineY += (int)(Font.Height * LineSpacing * line.LineHeight * 0.5f);
+				line.Layout(graphics, new Point(InfoMargin, (int)infoLineY), _resources!);
+				infoLineY += (int)(Font.Height * LineSpacing * line.LineHeight * 0.5f);
+			}
 		}
 
 		protected override void OnMouseMove(MouseEventArgs e)
@@ -747,6 +808,11 @@ namespace UnrealGameSync
 					line.Draw(e.Graphics, _resources!);
 				}
 
+				foreach (StatusLine line in _infoLines)
+				{
+					line.Draw(e.Graphics, _resources!);
+				}
+
 				_caption?.Draw(e.Graphics, _resources!);
 
 				if (_alert != null)
@@ -762,9 +828,11 @@ namespace UnrealGameSync
 			_mouseOverLocation = newMouseOverLocation;
 
 			StatusElement? newMouseOverElement = null;
+			StatusLine? newMouseOverLine = null;
+
 			if (_mouseOverLocation.HasValue)
 			{
-				HitTest(_mouseOverLocation.Value, out newMouseOverElement);
+				HitTest(_mouseOverLocation.Value, out newMouseOverElement, out newMouseOverLine);
 			}
 
 			if (newMouseOverElement != _mouseOverElement)
@@ -773,6 +841,16 @@ namespace UnrealGameSync
 				{
 					_mouseOverElement.MouseOver = false;
 					Cursor = Cursors.Arrow;
+
+					if (newMouseOverElement == null)
+					{
+						if (_currentToolTip != null)
+						{
+							_currentToolTip = null;
+							_toolTip.Hide(this);
+						}
+					}
+
 					Invalidate(_mouseOverElement.Bounds);
 				}
 
@@ -782,6 +860,34 @@ namespace UnrealGameSync
 				{
 					_mouseOverElement.MouseOver = true;
 					Cursor = _mouseOverElement.Cursor;
+					string? newToolTip = null;
+
+					if (String.IsNullOrEmpty(_mouseOverElement.ToolTip) == false)
+					{
+						newToolTip = _mouseOverElement.ToolTip;
+					}
+					else if (String.IsNullOrEmpty(newMouseOverLine?.ToolTip) == false)
+					{
+						newToolTip = newMouseOverLine.ToolTip;
+					}
+
+					if (String.IsNullOrEmpty(newToolTip) == false)
+					{
+						if (newToolTip != _currentToolTip)
+						{
+							_currentToolTip = newToolTip;
+							_toolTip.Show(newToolTip, this, _mouseOverLocation!.Value.X, _mouseOverElement.Bounds.Bottom + 5);
+						}
+					}
+					else
+					{
+						if (_currentToolTip != null)
+						{
+							_currentToolTip = null;
+							_toolTip.Hide(this);
+						}
+					}
+
 					Invalidate(_mouseOverElement.Bounds);
 				}
 			}
@@ -792,9 +898,10 @@ namespace UnrealGameSync
 			_mouseDownLocation = newMouseDownLocation;
 
 			StatusElement? newMouseDownElement = null;
+
 			if (_mouseDownLocation.HasValue)
 			{
-				HitTest(_mouseDownLocation.Value, out newMouseDownElement);
+				HitTest(_mouseDownLocation.Value, out newMouseDownElement, out _);
 			}
 
 			if (newMouseDownElement != _mouseDownElement)
diff --git a/Engine/Source/Programs/UnrealGameSync/UnrealGameSync/Controls/WorkspaceControl.cs b/Engine/Source/Programs/UnrealGameSync/UnrealGameSync/Controls/WorkspaceControl.cs
--- a/Engine/Source/Programs/UnrealGameSync/UnrealGameSync/Controls/WorkspaceControl.cs
+++ b/Engine/Source/Programs/UnrealGameSync/UnrealGameSync/Controls/WorkspaceControl.cs
@@ -4172,6 +4172,8 @@ namespace UnrealGameSync
 			StatusLine? alert = null;
 			Color? tintColor = null;
 
+			List<StatusLine> infoLines = new List<StatusLine>();
+
 			ConfigFile projectConfigFile = _perforceMonitor.LatestProjectConfigFile;
 			if (projectConfigFile != null)
 			{
@@ -4186,6 +4188,36 @@ namespace UnrealGameSync
 				{
 					tintColor = System.Drawing.ColorTranslator.FromHtml(statusPanelColor);
 				}
+
+				string[]? infoMessages;
+				if (TryGetProjectSetting(projectConfigFile, "InfoMessage", out infoMessages))
+				{
+					const int maxInfoMessages = 7;
+					if (infoMessages.Length > maxInfoMessages)
+					{
+						_logger.LogWarning("Maximum number of InfoMessage entries exceeded, found {MessageCount} messages but will only display {MaxMessages}", infoMessages.Length, maxInfoMessages);
+					}
+
+					foreach(string infoMessage in infoMessages.Take(Math.Min(infoMessages.Length, maxInfoMessages)))
+					{
+						ConfigObject infoMessageObj = new ConfigObject(infoMessage);
+
+						string? infoText = infoMessageObj.GetValue("Message");
+
+						if (String.IsNullOrEmpty(infoText) == false)
+						{						
+							StatusLine infoLine = CreateStatusLineFromMarkdown(infoText, applyDpiScalingToIcons: false);
+
+							string? infoTooltip = infoMessageObj.GetValue("Tooltip");
+							if (String.IsNullOrEmpty(infoTooltip) == false)
+							{
+								infoLine.ToolTip = infoTooltip;
+							}
+
+							infoLines.Add(infoLine);
+						}
+					}
+				}
 			}
 
 			using (Graphics graphics = CreateGraphics())
@@ -4201,7 +4233,7 @@ namespace UnrealGameSync
 				}
 			}
 
-			StatusPanel.Set(lines, caption, alert, tintColor);
+			StatusPanel.Set(lines, infoLines, caption, alert, tintColor);
 		}
 
 		private void ConnectToHorde(IHordeClient hordeClient)
@@ -4360,7 +4392,7 @@ namespace UnrealGameSync
 			}
 		}
 
-		private static StatusLine CreateStatusLineFromMarkdown(string text)
+		private static StatusLine CreateStatusLineFromMarkdown(string text, bool applyDpiScalingToIcons = true)
 		{
 			StatusLine line = new StatusLine();
 
@@ -4432,7 +4464,15 @@ namespace UnrealGameSync
 						if (String.Compare(":alert:", 0, text, idx, endIdx - idx, StringComparison.Ordinal) == 0)
 						{
 							FlushMarkdownText(line, elementText, style);
-							line.AddIcon(Properties.Resources.StatusIcons, new Size(16, 16), 4);
+							line.AddIcon(Properties.Resources.StatusIcons, new Size(16, 16), 4, applyDpiScalingToIcons);
+							idx = endIdx + 1;
+							continue;
+						}
+
+						if (String.Compare(":tick:", 0, text, idx, endIdx - idx, StringComparison.Ordinal) == 0)
+						{
+							FlushMarkdownText(line, elementText, style);
+							line.AddIcon(Properties.Resources.StatusIcons, new Size(16, 16), 0, applyDpiScalingToIcons);
 							idx = endIdx + 1;
 							continue;
 						}
diff --git a/Engine/Source/Programs/UnrealGameSync/UnrealGameSync/Forms/MainWindow.cs b/Engine/Source/Programs/UnrealGameSync/UnrealGameSync/Forms/MainWindow.cs
--- a/Engine/Source/Programs/UnrealGameSync/UnrealGameSync/Forms/MainWindow.cs
+++ b/Engine/Source/Programs/UnrealGameSync/UnrealGameSync/Forms/MainWindow.cs
@@ -734,7 +734,7 @@ namespace UnrealGameSync
 			openLine.AddLink("Application settings...", FontStyle.Bold | FontStyle.Underline, () => { ModifyApplicationSettings(); });
 			lines.Add(openLine);
 
-			DefaultControl.Set(lines, null, null, null);
+			DefaultControl.Set(lines, Enumerable.Empty<StatusLine>(), null, null, null);
 		}
 
 		private void CreateErrorPanel(int replaceTabIdx, UserSelectedProjectSettings project, string message)
@@ -797,7 +797,7 @@ namespace UnrealGameSync
 			});
 			lines.Add(actionLine);
 
-			errorPanel.Set(lines, null, null, null);
+			errorPanel.Set(lines, Enumerable.Empty<StatusLine>(), null, null, null);
 
 			string newProjectName = "Unknown";
 			if (project.Type == UserSelectedProjectType.Client && project.ClientPath != null)
