Moved Progress text reading out of UGS and into EpicGames.Core to allow for this to be reused

diff --git a/Engine/Source/Programs/Shared/EpicGames.Core/ProgressTextReader.cs b/Engine/Source/Programs/Shared/EpicGames.Core/ProgressTextReader.cs
new file mode 100644
--- /dev/null
+++ b/Engine/Source/Programs/Shared/EpicGames.Core/ProgressTextReader.cs
@@ -0,0 +1,247 @@
+// Copyright Epic Games, Inc. All Rights Reserved.
+
+using System;
+using System.Collections.Generic;
+using System.Linq;
+
+namespace EpicGames.Core
+{
+	public class ProgressValue
+	{
+		Tuple<string, float> _state = null!;
+		readonly Stack<Tuple<float, float>> _ranges = new Stack<Tuple<float, float>>();
+
+		public ProgressValue()
+		{
+			Clear();
+		}
+
+		public void Clear()
+		{
+			_state = new Tuple<string, float>("Starting...", 0.0f);
+
+			_ranges.Clear();
+			_ranges.Push(new Tuple<float, float>(0.0f, 1.0f));
+		}
+
+		public Tuple<string, float> Current => _state;
+
+		public void Set(string message)
+		{
+			if (_ranges.Count == 1)
+			{
+				_state = new Tuple<string, float>(message, _state.Item2);
+			}
+		}
+
+		public void Set(string message, float fraction)
+		{
+			if (_ranges.Count == 1)
+			{
+				_state = new Tuple<string, float>(message, RelativeToAbsoluteFraction(fraction));
+			}
+			else
+			{
+				_state = new Tuple<string, float>(_state.Item1, RelativeToAbsoluteFraction(fraction));
+			}
+		}
+
+		public void Set(float fraction)
+		{
+			_state = new Tuple<string, float>(_state.Item1, RelativeToAbsoluteFraction(fraction));
+		}
+
+		public void Increment(float fraction)
+		{
+			Set(_state.Item2 + RelativeToAbsoluteFraction(fraction));
+		}
+
+		public void Push(float maxFraction)
+		{
+			_ranges.Push(new Tuple<float, float>(_state.Item2, RelativeToAbsoluteFraction(maxFraction)));
+		}
+
+		public void Pop()
+		{
+			if (_ranges.Count > 1)
+			{
+				_state = new Tuple<string, float>(_state.Item1, _ranges.Pop().Item2);
+			}
+		}
+
+		float RelativeToAbsoluteFraction(float fraction)
+		{
+			Tuple<float, float> range = _ranges.Peek();
+			return range.Item1 + (range.Item2 - range.Item1) * fraction;
+		}
+	}
+
+	public static class ProgressTextReader
+	{
+		public const string DirectivePrefix = "@progress ";
+
+		public static string? ParseLine(string line, ProgressValue value)
+		{
+			string trimLine = line.Trim();
+			if (trimLine.StartsWith(DirectivePrefix, StringComparison.Ordinal))
+			{
+				// Line that just contains a progress directive
+				bool skipLine = false;
+				ProcessInternal(trimLine.Substring(DirectivePrefix.Length), ref skipLine, value);
+				return null;
+			}
+			else
+			{
+				bool skipLine = false;
+				string remainingLine = line;
+
+				// Look for a progress directive at the end of a line, in square brackets
+				if (trimLine.EndsWith("]", StringComparison.Ordinal))
+				{
+					for (int lastIdx = trimLine.Length - 2; lastIdx >= 0 && trimLine[lastIdx] != ']'; lastIdx--)
+					{
+						if (trimLine[lastIdx] == '[')
+						{
+							string directiveSubstring = trimLine.Substring(lastIdx + 1, trimLine.Length - lastIdx - 2);
+							if (directiveSubstring.StartsWith(DirectivePrefix, StringComparison.Ordinal))
+							{
+								ProcessInternal(directiveSubstring.Substring(DirectivePrefix.Length), ref skipLine, value);
+								remainingLine = line.Substring(0, lastIdx).TrimEnd();
+							}
+							break;
+						}
+					}
+				}
+
+				if (skipLine)
+				{
+					return null;
+				}
+				else
+				{
+					return remainingLine;
+				}
+			}
+		}
+
+		static void ProcessInternal(string line, ref bool skipLine, ProgressValue value)
+		{
+			List<string> tokens = ParseTokens(line);
+			for (int tokenIdx = 0; tokenIdx < tokens.Count;)
+			{
+				float fraction;
+				if (ReadFraction(tokens, ref tokenIdx, out fraction))
+				{
+					value.Set(fraction);
+				}
+				else if (tokens[tokenIdx] == "push")
+				{
+					tokenIdx++;
+					if (ReadFraction(tokens, ref tokenIdx, out fraction))
+					{
+						value.Push(fraction);
+					}
+				}
+				else if (tokens[tokenIdx] == "pop")
+				{
+					tokenIdx++;
+					value.Pop();
+				}
+				else if (tokens[tokenIdx] == "increment")
+				{
+					tokenIdx++;
+					if (ReadFraction(tokens, ref tokenIdx, out fraction))
+					{
+						value.Increment(fraction);
+					}
+				}
+				else if (tokens[tokenIdx] == "skipline")
+				{
+					tokenIdx++;
+					skipLine = true;
+				}
+				else if (tokens[tokenIdx].Length >= 2 && (tokens[tokenIdx][0] == '\'' || tokens[tokenIdx][0] == '\"') && tokens[tokenIdx].Last() == tokens[tokenIdx].First())
+				{
+					string message = tokens[tokenIdx++];
+					value.Set(message.Substring(1, message.Length - 2));
+				}
+				else
+				{
+					tokenIdx++;
+				}
+			}
+		}
+
+		static List<string> ParseTokens(string line)
+		{
+			List<string> tokens = new List<string>();
+			for (int idx = 0; ;)
+			{
+				// Skip whitespace
+				while (idx < line.Length && Char.IsWhiteSpace(line[idx]))
+				{
+					idx++;
+				}
+				if (idx == line.Length)
+				{
+					break;
+				}
+
+				// Read the next token
+				if (Char.IsLetterOrDigit(line[idx]))
+				{
+					int startIdx = idx++;
+					while (idx < line.Length && Char.IsLetterOrDigit(line[idx]))
+					{
+						idx++;
+					}
+					tokens.Add(line.Substring(startIdx, idx - startIdx));
+				}
+				else if (line[idx] == '\'' || line[idx] == '\"')
+				{
+					int startIdx = idx++;
+					while (idx < line.Length && line[idx] != line[startIdx])
+					{
+						idx++;
+					}
+					tokens.Add(line.Substring(startIdx, ++idx - startIdx));
+				}
+				else
+				{
+					tokens.Add(line.Substring(idx++, 1));
+				}
+			}
+			return tokens;
+		}
+
+		static bool ReadFraction(List<string> tokens, ref int tokenIdx, out float fraction)
+		{
+			// Read a fraction in the form x%
+			if (tokenIdx + 2 <= tokens.Count && tokens[tokenIdx + 1] == "%")
+			{
+				int numerator;
+				if (Int32.TryParse(tokens[tokenIdx], out numerator))
+				{
+					fraction = (float)numerator / 100.0f;
+					tokenIdx += 2;
+					return true;
+				}
+			}
+
+			// Read a fraction in the form x/y
+			if (tokenIdx + 3 <= tokens.Count && tokens[tokenIdx + 1] == "/")
+			{
+				int numerator, denominator;
+				if (Int32.TryParse(tokens[tokenIdx], out numerator) && Int32.TryParse(tokens[tokenIdx + 2], out denominator))
+				{
+					fraction = (float)numerator / (float)denominator;
+					tokenIdx += 3;
+					return true;
+				}
+			}
+
+			fraction = 0.0f;
+			return false;
+		}
+	}
+}
diff --git a/Engine/Source/Programs/UnrealGameSync/UnrealGameSyncShared/OutputAdapters.cs b/Engine/Source/Programs/UnrealGameSync/UnrealGameSyncShared/OutputAdapters.cs
--- a/Engine/Source/Programs/UnrealGameSync/UnrealGameSyncShared/OutputAdapters.cs
+++ b/Engine/Source/Programs/UnrealGameSync/UnrealGameSyncShared/OutputAdapters.cs
@@ -29,244 +29,5 @@ namespace UnrealGameSync
 			_inner.Log(logLevel, eventId, state, exception, (state, exception) => _prefix + formatter(state, exception));
 		}
 	}
-
-	public class ProgressValue
-	{
-		Tuple<string, float> _state = null!;
-		readonly Stack<Tuple<float, float>> _ranges = new Stack<Tuple<float, float>>();
-
-		public ProgressValue()
-		{
-			Clear();
-		}
-
-		public void Clear()
-		{
-			_state = new Tuple<string, float>("Starting...", 0.0f);
-
-			_ranges.Clear();
-			_ranges.Push(new Tuple<float, float>(0.0f, 1.0f));
-		}
-
-		public Tuple<string, float> Current => _state;
-
-		public void Set(string message)
-		{
-			if (_ranges.Count == 1)
-			{
-				_state = new Tuple<string, float>(message, _state.Item2);
-			}
-		}
-
-		public void Set(string message, float fraction)
-		{
-			if (_ranges.Count == 1)
-			{
-				_state = new Tuple<string, float>(message, RelativeToAbsoluteFraction(fraction));
-			}
-			else
-			{
-				_state = new Tuple<string, float>(_state.Item1, RelativeToAbsoluteFraction(fraction));
-			}
-		}
-
-		public void Set(float fraction)
-		{
-			_state = new Tuple<string, float>(_state.Item1, RelativeToAbsoluteFraction(fraction));
-		}
-
-		public void Increment(float fraction)
-		{
-			Set(_state.Item2 + RelativeToAbsoluteFraction(fraction));
-		}
-
-		public void Push(float maxFraction)
-		{
-			_ranges.Push(new Tuple<float, float>(_state.Item2, RelativeToAbsoluteFraction(maxFraction)));
-		}
-
-		public void Pop()
-		{
-			if (_ranges.Count > 1)
-			{
-				_state = new Tuple<string, float>(_state.Item1, _ranges.Pop().Item2);
-			}
-		}
-
-		float RelativeToAbsoluteFraction(float fraction)
-		{
-			Tuple<float, float> range = _ranges.Peek();
-			return range.Item1 + (range.Item2 - range.Item1) * fraction;
-		}
-	}
-
-	static class ProgressTextWriter
-	{
-		const string DirectivePrefix = "@progress ";
-
-		public static string? ParseLine(string line, ProgressValue value)
-		{
-			string trimLine = line.Trim();
-			if (trimLine.StartsWith(DirectivePrefix, StringComparison.Ordinal))
-			{
-				// Line that just contains a progress directive
-				bool skipLine = false;
-				ProcessInternal(trimLine.Substring(DirectivePrefix.Length), ref skipLine, value);
-				return null;
-			}
-			else
-			{
-				bool skipLine = false;
-				string remainingLine = line;
-
-				// Look for a progress directive at the end of a line, in square brackets
-				if (trimLine.EndsWith("]", StringComparison.Ordinal))
-				{
-					for (int lastIdx = trimLine.Length - 2; lastIdx >= 0 && trimLine[lastIdx] != ']'; lastIdx--)
-					{
-						if (trimLine[lastIdx] == '[')
-						{
-							string directiveSubstring = trimLine.Substring(lastIdx + 1, trimLine.Length - lastIdx - 2);
-							if (directiveSubstring.StartsWith(DirectivePrefix, StringComparison.Ordinal))
-							{
-								ProcessInternal(directiveSubstring.Substring(DirectivePrefix.Length), ref skipLine, value);
-								remainingLine = line.Substring(0, lastIdx).TrimEnd();
-							}
-							break;
-						}
-					}
-				}
-
-				if (skipLine)
-				{
-					return null;
-				}
-				else
-				{
-					return remainingLine;
-				}
-			}
-		}
-
-		static void ProcessInternal(string line, ref bool skipLine, ProgressValue value)
-		{
-			List<string> tokens = ParseTokens(line);
-			for (int tokenIdx = 0; tokenIdx < tokens.Count;)
-			{
-				float fraction;
-				if (ReadFraction(tokens, ref tokenIdx, out fraction))
-				{
-					value.Set(fraction);
-				}
-				else if (tokens[tokenIdx] == "push")
-				{
-					tokenIdx++;
-					if (ReadFraction(tokens, ref tokenIdx, out fraction))
-					{
-						value.Push(fraction);
-					}
-				}
-				else if (tokens[tokenIdx] == "pop")
-				{
-					tokenIdx++;
-					value.Pop();
-				}
-				else if (tokens[tokenIdx] == "increment")
-				{
-					tokenIdx++;
-					if (ReadFraction(tokens, ref tokenIdx, out fraction))
-					{
-						value.Increment(fraction);
-					}
-				}
-				else if (tokens[tokenIdx] == "skipline")
-				{
-					tokenIdx++;
-					skipLine = true;
-				}
-				else if (tokens[tokenIdx].Length >= 2 && (tokens[tokenIdx][0] == '\'' || tokens[tokenIdx][0] == '\"') && tokens[tokenIdx].Last() == tokens[tokenIdx].First())
-				{
-					string message = tokens[tokenIdx++];
-					value.Set(message.Substring(1, message.Length - 2));
-				}
-				else
-				{
-					tokenIdx++;
-				}
-			}
-		}
-
-		static List<string> ParseTokens(string line)
-		{
-			List<string> tokens = new List<string>();
-			for (int idx = 0; ;)
-			{
-				// Skip whitespace
-				while (idx < line.Length && Char.IsWhiteSpace(line[idx]))
-				{
-					idx++;
-				}
-				if (idx == line.Length)
-				{
-					break;
-				}
-
-				// Read the next token
-				if (Char.IsLetterOrDigit(line[idx]))
-				{
-					int startIdx = idx++;
-					while (idx < line.Length && Char.IsLetterOrDigit(line[idx]))
-					{
-						idx++;
-					}
-					tokens.Add(line.Substring(startIdx, idx - startIdx));
-				}
-				else if (line[idx] == '\'' || line[idx] == '\"')
-				{
-					int startIdx = idx++;
-					while (idx < line.Length && line[idx] != line[startIdx])
-					{
-						idx++;
-					}
-					tokens.Add(line.Substring(startIdx, ++idx - startIdx));
-				}
-				else
-				{
-					tokens.Add(line.Substring(idx++, 1));
-				}
-			}
-			return tokens;
-		}
-
-		static bool ReadFraction(List<string> tokens, ref int tokenIdx, out float fraction)
-		{
-			// Read a fraction in the form x%
-			if (tokenIdx + 2 <= tokens.Count && tokens[tokenIdx + 1] == "%")
-			{
-				int numerator;
-				if (Int32.TryParse(tokens[tokenIdx], out numerator))
-				{
-					fraction = (float)numerator / 100.0f;
-					tokenIdx += 2;
-					return true;
-				}
-			}
-
-			// Read a fraction in the form x/y
-			if (tokenIdx + 3 <= tokens.Count && tokens[tokenIdx + 1] == "/")
-			{
-				int numerator, denominator;
-				if (Int32.TryParse(tokens[tokenIdx], out numerator) && Int32.TryParse(tokens[tokenIdx + 2], out denominator))
-				{
-					fraction = (float)numerator / (float)denominator;
-					tokenIdx += 3;
-					return true;
-				}
-			}
-
-			fraction = 0.0f;
-			return false;
-		}
-	}
 }
 
diff --git a/Engine/Source/Programs/UnrealGameSync/UnrealGameSyncShared/WorkspaceUpdate.cs b/Engine/Source/Programs/UnrealGameSync/UnrealGameSyncShared/WorkspaceUpdate.cs
--- a/Engine/Source/Programs/UnrealGameSync/UnrealGameSyncShared/WorkspaceUpdate.cs
+++ b/Engine/Source/Programs/UnrealGameSync/UnrealGameSyncShared/WorkspaceUpdate.cs
@@ -1767,7 +1767,7 @@ namespace UnrealGameSync
 
 		static void ProcessOutput(string line, string prefix, ProgressValue progress, ILogger logger)
 		{
-			string? parsedLine = ProgressTextWriter.ParseLine(line, progress);
+			string? parsedLine = ProgressTextReader.ParseLine(line, progress);
 			if (parsedLine != null)
 			{
 				logger.LogInformation("{Prefix}{Message}", prefix, parsedLine);
