UGS: Fix compiler warnings

#skipci
#rnx

diff --git a/Engine/Source/Programs/UnrealGameSync/UnrealGameSyncShared/OutputAdapters.cs b/Engine/Source/Programs/UnrealGameSync/UnrealGameSyncShared/OutputAdapters.cs
--- a/Engine/Source/Programs/UnrealGameSync/UnrealGameSyncShared/OutputAdapters.cs
+++ b/Engine/Source/Programs/UnrealGameSync/UnrealGameSyncShared/OutputAdapters.cs
@@ -1,8 +1,6 @@
 // Copyright Epic Games, Inc. All Rights Reserved.
 
 using System;
-using System.Collections.Generic;
-using System.Linq;
 using Microsoft.Extensions.Logging;
 
 #nullable enable
