Fixed wrong preset used in sync categories

diff --git a/Engine/Source/Programs/UnrealGameSync/UnrealGameSyncCmd/Utils/SyncFilterUtils.cs b/Engine/Source/Programs/UnrealGameSync/UnrealGameSyncCmd/Utils/SyncFilterUtils.cs
--- a/Engine/Source/Programs/UnrealGameSync/UnrealGameSyncCmd/Utils/SyncFilterUtils.cs
+++ b/Engine/Source/Programs/UnrealGameSync/UnrealGameSyncCmd/Utils/SyncFilterUtils.cs
@@ -8,13 +8,18 @@ namespace UnrealGameSyncCmd.Utils
 {
 	internal static class SyncFilterUtils
 	{
-		internal static string[] ReadSyncFilter(UserWorkspaceSettings workspaceSettings, GlobalSettingsFile userSettings, ConfigFile projectConfig, string projectIdentifier)
+		internal static string[] ReadSyncFilter(
+			UserWorkspaceSettings workspaceSettings, 
+			UserProjectSettings contextUserProjectSettings,
+			GlobalSettingsFile userSettings,
+			ConfigFile projectConfig,
+			string projectIdentifier)
 		{
 			Dictionary<Guid, WorkspaceSyncCategory> syncCategories = ConfigUtils.GetSyncCategories(projectConfig);
-
+			
 			// check if any category is from the role
 			IDictionary<string, Preset> roles = ConfigUtils.GetPresets(projectConfig, projectIdentifier);
-			if (roles.TryGetValue(workspaceSettings.Preset, out Preset? role))
+			if (roles.TryGetValue(contextUserProjectSettings.Preset, out Preset? role))
 			{
 				foreach (RoleCategory roleCategory in role.Categories.Values)
 				{
@@ -27,7 +32,7 @@ namespace UnrealGameSyncCmd.Utils
 
 			ConfigSection? perforceSection = projectConfig.FindSection("Perforce");
 
-			string[] combinedSyncFilter = GlobalSettingsFile.GetCombinedSyncFilter(syncCategories, workspaceSettings.Preset, roles, userSettings.Global.Filter, workspaceSettings.Filter, perforceSection);
+			string[] combinedSyncFilter = GlobalSettingsFile.GetCombinedSyncFilter(syncCategories, contextUserProjectSettings.Preset, roles, userSettings.Global.Filter, workspaceSettings.Filter, perforceSection);
 
 			return combinedSyncFilter;
 		}
diff --git a/Engine/Source/Programs/UnrealGameSync/UnrealGameSyncShared/UserSettings.cs b/Engine/Source/Programs/UnrealGameSync/UnrealGameSyncShared/UserSettings.cs
--- a/Engine/Source/Programs/UnrealGameSync/UnrealGameSyncShared/UserSettings.cs
+++ b/Engine/Source/Programs/UnrealGameSync/UnrealGameSyncShared/UserSettings.cs
@@ -395,7 +395,6 @@ namespace UnrealGameSync
 		// The currently selected project, relative to the root directory
 		public string ProjectPath { get; set; } = String.Empty;
 
-		[Obsolete("This property is obsolete. Use UserProjectSettings:Preset instead.", false)]
 		public string Preset { get; set; } = String.Empty;
 		
 		// Workspace specific SyncFilters
