UGS: Fix compiler warnings

#rnx
#Synced-CL Unknown

diff --git a/Engine/Source/Programs/UnrealGameSync/UnrealGameSync/Controls/WorkspaceSyncFilterControl.cs b/Engine/Source/Programs/UnrealGameSync/UnrealGameSync/Controls/WorkspaceSyncFilterControl.cs
--- a/Engine/Source/Programs/UnrealGameSync/UnrealGameSync/Controls/WorkspaceSyncFilterControl.cs
+++ b/Engine/Source/Programs/UnrealGameSync/UnrealGameSync/Controls/WorkspaceSyncFilterControl.cs
@@ -27,7 +27,9 @@ namespace UnrealGameSync.Controls
 
 		private readonly BindingList<ViewItem> _views = new BindingList<ViewItem>();
 
+#pragma warning disable CA2227 // Collection properties should be read only
 		public ISet<string> BaseContentPaths { get; set; } = new HashSet<string>();
+#pragma warning restore CA2227
 
 		public WorkspaceSyncFilterControl()
 		{
diff --git a/Engine/Source/Programs/UnrealGameSync/UnrealGameSyncCmd/Commands/ArchiveCommand.cs b/Engine/Source/Programs/UnrealGameSync/UnrealGameSyncCmd/Commands/ArchiveCommand.cs
--- a/Engine/Source/Programs/UnrealGameSync/UnrealGameSyncCmd/Commands/ArchiveCommand.cs
+++ b/Engine/Source/Programs/UnrealGameSync/UnrealGameSyncCmd/Commands/ArchiveCommand.cs
@@ -6,10 +6,8 @@ using System.Linq;
 using System.Threading;
 using System.Threading.Tasks;
 using EpicGames.Core;
-using EpicGames.Perforce;
 using Microsoft.Extensions.Logging;
 using UnrealGameSync;
-using UnrealGameSyncCmd.Utils;
 
 namespace UnrealGameSyncCmd.Commands
 {
diff --git a/Engine/Source/Programs/UnrealGameSync/UnrealGameSyncCmd/Commands/BuildCommand.cs b/Engine/Source/Programs/UnrealGameSync/UnrealGameSyncCmd/Commands/BuildCommand.cs
--- a/Engine/Source/Programs/UnrealGameSync/UnrealGameSyncCmd/Commands/BuildCommand.cs
+++ b/Engine/Source/Programs/UnrealGameSync/UnrealGameSyncCmd/Commands/BuildCommand.cs
@@ -12,7 +12,6 @@ using EpicGames.Perforce;
 using Microsoft.Extensions.DependencyInjection;
 using Microsoft.Extensions.Logging;
 using UnrealGameSync;
-using UnrealGameSyncCmd.Utils;
 
 using UserErrorException = UnrealGameSyncCmd.Exceptions.UserErrorException;
 
diff --git a/Engine/Source/Programs/UnrealGameSync/UnrealGameSyncCmd/Commands/ChangesCommand.cs b/Engine/Source/Programs/UnrealGameSync/UnrealGameSyncCmd/Commands/ChangesCommand.cs
--- a/Engine/Source/Programs/UnrealGameSync/UnrealGameSyncCmd/Commands/ChangesCommand.cs
+++ b/Engine/Source/Programs/UnrealGameSync/UnrealGameSyncCmd/Commands/ChangesCommand.cs
@@ -8,7 +8,6 @@ using EpicGames.Core;
 using EpicGames.Perforce;
 using Microsoft.Extensions.Logging;
 using UnrealGameSync;
-using UnrealGameSyncCmd.Utils;
 
 namespace UnrealGameSyncCmd.Commands
 {
@@ -32,7 +31,6 @@ namespace UnrealGameSyncCmd.Commands
 			context.Arguments!.ApplyTo(options);
 			context.Arguments!.CheckAllArgumentsUsed(logger);
 
-
 			List<ChangesRecord> changes = await perforceClient.GetChangesAsync(EpicGames.Perforce.ChangesOptions.None, options.Count, ChangeStatus.Submitted, $"//{settings.ClientName}/...");
 			foreach (IEnumerable<ChangesRecord> changesBatch in changes.Batch(10))
 			{
diff --git a/Engine/Source/Programs/UnrealGameSync/UnrealGameSyncCmd/Commands/ClientsCommand.cs b/Engine/Source/Programs/UnrealGameSync/UnrealGameSyncCmd/Commands/ClientsCommand.cs
--- a/Engine/Source/Programs/UnrealGameSync/UnrealGameSyncCmd/Commands/ClientsCommand.cs
+++ b/Engine/Source/Programs/UnrealGameSync/UnrealGameSyncCmd/Commands/ClientsCommand.cs
@@ -21,7 +21,7 @@ namespace UnrealGameSyncCmd.Commands
 		{
 			ILogger logger = context.Logger!;
 			IPerforceConnection perforceClient = context.PerforceClient!;
-			ClientsCommandOptions options = context.Arguments!.ApplyTo<ClientsCommandOptions>(logger);
+			ClientsCommandOptions _ = context.Arguments!.ApplyTo<ClientsCommandOptions>(logger);
 			context.Arguments.CheckAllArgumentsUsed();
 
 			InfoRecord info = await perforceClient.GetInfoAsync(InfoOptions.ShortOutput);
diff --git a/Engine/Source/Programs/UnrealGameSync/UnrealGameSyncCmd/Commands/ConfigCommand.cs b/Engine/Source/Programs/UnrealGameSync/UnrealGameSyncCmd/Commands/ConfigCommand.cs
--- a/Engine/Source/Programs/UnrealGameSync/UnrealGameSyncCmd/Commands/ConfigCommand.cs
+++ b/Engine/Source/Programs/UnrealGameSync/UnrealGameSyncCmd/Commands/ConfigCommand.cs
@@ -8,7 +8,6 @@ using Microsoft.Extensions.Logging;
 
 using UnrealGameSync;
 using UnrealGameSyncCmd.Options;
-using UnrealGameSyncCmd.Utils;
 
 namespace UnrealGameSyncCmd.Commands
 {
diff --git a/Engine/Source/Programs/UnrealGameSync/UnrealGameSyncCmd/Commands/LoginCommand.cs b/Engine/Source/Programs/UnrealGameSync/UnrealGameSyncCmd/Commands/LoginCommand.cs
--- a/Engine/Source/Programs/UnrealGameSync/UnrealGameSyncCmd/Commands/LoginCommand.cs
+++ b/Engine/Source/Programs/UnrealGameSync/UnrealGameSyncCmd/Commands/LoginCommand.cs
@@ -10,7 +10,6 @@ using Microsoft.Extensions.Logging;
 using EpicGames.Core;
 using EpicGames.OIDC;
 using UnrealGameSync;
-using UnrealGameSyncCmd.Utils;
 
 using UserErrorException = UnrealGameSyncCmd.Exceptions.UserErrorException;
 
diff --git a/Engine/Source/Programs/UnrealGameSync/UnrealGameSyncCmd/Commands/RunCommand.cs b/Engine/Source/Programs/UnrealGameSync/UnrealGameSyncCmd/Commands/RunCommand.cs
--- a/Engine/Source/Programs/UnrealGameSync/UnrealGameSyncCmd/Commands/RunCommand.cs
+++ b/Engine/Source/Programs/UnrealGameSync/UnrealGameSyncCmd/Commands/RunCommand.cs
@@ -3,11 +3,8 @@
 using System;
 using System.Collections.Generic;
 using System.IO;
-using System.Linq;
-using System.Threading;
 using System.Threading.Tasks;
 using EpicGames.Core;
-using EpicGames.Perforce;
 using Microsoft.Extensions.Logging;
 using UnrealGameSync;
 using UnrealGameSyncCmd.Utils;
@@ -26,9 +23,7 @@ namespace UnrealGameSyncCmd.Commands
 		private static void ExecuteInternal(CommandContext context)
 		{
 			ILogger logger = context.Logger!;
-			IPerforceConnection perforceClient = context.PerforceClient!;
 			UserWorkspaceSettings settings = context.UserWorkspaceSettings!;
-			WorkspaceStateWrapper state = context.WorkspaceState!;
 			ProjectInfo projectInfo = context.ProjectInfo!;
 			ConfigFile projectConfig = context.ProjectConfig!;
 
diff --git a/Engine/Source/Programs/UnrealGameSync/UnrealGameSyncCmd/Commands/SettingsCommand.cs b/Engine/Source/Programs/UnrealGameSync/UnrealGameSyncCmd/Commands/SettingsCommand.cs
--- a/Engine/Source/Programs/UnrealGameSync/UnrealGameSyncCmd/Commands/SettingsCommand.cs
+++ b/Engine/Source/Programs/UnrealGameSync/UnrealGameSyncCmd/Commands/SettingsCommand.cs
@@ -6,8 +6,6 @@ using System.Text.Json;
 using System.Threading.Tasks;
 using EpicGames.Core;
 using Microsoft.Extensions.Logging;
-using UnrealGameSync;
-using UnrealGameSyncCmd.Utils;
 
 namespace UnrealGameSyncCmd.Commands
 {
diff --git a/Engine/Source/Programs/UnrealGameSync/UnrealGameSyncCmd/Commands/ToolsCommand.cs b/Engine/Source/Programs/UnrealGameSync/UnrealGameSyncCmd/Commands/ToolsCommand.cs
--- a/Engine/Source/Programs/UnrealGameSync/UnrealGameSyncCmd/Commands/ToolsCommand.cs
+++ b/Engine/Source/Programs/UnrealGameSync/UnrealGameSyncCmd/Commands/ToolsCommand.cs
@@ -14,7 +14,6 @@ using Microsoft.Extensions.DependencyInjection;
 using Microsoft.Extensions.Logging;
 
 using UnrealGameSync;
-using UnrealGameSyncCmd.Utils;
 
 namespace UnrealGameSyncCmd.Commands
 {
diff --git a/Engine/Source/Programs/UnrealGameSync/UnrealGameSyncShared/WorkspaceUpdate.cs b/Engine/Source/Programs/UnrealGameSync/UnrealGameSyncShared/WorkspaceUpdate.cs
--- a/Engine/Source/Programs/UnrealGameSync/UnrealGameSyncShared/WorkspaceUpdate.cs
+++ b/Engine/Source/Programs/UnrealGameSync/UnrealGameSyncShared/WorkspaceUpdate.cs
@@ -115,9 +115,11 @@ namespace UnrealGameSync
 		public string SnapshotNamespace { get; set; } = "";
 
 		// May be updated during sync
+#pragma warning disable CA2227 // Collection properties should be read only
 		public ConfigFile? ProjectConfigFile { get; set; }
 		public IReadOnlyList<string>? ProjectStreamFilter { get; set; }
 		public ISet<string> BaseContentPaths { get; set; } = new HashSet<string>();
+#pragma warning restore CA2227
 
 		public WorkspaceUpdateContext(int changeNumber, WorkspaceUpdateOptions options, BuildConfig editorConfig, string[]? syncFilter, List<ConfigObject> userBuildSteps, HashSet<Guid>? customBuildSteps)
 		{
