Reverted logging to debug to avoid spamming ugs cli

#Synced-CL 48449430

diff --git a/Engine/Source/Programs/UnrealGameSync/UnrealGameSyncShared/ConfigUtils.cs b/Engine/Source/Programs/UnrealGameSync/UnrealGameSyncShared/ConfigUtils.cs
--- a/Engine/Source/Programs/UnrealGameSync/UnrealGameSyncShared/ConfigUtils.cs
+++ b/Engine/Source/Programs/UnrealGameSync/UnrealGameSyncShared/ConfigUtils.cs
@@ -130,7 +130,7 @@ namespace UnrealGameSync
 						try
 						{
 							projectConfig.Parse(lines.ToArray(), projectIdentifier);
-							logger.LogInformation("Read config file from {DepotFile}", fileRecord.DepotFile);
+							logger.LogDebug("Read config file from {DepotFile}", fileRecord.DepotFile);
 						}
 						catch (Exception ex)
 						{
