added the depends on tools to the list of tools enabled explicitely
#Synced-CL 48449430

diff --git a/Engine/Source/Programs/UnrealGameSync/UnrealGameSyncCmd/Commands/ToolsCommand.cs b/Engine/Source/Programs/UnrealGameSync/UnrealGameSyncCmd/Commands/ToolsCommand.cs
--- a/Engine/Source/Programs/UnrealGameSync/UnrealGameSyncCmd/Commands/ToolsCommand.cs
+++ b/Engine/Source/Programs/UnrealGameSync/UnrealGameSyncCmd/Commands/ToolsCommand.cs
@@ -105,6 +105,7 @@ namespace UnrealGameSyncCmd.Commands
 			// update the list of enabled tools and update (this will update all the tools)
 			if (options.Install.Any() || options.Uninstall.Any())
 			{
+				IReadOnlyDictionary<Guid, ToolInfo> tools = toolUpdateMonitor.GetTools().ToDictionary(t => t.Id);
 				IReadOnlyList<ToolInfo> enabled = toolUpdateMonitor.GetEnabledTools();
 				HashSet<Guid> newEnabledTools = new HashSet<Guid>(enabled.Select(t => t.Id));
 
@@ -112,17 +113,24 @@ namespace UnrealGameSyncCmd.Commands
 				{
 					if (Guid.TryParse(id, out Guid guid))
 					{
-						foreach (ToolInfo toolInfo in toolUpdateMonitor.GetTools())
+						if (tools.TryGetValue(guid, out ToolInfo? toolInfo))
 						{
-							if (toolInfo.Id == guid)
+							logger.LogInformation("Adding tool '{Id}' - '{Name}' to the list of enabled tools", id, toolInfo.Name);
+							newEnabledTools.Add(guid);
+
+							foreach (Guid dependsOnToolId in toolInfo.DependsOnToolIds)
 							{
-								logger.LogInformation("Adding tool '{Id}' - '{Name}' to the list of enabled tools", id, toolInfo.Name);
-								newEnabledTools.Add(guid);
-								break;
+								if (tools.TryGetValue(dependsOnToolId, out ToolInfo? dependsOnToolInfo))
+								{
+									logger.LogInformation("Tool '{Id}' - '{Name}' will be installed as it is a dependency of '{ToolName}'", dependsOnToolId, dependsOnToolInfo.Name, toolInfo.Name);
+								}
+								else
+								{
+									logger.LogError("Could not install tool '{Id}' dependency '{DependsOnToolId}'", id, dependsOnToolId);
+								}
 							}
 						}
-
-						if (!newEnabledTools.Contains(guid))
+						else
 						{
 							logger.LogError("Could not install tool '{Id}'", id);
 						}
