[UGS] emergency fix, because command context is centralized we missed the case where the init is called and the user settings does not exist.

#Synced-CL 48482300

diff --git a/Engine/Source/Programs/UnrealGameSync/UnrealGameSyncCmd/Command.cs b/Engine/Source/Programs/UnrealGameSync/UnrealGameSyncCmd/Command.cs
--- a/Engine/Source/Programs/UnrealGameSync/UnrealGameSyncCmd/Command.cs
+++ b/Engine/Source/Programs/UnrealGameSync/UnrealGameSyncCmd/Command.cs
@@ -5,7 +5,12 @@ using UnrealGameSync;
 
 namespace UnrealGameSyncCmd
 {
-	internal abstract class Command
+	internal interface ICommand
+	{
+
+	}
+
+	internal abstract class Command : ICommand
 	{
 		internal static BuildConfig EditorConfig => BuildConfig.Development;
 
diff --git a/Engine/Source/Programs/UnrealGameSync/UnrealGameSyncCmd/Commands/InitCommand.cs b/Engine/Source/Programs/UnrealGameSync/UnrealGameSyncCmd/Commands/InitCommand.cs
--- a/Engine/Source/Programs/UnrealGameSync/UnrealGameSyncCmd/Commands/InitCommand.cs
+++ b/Engine/Source/Programs/UnrealGameSync/UnrealGameSyncCmd/Commands/InitCommand.cs
@@ -19,20 +19,21 @@ using UserErrorException = UnrealGameSyncCmd.Exceptions.UserErrorException;
 
 namespace UnrealGameSyncCmd.Commands
 {
-	internal class InitCommand : Command
+	internal class InitCommand : ICommand
 	{
-		public override async Task ExecuteAsync(CommandContext context)
+		public async Task ExecuteAsync(ILogger logger, CommandLineArguments args, ILoggerFactory loggerFactory, UserSettings settings)
 		{
-			ILogger logger = context.Logger!;
-			IPerforceConnection perforce = context.PerforceClient!;
 			// Get the positional argument indicating the file to look for
 			string? initName;
-			context.Arguments!.TryGetPositionalArgument(out initName);
+			args!.TryGetPositionalArgument(out initName);
 
 			// Get the config settings from the command line
 			InitCommandOptions options = new InitCommandOptions();
-			context.Arguments.ApplyTo(options);
-			context.Arguments.CheckAllArgumentsUsed();
+			args.ApplyTo(options);
+			args.CheckAllArgumentsUsed();
+
+			// do not use the command context perforce as it relies on the .ugs/settings files to exist
+			using IPerforceConnection perforce = await PerforceConnectionUtils.ConnectAsync(options.ServerAndPort, options.UserName, null, loggerFactory);
 
 			// Get the host name
 			InfoRecord perforceInfo = await perforce.GetInfoAsync(InfoOptions.ShortOutput);
@@ -41,11 +42,11 @@ namespace UnrealGameSyncCmd.Commands
 			// Create the perforce connection
 			if (initName != null)
 			{
-				await InitNewClientAsync(perforce, initName, hostName, options, context.GlobalSettings, logger);
+				await InitNewClientAsync(perforce, initName, hostName, options, settings, logger);
 			}
 			else
 			{
-				await InitExistingClientAsync(perforce, hostName, options, context.GlobalSettings, logger);
+				await InitExistingClientAsync(perforce, hostName, options, settings, logger);
 			}
 		}
 
diff --git a/Engine/Source/Programs/UnrealGameSync/UnrealGameSyncCmd/Commands/StatusCommand.cs b/Engine/Source/Programs/UnrealGameSync/UnrealGameSyncCmd/Commands/StatusCommand.cs
--- a/Engine/Source/Programs/UnrealGameSync/UnrealGameSyncCmd/Commands/StatusCommand.cs
+++ b/Engine/Source/Programs/UnrealGameSync/UnrealGameSyncCmd/Commands/StatusCommand.cs
@@ -23,7 +23,7 @@ namespace UnrealGameSyncCmd.Commands
 			bool update = context.Arguments!.HasOption("-Update");
 			context.Arguments.CheckAllArgumentsUsed();
 
-			UserWorkspaceSettings settings = UserSettingsUtils.ReadRequiredUserWorkspaceSettings();
+			UserWorkspaceSettings settings = context.UserWorkspaceSettings;
 			logger.LogInformation("User: {UserName}", settings.UserName);
 			logger.LogInformation("Server: {ServerAndPort}", settings.ServerAndPort);
 			logger.LogInformation("Project: {ClientProjectPath}", settings.ClientProjectPath);
diff --git a/Engine/Source/Programs/UnrealGameSync/UnrealGameSyncCmd/Program.cs b/Engine/Source/Programs/UnrealGameSync/UnrealGameSyncCmd/Program.cs
--- a/Engine/Source/Programs/UnrealGameSync/UnrealGameSyncCmd/Program.cs
+++ b/Engine/Source/Programs/UnrealGameSync/UnrealGameSyncCmd/Program.cs
@@ -22,7 +22,7 @@ using Serilog;
 using UnrealGameSync;
 using UnrealGameSyncCmd.Commands;
 using UnrealGameSyncCmd.Utils;
-
+using Log = EpicGames.Core.Log;
 using UserErrorException = UnrealGameSyncCmd.Exceptions.UserErrorException;
 
 namespace UnrealGameSyncCmd
@@ -165,10 +165,26 @@ namespace UnrealGameSyncCmd
 					}
 				}
 
-				Command instance = (Command)Activator.CreateInstance(command.Type)!;
-				using CommandContext commandContext = await CommandContext.CreateAsync(args, logger, loggerFactory, settings, globalSettings, hordeClient, cloudStorage);
+				ICommand instance = (ICommand)Activator.CreateInstance(command.Type)!;
+
+				switch (instance)
+				{
+					// init command is a special case and should not rely on the command context
+					case InitCommand initCommand:
+						{
+							await initCommand.ExecuteAsync(logger, args, loggerFactory, globalSettings);
+							break;
+						}
+						
+					case Command contextCommand:
+						{
+							using CommandContext commandContext = await CommandContext.CreateAsync(args, logger, loggerFactory, settings, globalSettings, hordeClient, cloudStorage);
+
+							await contextCommand.ExecuteAsync(commandContext!);
+							break;
+						}
+				}
 
-				await instance.ExecuteAsync(commandContext!);
 				return 0;
 			}
 			catch (UserErrorException ex)
