[UGS]

Changed the preset comboboxes' DropDownStyle to DropDownList to avoid being able to freely enter a value
Added and empty value to the workspace preset combobox to give the possibility to clear the preset

#Synced-CL 48482300

diff --git a/Engine/Source/Programs/UnrealGameSync/UnrealGameSync/Forms/SyncFilter.cs b/Engine/Source/Programs/UnrealGameSync/UnrealGameSync/Forms/SyncFilter.cs
--- a/Engine/Source/Programs/UnrealGameSync/UnrealGameSync/Forms/SyncFilter.cs
+++ b/Engine/Source/Programs/UnrealGameSync/UnrealGameSync/Forms/SyncFilter.cs
@@ -227,6 +227,7 @@ namespace UnrealGameSync
 		{
 			presetComboBox.BeginUpdate();
 			presetComboBox.Items.Clear();
+			presetComboBox.Items.Add(String.Empty);
 			presetComboBox.Items.AddRange(presets.ToArray());
 			
 			foreach(object? item in presetComboBox.Items)
