[UGS]

Suspend/Resume Layout when changing the menu

#Synced-CL 48482300

diff --git a/Engine/Source/Programs/UnrealGameSync/UnrealGameSync/Controls/WorkspaceControl.cs b/Engine/Source/Programs/UnrealGameSync/UnrealGameSync/Controls/WorkspaceControl.cs
--- a/Engine/Source/Programs/UnrealGameSync/UnrealGameSync/Controls/WorkspaceControl.cs
+++ b/Engine/Source/Programs/UnrealGameSync/UnrealGameSync/Controls/WorkspaceControl.cs
@@ -6344,6 +6344,9 @@ namespace UnrealGameSync
 		{
 			_hasBuildSteps = false;
 
+			// suspend layout to avoid flickering
+			MoreToolsContextMenu.SuspendLayout();
+
 			int moreToolsItemCount = MoreToolsContextMenu.Items.IndexOf(MoreActionsContextMenu_CustomToolSeparator);
 			while (moreToolsItemCount > 0)
 			{
@@ -6456,6 +6459,8 @@ namespace UnrealGameSync
 			}
 
 			MoreActionsContextMenu_CustomToolSeparator.Visible = (moreToolsItemCount > 0);
+
+			MoreToolsContextMenu.ResumeLayout();
 		}
 
 		private bool CanRunStep(BuildStep step)
