UnrealGameSync: Fix analysis warnings

#rnx
#Synced-CL 47605596

diff --git a/Engine/Source/Programs/UnrealGameSync/UnrealGameSyncCmd/Commands/InitCommand.cs b/Engine/Source/Programs/UnrealGameSync/UnrealGameSyncCmd/Commands/InitCommand.cs
--- a/Engine/Source/Programs/UnrealGameSync/UnrealGameSyncCmd/Commands/InitCommand.cs
+++ b/Engine/Source/Programs/UnrealGameSync/UnrealGameSyncCmd/Commands/InitCommand.cs
@@ -21,7 +21,7 @@ namespace UnrealGameSyncCmd.Commands
 {
 	internal class InitCommand : ICommand
 	{
-		public async Task ExecuteAsync(ILogger logger, CommandLineArguments args, ILoggerFactory loggerFactory, UserSettings settings)
+		public static async Task ExecuteAsync(ILogger logger, CommandLineArguments args, ILoggerFactory loggerFactory, UserSettings? settings)
 		{
 			// Get the positional argument indicating the file to look for
 			string? initName;
diff --git a/Engine/Source/Programs/UnrealGameSync/UnrealGameSyncCmd/Commands/StatusCommand.cs b/Engine/Source/Programs/UnrealGameSync/UnrealGameSyncCmd/Commands/StatusCommand.cs
--- a/Engine/Source/Programs/UnrealGameSync/UnrealGameSyncCmd/Commands/StatusCommand.cs
+++ b/Engine/Source/Programs/UnrealGameSync/UnrealGameSyncCmd/Commands/StatusCommand.cs
@@ -8,7 +8,7 @@ using Microsoft.Extensions.Logging;
 using EpicGames.Perforce;
 
 using UnrealGameSync;
-using UnrealGameSyncCmd.Utils;
+using Microsoft.Extensions.Logging.Abstractions;
 
 namespace UnrealGameSyncCmd.Commands
 {
@@ -16,11 +16,16 @@ namespace UnrealGameSyncCmd.Commands
 	{
 		public override async Task ExecuteAsync(CommandContext context)
 		{
-			ILogger logger = context.Logger!;
-			IPerforceConnection perforceClient = context.PerforceClient!;
-			WorkspaceStateWrapper state = context.WorkspaceState!;
+			ILogger logger = context.Logger ?? NullLogger.Instance;
+			if (context.PerforceClient == null || context.WorkspaceState == null || context.Arguments == null || context.UserWorkspaceSettings == null)
+			{
+				logger.LogWarning("Unable to determine status due to invalid CommandContext");
+				return;
+			}
+			IPerforceConnection perforceClient = context.PerforceClient;
+			WorkspaceStateWrapper state = context.WorkspaceState;
 
-			bool update = context.Arguments!.HasOption("-Update");
+			bool update = context.Arguments.HasOption("-Update");
 			context.Arguments.CheckAllArgumentsUsed();
 
 			UserWorkspaceSettings settings = context.UserWorkspaceSettings;
diff --git a/Engine/Source/Programs/UnrealGameSync/UnrealGameSyncCmd/Program.cs b/Engine/Source/Programs/UnrealGameSync/UnrealGameSyncCmd/Program.cs
--- a/Engine/Source/Programs/UnrealGameSync/UnrealGameSyncCmd/Program.cs
+++ b/Engine/Source/Programs/UnrealGameSync/UnrealGameSyncCmd/Program.cs
@@ -22,7 +22,6 @@ using Serilog;
 using UnrealGameSync;
 using UnrealGameSyncCmd.Commands;
 using UnrealGameSyncCmd.Utils;
-using Log = EpicGames.Core.Log;
 using UserErrorException = UnrealGameSyncCmd.Exceptions.UserErrorException;
 
 namespace UnrealGameSyncCmd
@@ -172,7 +171,7 @@ namespace UnrealGameSyncCmd
 					// init command is a special case and should not rely on the command context
 					case InitCommand initCommand:
 						{
-							await initCommand.ExecuteAsync(logger, args, loggerFactory, globalSettings);
+							await InitCommand.ExecuteAsync(logger, args, loggerFactory, globalSettings);
 							break;
 						}
 						
