[UGS]

Plastering the dialog to avoid showing in modal while the session is locked

#Synced-CL 48482300

diff --git a/Engine/Source/Programs/UnrealGameSync/UnrealGameSync/Program.cs b/Engine/Source/Programs/UnrealGameSync/UnrealGameSync/Program.cs
--- a/Engine/Source/Programs/UnrealGameSync/UnrealGameSync/Program.cs
+++ b/Engine/Source/Programs/UnrealGameSync/UnrealGameSync/Program.cs
@@ -36,6 +36,9 @@ namespace UnrealGameSync
 		/// </summary>
 		static CreateTelemetrySinkDelegate CreateTelemetrySink { get; } = (userName, sessionId, log) => new NullTelemetrySink();
 
+		// Session lock tracking
+		private static SessionLockTracker SessionLockTracker { get; } = new SessionLockTracker();
+
 		public static string GetVersionString()
 		{
 			AssemblyInformationalVersionAttribute? version = Assembly.GetExecutingAssembly().GetCustomAttribute<AssemblyInformationalVersionAttribute>();
@@ -370,7 +373,15 @@ namespace UnrealGameSync
 			SentrySdk.CaptureException(e.Exception);
 
 			using ThreadExceptionDialog dialog = new ThreadExceptionDialog(e.Exception);
-			dialog.ShowDialog();
+
+			if (SessionLockTracker.IsSessionLocked)
+			{
+				dialog.Show();
+			}
+			else
+			{
+				dialog.ShowDialog();
+			}
 		}
 
 		private static void Application_UnobservedException_Sentry(object? sender, UnobservedTaskExceptionEventArgs args)
diff --git a/Engine/Source/Programs/UnrealGameSync/UnrealGameSync/SessionLockTracker.cs b/Engine/Source/Programs/UnrealGameSync/UnrealGameSync/SessionLockTracker.cs
new file mode 100644
--- /dev/null
+++ b/Engine/Source/Programs/UnrealGameSync/UnrealGameSync/SessionLockTracker.cs
@@ -0,0 +1,30 @@
+﻿// Copyright Epic Games, Inc. All Rights Reserved.
+
+using Microsoft.Win32;
+
+namespace UnrealGameSync
+{
+	public class SessionLockTracker
+	{
+		private volatile bool _sessionLocked = false;
+		
+		public bool IsSessionLocked => _sessionLocked;
+
+		public SessionLockTracker()
+		{
+			SystemEvents.SessionSwitch += OnSessionSwitch;
+		}
+
+		private void OnSessionSwitch(object sender, SessionSwitchEventArgs e)
+		{
+			_sessionLocked = e.Reason switch
+			{
+				SessionSwitchReason.SessionLock => true,
+				SessionSwitchReason.SessionUnlock => false,
+				SessionSwitchReason.ConsoleDisconnect or SessionSwitchReason.RemoteDisconnect => true,
+				SessionSwitchReason.ConsoleConnect or SessionSwitchReason.RemoteConnect => false,
+				_ => _sessionLocked
+			};
+		}
+	}
+}
