Fixed typo

#Synced-CL 48482300

diff --git a/Engine/Source/Programs/UnrealGameSync/UnrealGameSync/Controls/WorkspaceSyncFilterControl.cs b/Engine/Source/Programs/UnrealGameSync/UnrealGameSync/Controls/WorkspaceSyncFilterControl.cs
--- a/Engine/Source/Programs/UnrealGameSync/UnrealGameSync/Controls/WorkspaceSyncFilterControl.cs
+++ b/Engine/Source/Programs/UnrealGameSync/UnrealGameSync/Controls/WorkspaceSyncFilterControl.cs
@@ -77,7 +77,7 @@ namespace UnrealGameSync.Controls
 				}
 				else
 				{
-					currentToolTip = "Base Content is disabled for this project."
+					currentToolTip = "Base Content is disabled for this project.";
 				}
 			}
 
