Allow UGS to specify build step dependencies that are only used when UGS is configured to locally compile executables.  Ensure that commands to compile locally are still functional, it is only the dependencies/requirements (ie: the edges of the graph) that are marked as something that are conditional on mode.

#Synced-CL 48523383

diff --git a/Engine/Source/Programs/UnrealGameSync/UnrealGameSync/Controls/WorkspaceControl.cs b/Engine/Source/Programs/UnrealGameSync/UnrealGameSync/Controls/WorkspaceControl.cs
--- a/Engine/Source/Programs/UnrealGameSync/UnrealGameSync/Controls/WorkspaceControl.cs
+++ b/Engine/Source/Programs/UnrealGameSync/UnrealGameSync/Controls/WorkspaceControl.cs
@@ -6491,9 +6491,15 @@ namespace UnrealGameSync
 				else
 				{
 					HashSet<Guid> stepSet = new HashSet<Guid> { step.UniqueId };
-
+					
+					bool UsingPreCompiledExecutables = GetSelectedArchiveChannels(GetArchiveChannels()).Any(channel => channel.Type == IArchiveChannel.EditorArchiveType);
+					int RequirementCount = step.Requires.Count;
+					if (!UsingPreCompiledExecutables)
+					{
+						RequirementCount += step.LocalCompileRequires.Count;
+					}
 					ConfigFile projectConfigFile = _workspace.ProjectConfigFile;
-					if (projectConfigFile != null && step.Requires.Count > 0)
+					if (projectConfigFile != null && RequirementCount > 0)
 					{
 						Stack<Guid> stack = new Stack<Guid>(stepSet);
 						while (stack.Count > 0)
@@ -6509,6 +6515,16 @@ namespace UnrealGameSync
 										stack.Push(requiresId);
 									}
 								}
+								if (!UsingPreCompiledExecutables)
+								{
+									foreach (Guid requiresId in nextStep.LocalCompileRequires)
+									{
+										if (stepSet.Add(requiresId))
+										{
+											stack.Push(requiresId);
+										}
+									}
+								}
 							}
 						}
 					}
diff --git a/Engine/Source/Programs/UnrealGameSync/UnrealGameSyncShared/BuildStep.cs b/Engine/Source/Programs/UnrealGameSync/UnrealGameSyncShared/BuildStep.cs
--- a/Engine/Source/Programs/UnrealGameSync/UnrealGameSyncShared/BuildStep.cs
+++ b/Engine/Source/Programs/UnrealGameSync/UnrealGameSyncShared/BuildStep.cs
@@ -35,6 +35,7 @@ namespace UnrealGameSync
 		public bool ScheduledSync { get; set; }
 		public string? StatusPanelLink { get; set; }
 		public List<Guid> Requires { get; } = new List<Guid>();
+		public List<Guid> LocalCompileRequires { get; } = new List<Guid>();
 		public Guid ToolId { get; set; }
 
 		public BuildStep(Guid inUniqueId, int inOrderIndex, string? inDescription, string? inStatusText, int inEstimatedDuration, string? inFileName, string? inArguments, string? inWorkingDir, bool inUseLogWindow)
@@ -114,6 +115,16 @@ namespace UnrealGameSync
 			}
 			Requires = requires;
 
+			List<Guid> localCompileRequires = new List<Guid>();
+			foreach (string localCompileRequireString in obj.GetValue("LocalCompileRequires", String.Empty).Split(';', StringSplitOptions.RemoveEmptyEntries))
+			{
+				if (Guid.TryParse(localCompileRequireString, out Guid localCompileRequire))
+				{
+					localCompileRequires.Add(localCompileRequire);
+				}
+			}
+			LocalCompileRequires = localCompileRequires;
+
 			ToolId = obj.GetValue("Tool", Guid.Empty);
 		}
 
