UGS: Fix compiler warnings

#Synced-CL Unknown

diff --git a/Engine/Source/Programs/UnrealGameSync/UnrealGameSync/Controls/SyncFilterControl.cs b/Engine/Source/Programs/UnrealGameSync/UnrealGameSync/Controls/SyncFilterControl.cs
--- a/Engine/Source/Programs/UnrealGameSync/UnrealGameSync/Controls/SyncFilterControl.cs
+++ b/Engine/Source/Programs/UnrealGameSync/UnrealGameSync/Controls/SyncFilterControl.cs
@@ -7,7 +7,6 @@ using System.IO;
 using System.Linq;
 using System.Text;
 using System.Windows.Forms;
-using Microsoft.VisualStudio.OLE.Interop;
 
 namespace UnrealGameSync.Controls
 {
diff --git a/Engine/Source/Programs/UnrealGameSync/UnrealGameSync/Controls/WorkspaceSyncFilterControl.cs b/Engine/Source/Programs/UnrealGameSync/UnrealGameSync/Controls/WorkspaceSyncFilterControl.cs
--- a/Engine/Source/Programs/UnrealGameSync/UnrealGameSync/Controls/WorkspaceSyncFilterControl.cs
+++ b/Engine/Source/Programs/UnrealGameSync/UnrealGameSync/Controls/WorkspaceSyncFilterControl.cs
@@ -5,10 +5,8 @@ using System.Collections.Generic;
 using System.ComponentModel;
 using System.IO;
 using System.Linq;
-using System.Runtime.Serialization;
 using System.Text;
 using System.Windows.Forms;
-using System.Xml;
 
 namespace UnrealGameSync.Controls
 {
@@ -323,7 +321,6 @@ namespace UnrealGameSync.Controls
 			{
 				_views.Add(presetView);
 			}
-
 		}
 	}
 }
