PR #12950: ugs: Fix incorrect sync filter handling

#Synced-CL 0

diff --git a/Engine/Source/Programs/UnrealGameSync/UnrealGameSyncShared/GlobalSettings.cs b/Engine/Source/Programs/UnrealGameSync/UnrealGameSyncShared/GlobalSettings.cs
--- a/Engine/Source/Programs/UnrealGameSync/UnrealGameSyncShared/GlobalSettings.cs
+++ b/Engine/Source/Programs/UnrealGameSync/UnrealGameSyncShared/GlobalSettings.cs
@@ -184,7 +184,7 @@ namespace UnrealGameSync
 			FilterSettings workspaceFilter,
 			ConfigSection? perforceSection)
 		{
-			ISet<string> lines = new HashSet<string>();
+			List<string> lines = new List<string>();
 			foreach (string viewLine in Enumerable.Concat(globalFilter.View, workspaceFilter.View).Select(x => x.Trim()).Where(x => x.Length > 0 && !x.StartsWith(";", StringComparison.Ordinal)))
 			{
 				lines.Add(viewLine);
@@ -242,12 +242,10 @@ namespace UnrealGameSync
 				}
 			}
 
-			// If there are no filtering lines then we can assume that AdditionalPathsToSync is covered and we do not
-			// need to add them manually.
-			if (lines.Count > 0 && perforceSection != null)
+			if (perforceSection != null)
 			{
 				IEnumerable<string> additionalPaths = perforceSection.GetValues("AdditionalPathsToSync", Array.Empty<string>());
-				lines.UnionWith(additionalPaths);
+				lines.AddRange(additionalPaths);
 			}
 
 			return lines.ToArray();
