fixed UGS-CLI commands that do not require a workspace

#Synced-CL 48482300

diff --git a/Engine/Source/Programs/UnrealGameSync/UnrealGameSyncCmd/Commands/ClientsCommand.cs b/Engine/Source/Programs/UnrealGameSync/UnrealGameSyncCmd/Commands/ClientsCommand.cs
--- a/Engine/Source/Programs/UnrealGameSync/UnrealGameSyncCmd/Commands/ClientsCommand.cs
+++ b/Engine/Source/Programs/UnrealGameSync/UnrealGameSyncCmd/Commands/ClientsCommand.cs
@@ -3,11 +3,12 @@
 using System;
 using System.Collections.Generic;
 using System.Threading.Tasks;
-using Microsoft.Extensions.Logging;
-
+using EpicGames.Core;
 using EpicGames.Perforce;
-
+using Microsoft.Extensions.Logging;
+using Polly;
 using UnrealGameSyncCmd.Options;
+using UnrealGameSyncCmd.Utils;
 
 namespace UnrealGameSyncCmd.Commands
 {
@@ -15,14 +16,14 @@ namespace UnrealGameSyncCmd.Commands
 	{
 	}
 
-	internal class ClientsCommand : Command
+	internal class ClientsCommand : ICommand
 	{
-		public override async Task ExecuteAsync(CommandContext context)
+		public static async Task ExecuteAsync(ILogger logger, CommandLineArguments args, ILoggerFactory loggerFactory)
 		{
-			ILogger logger = context.Logger!;
-			IPerforceConnection perforceClient = context.PerforceClient!;
-			ClientsCommandOptions _ = context.Arguments!.ApplyTo<ClientsCommandOptions>(logger);
-			context.Arguments.CheckAllArgumentsUsed();
+			ClientsCommandOptions options = args.ApplyTo<ClientsCommandOptions>(logger);
+			args.CheckAllArgumentsUsed();
+
+			using IPerforceConnection perforceClient = await PerforceConnectionUtils.ConnectAsync(options.ServerAndPort, options.UserName, null, loggerFactory);
 
 			InfoRecord info = await perforceClient.GetInfoAsync(InfoOptions.ShortOutput);
 
diff --git a/Engine/Source/Programs/UnrealGameSync/UnrealGameSyncCmd/Commands/SettingsCommand.cs b/Engine/Source/Programs/UnrealGameSync/UnrealGameSyncCmd/Commands/SettingsCommand.cs
--- a/Engine/Source/Programs/UnrealGameSync/UnrealGameSyncCmd/Commands/SettingsCommand.cs
+++ b/Engine/Source/Programs/UnrealGameSync/UnrealGameSyncCmd/Commands/SettingsCommand.cs
@@ -6,6 +6,8 @@ using System.Text.Json;
 using System.Threading.Tasks;
 using EpicGames.Core;
 using Microsoft.Extensions.Logging;
+using UnrealGameSync;
+using UnrealGameSyncCmd.Utils;
 
 namespace UnrealGameSyncCmd.Commands
 {
@@ -18,33 +20,32 @@ namespace UnrealGameSyncCmd.Commands
 		public string User { get; set; } = String.Empty;
 	}
 
-	internal class SettingsCommand : Command
+	internal class SettingsCommand : ICommand
 	{
-		public override async Task ExecuteAsync(CommandContext context)
+		public static async Task ExecuteAsync(ILogger logger, CommandLineArguments args, UserSettings? globalSettings, GlobalSettingsFile? userSettings)
 		{
-			ILogger logger = context.Logger!;
-
-			SettingsCommandOptions options = context.Arguments!.ApplyTo<SettingsCommandOptions>(logger);
-			context.Arguments.CheckAllArgumentsUsed();
+			SettingsCommandOptions options = args.ApplyTo<SettingsCommandOptions>(logger);
+			args.CheckAllArgumentsUsed();
 
 			if (!String.IsNullOrWhiteSpace(options.Global))
 			{
-				if (context.GlobalSettings != null)
+				if (globalSettings != null)
 				{
 					logger.LogInformation("Writing global settings to {Global}", options.Global);
-					await WriteSettings(options.Global, context.GlobalSettings);
+					await WriteSettings(options.Global, globalSettings);
 				}
-				else if (context.UserSettings != null)
+				else if (userSettings != null)
 				{
 					logger.LogInformation("Writing global settings to {Global}", options.Global);
-					await WriteSettings(options.Global, context.UserSettings);
+					await WriteSettings(options.Global, userSettings);
 				}
 			}
 
 			if (!String.IsNullOrWhiteSpace(options.User))
 			{
 				logger.LogInformation("Writing user settings to {User}", options.User);
-				await WriteSettings(options.User, context.UserWorkspaceSettings!);
+				UserWorkspaceSettings userWorkspaceSettings = UserSettingsUtils.ReadRequiredUserWorkspaceSettings();
+				await WriteSettings(options.User, userWorkspaceSettings);
 			}
 		}
 
diff --git a/Engine/Source/Programs/UnrealGameSync/UnrealGameSyncCmd/Commands/UpgradeCommand.cs b/Engine/Source/Programs/UnrealGameSync/UnrealGameSyncCmd/Commands/UpgradeCommand.cs
--- a/Engine/Source/Programs/UnrealGameSync/UnrealGameSyncCmd/Commands/UpgradeCommand.cs
+++ b/Engine/Source/Programs/UnrealGameSync/UnrealGameSyncCmd/Commands/UpgradeCommand.cs
@@ -25,16 +25,14 @@ namespace UnrealGameSyncCmd.Commands
 		public bool Force { get; set; }
 	}
 
-	internal class UpgradeCommand : Command
+	internal class UpgradeCommand : ICommand
 	{
-		public override async Task ExecuteAsync(CommandContext context)
+		public static async Task ExecuteAsync(ILogger logger, CommandLineArguments args)
 		{
-			ILogger logger = context.Logger!;
-
 			UpgradeCommandOptions options = new UpgradeCommandOptions();
-			context.Arguments!.ApplyTo(options);
-			string? targetDirStr = context.Arguments.GetStringOrDefault("-TargetDir=", null);
-			context.Arguments.CheckAllArgumentsUsed(logger);
+			args.ApplyTo(options);
+			string? targetDirStr = args.GetStringOrDefault("-TargetDir=", null);
+			args.CheckAllArgumentsUsed(logger);
 
 			string currentVersion = VersionUtils.GetVersion();
 
diff --git a/Engine/Source/Programs/UnrealGameSync/UnrealGameSyncCmd/Commands/VersionCommand.cs b/Engine/Source/Programs/UnrealGameSync/UnrealGameSyncCmd/Commands/VersionCommand.cs
--- a/Engine/Source/Programs/UnrealGameSync/UnrealGameSyncCmd/Commands/VersionCommand.cs
+++ b/Engine/Source/Programs/UnrealGameSync/UnrealGameSyncCmd/Commands/VersionCommand.cs
@@ -6,16 +6,13 @@ using UnrealGameSyncCmd.Utils;
 
 namespace UnrealGameSyncCmd.Commands
 {
-	internal class VersionCommand : Command
+	internal class VersionCommand : ICommand
 	{
-		public override Task ExecuteAsync(CommandContext context)
+		public static void Execute(ILogger logger)
 		{
-			ILogger logger = context.Logger!;
-
 			string version = VersionUtils.GetVersion();
 			logger.LogInformation("UnrealGameSync {Version}", version);
 
-			return Task.CompletedTask;
 		}
 	}
 }
diff --git a/Engine/Source/Programs/UnrealGameSync/UnrealGameSyncCmd/Program.cs b/Engine/Source/Programs/UnrealGameSync/UnrealGameSyncCmd/Program.cs
--- a/Engine/Source/Programs/UnrealGameSync/UnrealGameSyncCmd/Program.cs
+++ b/Engine/Source/Programs/UnrealGameSync/UnrealGameSyncCmd/Program.cs
@@ -168,13 +168,42 @@ namespace UnrealGameSyncCmd
 
 				switch (instance)
 				{
-					// init command is a special case and should not rely on the command context
-					case InitCommand initCommand:
+					//////////////////////////////////////
+					case ClientsCommand _:
+						{
+							await ClientsCommand.ExecuteAsync(logger, args, loggerFactory);
+							break;
+						}
+
+					//////////////////////////////////////
+					case InitCommand _:
 						{
 							await InitCommand.ExecuteAsync(logger, args, loggerFactory, globalSettings);
 							break;
 						}
-						
+
+					//////////////////////////////////////
+					case SettingsCommand _:
+						{
+							await SettingsCommand.ExecuteAsync(logger, args, globalSettings, settings);
+							break;
+						}
+
+					//////////////////////////////////////
+					case UpgradeCommand _:
+						{
+							await UpgradeCommand.ExecuteAsync(logger, args);
+							break;
+						}
+
+					//////////////////////////////////////
+					case VersionCommand _:
+						{
+							VersionCommand.Execute(logger);
+							break;
+						}
+
+					//////////////////////////////////////
 					case Command contextCommand:
 						{
 							using CommandContext commandContext = await CommandContext.CreateAsync(args, logger, loggerFactory, settings, globalSettings, hordeClient, cloudStorage);
