changed info message size from static to dynamic
added auto-hide system for very small windows size
added separate control to select when the info messages are shown

#Synced-CL 48482300

diff --git a/Engine/Source/Programs/UnrealGameSync/UnrealGameSync/Controls/StatusPanel.cs b/Engine/Source/Programs/UnrealGameSync/UnrealGameSync/Controls/StatusPanel.cs
--- a/Engine/Source/Programs/UnrealGameSync/UnrealGameSync/Controls/StatusPanel.cs
+++ b/Engine/Source/Programs/UnrealGameSync/UnrealGameSync/Controls/StatusPanel.cs
@@ -454,7 +454,6 @@ namespace UnrealGameSync
 	class StatusPanel : Panel, ICustomThemeHandler
 	{
 		const float LineSpacing = 1.35f;
-		const int MaxInfoLineWidth = 700;
 		const int InfoMargin = 15;
 
 		InvertibleImage _projectLogo = new InvertibleImage(Properties.Resources.DefaultProjectLogo, false);
@@ -476,6 +475,8 @@ namespace UnrealGameSync
 		readonly System.Windows.Forms.ToolTip _toolTip = new System.Windows.Forms.ToolTip();
 		string? _currentToolTip = null;
 
+		private bool _showInfoLines = true;
+		
 		public StatusPanel()
 		{
 			DoubleBuffered = true;
@@ -682,7 +683,33 @@ namespace UnrealGameSync
 			int logoHeight = (int)(drawProjectLogo.Height * logoScale);
 
 			// Figure out where the split between content and the logo is going to be
-			int dividerX = Math.Max( ((Width - logoWidth - _contentWidth) / 2) + logoWidth, MaxInfoLineWidth);
+			int dividerX = ((Width - logoWidth - _contentWidth) / 2) + logoWidth;
+
+			if (_infoLines.Any())
+			{
+				int maxWidth = 0;
+
+				float tempTotalInfoLineHeight = _infoLines.Sum(x => x.LineHeight);
+				float tempInfoLineY = (bodyHeight - InfoMargin - tempTotalInfoLineHeight * (int)(Font.Height * LineSpacing));
+				foreach (StatusLine line in _infoLines)
+				{
+					line.Layout(graphics, new Point(InfoMargin, (int)tempInfoLineY), _resources!);
+					maxWidth = Math.Max(maxWidth, line.Bounds.Right);
+				}
+
+				// init to logo position
+				int minContentX = dividerX - logoWidth;
+
+				if ( minContentX <= maxWidth)
+				{
+					_showInfoLines = false;
+				}
+				else
+				{
+					_showInfoLines = true;
+					dividerX = Math.Max(dividerX, maxWidth);
+				}
+			}
 
 			// Get the logo position
 			int logoX = dividerX - logoWidth;
@@ -808,11 +835,14 @@ namespace UnrealGameSync
 					line.Draw(e.Graphics, _resources!);
 				}
 
-				foreach (StatusLine line in _infoLines)
+				if (_showInfoLines)
 				{
-					line.Draw(e.Graphics, _resources!);
+					foreach (StatusLine line in _infoLines)
+					{
+						line.Draw(e.Graphics, _resources!);
+					}
 				}
-
+				
 				_caption?.Draw(e.Graphics, _resources!);
 
 				if (_alert != null)
