Disabling fast sync when the distance between the current CL and the target is more than 100k

#Synced-CL 48482300

diff --git a/Engine/Source/Programs/UnrealGameSync/UnrealGameSyncShared/WorkspaceUpdate.cs b/Engine/Source/Programs/UnrealGameSync/UnrealGameSyncShared/WorkspaceUpdate.cs
--- a/Engine/Source/Programs/UnrealGameSync/UnrealGameSyncShared/WorkspaceUpdate.cs
+++ b/Engine/Source/Programs/UnrealGameSync/UnrealGameSyncShared/WorkspaceUpdate.cs
@@ -874,7 +874,10 @@ namespace UnrealGameSync
 
 							if (Context.Options.HasFlag(WorkspaceUpdateOptions.Sync) && Context.Options.HasFlag(WorkspaceUpdateOptions.FastSync))
 							{
-								if (state.CurrentChangeNumber > 0 && Context.ChangeNumber >= state.CurrentChangeNumber)
+								bool changelistDistanceInRange = Math.Abs(Context.ChangeNumber - state.CurrentChangeNumber) < 100_000;
+								bool targetChangelistBiggerThanCurrent = (Context.ChangeNumber >= state.CurrentChangeNumber);
+
+								if (state.CurrentChangeNumber > 0 && changelistDistanceInRange && targetChangelistBiggerThanCurrent)
 								{
 									syncFilter = $"{state.CurrentChangeNumber},{Context.ChangeNumber}";
 									logger.LogInformation("Using Fast Sync between current changelist {CurrentChangelist} and target changelist {TargetChangelist}", state.CurrentChangeNumber, Context.ChangeNumber);
@@ -885,17 +888,20 @@ namespace UnrealGameSync
 									{
 										logger.LogInformation("Fast Sync cannot be used with current changelist {CurrentChangelist} and target changelist {TargetChangelist} because there is no current change list to compare to.", state.CurrentChangeNumber, Context.ChangeNumber);
 									}
-
-									if (Context.ChangeNumber < state.CurrentChangeNumber)
+									else if (!targetChangelistBiggerThanCurrent)
 									{
 										logger.LogInformation("Fast Sync cannot be used with current changelist {CurrentChangelist} and target changelist {TargetChangelist} because workspace syncing backwards.", state.CurrentChangeNumber, Context.ChangeNumber);
 									}
+									else if (!changelistDistanceInRange)
+									{
+										logger.LogInformation("Fast Sync cannot be used with current changelist {CurrentChangelist} and target changelist {TargetChangelist} because the distance between then is > 100K changelists.", state.CurrentChangeNumber, Context.ChangeNumber);
+									}
 								}
 							}
 
 							foreach (string syncPath in syncPaths)
 							{
-								async Task<(WorkspaceUpdateResult, string)> inlineGetSyncRecords(bool firstTry)
+								async Task<(WorkspaceUpdateResult, string)> inlineGetSyncRecords(bool firstTry, RecordCounter counter)
 								{
 									List<SyncFile> syncFiles = new List<SyncFile>();
 									await foreach (PerforceResponse<SyncRecord> response in perforce.TrySyncAsync(SyncOptions.PreviewOnly, -1, 0, -1, -1, -1, -1, $"{syncPath}@{syncFilter}", cancellationToken))
@@ -913,7 +919,7 @@ namespace UnrealGameSync
 												syncTree.CanContainExcludedFiles = false;
 												logger.LogWarning("Couldn't enumerate changes matching with fast sync enabled: {SyncPath}. Disabling fast sync and retrying.", syncPath);
 
-												return await inlineGetSyncRecords(false);
+												return await inlineGetSyncRecords(false, counter);
 											}
 
 											return (WorkspaceUpdateResult.FailedToSync, $"Couldn't enumerate changes matching {syncPath}.");
@@ -994,7 +1000,7 @@ namespace UnrealGameSync
 									return (WorkspaceUpdateResult.Success, String.Empty);
 								}
 
-								(WorkspaceUpdateResult, string) ret = await inlineGetSyncRecords(true);
+								(WorkspaceUpdateResult, string) ret = await inlineGetSyncRecords(true, counter);
 								if (ret.Item1 != WorkspaceUpdateResult.Success)
 								{
 									return ret;
