unused usings

#Synced-CL 48986248

diff --git a/Engine/Source/Programs/UnrealGameSync/UnrealGameSyncCmd/Commands/ClientsCommand.cs b/Engine/Source/Programs/UnrealGameSync/UnrealGameSyncCmd/Commands/ClientsCommand.cs
--- a/Engine/Source/Programs/UnrealGameSync/UnrealGameSyncCmd/Commands/ClientsCommand.cs
+++ b/Engine/Source/Programs/UnrealGameSync/UnrealGameSyncCmd/Commands/ClientsCommand.cs
@@ -6,7 +6,6 @@ using System.Threading.Tasks;
 using EpicGames.Core;
 using EpicGames.Perforce;
 using Microsoft.Extensions.Logging;
-using Polly;
 using UnrealGameSyncCmd.Options;
 using UnrealGameSyncCmd.Utils;
 
diff --git a/Engine/Source/Programs/UnrealGameSync/UnrealGameSyncCmd/Commands/VersionCommand.cs b/Engine/Source/Programs/UnrealGameSync/UnrealGameSyncCmd/Commands/VersionCommand.cs
--- a/Engine/Source/Programs/UnrealGameSync/UnrealGameSyncCmd/Commands/VersionCommand.cs
+++ b/Engine/Source/Programs/UnrealGameSync/UnrealGameSyncCmd/Commands/VersionCommand.cs
@@ -1,6 +1,5 @@
 ﻿// Copyright Epic Games, Inc. All Rights Reserved.
 
-using System.Threading.Tasks;
 using Microsoft.Extensions.Logging;
 using UnrealGameSyncCmd.Utils;
 
