fixed test to detect if a preset is the default one

#Synced-CL 49039596

diff --git a/Engine/Source/Programs/UnrealGameSync/UnrealGameSync/Forms/SyncFilter.cs b/Engine/Source/Programs/UnrealGameSync/UnrealGameSync/Forms/SyncFilter.cs
--- a/Engine/Source/Programs/UnrealGameSync/UnrealGameSync/Forms/SyncFilter.cs
+++ b/Engine/Source/Programs/UnrealGameSync/UnrealGameSync/Forms/SyncFilter.cs
@@ -311,7 +311,7 @@ namespace UnrealGameSync
 			WorkspaceControl.GenerateUprojectSpecificSolution.Checked = GlobalControl.GenerateUprojectSpecificSolution.Checked;
 		}
 
-		private void SetPresetData(ComboBox presetComboBox, string preset, IEnumerable<string> availablePresets)
+		private void SetPresetData(ComboBox presetComboBox, string currentPreset, IEnumerable<string> availablePresets)
 		{
 			presetComboBox.BeginUpdate();
 			presetComboBox.Items.Clear();
@@ -319,7 +319,7 @@ namespace UnrealGameSync
 			foreach (string presetName in availablePresets)
 			{
 				string suffix = String.Empty;
-				if (!String.IsNullOrWhiteSpace(_defaultPreset) && preset.Equals(_defaultPreset, StringComparison.OrdinalIgnoreCase))
+				if (!String.IsNullOrWhiteSpace(_defaultPreset) && presetName.Equals(_defaultPreset, StringComparison.OrdinalIgnoreCase))
 				{
 					suffix = "(Default)";
 				}
@@ -331,7 +331,7 @@ namespace UnrealGameSync
 				};
 				presetComboBox.Items.Add(item);
 
-				if (presetName.Equals(preset, StringComparison.OrdinalIgnoreCase))
+				if (presetName.Equals(currentPreset, StringComparison.OrdinalIgnoreCase))
 				{
 					presetComboBox.SelectedItem = item;
 				}
