Disable infoLines for now until we have the accurate patch data.

#Synced-CL 48803827

diff --git a/Engine/Source/Programs/UnrealGameSync/UnrealGameSync/Controls/WorkspaceControl.cs b/Engine/Source/Programs/UnrealGameSync/UnrealGameSync/Controls/WorkspaceControl.cs
--- a/Engine/Source/Programs/UnrealGameSync/UnrealGameSync/Controls/WorkspaceControl.cs
+++ b/Engine/Source/Programs/UnrealGameSync/UnrealGameSync/Controls/WorkspaceControl.cs
@@ -4295,7 +4295,7 @@ namespace UnrealGameSync
 					tintColor = System.Drawing.ColorTranslator.FromHtml(statusPanelColor);
 				}
 
-				string[]? infoMessages;
+				/*string[]? infoMessages;
 				if (TryGetProjectSetting(projectConfigFile, "InfoMessage", out infoMessages))
 				{
 					const int maxInfoMessages = 7;
@@ -4323,7 +4323,7 @@ namespace UnrealGameSync
 							infoLines.Add(infoLine);
 						}
 					}
-				}
+				}*/
 			}
 
 			using (Graphics graphics = CreateGraphics())
