added ArchiveSource parameter to sync command and archive command to limit the archive channels required

#Synced-CL 49039596

diff --git a/Engine/Source/Programs/UnrealGameSync/UnrealGameSync/PerforceMonitor.cs b/Engine/Source/Programs/UnrealGameSync/UnrealGameSync/PerforceMonitor.cs
--- a/Engine/Source/Programs/UnrealGameSync/UnrealGameSync/PerforceMonitor.cs
+++ b/Engine/Source/Programs/UnrealGameSync/UnrealGameSync/PerforceMonitor.cs
@@ -609,7 +609,7 @@ namespace UnrealGameSync
 
 		async Task<bool> UpdateArchivesAsync(IPerforceConnection perforce, CancellationToken cancellationToken)
 		{
-			List<BaseArchiveChannel> newArchives = await BaseArchive.EnumerateChannelsAsync(perforce, _hordeClient, _cloudStorage, LatestProjectConfigFile, _selectedProjectIdentifier, cancellationToken);
+			List<BaseArchiveChannel> newArchives = await BaseArchive.EnumerateChannelsAsync(perforce, _hordeClient, _cloudStorage, LatestProjectConfigFile, _selectedProjectIdentifier, BaseArchive.DefaultArchiveSources, cancellationToken);
 
 			// Check if the information has changed
 			if (!Enumerable.SequenceEqual(_archives, newArchives))
diff --git a/Engine/Source/Programs/UnrealGameSync/UnrealGameSyncCmd/Commands/ArchiveCommand.cs b/Engine/Source/Programs/UnrealGameSync/UnrealGameSyncCmd/Commands/ArchiveCommand.cs
--- a/Engine/Source/Programs/UnrealGameSync/UnrealGameSyncCmd/Commands/ArchiveCommand.cs
+++ b/Engine/Source/Programs/UnrealGameSync/UnrealGameSyncCmd/Commands/ArchiveCommand.cs
@@ -6,6 +6,7 @@ using System.Linq;
 using System.Threading;
 using System.Threading.Tasks;
 using EpicGames.Core;
+
 using Microsoft.Extensions.Logging;
 using UnrealGameSync;
 
@@ -30,6 +31,9 @@ namespace UnrealGameSyncCmd.Commands
 
 		[CommandLine("-ArchiveName")]
 		public string ArchiveName { get; set; } = String.Empty;
+
+		[CommandLine("-ArchiveSources", ListSeparator = ';')]
+		public List<string> ArchiveSources { get; set; } = new List<string>();
 	}
 
 	class ArchiveCommand : Command
@@ -43,13 +47,21 @@ namespace UnrealGameSyncCmd.Commands
 			context.Arguments.ApplyTo(options);
 			context.Arguments.CheckAllArgumentsUsed();
 
+			HashSet<string> archiveSources = BaseArchive.DefaultArchiveSources;
+			if (options.ArchiveSources.Any())
+			{
+				archiveSources = options.ArchiveSources.ToHashSet(StringComparer.OrdinalIgnoreCase);
+			}
+
 			List<BaseArchiveChannel> archives = await BaseArchive.EnumerateChannelsAsync(
 				context.PerforceClient!, 
 				context.HordeClient!, 
 				context.CloudStorage, 
 				context.ProjectConfig!, 
-				context.WorkspaceState!.Current.ProjectIdentifier, 
-				CancellationToken.None);
+				context.WorkspaceState!.Current.ProjectIdentifier,
+				archiveSources,
+				CancellationToken.None
+				);
 			List<IArchiveChannel> selected = GetSelectedArchiveChannels(context.GlobalSettings!, archives);
 
 			if (options.List)
diff --git a/Engine/Source/Programs/UnrealGameSync/UnrealGameSyncCmd/Commands/SyncCommand.cs b/Engine/Source/Programs/UnrealGameSync/UnrealGameSyncCmd/Commands/SyncCommand.cs
--- a/Engine/Source/Programs/UnrealGameSync/UnrealGameSyncCmd/Commands/SyncCommand.cs
+++ b/Engine/Source/Programs/UnrealGameSync/UnrealGameSyncCmd/Commands/SyncCommand.cs
@@ -40,6 +40,9 @@ namespace UnrealGameSyncCmd.Commands
 
 		[CommandLine("-Only")]
 		public bool SingleChange { get; set; }
+
+		[CommandLine("-ArchiveSources", ListSeparator = ';')]
+		public List<string> ArchiveSources { get; set; } = new List<string>();
 	}
 
 	internal class SyncCommand : Command
@@ -140,7 +143,13 @@ namespace UnrealGameSyncCmd.Commands
 
 			if (syncOptions.Binaries)
 			{
-				List<BaseArchiveChannel> archives = await BaseArchive.EnumerateChannelsAsync(perforceClient, context.HordeClient, context.CloudStorage, projectConfig, state.Current.ProjectIdentifier, CancellationToken.None);
+				HashSet<string> archiveSources = BaseArchive.DefaultArchiveSources;
+				if (syncOptions.ArchiveSources.Any())
+				{
+					archiveSources = syncOptions.ArchiveSources.ToHashSet(StringComparer.OrdinalIgnoreCase);
+				}
+
+				List<BaseArchiveChannel> archives = await BaseArchive.EnumerateChannelsAsync(perforceClient, context.HordeClient, context.CloudStorage, projectConfig, state.Current.ProjectIdentifier, archiveSources, CancellationToken.None);
 
 				BaseArchiveChannel? editorArchiveInfo = archives.FirstOrDefault(x => x.Type == IArchiveChannel.EditorArchiveType);
 				if (editorArchiveInfo == null)
diff --git a/Engine/Source/Programs/UnrealGameSync/UnrealGameSyncShared/ArchiveInfo.cs b/Engine/Source/Programs/UnrealGameSync/UnrealGameSyncShared/ArchiveInfo.cs
--- a/Engine/Source/Programs/UnrealGameSync/UnrealGameSyncShared/ArchiveInfo.cs
+++ b/Engine/Source/Programs/UnrealGameSync/UnrealGameSyncShared/ArchiveInfo.cs
@@ -549,20 +549,41 @@ namespace UnrealGameSync
 
 	public static class BaseArchive
 	{
-		public static async Task<List<BaseArchiveChannel>> EnumerateChannelsAsync(IPerforceConnection perforce, IHordeClient? hordeClient, ICloudStorage? cloudStorage, ConfigFile latestProjectConfigFile, string projectIdentifier, CancellationToken cancellationToken)
+		public static readonly HashSet<string> DefaultArchiveSources = new HashSet<string>(StringComparer.OrdinalIgnoreCase) { "perforce", "cloud", "horde" };
+
+		public static async Task<List<BaseArchiveChannel>> EnumerateChannelsAsync(
+			IPerforceConnection perforce, 
+			IHordeClient? hordeClient, 
+			ICloudStorage? cloudStorage, 
+			ConfigFile latestProjectConfigFile, 
+			string projectIdentifier,
+			HashSet<string> archiveSources,
+			CancellationToken cancellationToken)
 		{
 			List<BaseArchiveChannel> newArchives = new List<BaseArchiveChannel>();
-			newArchives.AddRange(await PerforceArchiveChannel.GetChannelsAsync(perforce, latestProjectConfigFile, projectIdentifier, cancellationToken));
+
+			if (archiveSources.Contains("perforce"))
+			{
+				newArchives.AddRange(await PerforceArchiveChannel.GetChannelsAsync(perforce, latestProjectConfigFile, projectIdentifier, cancellationToken));
+			}
+
 			// prefer Cloud Storage pcbs over horde artifacts
-			if (cloudStorage != null && cloudStorage.IsEnabled(latestProjectConfigFile)) 
+			if (archiveSources.Contains("cloud"))
 			{
-				newArchives.AddRange(await CloudArchiveChannel.GetChannelsAsync(cloudStorage, latestProjectConfigFile, projectIdentifier, cancellationToken));
+				if (cloudStorage != null && cloudStorage.IsEnabled(latestProjectConfigFile))
+				{
+					newArchives.AddRange(await CloudArchiveChannel.GetChannelsAsync(cloudStorage, latestProjectConfigFile, projectIdentifier, cancellationToken));
+				}
 			}
 
-			if (hordeClient != null)
+			if (archiveSources.Contains("horde"))
 			{
-				newArchives.AddRange(await HordeArchiveChannel.GetChannelsAsync(hordeClient, projectIdentifier, cancellationToken));
+				if (hordeClient != null)
+				{
+					newArchives.AddRange(await HordeArchiveChannel.GetChannelsAsync(hordeClient, projectIdentifier, cancellationToken));
+				}
 			}
+
 			return newArchives;
 		}
 	}
