do not open the UGS launcher if there is already one instance running

#Synced-CL 49070965

diff --git a/Engine/Source/Programs/UnrealGameSync/UnrealGameSync/Program.cs b/Engine/Source/Programs/UnrealGameSync/UnrealGameSync/Program.cs
--- a/Engine/Source/Programs/UnrealGameSync/UnrealGameSync/Program.cs
+++ b/Engine/Source/Programs/UnrealGameSync/UnrealGameSync/Program.cs
@@ -109,7 +109,7 @@ namespace UnrealGameSync
 					bool runUpdateCheck = ShouldRunAutoUpdate(args);
 
 					// Check for a newer version of the application
-					if (runUpdateCheck && Launcher.SyncAndRunLatest(instanceMutex, args) != LauncherResult.Continue)
+					if (runUpdateCheck && firstInstance && Launcher.SyncAndRunLatest(instanceMutex, args) != LauncherResult.Continue)
 					{
 						return;
 					}
