[UGS] Added PCB override / project

Added a Archive label refresh when buildmetadata is updated.

#Synced-CL 49648897

diff --git a/Engine/Source/Programs/UnrealGameSync/UnrealGameSync/Controls/WorkspaceControl.cs b/Engine/Source/Programs/UnrealGameSync/UnrealGameSync/Controls/WorkspaceControl.cs
--- a/Engine/Source/Programs/UnrealGameSync/UnrealGameSync/Controls/WorkspaceControl.cs
+++ b/Engine/Source/Programs/UnrealGameSync/UnrealGameSync/Controls/WorkspaceControl.cs
@@ -1280,6 +1280,7 @@ namespace UnrealGameSync
 			{
 				_changeNumberToLayoutInfo.Clear();
 				UpdateBuildList();
+				UpdateSelectedArchivesLabel();
 				UpdateStatusPanel();
 			}
 		}
@@ -1935,6 +1936,7 @@ namespace UnrealGameSync
 			UpdateStatusPanel();
 			UpdateBuildFailureNotification();
 			CheckForStartupComplete();
+			UpdateSelectedArchivesLabel();
 
 			// If we are filtering by badges, we may also need to update the build list
 			if (_projectSettings.FilterType != FilterType.None || _projectSettings.FilterBadges.Count > 0)
