Rename Sync Categories (read only) to ?disabled?

#Synced-CL 49648897

diff --git a/Engine/Source/Programs/UnrealGameSync/UnrealGameSync/Forms/SyncFilter.cs b/Engine/Source/Programs/UnrealGameSync/UnrealGameSync/Forms/SyncFilter.cs
--- a/Engine/Source/Programs/UnrealGameSync/UnrealGameSync/Forms/SyncFilter.cs
+++ b/Engine/Source/Programs/UnrealGameSync/UnrealGameSync/Forms/SyncFilter.cs
@@ -50,7 +50,7 @@ namespace UnrealGameSync
 				{
 					if (IsReadOnly)
 					{
-						return $"{CategoryName} (Read Only)";
+						return $"{CategoryName} (Disabled)";
 					}
 
 					return CategoryName;
