Sync Base Content Feedback requested changes:

additional checks + resizable dialog

#Synced-CL 49648897

diff --git a/Engine/Source/Programs/UnrealGameSync/UnrealGameSync/Controls/WorkspaceSyncFilterControl.cs b/Engine/Source/Programs/UnrealGameSync/UnrealGameSync/Controls/WorkspaceSyncFilterControl.cs
--- a/Engine/Source/Programs/UnrealGameSync/UnrealGameSync/Controls/WorkspaceSyncFilterControl.cs
+++ b/Engine/Source/Programs/UnrealGameSync/UnrealGameSync/Controls/WorkspaceSyncFilterControl.cs
@@ -136,7 +136,7 @@ namespace UnrealGameSync.Controls
 				{
 					if (property.GetValue(item) is string toolTip)
 					{
-						if (SyncBaseContentCheckBox.Checked == false)
+						if (SyncBaseContentCheckBox.Visible && SyncBaseContentCheckBox.Checked == false)
 						{
 							toolTip = $"{toolTip}\n\n" +
 									  $@"(Click the Sync Base Content checkbox to allow enabling/disabling of sync categories.)";
