Added missing reference
Static Analysis fixed

#Synced-CL 50047913

diff --git a/Engine/Source/Programs/UnrealGameSync/UnrealGameSync/Controls/WorkspaceControl.cs b/Engine/Source/Programs/UnrealGameSync/UnrealGameSync/Controls/WorkspaceControl.cs
--- a/Engine/Source/Programs/UnrealGameSync/UnrealGameSync/Controls/WorkspaceControl.cs
+++ b/Engine/Source/Programs/UnrealGameSync/UnrealGameSync/Controls/WorkspaceControl.cs
@@ -2059,17 +2059,6 @@ namespace UnrealGameSync
 
 			return false;
 		}
-
-		bool GetForceSyncBaseContent()
-		{
-			if (ConfigUtils.GetForceSyncBaseContent(_perforceMonitor.LatestProjectConfigFile, SelectedProjectIdentifier, out bool forceSyncBaseContent))
-			{
-				return forceSyncBaseContent;
-			}
-
-			return false;
-		}
-
 		ISet<string> GetBaseContentPaths()
 		{
 			return ConfigUtils.GetBaseContentPaths(_perforceMonitor.LatestProjectConfigFile, SelectedProjectIdentifier);
diff --git a/Engine/Source/Programs/UnrealGameSync/UnrealGameSyncShared/ArchiveOverride.cs b/Engine/Source/Programs/UnrealGameSync/UnrealGameSyncShared/ArchiveOverride.cs
--- a/Engine/Source/Programs/UnrealGameSync/UnrealGameSyncShared/ArchiveOverride.cs
+++ b/Engine/Source/Programs/UnrealGameSync/UnrealGameSyncShared/ArchiveOverride.cs
@@ -46,9 +46,6 @@ namespace UnrealGameSync
 		private readonly IPerforceSettings _perforceSettings;
 		private readonly ILogger _logger;
 
-		[GeneratedRegex("^[A-Za-z]+:")]
-		private static partial Regex SectionRegex();
-
 		[GeneratedRegex("^\\s*import\\s+\\.\\.\\.")]
 		private static partial Regex ImportRegex();
 
