fixed issues with user settings's archives deserialization
fixed issue with sub menu only showing in archives overrides if channels > 1

#Synced-CL 50047913

diff --git a/Engine/Source/Programs/UnrealGameSync/UnrealGameSync/Controls/WorkspaceControl.cs b/Engine/Source/Programs/UnrealGameSync/UnrealGameSync/Controls/WorkspaceControl.cs
--- a/Engine/Source/Programs/UnrealGameSync/UnrealGameSync/Controls/WorkspaceControl.cs
+++ b/Engine/Source/Programs/UnrealGameSync/UnrealGameSync/Controls/WorkspaceControl.cs
@@ -5427,7 +5427,7 @@ namespace UnrealGameSync
 						OptionsContextMenu_SyncPrecompiledBinariesProjectOverride.ToolTipText = null;
 						OptionsContextMenu_SyncPrecompiledBinariesProjectOverride.Checked = selectedArchives.Count > 0;
 
-						if (channels.Count > 1 || channels[0].Type != IArchiveChannel.EditorArchiveType)
+						if (channels.Any())
 						{
 							ToolStripMenuItem disableItem = new ToolStripMenuItem("Unset (no override)");
 							disableItem.Checked = (selectedArchives.Count == 0) && (_projectSettings.ArchivesDeactivated == false);
diff --git a/Engine/Source/Programs/UnrealGameSync/UnrealGameSyncShared/UserSettings.cs b/Engine/Source/Programs/UnrealGameSync/UnrealGameSyncShared/UserSettings.cs
--- a/Engine/Source/Programs/UnrealGameSync/UnrealGameSyncShared/UserSettings.cs
+++ b/Engine/Source/Programs/UnrealGameSync/UnrealGameSyncShared/UserSettings.cs
@@ -207,7 +207,7 @@ namespace UnrealGameSync
 		public string Type { get; set; } = String.Empty;
 
 		[JsonPropertyName("Order")]
-		public List<string> Order { get; } = new ();
+		public List<string> Order { get; init; } = new ();
 
 		public ArchiveSettings()
 		{
@@ -505,7 +505,7 @@ namespace UnrealGameSync
 		public bool PresetSetByUser { get; set; } = false;
 
 		[JsonPropertyName("Archives")]
-		public List<ArchiveSettings> Archives { get; } = new List<ArchiveSettings>();
+		public List<ArchiveSettings> Archives { get; init; } = new List<ArchiveSettings>();
 
 		[JsonPropertyName("ArchivesDeactivated")]
 		public bool ArchivesDeactivated { get; set; } = false;
