using project's settings to pass editors argument in ugs run command

#Synced-CL 50047913

diff --git a/Engine/Source/Programs/UnrealGameSync/UnrealGameSyncCmd/Commands/RunCommand.cs b/Engine/Source/Programs/UnrealGameSync/UnrealGameSyncCmd/Commands/RunCommand.cs
--- a/Engine/Source/Programs/UnrealGameSync/UnrealGameSyncCmd/Commands/RunCommand.cs
+++ b/Engine/Source/Programs/UnrealGameSync/UnrealGameSyncCmd/Commands/RunCommand.cs
@@ -3,6 +3,7 @@
 using System;
 using System.Collections.Generic;
 using System.IO;
+using System.Linq;
 using System.Threading.Tasks;
 using EpicGames.Core;
 using Microsoft.Extensions.Logging;
@@ -13,6 +14,12 @@ using UserErrorException = UnrealGameSyncCmd.Exceptions.UserErrorException;
 
 namespace UnrealGameSyncCmd.Commands
 {
+	internal class RunCommandOptions
+	{
+		[CommandLine("-UseEditorArgsFromSettings")]
+		public bool UseEditorArgsFromSettings { get; set; } = false;
+	}
+
 	internal class RunCommand : Command
 	{
 		public override async Task ExecuteAsync(CommandContext context)
@@ -48,6 +55,8 @@ namespace UnrealGameSyncCmd.Commands
 			{
 				launchArguments.Add("-debug");
 			}
+
+			RunCommandOptions options = context.Arguments!.ApplyTo<RunCommandOptions>(logger);
 			for (int idx = 0; idx < context.Arguments!.Count; idx++)
 			{
 				if (!context.Arguments.HasBeenUsed(idx))
@@ -56,6 +65,14 @@ namespace UnrealGameSyncCmd.Commands
 				}
 			}
 
+			if (options.UseEditorArgsFromSettings)
+			{
+				foreach (LockableEditorArgument editorArgument in GetEditorArgumentsWithDefaults(context))
+				{
+					launchArguments.Add(editorArgument.Name);
+				}
+			}
+			
 			string commandLine = CommandLineArguments.Join(launchArguments);
 			logger.LogInformation("Spawning: {LaunchFile} {CommandLine}", CommandLineArguments.Quote(receipt.Launch), commandLine);
 
@@ -64,5 +81,48 @@ namespace UnrealGameSyncCmd.Commands
 				logger.LogError("Unable to spawn {App} {Args}", receipt.Launch, launchArguments.ToString());
 			}
 		}
+
+		private static List<LockableEditorArgument> GetDefaultEditorArguments(CommandContext context)
+		{
+			List<string> defaultEditorArgumentDefinitions = new List<string>();
+			ConfigUtils.GetProjectSettings(context.ProjectConfig!, context.ProjectInfo!.ProjectIdentifier, "DefaultEditorArgument", defaultEditorArgumentDefinitions);
+
+			List<LockableEditorArgument> defaultEditorArguments = new List<LockableEditorArgument>();
+			foreach (string editorArgumentDefinition in defaultEditorArgumentDefinitions.Distinct())
+			{
+				if (LockableEditorArgument.TryParseConfigEntry(editorArgumentDefinition, out LockableEditorArgument? editorArgument))
+				{
+					defaultEditorArguments.Add(editorArgument);
+				}
+			}
+
+			return defaultEditorArguments;
+		}
+
+		private static List<LockableEditorArgument> GetEditorArgumentsWithDefaults(CommandContext context)
+		{
+			List<LockableEditorArgument> defaultEditorArguments = GetDefaultEditorArguments(context);
+			List<LockableEditorArgument> currentEditorArguments = context.GlobalSettings!.EditorArguments.ToList();
+
+			foreach (LockableEditorArgument defaultEditorArgument in defaultEditorArguments)
+			{
+				// Check to see if the user already has this default argument in their list
+				int index = currentEditorArguments.FindIndex(x => x.Name == defaultEditorArgument.Name);
+
+				if (index != -1)
+				{
+					if (defaultEditorArgument.Locked)
+					{
+						currentEditorArguments[index] = new LockableEditorArgument(defaultEditorArgument.Name, defaultEditorArgument.Enabled, /* Locked = */ true);
+					}
+				}
+				else
+				{
+					currentEditorArguments.Add(new LockableEditorArgument(defaultEditorArgument));
+				}
+			}
+
+			return currentEditorArguments;
+		}
 	}
 }
diff --git a/Engine/Source/Programs/UnrealGameSync/UnrealGameSyncCmd/CommandsFactory.cs b/Engine/Source/Programs/UnrealGameSync/UnrealGameSyncCmd/CommandsFactory.cs
--- a/Engine/Source/Programs/UnrealGameSync/UnrealGameSyncCmd/CommandsFactory.cs
+++ b/Engine/Source/Programs/UnrealGameSync/UnrealGameSyncCmd/CommandsFactory.cs
@@ -58,7 +58,7 @@ namespace UnrealGameSyncCmd
 				"ugs clients",
 				"Lists all clients suitable for use on the current machine."
 			),
-				new CommandInfo("run", typeof(RunCommand), null,
+				new CommandInfo("run", typeof(RunCommand), typeof(RunCommandOptions),
 				"ugs run",
 				"Runs the editor for the current branch."
 			),
