UGS: Add descriptive label indicating why p4 config generation is disabled

#Synced-CL 50246319

diff --git a/Engine/Source/Programs/UnrealGameSync/UnrealGameSync/Forms/OpenProjectWindow.cs b/Engine/Source/Programs/UnrealGameSync/UnrealGameSync/Forms/OpenProjectWindow.cs
--- a/Engine/Source/Programs/UnrealGameSync/UnrealGameSync/Forms/OpenProjectWindow.cs
+++ b/Engine/Source/Programs/UnrealGameSync/UnrealGameSync/Forms/OpenProjectWindow.cs
@@ -84,10 +84,11 @@ namespace UnrealGameSync
 				}
 			}
 
+			GenerateP4ConfigCheckbox.Checked = IsP4ConfigSet();
+
 			UpdateEnabledControls();
 			UpdateServerLabel();
 			UpdateWorkspacePathBrowseButton();
-			UpdateP4ConfigCheckBox();
 			UpdateOkButton();
 		}
 
@@ -176,14 +177,6 @@ namespace UnrealGameSync
 			WorkspacePathBrowseBtn.Enabled = TryGetWorkspaceName(out _);
 		}
 
-		private void UpdateP4ConfigCheckBox()
-		{
-			bool isP4ConfigSet = !String.IsNullOrEmpty(PerforceEnvironment.Default.GetValue("P4CONFIG"));
-
-			GenerateP4ConfigCheckbox.Visible = isP4ConfigSet;
-			GenerateP4ConfigCheckbox.Checked = isP4ConfigSet;
-		}
-
 		private void UpdateOkButton()
 		{
 			OkBtn.Enabled = WorkspaceRadioBtn.Checked ? TryGetClientPath(out _) : TryGetLocalPath(out _);
@@ -299,6 +292,30 @@ namespace UnrealGameSync
 			return false;
 		}
 
+		private void GenerateP4ConfigCheckbox_Click(object sender, EventArgs e)
+		{
+			if (GenerateP4ConfigCheckbox.Checked && !IsP4ConfigSet())
+			{
+				string message =
+					"The P4CONFIG variable is not set but you've chosen to generate a P4Config file."
+					+ Environment.NewLine + Environment.NewLine
+					+ "You must set a P4CONFIG value and restart UGS before enabling this option."
+					+ Environment.NewLine + Environment.NewLine
+					+ "This is not a required feature. If you're not familiar with what this is, it's perfectly "
+					+ "safe to leave this option disabled and continue.";
+				
+				MessageBox.Show(
+					this,
+					message,
+					"P4CONFIG Not Set",
+					MessageBoxButtons.OK,
+					MessageBoxIcon.Information
+				);
+
+				GenerateP4ConfigCheckbox.Checked = false;
+			}
+		}
+
 		private void OkBtn_Click(object sender, EventArgs e)
 		{
 			UserSelectedProjectSettings? selectedProject;
@@ -419,5 +436,10 @@ namespace UnrealGameSync
 		{
 			WorkspaceRadioBtn.Checked = true;
 		}
+		
+		private bool IsP4ConfigSet()
+		{
+			return !String.IsNullOrEmpty(PerforceEnvironment.Default.GetValue("P4CONFIG"));
+		}
 	}
 }
