only passing enabled editor's argument in ugs run command

#Synced-CL 50047913

diff --git a/Engine/Source/Programs/UnrealGameSync/UnrealGameSyncCmd/Commands/RunCommand.cs b/Engine/Source/Programs/UnrealGameSync/UnrealGameSyncCmd/Commands/RunCommand.cs
--- a/Engine/Source/Programs/UnrealGameSync/UnrealGameSyncCmd/Commands/RunCommand.cs
+++ b/Engine/Source/Programs/UnrealGameSync/UnrealGameSyncCmd/Commands/RunCommand.cs
@@ -67,7 +67,7 @@ namespace UnrealGameSyncCmd.Commands
 
 			if (options.UseEditorArgsFromSettings)
 			{
-				foreach (LockableEditorArgument editorArgument in GetEditorArgumentsWithDefaults(context))
+				foreach (LockableEditorArgument editorArgument in GetEditorArgumentsWithDefaults(context).Where(e => e.Enabled))
 				{
 					launchArguments.Add(editorArgument.Name);
 				}
