added handling of null global settings in Run command

#Synced-CL 50257202

diff --git a/Engine/Source/Programs/UnrealGameSync/UnrealGameSyncCmd/Commands/RunCommand.cs b/Engine/Source/Programs/UnrealGameSync/UnrealGameSyncCmd/Commands/RunCommand.cs
--- a/Engine/Source/Programs/UnrealGameSync/UnrealGameSyncCmd/Commands/RunCommand.cs
+++ b/Engine/Source/Programs/UnrealGameSync/UnrealGameSyncCmd/Commands/RunCommand.cs
@@ -84,6 +84,11 @@ namespace UnrealGameSyncCmd.Commands
 
 		private static List<LockableEditorArgument> GetDefaultEditorArguments(CommandContext context)
 		{
+			if(context.GlobalSettings == null || context.ProjectInfo == null)
+			{
+				return new List<LockableEditorArgument>();
+			}
+
 			List<string> defaultEditorArgumentDefinitions = new List<string>();
 			ConfigUtils.GetProjectSettings(context.ProjectConfig!, context.ProjectInfo!.ProjectIdentifier, "DefaultEditorArgument", defaultEditorArgumentDefinitions);
 
@@ -101,6 +106,11 @@ namespace UnrealGameSyncCmd.Commands
 
 		private static List<LockableEditorArgument> GetEditorArgumentsWithDefaults(CommandContext context)
 		{
+			if (context.GlobalSettings == null)
+			{
+				return new List<LockableEditorArgument>();
+			}
+
 			List<LockableEditorArgument> defaultEditorArguments = GetDefaultEditorArguments(context);
 			List<LockableEditorArgument> currentEditorArguments = context.GlobalSettings!.EditorArguments.ToList();
 
