UGS: Fix compiler warning

#skipci

diff --git a/Engine/Source/Programs/UnrealGameSync/UnrealGameSync/Forms/OpenProjectWindow.cs b/Engine/Source/Programs/UnrealGameSync/UnrealGameSync/Forms/OpenProjectWindow.cs
--- a/Engine/Source/Programs/UnrealGameSync/UnrealGameSync/Forms/OpenProjectWindow.cs
+++ b/Engine/Source/Programs/UnrealGameSync/UnrealGameSync/Forms/OpenProjectWindow.cs
@@ -437,7 +437,7 @@ namespace UnrealGameSync
 			WorkspaceRadioBtn.Checked = true;
 		}
 		
-		private bool IsP4ConfigSet()
+		private static bool IsP4ConfigSet()
 		{
 			return !String.IsNullOrEmpty(PerforceEnvironment.Default.GetValue("P4CONFIG"));
 		}
