PostBadgeStatus: Update project file to fix warnings when built under SubmitTool.

#Synced-CL 50766644

diff --git a/Engine/Source/Programs/UnrealGameSync/PostBadgeStatus/PostBadgeStatus.csproj b/Engine/Source/Programs/UnrealGameSync/PostBadgeStatus/PostBadgeStatus.csproj
--- a/Engine/Source/Programs/UnrealGameSync/PostBadgeStatus/PostBadgeStatus.csproj
+++ b/Engine/Source/Programs/UnrealGameSync/PostBadgeStatus/PostBadgeStatus.csproj
@@ -1,68 +1,11 @@
-﻿<?xml version="1.0" encoding="utf-8"?>
-<Project ToolsVersion="12.0" DefaultTargets="Build" xmlns="http://schemas.microsoft.com/developer/msbuild/2003">
-  <Import Project="$(MSBuildExtensionsPath)\$(MSBuildToolsVersion)\Microsoft.Common.props" Condition="Exists('$(MSBuildExtensionsPath)\$(MSBuildToolsVersion)\Microsoft.Common.props')" />
+﻿<Project Sdk="Microsoft.NET.Sdk">
+  <Import Project="..\..\Shared\UnrealEngine.csproj.props" />
   <PropertyGroup>
-    <Configuration Condition=" '$(Configuration)' == '' ">Debug</Configuration>
-    <Platform Condition=" '$(Platform)' == '' ">AnyCPU</Platform>
-    <ProjectGuid>{DCBE3437-5DBF-41B4-B854-967BC0C93C48}</ProjectGuid>
     <OutputType>Exe</OutputType>
-    <AppDesignerFolder>Properties</AppDesignerFolder>
-    <RootNamespace>PostBadgeStatus</RootNamespace>
-    <AssemblyName>PostBadgeStatus</AssemblyName>
-    <TargetFrameworkVersion>v4.8</TargetFrameworkVersion>
-    <FileAlignment>512</FileAlignment>
-    <TargetFrameworkProfile />
-  </PropertyGroup>
-  <PropertyGroup Condition="'$(Configuration)|$(Platform)' == 'Debug|AnyCPU'">
-    <DebugSymbols>true</DebugSymbols>
-    <OutputPath>bin\Debug\</OutputPath>
-    <DefineConstants>DEBUG;TRACE</DefineConstants>
-    <DebugType>full</DebugType>
-    <PlatformTarget>AnyCPU</PlatformTarget>
-    <ErrorReport>prompt</ErrorReport>
-    <CodeAnalysisRuleSet>MinimumRecommendedRules.ruleset</CodeAnalysisRuleSet>
-    <Prefer32Bit>true</Prefer32Bit>
-  </PropertyGroup>
-  <PropertyGroup Condition="'$(Configuration)|$(Platform)' == 'Release|AnyCPU'">
-    <OutputPath>bin\Release\</OutputPath>
-    <DefineConstants>TRACE</DefineConstants>
-    <Optimize>true</Optimize>
-    <DebugType>pdbonly</DebugType>
-    <PlatformTarget>AnyCPU</PlatformTarget>
-    <ErrorReport>prompt</ErrorReport>
-    <CodeAnalysisRuleSet>MinimumRecommendedRules.ruleset</CodeAnalysisRuleSet>
-    <Prefer32Bit>true</Prefer32Bit>
+    <GenerateAssemblyInfo>false</GenerateAssemblyInfo>
+    <TargetFramework>net48</TargetFramework>
   </PropertyGroup>
   <ItemGroup>
-    <Reference Include="System" />
-    <Reference Include="System.Core" />
     <Reference Include="System.Web.Extensions" />
-    <Reference Include="System.Xml.Linq" />
-    <Reference Include="System.Data.DataSetExtensions" />
-    <Reference Include="Microsoft.CSharp" />
-    <Reference Include="System.Data" />
-    <Reference Include="System.Xml" />
-  </ItemGroup>
-  <ItemGroup>
-    <Compile Include="Program.cs" />
-    <Compile Include="Properties\AssemblyInfo.cs" />
-  </ItemGroup>
-  <ItemGroup>
-    <None Include="App.config" />
   </ItemGroup>
-  <Choose>
-    <When Condition="Exists('NotForLicensees\ProgramSettings.cs')">
-      <ItemGroup>
-        <Compile Include="NotForLicensees\ProgramSettings.cs" />
-      </ItemGroup>
-    </When>
-  </Choose>
-  <Import Project="$(MSBuildToolsPath)\Microsoft.CSharp.targets" />
-  <!-- To modify your build process, add your task inside one of the targets below and uncomment it. 
-       Other similar extension points exist, see Microsoft.Common.targets.
-  <Target Name="BeforeBuild">
-  </Target>
-  <Target Name="AfterBuild">
-  </Target>
-  -->
-</Project>
\ No newline at end of file
+</Project>
