UnrealGameSync: Add dynamic config injection from external scripts

Implements a new [IniDynamicInjection] section in UGS config files that allows
executing scripts to dynamically generate INI configuration at runtime.

Features:
- Platform-specific script execution via +Inject=(Platform="Win64", Script="path")
- Supports Python (.py), PowerShell (.ps1), Bash (.sh), Batch (.bat), and executables
- Scripts execute sequentially during config load, output captured as INI lines
- Handles Perforce-tracked, locally-edited, and local-only scripts
- Graceful degradation: config load succeeds even if scripts fail
- Security: Path traversal protection, proper temp file cleanup on cancellation
- Cross-platform: Works on Windows, macOS, and Linux

ConfigUtils.cs changes:
- ExecuteConfigInjectionScriptsAsync: Main orchestration, derives local root from Perforce client
- ExecuteSingleInjectionScriptAsync: Executes individual script, handles all file sources
- GetInterpreterForScript: Cross-platform interpreter detection
- IsCommandAvailable: Check for interpreter availability in PATH
- MakeScriptExecutable: Unix chmod +x support
- GetLocalRootPathAsync: Derives workspace root from IPerforceConnection

Example usage in UnrealGameSync.ini:
[IniDynamicInjection]
+Inject=(Platform="Win64", Script="Build/generate-config.ps1")
+Inject=(Platform="Mac", Script="Build/generate-config.sh")
+Inject=(Script="Tools/common-config.py")

#Synced-CL 50436066

#Synced-CL 50436066

diff --git a/Engine/Source/Programs/UnrealGameSync/UnrealGameSyncShared/ConfigUtils.cs b/Engine/Source/Programs/UnrealGameSync/UnrealGameSyncShared/ConfigUtils.cs
--- a/Engine/Source/Programs/UnrealGameSync/UnrealGameSyncShared/ConfigUtils.cs
+++ b/Engine/Source/Programs/UnrealGameSync/UnrealGameSyncShared/ConfigUtils.cs
@@ -2,6 +2,7 @@
 
 using System;
 using System.Collections.Generic;
+using System.Diagnostics;
 using System.Diagnostics.CodeAnalysis;
 using System.IO;
 using System.Linq;
@@ -77,9 +78,9 @@ namespace UnrealGameSync
 			return ReadProjectConfigFileAsync(perforce, projectInfo, new List<KeyValuePair<FileReference, DateTime>>(), logger, cancellationToken);
 		}
 
-		public static Task<ConfigFile> ReadProjectConfigFileAsync(IPerforceConnection perforce, ProjectInfo projectInfo, List<KeyValuePair<FileReference, DateTime>> localConfigFiles, ILogger logger, CancellationToken cancellationToken)
+		public static async Task<ConfigFile> ReadProjectConfigFileAsync(IPerforceConnection perforce, ProjectInfo projectInfo, List<KeyValuePair<FileReference, DateTime>> localConfigFiles, ILogger logger, CancellationToken cancellationToken)
 		{
-			return ReadProjectConfigFileAsync(perforce, projectInfo.ClientRootPath, projectInfo.ClientFileName, projectInfo.ProjectIdentifier, projectInfo.CacheFolder, localConfigFiles, logger, cancellationToken);
+			return await ReadProjectConfigFileAsync(perforce, projectInfo.ClientRootPath, projectInfo.ClientFileName, projectInfo.ProjectIdentifier, projectInfo.CacheFolder, localConfigFiles, logger, cancellationToken);
 		}
 
 		public static async Task<ConfigFile> ReadProjectConfigFileAsync(IPerforceConnection perforce, string branchClientPath, string selectedClientFileName, string projectIdentifier, DirectoryReference cacheFolder, List<KeyValuePair<FileReference, DateTime>> localConfigFiles, ILogger logger, CancellationToken cancellationToken)
@@ -139,6 +140,10 @@ namespace UnrealGameSync
 					}
 				}
 			}
+
+			// Execute dynamic config injection scripts
+			await ExecuteConfigInjectionScriptsAsync(perforce, branchClientPath, projectIdentifier, cacheFolder, projectConfig, localConfigFiles, logger, cancellationToken);
+
 			return projectConfig;
 		}
 
@@ -787,5 +792,478 @@ namespace UnrealGameSync
 				}
 			}
 		}
+
+		private static async Task<DirectoryReference?> GetLocalRootPathAsync(IPerforceConnection perforce, ILogger logger, CancellationToken cancellationToken)
+		{
+			try
+			{
+				string? clientName = perforce.Settings.ClientName;
+				if (String.IsNullOrEmpty(clientName))
+				{
+					logger.LogDebug("No client name available in Perforce settings");
+					return null;
+				}
+
+				ClientRecord client = await perforce.GetClientAsync(clientName, cancellationToken);
+				if (client == null || String.IsNullOrEmpty(client.Root))
+				{
+					logger.LogDebug("Could not get client root from Perforce client {ClientName}", clientName);
+					return null;
+				}
+
+				return new DirectoryReference(client.Root);
+			}
+			catch (Exception ex)
+			{
+				logger.LogDebug(ex, "Failed to get local root path from Perforce client");
+				return null;
+			}
+		}
+
+		private static async Task ExecuteConfigInjectionScriptsAsync(
+			IPerforceConnection perforce,
+			string branchClientPath,
+			string projectIdentifier,
+			DirectoryReference cacheFolder,
+			ConfigFile configFile,
+			List<KeyValuePair<FileReference, DateTime>> localConfigFiles,
+			ILogger logger,
+			CancellationToken cancellationToken)
+		{
+			// Find IniDynamicInjection section in config
+			ConfigSection? injectionSection = configFile.FindSection("IniDynamicInjection");
+			if (injectionSection == null)
+			{
+				return;
+			}
+
+			// Get all Inject values
+			string[] injectLines = injectionSection.GetValues("Inject", Array.Empty<string>());
+			if (injectLines.Length == 0)
+			{
+				return;
+			}
+
+			// Get local root path from Perforce client
+			DirectoryReference? localRootPath = await GetLocalRootPathAsync(perforce, logger, cancellationToken);
+			if (localRootPath == null)
+			{
+				logger.LogWarning("Could not determine local root path from Perforce client, skipping dynamic config injection");
+				return;
+			}
+
+			logger.LogDebug("Found {Count} dynamic injection script(s) to execute", injectLines.Length);
+
+			// Execute each script in order
+			foreach (string injectLine in injectLines)
+			{
+				try
+				{
+					// Parse the ConfigObject format: (Platform="Win64", Script="Tools/Scripts/inject.py")
+					ConfigObject injectObj = new ConfigObject(injectLine);
+
+					string platform = injectObj.GetValue("Platform", "");
+					string scriptPath = injectObj.GetValue("Script", "");
+
+					if (String.IsNullOrWhiteSpace(scriptPath))
+					{
+						logger.LogWarning("Inject entry missing Script value: {Line}", injectLine);
+						continue;
+					}
+
+					// Check if platform matches (empty platform means all platforms)
+					if (!String.IsNullOrEmpty(platform) && !platform.Equals(HostPlatform, StringComparison.OrdinalIgnoreCase))
+					{
+						logger.LogDebug("Skipping script {Script} (platform {Platform} does not match {HostPlatform})", scriptPath, platform, HostPlatform);
+						continue;
+					}
+
+					await ExecuteSingleInjectionScriptAsync(perforce, branchClientPath, localRootPath, projectIdentifier, cacheFolder, configFile, localConfigFiles, scriptPath, logger, cancellationToken);
+				}
+				catch (Exception ex)
+				{
+					logger.LogWarning(ex, "Failed to execute injection script from line: {Line}", injectLine);
+				}
+			}
+		}
+
+		private static async Task ExecuteSingleInjectionScriptAsync(
+			IPerforceConnection perforce,
+			string branchClientPath,
+			DirectoryReference localRootPath,
+			string projectIdentifier,
+			DirectoryReference cacheFolder,
+			ConfigFile configFile,
+			List<KeyValuePair<FileReference, DateTime>> localConfigFiles,
+			string relativePath,
+			ILogger logger,
+			CancellationToken cancellationToken)
+		{
+			// Determine local path
+			string localPath = Path.Combine(localRootPath.FullName, relativePath.Replace('/', Path.DirectorySeparatorChar));
+
+			// Validate path doesn't escape workspace root (prevent path traversal)
+			string normalizedPath = Path.GetFullPath(localPath);
+			string rootPrefix = localRootPath.FullName.TrimEnd(Path.DirectorySeparatorChar) + Path.DirectorySeparatorChar;
+			StringComparison pathComparison = RuntimeInformation.IsOSPlatform(OSPlatform.Windows) ? StringComparison.OrdinalIgnoreCase : StringComparison.Ordinal;
+			if (!normalizedPath.StartsWith(rootPrefix, pathComparison))
+			{
+				logger.LogWarning("Script path {Script} attempts to escape workspace root, skipping", relativePath);
+				return;
+			}
+
+			bool isLocalFile = false;
+			bool needsTempFile = false;
+
+			// Convert relative path to client path for Perforce (always use forward slashes)
+			string clientPath = $"{branchClientPath}/{relativePath
+				.Replace('\\', '/')
+				.TrimStart('/')}";
+
+			// Try to get file metadata from Perforce
+			List<PerforceResponse<FStatRecord>> responses = await perforce.TryFStatAsync(FStatOptions.IncludeFileSizes, clientPath, cancellationToken).ToListAsync(cancellationToken);
+
+			if (responses.Count > 0 && responses[0].Succeeded)
+			{
+				// File exists in Perforce
+				FStatRecord fileRecord = responses[0].Data;
+
+				// Skip if synced from different branch
+				if (fileRecord.HeadAction == FileAction.None && !fileRecord.IsMapped)
+				{
+					logger.LogDebug("Skipping script {Script} (synced from different branch)", relativePath);
+					return;
+				}
+
+				// Check if local file is writable (edited or opened for add)
+				string? localFileName = fileRecord.ClientFile;
+				if (localFileName != null && File.Exists(localFileName) && (File.GetAttributes(localFileName) & FileAttributes.ReadOnly) == 0)
+				{
+					try
+					{
+						DateTime lastModifiedTime = File.GetLastWriteTimeUtc(localFileName);
+						localConfigFiles.Add(new KeyValuePair<FileReference, DateTime>(new FileReference(localFileName), lastModifiedTime));
+						localPath = localFileName;
+						isLocalFile = true;
+						logger.LogDebug("Using locally edited script from Perforce workspace: {Script}", relativePath);
+					}
+					catch (Exception ex)
+					{
+						logger.LogWarning(ex, "Failed to read local script file {Path}, falling back to Perforce", localFileName);
+					}
+				}
+
+				// If not using local edited version, get from Perforce cache
+				if (!isLocalFile && !String.IsNullOrWhiteSpace(fileRecord.DepotFile))
+				{
+					string[]? scriptLines = await Utility.TryPrintFileUsingCacheAsync(perforce, fileRecord.DepotFile, cacheFolder, fileRecord.Digest, logger, cancellationToken);
+					if (scriptLines == null)
+					{
+						logger.LogWarning("Failed to retrieve script from Perforce: {Script}", relativePath);
+						return;
+					}
+
+					// Write cached script to temp location for execution
+					localPath = Path.Combine(Path.GetTempPath(), $"ugs_inject_{Guid.NewGuid()}{Path.GetExtension(relativePath)}");
+					await File.WriteAllLinesAsync(localPath, scriptLines, System.Text.Encoding.UTF8, cancellationToken);
+					needsTempFile = true;
+				}
+				else if (!isLocalFile)
+				{
+					logger.LogWarning("Script {Script} exists in Perforce but has no depot path and no usable local file, skipping", relativePath);
+					return;
+				}
+			}
+			else
+			{
+				// File not in Perforce - check if it exists locally only
+				if (File.Exists(localPath))
+				{
+					try
+					{
+						DateTime lastModifiedTime = File.GetLastWriteTimeUtc(localPath);
+						localConfigFiles.Add(new KeyValuePair<FileReference, DateTime>(new FileReference(localPath), lastModifiedTime));
+						isLocalFile = true;
+						logger.LogInformation("Using local-only script (not in Perforce): {Script}", relativePath);
+					}
+					catch (Exception ex)
+					{
+						logger.LogWarning(ex, "Failed to read local-only script file {Path}", localPath);
+						return;
+					}
+				}
+				else
+				{
+					logger.LogWarning("Script not found in Perforce or locally: {Script}", relativePath);
+					return;
+				}
+			}
+
+			// Determine interpreter for script
+			(string? interpreter, string? interpreterArgs) = GetInterpreterForScript(localPath, logger);
+			if (interpreter == null)
+			{
+				logger.LogWarning("No interpreter found for script: {Script}", relativePath);
+				if (needsTempFile)
+				{
+					try
+					{
+						File.Delete(localPath);
+					}
+					catch
+					{
+						// Ignore cleanup errors
+					}
+				}
+				return;
+			}
+
+			// Make script executable on Unix platforms
+			if (!RuntimeInformation.IsOSPlatform(OSPlatform.Windows) && interpreter == localPath)
+			{
+				MakeScriptExecutable(localPath, logger);
+			}
+
+			try
+			{
+				// Execute script
+				List<string> capturedLines = new List<string>();
+				string arguments;
+				if (String.Equals(interpreter, localPath, StringComparison.OrdinalIgnoreCase))
+				{
+					// Direct execution (.exe or extensionless): the script is the
+					// process itself, so only pass the workspace root as an argument.
+					arguments = $"\"{localRootPath.FullName}\"";
+				}
+				else if (String.Equals(interpreter, "cmd", StringComparison.OrdinalIgnoreCase))
+				{
+					// cmd.exe /c strips the first and last quote characters when the
+					// command string contains more than two quotes. Wrap the entire
+					// command in an extra set of quotes so inner quoting is preserved
+					// after stripping.
+					arguments = $"{interpreterArgs} \"\"{localPath}\" \"{localRootPath.FullName}\"\"";
+				}
+				else if (String.IsNullOrEmpty(interpreterArgs))
+				{
+					// Interpreter without extra args (e.g. python, bash): pass the
+					// script path followed by the workspace root.
+					arguments = $"\"{localPath}\" \"{localRootPath.FullName}\"";
+				}
+				else
+				{
+					// Interpreter with args (e.g. pwsh -ExecutionPolicy Bypass -File):
+					// pass interpreter args, then script path, then workspace root.
+					arguments = $"{interpreterArgs} \"{localPath}\" \"{localRootPath.FullName}\"";
+				}
+
+				int exitCode = await Utility.ExecuteProcessAsync(
+					interpreter,
+					localRootPath.FullName,
+					arguments,
+					line => capturedLines.Add(line),
+					env: null,
+					cancellationToken);
+
+				// Parse output if successful
+				if (exitCode == 0)
+				{
+					try
+					{
+						foreach (string line in capturedLines)
+						{
+							logger.LogInformation("Script '{Script}' output: {Line}", localPath, line);
+						}
+
+						configFile.Parse(capturedLines.ToArray(), projectIdentifier);
+						logger.LogInformation("Injected config from script {Script} ({LineCount} lines)", relativePath, capturedLines.Count);
+					}
+					catch (Exception ex)
+					{
+						logger.LogWarning(ex, "Failed to parse output from script {Script}", relativePath);
+					}
+				}
+				else
+				{
+					logger.LogWarning("Script {Script} exited with code {ExitCode}", relativePath, exitCode);
+				}
+			}
+			finally
+			{
+				// Clean up temp file if needed (even on cancellation)
+				if (needsTempFile)
+				{
+					try
+					{
+						File.Delete(localPath);
+					}
+					catch
+					{
+						// Ignore cleanup errors
+					}
+				}
+			}
+		}
+
+		private static (string? interpreter, string? interpreterArgs) GetInterpreterForScript(string scriptPath, ILogger logger)
+		{
+			string extension = Path.GetExtension(scriptPath).ToLowerInvariant();
+
+			switch (extension)
+			{
+				case ".exe":
+				case "":
+					// Direct execution
+					return (scriptPath, null);
+
+				case ".py":
+					// Python script - check for python3 then python
+					if (IsCommandAvailable("python3"))
+					{
+						return ("python3", null);
+					}
+					else if (IsCommandAvailable("python"))
+					{
+						return ("python", null);
+					}
+					else
+					{
+						logger.LogWarning("Python interpreter not found in PATH for script {Script}", scriptPath);
+						return (null, null);
+					}
+
+				case ".ps1":
+					// PowerShell script - check for pwsh then powershell
+					if (IsCommandAvailable("pwsh"))
+					{
+						return ("pwsh", "-ExecutionPolicy Bypass -File");
+					}
+					else if (RuntimeInformation.IsOSPlatform(OSPlatform.Windows) && IsCommandAvailable("powershell"))
+					{
+						return ("powershell", "-ExecutionPolicy Bypass -File");
+					}
+					else
+					{
+						logger.LogWarning("PowerShell interpreter not found in PATH for script {Script}", scriptPath);
+						return (null, null);
+					}
+
+				case ".sh":
+					// Shell script - check for bash then sh (Unix only)
+					if (!RuntimeInformation.IsOSPlatform(OSPlatform.Windows))
+					{
+						if (IsCommandAvailable("bash"))
+						{
+							return ("bash", null);
+						}
+						else if (IsCommandAvailable("sh"))
+						{
+							return ("sh", null);
+						}
+						else
+						{
+							logger.LogWarning("Shell interpreter not found in PATH for script {Script}", scriptPath);
+							return (null, null);
+						}
+					}
+					else
+					{
+						logger.LogWarning("Shell scripts (.sh) are not supported on Windows");
+						return (null, null);
+					}
+
+				case ".bat":
+				case ".cmd":
+					// Batch file (Windows only)
+					if (RuntimeInformation.IsOSPlatform(OSPlatform.Windows))
+					{
+						return ("cmd", "/c");
+					}
+					else
+					{
+						logger.LogWarning("Batch files are not supported on non-Windows platforms");
+						return (null, null);
+					}
+
+				default:
+					logger.LogWarning("Unknown script extension: {Extension} for script {Script}", extension, scriptPath);
+					return (null, null);
+			}
+		}
+
+		private static bool IsCommandAvailable(string command)
+		{
+			try
+			{
+				string fileName = RuntimeInformation.IsOSPlatform(OSPlatform.Windows) ? "where" : "which";
+				using Process process = new Process();
+				process.StartInfo.FileName = fileName;
+				process.StartInfo.Arguments = command;
+				process.StartInfo.UseShellExecute = false;
+				process.StartInfo.RedirectStandardOutput = true;
+				process.StartInfo.RedirectStandardError = true;
+				process.StartInfo.CreateNoWindow = true;
+				process.Start();
+
+				if (!process.WaitForExit(1000))
+				{
+					try
+					{
+						process.Kill();
+					}
+					catch (InvalidOperationException)
+					{
+						// Process already exited between WaitForExit and Kill
+					}
+					return false;
+				}
+
+				return process.ExitCode == 0;
+			}
+			catch
+			{
+				return false;
+			}
+		}
+
+		private static void MakeScriptExecutable(string scriptPath, ILogger logger)
+		{
+			if (RuntimeInformation.IsOSPlatform(OSPlatform.Windows))
+			{
+				return;
+			}
+
+			try
+			{
+				using Process process = new Process();
+				process.StartInfo.FileName = "chmod";
+				process.StartInfo.Arguments = $"+x \"{scriptPath}\"";
+				process.StartInfo.UseShellExecute = false;
+				process.StartInfo.RedirectStandardOutput = true;
+				process.StartInfo.RedirectStandardError = true;
+				process.StartInfo.CreateNoWindow = true;
+				process.Start();
+
+				if (!process.WaitForExit(5000))
+				{
+					try
+					{
+						process.Kill();
+					}
+					catch (InvalidOperationException)
+					{
+						// Process already exited between WaitForExit and Kill
+					}
+					logger.LogDebug("chmod command timed out for {Script}", scriptPath);
+				}
+				else if (process.ExitCode != 0)
+				{
+					logger.LogDebug("chmod command failed for {Script} with exit code {ExitCode}", scriptPath, process.ExitCode);
+				}
+			}
+			catch (Exception ex)
+			{
+				logger.LogDebug(ex, "Failed to make script executable: {Script}", scriptPath);
+			}
+		}
 	}
 }
