UnrealGameSync: Remove redundant properties already imported from UnrealEngine.csproj.props
Update PostBadgeStatus to .net8.0

#rnx

#Synced-CL 50830039

diff --git a/Engine/Source/Programs/UnrealGameSync/PostBadgeStatus/PostBadgeStatus.csproj b/Engine/Source/Programs/UnrealGameSync/PostBadgeStatus/PostBadgeStatus.csproj
--- a/Engine/Source/Programs/UnrealGameSync/PostBadgeStatus/PostBadgeStatus.csproj
+++ b/Engine/Source/Programs/UnrealGameSync/PostBadgeStatus/PostBadgeStatus.csproj
@@ -3,9 +3,6 @@
   <PropertyGroup>
     <OutputType>Exe</OutputType>
     <GenerateAssemblyInfo>false</GenerateAssemblyInfo>
-    <TargetFramework>net48</TargetFramework>
+    <NoWarn>$(NoWarn);SYSLIB0014</NoWarn>
   </PropertyGroup>
-  <ItemGroup>
-    <Reference Include="System.Web.Extensions" />
-  </ItemGroup>
 </Project>
diff --git a/Engine/Source/Programs/UnrealGameSync/PostBadgeStatus/Program.cs b/Engine/Source/Programs/UnrealGameSync/PostBadgeStatus/Program.cs
--- a/Engine/Source/Programs/UnrealGameSync/PostBadgeStatus/Program.cs
+++ b/Engine/Source/Programs/UnrealGameSync/PostBadgeStatus/Program.cs
@@ -5,8 +5,8 @@ using System.Collections.Generic;
 using System.IO;
 using System.Net;
 using System.Text;
+using System.Text.Json;
 using System.Threading;
-using System.Web.Script.Serialization;
 
 namespace WriteBadgeStatus
 {
@@ -52,14 +52,14 @@ namespace WriteBadgeStatus
 			}
 			if (!String.IsNullOrWhiteSpace(Metadata))
 			{
-				Build.Metadata = new JavaScriptSerializer().Deserialize<object>(Metadata);
+				Build.Metadata = JsonSerializer.Deserialize<JsonElement>(Metadata);
 			}
 			int NumRetries = 0;
 			while (true)
 			{
 				try
 				{
-					return SendRequest(RestUrl, "Build", "POST", new JavaScriptSerializer().Serialize(Build));
+					return SendRequest(RestUrl, "Build", "POST", JsonSerializer.Serialize(Build));
 				}
 				catch (Exception ex)
 				{
@@ -110,7 +110,7 @@ namespace WriteBadgeStatus
 			public string Url;
 			public string Project;
 			public string ArchivePath;
-			public object Metadata;
+			public JsonElement? Metadata;
 
 			public bool IsSuccess => Result == BuildDataResult.Success || Result == BuildDataResult.Warning;
 
@@ -123,17 +123,17 @@ namespace WriteBadgeStatus
 			StringBuilder TargetURI = new StringBuilder(String.Format("{0}/api/{1}", URI, Resource));
 			if (QueryParams.Length != 0)
 			{
-				TargetURI.Append("?");
+				TargetURI.Append('?');
 				for (int Idx = 0; Idx < QueryParams.Length; Idx++)
 				{
 					TargetURI.Append(QueryParams[Idx]);
 					if (Idx != QueryParams.Length - 1)
 					{
-						TargetURI.Append("&");
+						TargetURI.Append('&');
 					}
 				}
 			}
-			HttpWebRequest Request = (HttpWebRequest)WebRequest.Create(TargetURI.ToString());
+			HttpWebRequest Request = (HttpWebRequest)WebRequest.Create(new Uri(TargetURI.ToString()));
 			Request.ContentType = "application/json";
 			Request.Method = Method;
 
diff --git a/Engine/Source/Programs/UnrealGameSync/UnrealGameSync/UnrealGameSync.csproj b/Engine/Source/Programs/UnrealGameSync/UnrealGameSync/UnrealGameSync.csproj
--- a/Engine/Source/Programs/UnrealGameSync/UnrealGameSync/UnrealGameSync.csproj
+++ b/Engine/Source/Programs/UnrealGameSync/UnrealGameSync/UnrealGameSync.csproj
@@ -5,7 +5,6 @@
     <OutputType>WinExe</OutputType>
     <TargetFramework>net8.0-windows</TargetFramework>
     <UseWindowsForms>true</UseWindowsForms>
-    <Deterministic>true</Deterministic>
     <Product>UnrealGameSync</Product>
     <Company>Epic Games, Inc</Company>
     <Copyright>Copyright Epic Games, Inc. All Rights Reserved.</Copyright>
@@ -16,16 +15,7 @@
     <ApplicationManifest>Properties\app.manifest</ApplicationManifest>
     <Nullable>enable</Nullable>
     <NeutralLanguage>en-US</NeutralLanguage>
-    <RunAnalyzersDuringBuild>False</RunAnalyzersDuringBuild>
-    <Configurations>Debug;Release;Analyze;DebugAutoUpdate;ReleaseAutoUpdate</Configurations>
-  </PropertyGroup>
-
-  <PropertyGroup Condition=" '$(Configuration)|$(Platform)' == 'Analyze|AnyCPU' ">
-    <RunAnalyzersDuringBuild>True</RunAnalyzersDuringBuild>
-    <TreatWarningsAsErrors>false</TreatWarningsAsErrors>
-    <EnforceCodeStyleInBuild>True</EnforceCodeStyleInBuild>
-    <GenerateDocumentationFile>True</GenerateDocumentationFile>
-    <DocumentationFile></DocumentationFile>
+    <Configurations>$(Configurations);DebugAutoUpdate;ReleaseAutoUpdate</Configurations>
   </PropertyGroup>
 
   <PropertyGroup Condition="'$(Configuration)|$(Platform)'=='Release|AnyCPU'">
diff --git a/Engine/Source/Programs/UnrealGameSync/UnrealGameSyncAutomation/UnrealGameSyncAutomation.csproj b/Engine/Source/Programs/UnrealGameSync/UnrealGameSyncAutomation/UnrealGameSyncAutomation.csproj
--- a/Engine/Source/Programs/UnrealGameSync/UnrealGameSyncAutomation/UnrealGameSyncAutomation.csproj
+++ b/Engine/Source/Programs/UnrealGameSync/UnrealGameSyncAutomation/UnrealGameSyncAutomation.csproj
@@ -9,15 +9,5 @@
     <GenerateAssemblyInfo>false</GenerateAssemblyInfo>
     <RootNamespace>UnrealGameSyncAutomation</RootNamespace>
     <AssemblyName>UnrealGameSyncAutomation</AssemblyName>
-    <RunAnalyzersDuringBuild>False</RunAnalyzersDuringBuild>
-    <Configurations>Debug;Release;Analyze</Configurations>
-  </PropertyGroup>
-
-  <PropertyGroup Condition=" '$(Configuration)|$(Platform)' == 'Analyze|AnyCPU' ">
-    <RunAnalyzersDuringBuild>True</RunAnalyzersDuringBuild>
-    <TreatWarningsAsErrors>false</TreatWarningsAsErrors>
-    <EnforceCodeStyleInBuild>True</EnforceCodeStyleInBuild>
-    <GenerateDocumentationFile>True</GenerateDocumentationFile>
-    <DocumentationFile></DocumentationFile>
   </PropertyGroup>
 </Project>
\ No newline at end of file
diff --git a/Engine/Source/Programs/UnrealGameSync/UnrealGameSyncCmd/UnrealGameSyncCmd.csproj b/Engine/Source/Programs/UnrealGameSync/UnrealGameSyncCmd/UnrealGameSyncCmd.csproj
--- a/Engine/Source/Programs/UnrealGameSync/UnrealGameSyncCmd/UnrealGameSyncCmd.csproj
+++ b/Engine/Source/Programs/UnrealGameSync/UnrealGameSyncCmd/UnrealGameSyncCmd.csproj
@@ -6,16 +6,6 @@
     <TargetFramework>net8.0</TargetFramework>
     <Nullable>enable</Nullable>
     <AssemblyName>ugs</AssemblyName>
-    <RunAnalyzersDuringBuild>False</RunAnalyzersDuringBuild>
-    <Configurations>Debug;Release;Analyze</Configurations>
-  </PropertyGroup>
-
-  <PropertyGroup Condition=" '$(Configuration)|$(Platform)' == 'Analyze|AnyCPU' ">
-    <RunAnalyzersDuringBuild>True</RunAnalyzersDuringBuild>
-    <TreatWarningsAsErrors>false</TreatWarningsAsErrors>
-    <EnforceCodeStyleInBuild>True</EnforceCodeStyleInBuild>
-    <GenerateDocumentationFile>True</GenerateDocumentationFile>
-    <DocumentationFile></DocumentationFile>
   </PropertyGroup>
 
   <ItemGroup>
diff --git a/Engine/Source/Programs/UnrealGameSync/UnrealGameSyncShared/UnrealGameSyncShared.csproj b/Engine/Source/Programs/UnrealGameSync/UnrealGameSyncShared/UnrealGameSyncShared.csproj
--- a/Engine/Source/Programs/UnrealGameSync/UnrealGameSyncShared/UnrealGameSyncShared.csproj
+++ b/Engine/Source/Programs/UnrealGameSync/UnrealGameSyncShared/UnrealGameSyncShared.csproj
@@ -5,16 +5,6 @@
     <TargetFramework>net8.0</TargetFramework>
     <Nullable>enable</Nullable>
     <RootNamespace>UnrealGameSync</RootNamespace>
-    <RunAnalyzersDuringBuild>False</RunAnalyzersDuringBuild>
-    <Configurations>Debug;Release;Analyze</Configurations>
-  </PropertyGroup>
-
-  <PropertyGroup Condition=" '$(Configuration)|$(Platform)' == 'Analyze|AnyCPU' ">
-    <RunAnalyzersDuringBuild>True</RunAnalyzersDuringBuild>
-    <TreatWarningsAsErrors>false</TreatWarningsAsErrors>
-    <EnforceCodeStyleInBuild>True</EnforceCodeStyleInBuild>
-    <GenerateDocumentationFile>True</GenerateDocumentationFile>
-    <DocumentationFile></DocumentationFile>
   </PropertyGroup>
 
   <ItemGroup>
