added support to project extension specific section to dynamic ini system

#Synced-CL 50436066

diff --git a/Engine/Source/Programs/UnrealGameSync/UnrealGameSyncShared/ConfigUtils.cs b/Engine/Source/Programs/UnrealGameSync/UnrealGameSyncShared/ConfigUtils.cs
--- a/Engine/Source/Programs/UnrealGameSync/UnrealGameSyncShared/ConfigUtils.cs
+++ b/Engine/Source/Programs/UnrealGameSync/UnrealGameSyncShared/ConfigUtils.cs
@@ -830,15 +830,8 @@ namespace UnrealGameSync
 			ILogger logger,
 			CancellationToken cancellationToken)
 		{
-			// Find IniDynamicInjection section in config
-			ConfigSection? injectionSection = configFile.FindSection("IniDynamicInjection");
-			if (injectionSection == null)
-			{
-				return;
-			}
-
 			// Get all Inject values
-			string[] injectLines = injectionSection.GetValues("Inject", Array.Empty<string>());
+			string[] injectLines = configFile.GetValues("IniDynamicInjection.Inject", Array.Empty<string>(), projectIdentifier);
 			if (injectLines.Length == 0)
 			{
 				return;
