Ensureing the dynamic INI injection is only logged in the file log and not the synclog

#Synced-CL 50941491

#Synced-CL 50941491

diff --git a/Engine/Source/Programs/UnrealGameSync/UnrealGameSync/Controls/WorkspaceControl.cs b/Engine/Source/Programs/UnrealGameSync/UnrealGameSync/Controls/WorkspaceControl.cs
--- a/Engine/Source/Programs/UnrealGameSync/UnrealGameSync/Controls/WorkspaceControl.cs
+++ b/Engine/Source/Programs/UnrealGameSync/UnrealGameSync/Controls/WorkspaceControl.cs
@@ -1140,6 +1140,7 @@ namespace UnrealGameSync
 					_perforceMonitor.LatestPerforceConfigSection());
 
 				WorkspaceUpdateContext context = new WorkspaceUpdateContext(
+					_serviceProvider,
 					changeNumber,
 					options,
 					GetEditorBuildConfig(),
@@ -3750,7 +3751,7 @@ namespace UnrealGameSync
 				_owner.ShowAndActivate();
 
 				WorkspaceUpdateOptions options = WorkspaceUpdateOptions.Build | WorkspaceUpdateOptions.RunAfterSync;
-				WorkspaceUpdateContext context = new WorkspaceUpdateContext(_workspace.CurrentChangeNumber, options, _settings.CompiledEditorBuildConfig, null, _projectSettings.BuildSteps, null);
+				WorkspaceUpdateContext context = new WorkspaceUpdateContext(_serviceProvider, _workspace.CurrentChangeNumber, options, _settings.CompiledEditorBuildConfig, null, _projectSettings.BuildSteps, null);
 				StartWorkspaceUpdate(context, null);
 			}
 		}
@@ -5687,7 +5688,7 @@ namespace UnrealGameSync
 				List<ConfigObject> userBuildSteps = new List<ConfigObject>();
 				userBuildSteps.Add(customBuildStep.ToConfigObject());
 
-				WorkspaceUpdateContext context = new WorkspaceUpdateContext(_workspace.CurrentChangeNumber, WorkspaceUpdateOptions.Build, GetEditorBuildConfig(), null, userBuildSteps, new HashSet<Guid> { customBuildStep.UniqueId });
+				WorkspaceUpdateContext context = new WorkspaceUpdateContext(_serviceProvider, _workspace.CurrentChangeNumber, WorkspaceUpdateOptions.Build, GetEditorBuildConfig(), null, userBuildSteps, new HashSet<Guid> { customBuildStep.UniqueId });
 				StartWorkspaceUpdate(context, null);
 			}
 		}
@@ -6756,7 +6757,7 @@ namespace UnrealGameSync
 						}
 					}
 
-					WorkspaceUpdateContext context = new WorkspaceUpdateContext(_workspace.CurrentChangeNumber, WorkspaceUpdateOptions.Build, GetEditorBuildConfig(), null, _projectSettings.BuildSteps, stepSet);
+					WorkspaceUpdateContext context = new WorkspaceUpdateContext(_serviceProvider, _workspace.CurrentChangeNumber, WorkspaceUpdateOptions.Build, GetEditorBuildConfig(), null, _projectSettings.BuildSteps, stepSet);
 
 					if (step.ToolId != Guid.Empty)
 					{
diff --git a/Engine/Source/Programs/UnrealGameSync/UnrealGameSyncCmd/CommandContext.cs b/Engine/Source/Programs/UnrealGameSync/UnrealGameSyncCmd/CommandContext.cs
--- a/Engine/Source/Programs/UnrealGameSync/UnrealGameSyncCmd/CommandContext.cs
+++ b/Engine/Source/Programs/UnrealGameSync/UnrealGameSyncCmd/CommandContext.cs
@@ -17,6 +17,8 @@ namespace UnrealGameSyncCmd
 	{
 		public CommandLineArguments? Arguments { get; private set; }
 		
+		public IServiceProvider ServiceProvider { get; private set; }
+
 		public ILogger? Logger { get; private set; }
 
 		public ILoggerFactory? LoggerFactory { get; private set; }
@@ -41,8 +43,9 @@ namespace UnrealGameSyncCmd
 
 		public ConfigFile? ProjectConfig { get; private set; }
 
-		private CommandContext()
+		private CommandContext(IServiceProvider serviceProvider)
 		{
+			ServiceProvider = serviceProvider;
 		}
 
 		public void Dispose()
@@ -52,6 +55,7 @@ namespace UnrealGameSyncCmd
 
 		public static async Task<CommandContext> CreateAsync(
 			CommandLineArguments arguments,
+			IServiceProvider serviceProvider,
 			ILogger logger,
 			ILoggerFactory loggerFactory,
 			GlobalSettingsFile userSettings,
@@ -59,7 +63,7 @@ namespace UnrealGameSyncCmd
 			IHordeClient? hordeClient,
 			ICloudStorage? cloudStorage)
 		{
-			CommandContext context = new CommandContext();
+			CommandContext context = new CommandContext(serviceProvider);
 			context.Arguments = arguments;
 			context.Logger = logger;
 			context.LoggerFactory = loggerFactory;
diff --git a/Engine/Source/Programs/UnrealGameSync/UnrealGameSyncCmd/Commands/BuildCommand.cs b/Engine/Source/Programs/UnrealGameSync/UnrealGameSyncCmd/Commands/BuildCommand.cs
--- a/Engine/Source/Programs/UnrealGameSync/UnrealGameSyncCmd/Commands/BuildCommand.cs
+++ b/Engine/Source/Programs/UnrealGameSync/UnrealGameSyncCmd/Commands/BuildCommand.cs
@@ -33,6 +33,7 @@ namespace UnrealGameSyncCmd.Commands
 			ConfigFile projectConfig = context.ProjectConfig!;
 			UserProjectSettings projectSettings = context.UserProjectSettings!;
 			WorkspaceStateWrapper state = context.WorkspaceState!;
+			IServiceProvider serviceProvider = context.ServiceProvider;
 
 			context.Arguments!.TryGetPositionalArgument(out string? target);
 
@@ -121,6 +122,7 @@ namespace UnrealGameSyncCmd.Commands
 			}
 
 			WorkspaceUpdateContext updateContext = new WorkspaceUpdateContext(
+				context.ServiceProvider,
 				state.Current.CurrentChangeNumber,
 				WorkspaceUpdateOptions.Build,
 				BuildConfig.Development,
diff --git a/Engine/Source/Programs/UnrealGameSync/UnrealGameSyncCmd/Commands/SyncCommand.cs b/Engine/Source/Programs/UnrealGameSync/UnrealGameSyncCmd/Commands/SyncCommand.cs
--- a/Engine/Source/Programs/UnrealGameSync/UnrealGameSyncCmd/Commands/SyncCommand.cs
+++ b/Engine/Source/Programs/UnrealGameSync/UnrealGameSyncCmd/Commands/SyncCommand.cs
@@ -138,7 +138,7 @@ namespace UnrealGameSyncCmd.Commands
 				return;
 			}
 
-			WorkspaceUpdateContext updateContext = new WorkspaceUpdateContext(change, options, BuildConfig.Development, syncFilter, projectSettings.BuildSteps, null);
+			WorkspaceUpdateContext updateContext = new WorkspaceUpdateContext(context.ServiceProvider, change, options, BuildConfig.Development, syncFilter, projectSettings.BuildSteps, null);
 			updateContext.PerforceSyncOptions = userSettings.Global.Perforce;
 
 			if (syncOptions.Binaries)
diff --git a/Engine/Source/Programs/UnrealGameSync/UnrealGameSyncCmd/Program.cs b/Engine/Source/Programs/UnrealGameSync/UnrealGameSyncCmd/Program.cs
--- a/Engine/Source/Programs/UnrealGameSync/UnrealGameSyncCmd/Program.cs
+++ b/Engine/Source/Programs/UnrealGameSync/UnrealGameSyncCmd/Program.cs
@@ -206,7 +206,7 @@ namespace UnrealGameSyncCmd
 					//////////////////////////////////////
 					case Command contextCommand:
 						{
-							using CommandContext commandContext = await CommandContext.CreateAsync(args, logger, loggerFactory, settings, globalSettings, hordeClient, cloudStorage);
+							using CommandContext commandContext = await CommandContext.CreateAsync(args, serviceProvider,logger, loggerFactory, settings, globalSettings, hordeClient, cloudStorage);
 
 							await contextCommand.ExecuteAsync(commandContext!);
 							break;
diff --git a/Engine/Source/Programs/UnrealGameSync/UnrealGameSyncShared/WorkspaceUpdate.cs b/Engine/Source/Programs/UnrealGameSync/UnrealGameSyncShared/WorkspaceUpdate.cs
--- a/Engine/Source/Programs/UnrealGameSync/UnrealGameSyncShared/WorkspaceUpdate.cs
+++ b/Engine/Source/Programs/UnrealGameSync/UnrealGameSyncShared/WorkspaceUpdate.cs
@@ -6,7 +6,6 @@ using System.Collections.Generic;
 using System.Diagnostics;
 using System.IO;
 using System.Linq;
-using System.Runtime;
 using System.Runtime.InteropServices;
 using System.Security.Cryptography;
 using System.Text;
@@ -18,7 +17,7 @@ using System.Threading.Tasks;
 using EpicGames.Core;
 using EpicGames.Perforce;
 using EpicGames.ProjectStore;
-using Google.Protobuf.Compiler;
+using Microsoft.Extensions.DependencyInjection;
 using Microsoft.Extensions.Logging;
 
 namespace UnrealGameSync
@@ -115,6 +114,7 @@ namespace UnrealGameSync
 		public List<SnapshotSettings> Snapshots {  get; } = new List<SnapshotSettings>();
 		public string SnapshotHost { get; set; } = "";
 		public string SnapshotNamespace { get; set; } = "";
+		public IServiceProvider ServiceProvider { get; init; }
 
 		// May be updated during sync
 #pragma warning disable CA2227 // Collection properties should be read only
@@ -123,8 +123,9 @@ namespace UnrealGameSync
 		public ISet<string> BaseContentPaths { get; set; } = new HashSet<string>();
 #pragma warning restore CA2227
 
-		public WorkspaceUpdateContext(int changeNumber, WorkspaceUpdateOptions options, BuildConfig editorConfig, string[]? syncFilter, List<ConfigObject> userBuildSteps, HashSet<Guid>? customBuildSteps)
+		public WorkspaceUpdateContext(IServiceProvider serviceProvider, int changeNumber, WorkspaceUpdateOptions options, BuildConfig editorConfig, string[]? syncFilter, List<ConfigObject> userBuildSteps, HashSet<Guid>? customBuildSteps)
 		{
+			ServiceProvider = serviceProvider;
 			ChangeNumber = changeNumber;
 			Options = options;
 			EditorConfig = editorConfig;
@@ -1071,7 +1072,7 @@ namespace UnrealGameSync
 					if (Context.Options.HasFlag(WorkspaceUpdateOptions.Sync) && !Context.Options.HasFlag(WorkspaceUpdateOptions.UpdateFilter))
 					{
 						// Read the new config file
-						Context.ProjectConfigFile = await ReadProjectConfigFile(perforce, project, logger);
+						Context.ProjectConfigFile = await ReadProjectConfigFile(perforce, project, Context.ServiceProvider.GetRequiredService<ILogger<WorkspaceUpdate>>());
 						Context.ProjectStreamFilter = await ReadProjectStreamFilter(perforce, Context.ProjectConfigFile, cancellationToken);
 
 						// Get the branch name
@@ -1258,7 +1259,7 @@ namespace UnrealGameSync
 					// Continue processing sync-only actions
 					if (Context.Options.HasFlag(WorkspaceUpdateOptions.Sync) && !Context.Options.HasFlag(WorkspaceUpdateOptions.UpdateFilter))
 					{
-						Context.ProjectConfigFile ??= await ConfigUtils.ReadProjectConfigFileAsync(perforce, project, logger, cancellationToken);
+						Context.ProjectConfigFile ??= await ConfigUtils.ReadProjectConfigFileAsync(perforce, project, Context.ServiceProvider.GetRequiredService<ILogger<WorkspaceUpdate>>(), cancellationToken);
 
 						// Execute any project specific post-sync steps
 						string[]? postSyncSteps = Context.ProjectConfigFile.GetValues("Sync.Step", null, project.ProjectIdentifier);
@@ -1604,7 +1605,7 @@ namespace UnrealGameSync
 			// Build everything using MegaXGE
 			if (Context.Options.HasFlag(WorkspaceUpdateOptions.Build))
 			{
-				Context.ProjectConfigFile ??= await ReadProjectConfigFile(perforce, project, logger);
+				Context.ProjectConfigFile ??= await ReadProjectConfigFile(perforce, project, Context.ServiceProvider.GetRequiredService<ILogger<WorkspaceUpdate>>());
 
 				// Figure out the new editor target
 				TargetReceipt defaultEditorReceipt = ConfigUtils.CreateDefaultEditorReceipt(project, Context.ProjectConfigFile, Context.EditorConfig);
