Updating label for setting pcb's to disabled in the pcb override.

#Synced-CL 49593614

diff --git a/Engine/Source/Programs/UnrealGameSync/UnrealGameSync/Controls/WorkspaceControl.cs b/Engine/Source/Programs/UnrealGameSync/UnrealGameSync/Controls/WorkspaceControl.cs
--- a/Engine/Source/Programs/UnrealGameSync/UnrealGameSync/Controls/WorkspaceControl.cs
+++ b/Engine/Source/Programs/UnrealGameSync/UnrealGameSync/Controls/WorkspaceControl.cs
@@ -5362,7 +5362,7 @@ namespace UnrealGameSync
 
 			if (_projectSettings.ArchivesDeactivated)
 			{
-				text = "PCB: Compile locally (project)";
+				text = "PCB: Do not sync PCBs (project)";
 			}
 			else
 			{
@@ -5444,7 +5444,7 @@ namespace UnrealGameSync
 							disableItem.Click += (sender, args) => ClearSelectedProjectOverrideArchives();
 							OptionsContextMenu_SyncPrecompiledBinariesProjectOverride.DropDownItems.Add(disableItem);
 
-							ToolStripMenuItem deactivateItem = new ToolStripMenuItem("Disable (compile locally)");
+							ToolStripMenuItem deactivateItem = new ToolStripMenuItem("Do not sync PCBs");
 							deactivateItem.Checked = _projectSettings.ArchivesDeactivated;
 							deactivateItem.Click += (sender, args) => DeactivateSelectedProjectOverrideArchives();
 							OptionsContextMenu_SyncPrecompiledBinariesProjectOverride.DropDownItems.Add(deactivateItem);
