- fix local project path being wrong for virtual streams
- pass the user's workspace and project settings to post sync steps

#Synced-CL 50941491

diff --git a/Engine/Source/Programs/UnrealGameSync/UnrealGameSyncCmd/Commands/SettingsCommand.cs b/Engine/Source/Programs/UnrealGameSync/UnrealGameSyncCmd/Commands/SettingsCommand.cs
--- a/Engine/Source/Programs/UnrealGameSync/UnrealGameSyncCmd/Commands/SettingsCommand.cs
+++ b/Engine/Source/Programs/UnrealGameSync/UnrealGameSyncCmd/Commands/SettingsCommand.cs
@@ -1,10 +1,14 @@
 ﻿// Copyright Epic Games, Inc. All Rights Reserved.
 
 using System;
+using System.Collections.Generic;
 using System.IO;
+using System.Linq;
 using System.Text.Json;
+using System.Threading;
 using System.Threading.Tasks;
 using EpicGames.Core;
+using EpicGames.Perforce;
 using Microsoft.Extensions.Logging;
 using UnrealGameSync;
 using UnrealGameSyncCmd.Utils;
@@ -22,7 +26,7 @@ namespace UnrealGameSyncCmd.Commands
 
 	internal class SettingsCommand : ICommand
 	{
-		public static async Task ExecuteAsync(ILogger logger, CommandLineArguments args, UserSettings? globalSettings, GlobalSettingsFile? userSettings)
+		public static async Task ExecuteAsync(ILogger logger, ILoggerFactory loggerFactory, CommandLineArguments args, UserSettings? globalSettings, GlobalSettingsFile? userSettings)
 		{
 			SettingsCommandOptions options = args.ApplyTo<SettingsCommandOptions>(logger);
 			args.CheckAllArgumentsUsed();
@@ -45,6 +49,18 @@ namespace UnrealGameSyncCmd.Commands
 			{
 				logger.LogInformation("Writing user settings to {User}", options.User);
 				UserWorkspaceSettings userWorkspaceSettings = UserSettingsUtils.ReadRequiredUserWorkspaceSettings();
+
+				IPerforceConnection perforce = await PerforceConnectionUtils.ConnectAsync(userWorkspaceSettings, loggerFactory);
+
+				// get the local path for the project
+				List<WhereRecord> records = await perforce.WhereAsync($"//{userWorkspaceSettings.ClientName}{userWorkspaceSettings.ProjectPath}", CancellationToken.None).Where(x => !x.Unmap).ToListAsync();
+				if (records.Count == 0)
+				{
+					throw new Exception($"Unable to find a local path for the project. Client: {userWorkspaceSettings.ClientName}, ProjectPath: {userWorkspaceSettings.ProjectPath}");
+				}
+
+				userWorkspaceSettings.LocalProjectPath = new FileReference(records[0].Path);
+
 				await WriteSettings(options.User, userWorkspaceSettings);
 			}
 		}
diff --git a/Engine/Source/Programs/UnrealGameSync/UnrealGameSyncCmd/Program.cs b/Engine/Source/Programs/UnrealGameSync/UnrealGameSyncCmd/Program.cs
--- a/Engine/Source/Programs/UnrealGameSync/UnrealGameSyncCmd/Program.cs
+++ b/Engine/Source/Programs/UnrealGameSync/UnrealGameSyncCmd/Program.cs
@@ -185,7 +185,7 @@ namespace UnrealGameSyncCmd
 					//////////////////////////////////////
 					case SettingsCommand _:
 						{
-							await SettingsCommand.ExecuteAsync(logger, args, globalSettings, settings);
+							await SettingsCommand.ExecuteAsync(logger, loggerFactory, args, globalSettings, settings);
 							break;
 						}
 
