- fixed issue with WorkspaceState being initialized before the localProjectPath was determined
- added parameters to dynamic ini injection

#Synced-CL 50941491

diff --git a/Engine/Source/Programs/UnrealGameSync/UnrealGameSyncCmd/CommandContext.cs b/Engine/Source/Programs/UnrealGameSync/UnrealGameSyncCmd/CommandContext.cs
--- a/Engine/Source/Programs/UnrealGameSync/UnrealGameSyncCmd/CommandContext.cs
+++ b/Engine/Source/Programs/UnrealGameSync/UnrealGameSyncCmd/CommandContext.cs
@@ -75,19 +75,21 @@ namespace UnrealGameSyncCmd
 			context.CloudStorage = cloudStorage;
 			context.UserWorkspaceSettings = UserSettingsUtils.ReadRequiredUserWorkspaceSettings();
 			context.PerforceClient = await PerforceConnectionUtils.ConnectAsync(context.UserWorkspaceSettings!, context.LoggerFactory!);
-			context.WorkspaceState = await WorkspaceStateUtils.ReadWorkspaceState(context.PerforceClient, context.UserWorkspaceSettings, context.UserSettings!, context.Logger!);
+			
 
 			// get the local path for the project
-			List<WhereRecord> records = await context.PerforceClient.WhereAsync($"//{context.WorkspaceState.Current.ClientName}{context.WorkspaceState.Current.ProjectPath}", CancellationToken.None).Where(x => !x.Unmap).ToListAsync();
+			List<WhereRecord> records = await context.PerforceClient.WhereAsync($"//{context.UserWorkspaceSettings.ClientName}{context.UserWorkspaceSettings.ProjectPath}", CancellationToken.None).Where(x => !x.Unmap).ToListAsync();
 			if(records.Count == 0)
 			{
-				throw new Exception($"Unable to find a local path for the project. Client: {context.WorkspaceState.Current.ClientName}, ProjectPath: {context.WorkspaceState.Current.ProjectPath}");
+				throw new Exception($"Unable to find a local path for the project. Client: {context.UserWorkspaceSettings.ClientName}, ProjectPath: {context.UserWorkspaceSettings.ProjectPath}");
 			}
 
 			string localProjectPath = records[0].Path;
 
 			context.UserWorkspaceSettings.LocalProjectPath = new FileReference(localProjectPath);
 
+			context.WorkspaceState = await WorkspaceStateUtils.ReadWorkspaceState(context.PerforceClient, context.UserWorkspaceSettings, context.UserSettings!, context.Logger!);
+
 			context.ProjectInfo = new ProjectInfo(
 				context.UserWorkspaceSettings.RootDir,
 				context.WorkspaceState.Current.ClientName,
diff --git a/Engine/Source/Programs/UnrealGameSync/UnrealGameSyncShared/ConfigUtils.cs b/Engine/Source/Programs/UnrealGameSync/UnrealGameSyncShared/ConfigUtils.cs
--- a/Engine/Source/Programs/UnrealGameSync/UnrealGameSyncShared/ConfigUtils.cs
+++ b/Engine/Source/Programs/UnrealGameSync/UnrealGameSyncShared/ConfigUtils.cs
@@ -142,7 +142,7 @@ namespace UnrealGameSync
 			}
 
 			// Execute dynamic config injection scripts
-			await ExecuteConfigInjectionScriptsAsync(perforce, branchClientPath, projectIdentifier, cacheFolder, projectConfig, localConfigFiles, logger, cancellationToken);
+			await ExecuteConfigInjectionScriptsAsync(perforce, branchClientPath, selectedClientFileName, projectIdentifier, cacheFolder, projectConfig, localConfigFiles, logger, cancellationToken);
 
 			return projectConfig;
 		}
@@ -823,6 +823,7 @@ namespace UnrealGameSync
 		public static async Task ExecuteConfigInjectionScriptsAsync(
 			IPerforceConnection perforce,
 			string branchClientPath,
+			string selectedClientFileName,
 			string projectIdentifier,
 			DirectoryReference cacheFolder,
 			ConfigFile configFile,
@@ -857,6 +858,7 @@ namespace UnrealGameSync
 
 					string platform = injectObj.GetValue("Platform", "");
 					string scriptPath = injectObj.GetValue("Script", "");
+					string scriptArgs = injectObj.GetValue("Args", "");
 
 					if (String.IsNullOrWhiteSpace(scriptPath))
 					{
@@ -871,7 +873,7 @@ namespace UnrealGameSync
 						continue;
 					}
 
-					await ExecuteSingleInjectionScriptAsync(perforce, branchClientPath, localRootPath, projectIdentifier, cacheFolder, configFile, localConfigFiles, scriptPath, logger, cancellationToken);
+					await ExecuteSingleInjectionScriptAsync(perforce, branchClientPath, localRootPath, selectedClientFileName, projectIdentifier, cacheFolder, configFile, localConfigFiles, scriptPath, scriptArgs, logger, cancellationToken);
 				}
 				catch (Exception ex)
 				{
@@ -884,11 +886,13 @@ namespace UnrealGameSync
 			IPerforceConnection perforce,
 			string branchClientPath,
 			DirectoryReference localRootPath,
+			string selectedClientFileName,
 			string projectIdentifier,
 			DirectoryReference cacheFolder,
 			ConfigFile configFile,
 			List<KeyValuePair<FileReference, DateTime>> localConfigFiles,
 			string relativePath,
+			string scriptArgs,
 			ILogger logger,
 			CancellationToken cancellationToken)
 		{
@@ -1048,6 +1052,20 @@ namespace UnrealGameSync
 					// pass interpreter args, then script path, then workspace root.
 					arguments = $"{interpreterArgs} \"{localPath}\" \"{localRootPath.FullName}\"";
 				}
+				
+				ClientRecord clientRecord = await perforce.GetClientAsync(null, cancellationToken);
+
+				string clientName = clientRecord.Name;
+
+				// Append additional script arguments if specified
+				if (!String.IsNullOrEmpty(scriptArgs))
+				{
+					scriptArgs = scriptArgs.Replace("$(ClientName)", clientRecord.Name, StringComparison.OrdinalIgnoreCase);
+					scriptArgs = scriptArgs.Replace("$(ClientRoot)", clientRecord.Root, StringComparison.OrdinalIgnoreCase);
+					scriptArgs = scriptArgs.Replace("$(ProjectClientPath)", selectedClientFileName, StringComparison.OrdinalIgnoreCase);
+
+					arguments = $"{arguments} {scriptArgs}";
+				}
 
 				int exitCode = await Utility.ExecuteProcessAsync(
 					interpreter,
diff --git a/Engine/Source/Programs/UnrealGameSync/UnrealGameSyncShared/WorkspaceUpdate.cs b/Engine/Source/Programs/UnrealGameSync/UnrealGameSyncShared/WorkspaceUpdate.cs
--- a/Engine/Source/Programs/UnrealGameSync/UnrealGameSyncShared/WorkspaceUpdate.cs
+++ b/Engine/Source/Programs/UnrealGameSync/UnrealGameSyncShared/WorkspaceUpdate.cs
@@ -2259,7 +2259,8 @@ namespace UnrealGameSync
 
 			await ConfigUtils.ExecuteConfigInjectionScriptsAsync(
 				perforce, 
-				projectInfo.BranchPath, 
+				projectInfo.BranchPath,
+				projectInfo.ClientFileName,
 				projectInfo.ProjectIdentifier, 
 				projectInfo.CacheFolder, 
 				projectConfig,
