UnrealGameSync: Bump max changelist distance

#rnx

#Synced-CL 51357799

diff --git a/Engine/Source/Programs/UnrealGameSync/UnrealGameSyncShared/WorkspaceUpdate.cs b/Engine/Source/Programs/UnrealGameSync/UnrealGameSyncShared/WorkspaceUpdate.cs
--- a/Engine/Source/Programs/UnrealGameSync/UnrealGameSyncShared/WorkspaceUpdate.cs
+++ b/Engine/Source/Programs/UnrealGameSync/UnrealGameSyncShared/WorkspaceUpdate.cs
@@ -899,7 +899,8 @@ namespace UnrealGameSync
 
 							if (Context.Options.HasFlag(WorkspaceUpdateOptions.Sync) && Context.Options.HasFlag(WorkspaceUpdateOptions.FastSync))
 							{
-								bool changelistDistanceInRange = Math.Abs(Context.ChangeNumber - state.CurrentChangeNumber) < 100_000;
+								int maxChangelistDistance = 500_000;
+								bool changelistDistanceInRange = Math.Abs(Context.ChangeNumber - state.CurrentChangeNumber) < maxChangelistDistance;
 								bool targetChangelistBiggerThanCurrent = (Context.ChangeNumber >= state.CurrentChangeNumber);
 
 								if (state.CurrentChangeNumber > 0 && changelistDistanceInRange && targetChangelistBiggerThanCurrent)
@@ -919,7 +920,7 @@ namespace UnrealGameSync
 									}
 									else if (!changelistDistanceInRange)
 									{
-										logger.LogInformation("Fast Sync cannot be used with current changelist {CurrentChangelist} and target changelist {TargetChangelist} because the distance between then is > 100K changelists.", state.CurrentChangeNumber, Context.ChangeNumber);
+										logger.LogInformation("Fast Sync cannot be used with current changelist {CurrentChangelist} and target changelist {TargetChangelist} because the distance between them is > {MaxDistance}K changelists.", state.CurrentChangeNumber, Context.ChangeNumber, maxChangelistDistance / 1000);
 									}
 								}
 							}
