fixing bug with compile locally

#Synced-CL 51292625

diff --git a/Engine/Source/Programs/UnrealGameSync/UnrealGameSyncCmd/Commands/ArchiveOverrideCommand.cs b/Engine/Source/Programs/UnrealGameSync/UnrealGameSyncCmd/Commands/ArchiveOverrideCommand.cs
--- a/Engine/Source/Programs/UnrealGameSync/UnrealGameSyncCmd/Commands/ArchiveOverrideCommand.cs
+++ b/Engine/Source/Programs/UnrealGameSync/UnrealGameSyncCmd/Commands/ArchiveOverrideCommand.cs
@@ -116,7 +116,7 @@ namespace UnrealGameSyncCmd.Commands
 			else if(options.SetCompileLocally)
 			{
 				ResetArchiveSettings(logger, context.UserProjectSettings!);
-				context.UserProjectSettings!.ArchivesDeactivated = false;
+				context.UserProjectSettings!.ArchivesDeactivated = true;
 				context.UserProjectSettings!.Save(logger);
 				logger.LogInformation("Archive override set to compile locally.");
 			}
