UnrealGameSync: Enable nullable

#rnx

#Synced-CL 51442092

diff --git a/Engine/Source/Programs/UnrealGameSync/PostBadgeStatus/PostBadgeStatus.csproj b/Engine/Source/Programs/UnrealGameSync/PostBadgeStatus/PostBadgeStatus.csproj
--- a/Engine/Source/Programs/UnrealGameSync/PostBadgeStatus/PostBadgeStatus.csproj
+++ b/Engine/Source/Programs/UnrealGameSync/PostBadgeStatus/PostBadgeStatus.csproj
@@ -4,5 +4,6 @@
     <OutputType>Exe</OutputType>
     <GenerateAssemblyInfo>false</GenerateAssemblyInfo>
     <NoWarn>$(NoWarn);SYSLIB0014</NoWarn>
+    <Nullable>enable</Nullable>
   </PropertyGroup>
 </Project>
diff --git a/Engine/Source/Programs/UnrealGameSync/PostBadgeStatus/Program.cs b/Engine/Source/Programs/UnrealGameSync/PostBadgeStatus/Program.cs
--- a/Engine/Source/Programs/UnrealGameSync/PostBadgeStatus/Program.cs
+++ b/Engine/Source/Programs/UnrealGameSync/PostBadgeStatus/Program.cs
@@ -16,13 +16,13 @@ namespace WriteBadgeStatus
 		{
 			// Parse all the parameters
 			List<string> Arguments = new List<string>(Args);
-			string Name = ParseParam(Arguments, "Name");
-			string Change = ParseParam(Arguments, "Change");
-			string Project = ParseParam(Arguments, "Project");
-			string RestUrl = ParseParam(Arguments, "RestUrl");
-			string Status = ParseParam(Arguments, "Status");
-			string Url = ParseParam(Arguments, "Url");
-			string Metadata = ParseParam(Arguments, "Metadata");
+			string? Name = ParseParam(Arguments, "Name");
+			string? Change = ParseParam(Arguments, "Change");
+			string? Project = ParseParam(Arguments, "Project");
+			string? RestUrl = ParseParam(Arguments, "RestUrl");
+			string? Status = ParseParam(Arguments, "Status");
+			string? Url = ParseParam(Arguments, "Url");
+			string? Metadata = ParseParam(Arguments, "Metadata");
 
 			// Check we've got all the arguments we need (and no more)
 			if (Arguments.Count > 0 || Name == null || Change == null || Project == null || RestUrl == null || Status == null || Url == null)
@@ -78,7 +78,7 @@ namespace WriteBadgeStatus
 			}
 		}
 
-		static string ParseParam(List<string> Arguments, string ParamName)
+		static string? ParseParam(List<string> Arguments, string ParamName)
 		{
 			string ParamPrefix = String.Format("-{0}=", ParamName);
 			for (int Idx = 0; Idx < Arguments.Count; Idx++)
@@ -105,11 +105,11 @@ namespace WriteBadgeStatus
 			}
 
 			public int ChangeNumber;
-			public string BuildType;
+			public required string BuildType;
 			public BuildDataResult Result;
-			public string Url;
-			public string Project;
-			public string ArchivePath;
+			public required string Url;
+			public required string Project;
+			public required string ArchivePath;
 			public JsonElement? Metadata;
 
 			public bool IsSuccess => Result == BuildDataResult.Success || Result == BuildDataResult.Warning;
@@ -117,7 +117,7 @@ namespace WriteBadgeStatus
 			public bool IsFailure => Result == BuildDataResult.Failure;
 		}
 
-		static int SendRequest(string URI, string Resource, string Method, string RequestBody = null, params string[] QueryParams)
+		static int SendRequest(string URI, string Resource, string Method, string? RequestBody = null, params string[] QueryParams)
 		{
 			// set up the query string
 			StringBuilder TargetURI = new StringBuilder(String.Format("{0}/api/{1}", URI, Resource));
@@ -153,7 +153,7 @@ namespace WriteBadgeStatus
 			{
 				using (HttpWebResponse Response = (HttpWebResponse)Request.GetResponse())
 				{
-					string ResponseContent = null;
+					string? ResponseContent = null;
 					using (StreamReader ResponseReader = new System.IO.StreamReader(Response.GetResponseStream(), Encoding.UTF8))
 					{
 						ResponseContent = ResponseReader.ReadToEnd();
@@ -170,7 +170,7 @@ namespace WriteBadgeStatus
 				}
 				else
 				{
-					throw new Exception(String.Format("Request returned message: {0}", ex.InnerException.Message));
+					throw new Exception(String.Format("Request returned message: {0}", ex.InnerException?.Message ?? ex.Message));
 				}
 			}
 			catch (Exception ex)
diff --git a/Engine/Source/Programs/UnrealGameSync/UnrealGameSyncAutomation/Program.cs b/Engine/Source/Programs/UnrealGameSync/UnrealGameSyncAutomation/Program.cs
--- a/Engine/Source/Programs/UnrealGameSync/UnrealGameSyncAutomation/Program.cs
+++ b/Engine/Source/Programs/UnrealGameSync/UnrealGameSyncAutomation/Program.cs
@@ -46,7 +46,12 @@ namespace UnrealGameSyncAutomation
 
 		static Tuple<AutomationRequestResult, byte[]> SendRequest(AutomationRequestType Request, byte[] RequestData)
 		{
-			int PortNumber = (int)Registry.GetValue("HKEY_CURRENT_USER\\Software\\Epic Games\\UnrealGameSync", "AutomationPort", null);
+			object PortValue = Registry.GetValue("HKEY_CURRENT_USER\\Software\\Epic Games\\UnrealGameSync", "AutomationPort", null)
+				?? throw new InvalidOperationException("AutomationPort registry value is not set.");
+			if (!int.TryParse(PortValue.ToString(), out int PortNumber) || PortNumber < 1 || PortNumber > 65535)
+			{
+				throw new InvalidOperationException($"AutomationPort registry value '{PortValue}' is not a valid TCP port number (1-65535).");
+			}
 			using (TcpClient Client = new TcpClient())
 			{
 				Client.Connect(new IPEndPoint(IPAddress.Loopback, PortNumber));
diff --git a/Engine/Source/Programs/UnrealGameSync/UnrealGameSyncAutomation/UnrealGameSyncAutomation.csproj b/Engine/Source/Programs/UnrealGameSync/UnrealGameSyncAutomation/UnrealGameSyncAutomation.csproj
--- a/Engine/Source/Programs/UnrealGameSync/UnrealGameSyncAutomation/UnrealGameSyncAutomation.csproj
+++ b/Engine/Source/Programs/UnrealGameSync/UnrealGameSyncAutomation/UnrealGameSyncAutomation.csproj
@@ -9,5 +9,6 @@
     <GenerateAssemblyInfo>false</GenerateAssemblyInfo>
     <RootNamespace>UnrealGameSyncAutomation</RootNamespace>
     <AssemblyName>UnrealGameSyncAutomation</AssemblyName>
+    <Nullable>enable</Nullable>
   </PropertyGroup>
 </Project>
\ No newline at end of file
