added a simple feature flag system in the INI + handy cache

#Synced-CL 51549253

diff --git a/Engine/Source/Programs/UnrealGameSync/UnrealGameSyncShared/ConfigUtils.cs b/Engine/Source/Programs/UnrealGameSync/UnrealGameSyncShared/ConfigUtils.cs
--- a/Engine/Source/Programs/UnrealGameSync/UnrealGameSyncShared/ConfigUtils.cs
+++ b/Engine/Source/Programs/UnrealGameSync/UnrealGameSyncShared/ConfigUtils.cs
@@ -882,6 +882,72 @@ namespace UnrealGameSync
 			}
 		}
 
+
+		public static IDictionary<string, FeatureFlag> GetFeatureFlags(ConfigFile? projectConfigFile, string projectIdentifier)
+		{
+			Dictionary<string, FeatureFlag> output = new(StringComparer.OrdinalIgnoreCase);
+
+			if (projectConfigFile == null)
+			{
+				return output;
+			}
+			List<string> features = new List<string>();
+
+			features.AddRange(projectConfigFile.GetValues("Features.FeatureFlag", []));
+
+			if (TryGetProjectSetting(projectConfigFile, projectIdentifier, "FeatureFlag", out string? projectFeatures))
+			{
+				using StringReader sr = new StringReader(projectFeatures);
+				while (sr.Peek() != -1)
+				{
+					string? line = sr.ReadLine();
+					if (String.IsNullOrWhiteSpace(line))
+					{
+						continue;
+					}
+					features.Add(line);
+				}
+			}
+
+			foreach (string presetDefinition in features)
+			{
+				ConfigObject obj = new ConfigObject(presetDefinition);
+
+				FeatureFlag? ff = new FeatureFlag();
+
+				ff.Name = obj.GetValue("Name", String.Empty);
+
+				// do not allow empty name
+				if (String.IsNullOrWhiteSpace(ff.Name))
+				{
+					continue;
+				}
+
+				ff.Enabled = obj.GetValue("Enabled", ff.Enabled);
+
+				if (output.ContainsKey(ff.Name))
+				{
+					output[ff.Name].Import(ff);
+				}
+				else
+				{
+					output.TryAdd(ff.Name, ff);
+				}
+			}
+
+			return output;
+		}
+
+		public static bool IsFeatureEnabled(IDictionary<string, FeatureFlag> features, string featureName)
+		{
+			if (features.TryGetValue(featureName, out FeatureFlag? featureFlag))
+			{
+				return featureFlag.Enabled;
+			}
+
+			return false;
+		}
+
 		private static async Task ExecuteSingleInjectionScriptAsync(
 			IPerforceConnection perforce,
 			string branchClientPath,
diff --git a/Engine/Source/Programs/UnrealGameSync/UnrealGameSyncShared/FeatureFlag.cs b/Engine/Source/Programs/UnrealGameSync/UnrealGameSyncShared/FeatureFlag.cs
new file mode 100644
--- /dev/null
+++ b/Engine/Source/Programs/UnrealGameSync/UnrealGameSyncShared/FeatureFlag.cs
@@ -0,0 +1,39 @@
+﻿// Copyright Epic Games, Inc. All Rights Reserved.
+
+using System;
+using System.Collections.Generic;
+
+namespace UnrealGameSync
+{
+	public class FeatureFlag
+	{
+		public string Name { get; set; } = String.Empty;
+		public bool Enabled { get; set; } = false;
+
+		public void Import(FeatureFlag ff)
+		{
+			Name = ff.Name;
+			Enabled = ff.Enabled;
+		}
+	}
+
+	public class FeatureFlagCache
+	{
+		private IDictionary<string, FeatureFlag> Features { get; }
+
+		public FeatureFlagCache(IDictionary<string, FeatureFlag> features)
+		{
+			Features = features;
+		}
+
+		public bool IsFeatureEnabled(string featureName)
+		{
+			if (Features.TryGetValue(featureName, out FeatureFlag? featureFlag))
+			{
+				return featureFlag.Enabled;
+			}
+
+			return false;
+		}
+	}
+}
