UnrealGameSync: Update to TargetFramework net10.0

#rnx
#Synced-CL 51592834

diff --git a/Engine/Source/Programs/UnrealGameSync/PostBadgeStatus/PostBadgeStatus.csproj b/Engine/Source/Programs/UnrealGameSync/PostBadgeStatus/PostBadgeStatus.csproj
--- a/Engine/Source/Programs/UnrealGameSync/PostBadgeStatus/PostBadgeStatus.csproj
+++ b/Engine/Source/Programs/UnrealGameSync/PostBadgeStatus/PostBadgeStatus.csproj
@@ -1,9 +1,7 @@
 ﻿<Project Sdk="Microsoft.NET.Sdk">
-  <Import Project="..\..\Shared\UnrealEngine.csproj.props" />
   <PropertyGroup>
     <OutputType>Exe</OutputType>
     <GenerateAssemblyInfo>false</GenerateAssemblyInfo>
     <NoWarn>$(NoWarn);SYSLIB0014</NoWarn>
-    <Nullable>enable</Nullable>
   </PropertyGroup>
 </Project>
diff --git a/Engine/Source/Programs/UnrealGameSync/UnrealGameSync/AutomationServer.cs b/Engine/Source/Programs/UnrealGameSync/UnrealGameSync/AutomationServer.cs
--- a/Engine/Source/Programs/UnrealGameSync/UnrealGameSync/AutomationServer.cs
+++ b/Engine/Source/Programs/UnrealGameSync/UnrealGameSync/AutomationServer.cs
@@ -342,7 +342,7 @@ namespace UnrealGameSync
 			int len = stream.ReadByte() * 256;
 			len += stream.ReadByte();
 			byte[] inBuffer = new byte[len];
-			stream.Read(inBuffer, 0, len);
+			stream.ReadExactly(inBuffer, 0, len);
 
 			return _streamEncoding.GetString(inBuffer);
 		}
diff --git a/Engine/Source/Programs/UnrealGameSync/UnrealGameSync/Controls/AlertButtonControl.cs b/Engine/Source/Programs/UnrealGameSync/UnrealGameSync/Controls/AlertButtonControl.cs
--- a/Engine/Source/Programs/UnrealGameSync/UnrealGameSync/Controls/AlertButtonControl.cs
+++ b/Engine/Source/Programs/UnrealGameSync/UnrealGameSync/Controls/AlertButtonControl.cs
@@ -69,6 +69,7 @@ namespace UnrealGameSync
 		bool _mouseOver;
 		bool _mouseDown;
 
+		[DesignerSerializationVisibility(DesignerSerializationVisibility.Hidden)]
 		public AlertButtonTheme Theme
 		{
 			get => _themeValue;
@@ -79,6 +80,7 @@ namespace UnrealGameSync
 			}
 		}
 
+		[DesignerSerializationVisibility(DesignerSerializationVisibility.Hidden)]
 		public AlertButtonColors CustomColors
 		{
 			get => _customColorsValue;
diff --git a/Engine/Source/Programs/UnrealGameSync/UnrealGameSync/Controls/AlertImageControl.cs b/Engine/Source/Programs/UnrealGameSync/UnrealGameSync/Controls/AlertImageControl.cs
--- a/Engine/Source/Programs/UnrealGameSync/UnrealGameSync/Controls/AlertImageControl.cs
+++ b/Engine/Source/Programs/UnrealGameSync/UnrealGameSync/Controls/AlertImageControl.cs
@@ -1,6 +1,7 @@
 // Copyright Epic Games, Inc. All Rights Reserved.
 
 using System;
+using System.ComponentModel;
 using System.Drawing;
 using System.Drawing.Drawing2D;
 using System.Windows.Forms;
@@ -11,6 +12,7 @@ namespace UnrealGameSync
 	{
 		Image? _imageValue;
 
+		[DesignerSerializationVisibility(DesignerSerializationVisibility.Hidden)]
 		public Image? Image
 		{
 			get => _imageValue;
diff --git a/Engine/Source/Programs/UnrealGameSync/UnrealGameSync/Controls/LogSplitContainer.cs b/Engine/Source/Programs/UnrealGameSync/UnrealGameSync/Controls/LogSplitContainer.cs
--- a/Engine/Source/Programs/UnrealGameSync/UnrealGameSync/Controls/LogSplitContainer.cs
+++ b/Engine/Source/Programs/UnrealGameSync/UnrealGameSync/Controls/LogSplitContainer.cs
@@ -31,10 +31,11 @@ namespace UnrealGameSync
 		public event Action<bool>? OnVisibilityChanged;
 
 		[Browsable(true)]
+		[DesignerSerializationVisibility(DesignerSerializationVisibility.Hidden)]
 		public string Caption
 		{
 			get;
-			set;
+			set => field = value ?? string.Empty;
 		}
 
 		public LogSplitContainer()
diff --git a/Engine/Source/Programs/UnrealGameSync/UnrealGameSync/Controls/TextBoxWithCueBanner.cs b/Engine/Source/Programs/UnrealGameSync/UnrealGameSync/Controls/TextBoxWithCueBanner.cs
--- a/Engine/Source/Programs/UnrealGameSync/UnrealGameSync/Controls/TextBoxWithCueBanner.cs
+++ b/Engine/Source/Programs/UnrealGameSync/UnrealGameSync/Controls/TextBoxWithCueBanner.cs
@@ -1,6 +1,7 @@
 // Copyright Epic Games, Inc. All Rights Reserved.
 
 using System;
+using System.ComponentModel;
 using System.Runtime.InteropServices;
 using System.Windows.Forms;
 
@@ -15,6 +16,7 @@ namespace UnrealGameSync
 
 		private string _cueBannerValue = String.Empty;
 
+		[DesignerSerializationVisibility(DesignerSerializationVisibility.Hidden)]
 		public string CueBanner
 		{
 			get => _cueBannerValue;
diff --git a/Engine/Source/Programs/UnrealGameSync/UnrealGameSync/Controls/WorkspaceSyncFilterControl.cs b/Engine/Source/Programs/UnrealGameSync/UnrealGameSync/Controls/WorkspaceSyncFilterControl.cs
--- a/Engine/Source/Programs/UnrealGameSync/UnrealGameSync/Controls/WorkspaceSyncFilterControl.cs
+++ b/Engine/Source/Programs/UnrealGameSync/UnrealGameSync/Controls/WorkspaceSyncFilterControl.cs
@@ -34,7 +34,9 @@ namespace UnrealGameSync.Controls
 		private readonly BindingList<ViewItem> _views = new BindingList<ViewItem>();
 		private readonly List<ViewItem> _presetViews = new List<ViewItem>();
 
+		[DesignerSerializationVisibility(DesignerSerializationVisibility.Hidden)]
 		public string SyncBaseContentToolTipTitle { get; set; } = String.Empty;
+		[DesignerSerializationVisibility(DesignerSerializationVisibility.Hidden)]
 		public string SyncBaseContentToolTipContent { get; set; } = String.Empty;
 
 		public WorkspaceSyncFilterControl()
diff --git a/Engine/Source/Programs/UnrealGameSync/UnrealGameSync/Forms/CleanWorkspaceWindow.cs b/Engine/Source/Programs/UnrealGameSync/UnrealGameSync/Forms/CleanWorkspaceWindow.cs
--- a/Engine/Source/Programs/UnrealGameSync/UnrealGameSync/Forms/CleanWorkspaceWindow.cs
+++ b/Engine/Source/Programs/UnrealGameSync/UnrealGameSync/Forms/CleanWorkspaceWindow.cs
@@ -217,7 +217,7 @@ namespace UnrealGameSync
 				TreeNode childNode = BuildTreeViewStructure(subFolder, folderPath + subFolder.Name.ToLowerInvariant() + "/", selectFolder, depth + 1);
 				folderNode.Nodes.Add(childNode);
 
-				TreeNodeData childNodeData = (TreeNodeData)childNode.Tag;
+				TreeNodeData childNodeData = childNode.Tag as TreeNodeData ?? throw new InvalidOperationException("Missing node data");
 				folderNodeData._numFiles += childNodeData._numFiles;
 				folderNodeData._numSelectedFiles += childNodeData._numSelectedFiles;
 				folderNodeData._numEmptySelectedFiles += childNodeData._numEmptySelectedFiles;
@@ -317,23 +317,26 @@ namespace UnrealGameSync
 
 		private void TreeView_NodeMouseClick(object sender, TreeNodeMouseClickEventArgs e)
 		{
-			TreeNode node = e.Node;
+			if (e.Node is not TreeNode node)
+			{
+				return;
+			}
 			if (e.Button == System.Windows.Forms.MouseButtons.Right)
 			{
-				TreeView.SelectedNode = e.Node;
-				FolderContextMenu.Tag = e.Node;
+				TreeView.SelectedNode = node;
+				FolderContextMenu.Tag = node;
 				FolderContextMenu.Show(TreeView.PointToScreen(e.Location));
 			}
-			else if (e.X >= e.Node.Bounds.Left - 32 && e.X < e.Node.Bounds.Left - 16)
+			else if (e.X >= node.Bounds.Left - 32 && e.X < node.Bounds.Left - 16)
 			{
-				TreeNodeData nodeData = (TreeNodeData)node.Tag;
+				TreeNodeData nodeData = node.Tag as TreeNodeData ?? throw new InvalidOperationException("Missing node data");
 				SetSelected(node, (nodeData._numSelectedFiles == 0) ? SelectionType.All : SelectionType.None);
 			}
 		}
 
 		private static void SetSelected(TreeNode parentNode, SelectionType type)
 		{
-			TreeNodeData parentNodeData = (TreeNodeData)parentNode.Tag;
+			TreeNodeData parentNodeData = parentNode.Tag as TreeNodeData ?? throw new InvalidOperationException("Missing node data");
 
 			int prevNumSelectedFiles = parentNodeData._numSelectedFiles;
 			SetSelectedOnChildren(parentNode, type);
@@ -341,9 +344,9 @@ namespace UnrealGameSync
 			int deltaNumSelectedFiles = parentNodeData._numSelectedFiles - prevNumSelectedFiles;
 			if (deltaNumSelectedFiles != 0)
 			{
-				for (TreeNode nextParentNode = parentNode.Parent; nextParentNode != null; nextParentNode = nextParentNode.Parent)
+				for (TreeNode? nextParentNode = parentNode.Parent; nextParentNode != null; nextParentNode = nextParentNode.Parent)
 				{
-					TreeNodeData nextParentNodeData = (TreeNodeData)nextParentNode.Tag;
+					TreeNodeData nextParentNodeData = nextParentNode.Tag as TreeNodeData ?? throw new InvalidOperationException("Missing node data");
 					nextParentNodeData._numSelectedFiles += deltaNumSelectedFiles;
 					UpdateImage(nextParentNode);
 				}
@@ -352,7 +355,7 @@ namespace UnrealGameSync
 
 		private static void SetSelectedOnChildren(TreeNode parentNode, SelectionType type)
 		{
-			TreeNodeData parentNodeData = (TreeNodeData)parentNode.Tag;
+			TreeNodeData parentNodeData = parentNode.Tag as TreeNodeData ?? throw new InvalidOperationException("Missing node data");
 
 			int newNumSelectedFiles = 0;
 			switch (type)
@@ -390,7 +393,7 @@ namespace UnrealGameSync
 
 		private static void UpdateImage(TreeNode node)
 		{
-			TreeNodeData nodeData = (TreeNodeData)node.Tag;
+			TreeNodeData nodeData = node.Tag as TreeNodeData ?? throw new InvalidOperationException("Missing node data");
 			int imageIndex = (nodeData._folder != null) ? 0 : 3;
 			imageIndex += (nodeData._numSelectedFiles == 0) ? 0 : (nodeData._numSelectedFiles < nodeData._numFiles || (nodeData._folder != null && !nodeData._folder._emptyAfterClean)) ? 1 : 2;
 			node.ImageIndex = imageIndex;
@@ -423,7 +426,7 @@ namespace UnrealGameSync
 
 		private static void FindSelection(TreeNode node, List<FileInfo> filesToSync, List<FileInfo> filesToDelete, List<DirectoryInfo> directoriesToDelete)
 		{
-			TreeNodeData nodeData = (TreeNodeData)node.Tag;
+			TreeNodeData nodeData = node.Tag as TreeNodeData ?? throw new InvalidOperationException("Missing node data");
 			if (nodeData._file != null)
 			{
 				if (nodeData._numSelectedFiles > 0)
@@ -503,13 +506,13 @@ namespace UnrealGameSync
 		private void FolderContextMenu_OpenWithExplorer_Click(object sender, EventArgs e)
 		{
 			TreeNode node = (TreeNode)FolderContextMenu.Tag!;
-			TreeNodeData nodeData = (TreeNodeData)node.Tag;
+			TreeNodeData? nodeData = node.Tag as TreeNodeData;
 
-			if (nodeData._folder != null)
+			if (nodeData?._folder != null)
 			{
 				Process.Start("explorer.exe", String.Format("\"{0}\"", nodeData._folder._directory.FullName));
 			}
-			else if (nodeData._file != null)
+			else if (nodeData?._file != null)
 			{
 				Process.Start("explorer.exe", String.Format("\"{0}\"", nodeData._file.Directory!.FullName));
 			}
diff --git a/Engine/Source/Programs/UnrealGameSync/UnrealGameSync/Forms/IssueAlertWindow.cs b/Engine/Source/Programs/UnrealGameSync/UnrealGameSync/Forms/IssueAlertWindow.cs
--- a/Engine/Source/Programs/UnrealGameSync/UnrealGameSync/Forms/IssueAlertWindow.cs
+++ b/Engine/Source/Programs/UnrealGameSync/UnrealGameSync/Forms/IssueAlertWindow.cs
@@ -2,6 +2,7 @@
 
 using System;
 using System.Collections.Generic;
+using System.ComponentModel;
 using System.Drawing;
 using System.Runtime.InteropServices;
 using System.Text;
@@ -22,8 +23,10 @@ namespace UnrealGameSync
 
 		public IssueMonitor IssueMonitor { get; }
 
+		[DesignerSerializationVisibility(DesignerSerializationVisibility.Hidden)]
 		public bool IsWarning { get; set; }
 
+		[DesignerSerializationVisibility(DesignerSerializationVisibility.Hidden)]
 		public bool? StrongAlert { get; set; }
 
 		public IssueAlertWindow(IssueMonitor issueMonitor, IssueData issue, IssueAlertReason reason)
diff --git a/Engine/Source/Programs/UnrealGameSync/UnrealGameSync/Forms/NotificationWindow.cs b/Engine/Source/Programs/UnrealGameSync/UnrealGameSync/Forms/NotificationWindow.cs
--- a/Engine/Source/Programs/UnrealGameSync/UnrealGameSync/Forms/NotificationWindow.cs
+++ b/Engine/Source/Programs/UnrealGameSync/UnrealGameSync/Forms/NotificationWindow.cs
@@ -1,6 +1,7 @@
 // Copyright Epic Games, Inc. All Rights Reserved.
 
 using System;
+using System.ComponentModel;
 using System.Drawing;
 using System.Drawing.Drawing2D;
 using System.Windows.Forms;
@@ -37,7 +38,9 @@ namespace UnrealGameSync
 		bool _mouseDownOverClose;
 
 		// Notifications
+		[DesignerSerializationVisibility(DesignerSerializationVisibility.Hidden)]
 		public Action? OnMoreInformation { get; set; }
+		[DesignerSerializationVisibility(DesignerSerializationVisibility.Hidden)]
 		public Action? OnDismiss { get; set; }
 
 		public NotificationWindow(Image inLogo)
@@ -83,13 +86,6 @@ namespace UnrealGameSync
 			base.Dispose(disposing);
 		}
 
-		protected override void OnClosed(EventArgs e)
-		{
-			base.OnClosed(e);
-
-			_logoImage.Dispose();
-		}
-
 		public void CalculateBounds()
 		{
 			int referenceSize8 = (int)(8 * Font.Height) / 16;
diff --git a/Engine/Source/Programs/UnrealGameSync/UnrealGameSync/Forms/PresetsWindow.cs b/Engine/Source/Programs/UnrealGameSync/UnrealGameSync/Forms/PresetsWindow.cs
--- a/Engine/Source/Programs/UnrealGameSync/UnrealGameSync/Forms/PresetsWindow.cs
+++ b/Engine/Source/Programs/UnrealGameSync/UnrealGameSync/Forms/PresetsWindow.cs
@@ -2,6 +2,7 @@
 
 using System;
 using System.Collections.Generic;
+using System.ComponentModel;
 using System.Windows.Forms;
 
 namespace UnrealGameSync.Forms
@@ -35,6 +36,7 @@ namespace UnrealGameSync.Forms
 			}
 		}
 
+		[DesignerSerializationVisibility(DesignerSerializationVisibility.Hidden)]
 		public string PresetName { get; set; } = String.Empty;
 
 		public PresetsWindow(
diff --git a/Engine/Source/Programs/UnrealGameSync/UnrealGameSync/Forms/SelectProjectFromWorkspaceWindow.cs b/Engine/Source/Programs/UnrealGameSync/UnrealGameSync/Forms/SelectProjectFromWorkspaceWindow.cs
--- a/Engine/Source/Programs/UnrealGameSync/UnrealGameSync/Forms/SelectProjectFromWorkspaceWindow.cs
+++ b/Engine/Source/Programs/UnrealGameSync/UnrealGameSync/Forms/SelectProjectFromWorkspaceWindow.cs
@@ -180,7 +180,7 @@ namespace UnrealGameSync
 				if (String.Equals(projectNodes[idx].FullName, _selectedProjectFile, StringComparison.OrdinalIgnoreCase))
 				{
 					ProjectTreeView.SelectedNode = node;
-					for (TreeNode parentNode = node.Parent; parentNode != rootNode; parentNode = parentNode.Parent)
+					for (TreeNode? parentNode = node.Parent; parentNode != null && parentNode != rootNode; parentNode = parentNode.Parent)
 					{
 						parentNode.Expand();
 					}
diff --git a/Engine/Source/Programs/UnrealGameSync/UnrealGameSync/Forms/SelectStreamWindow.cs b/Engine/Source/Programs/UnrealGameSync/UnrealGameSync/Forms/SelectStreamWindow.cs
--- a/Engine/Source/Programs/UnrealGameSync/UnrealGameSync/Forms/SelectStreamWindow.cs
+++ b/Engine/Source/Programs/UnrealGameSync/UnrealGameSync/Forms/SelectStreamWindow.cs
@@ -359,8 +359,7 @@ namespace UnrealGameSync
 			string? newSelectedStream = null;
 			if (StreamsTreeView.SelectedNode != null)
 			{
-				StreamNode streamNode = (StreamNode)StreamsTreeView.SelectedNode.Tag;
-				if (streamNode != null)
+				if (StreamsTreeView.SelectedNode.Tag is StreamNode streamNode)
 				{
 					newSelectedStream = streamNode.Record.Stream;
 				}
diff --git a/Engine/Source/Programs/UnrealGameSync/UnrealGameSync/Launcher.cs b/Engine/Source/Programs/UnrealGameSync/UnrealGameSync/Launcher.cs
--- a/Engine/Source/Programs/UnrealGameSync/UnrealGameSync/Launcher.cs
+++ b/Engine/Source/Programs/UnrealGameSync/UnrealGameSync/Launcher.cs
@@ -309,7 +309,7 @@ namespace UnrealGameSync
 						if (syncPath.EndsWith(".zip", StringComparison.OrdinalIgnoreCase))
 						{
 							string localPath = Path.Combine(applicationFolder, syncPath.Substring(depotPathPrefix.Length).Replace('/', Path.DirectorySeparatorChar));
-							ZipFile.ExtractToDirectory(localPath, applicationFolder);
+							await ZipFile.ExtractToDirectoryAsync(localPath, applicationFolder, cancellationToken);
 						}
 
 						// Check the application exists
diff --git a/Engine/Source/Programs/UnrealGameSync/UnrealGameSync/UnrealGameSync.csproj b/Engine/Source/Programs/UnrealGameSync/UnrealGameSync/UnrealGameSync.csproj
--- a/Engine/Source/Programs/UnrealGameSync/UnrealGameSync/UnrealGameSync.csproj
+++ b/Engine/Source/Programs/UnrealGameSync/UnrealGameSync/UnrealGameSync.csproj
@@ -1,9 +1,7 @@
 <Project Sdk="Microsoft.NET.Sdk">
-  <Import Project="..\..\Shared\UnrealEngine.csproj.props" />
-
   <PropertyGroup>
     <OutputType>WinExe</OutputType>
-    <TargetFramework>net8.0-windows</TargetFramework>
+    <TargetFramework>net10.0-windows</TargetFramework>
     <UseWindowsForms>true</UseWindowsForms>
     <Product>UnrealGameSync</Product>
     <Company>Epic Games, Inc</Company>
@@ -13,16 +11,15 @@
     <StartupObject>UnrealGameSync.Program</StartupObject>
     <ApplicationIcon>Resources\Icon.ico</ApplicationIcon>
     <ApplicationManifest>Properties\app.manifest</ApplicationManifest>
-    <Nullable>enable</Nullable>
     <NeutralLanguage>en-US</NeutralLanguage>
     <Configurations>$(Configurations);DebugAutoUpdate;ReleaseAutoUpdate</Configurations>
   </PropertyGroup>
 
   <PropertyGroup Condition="'$(Configuration)|$(Platform)'=='Release|AnyCPU'">
-    <DefineConstants>TRACE;DEBUG;WITH_TELEMETRY</DefineConstants>
+    <DefineConstants>$(DefineConstants);WITH_TELEMETRY</DefineConstants>
   </PropertyGroup>
   <PropertyGroup Condition="'$(Configuration)|$(Platform)'=='Debug|AnyCPU'">
-    <DefineConstants>TRACE;WITH_TELEMETRY</DefineConstants>
+    <DefineConstants>$(DefineConstants);WITH_TELEMETRY</DefineConstants>
   </PropertyGroup>
   <PropertyGroup Condition=" '$(Configuration)' == 'DebugAutoUpdate' or '$(Configuration)' == 'ReleaseAutoUpdate'">
     <DefineConstants>$(DefineConstants);WITH_AUTOUPDATE</DefineConstants>
diff --git a/Engine/Source/Programs/UnrealGameSync/UnrealGameSyncAutomation/UnrealGameSyncAutomation.csproj b/Engine/Source/Programs/UnrealGameSync/UnrealGameSyncAutomation/UnrealGameSyncAutomation.csproj
--- a/Engine/Source/Programs/UnrealGameSync/UnrealGameSyncAutomation/UnrealGameSyncAutomation.csproj
+++ b/Engine/Source/Programs/UnrealGameSync/UnrealGameSyncAutomation/UnrealGameSyncAutomation.csproj
@@ -1,7 +1,5 @@
 <?xml version="1.0" encoding="utf-8"?>
 <Project Sdk="Microsoft.NET.Sdk">
-  <Import Project="..\..\Shared\UnrealEngine.csproj.props" />
-
   <PropertyGroup>
     <OutputType>WinExe</OutputType>
     <TargetFramework>net8.0-windows</TargetFramework>
@@ -9,6 +7,5 @@
     <GenerateAssemblyInfo>false</GenerateAssemblyInfo>
     <RootNamespace>UnrealGameSyncAutomation</RootNamespace>
     <AssemblyName>UnrealGameSyncAutomation</AssemblyName>
-    <Nullable>enable</Nullable>
   </PropertyGroup>
 </Project>
\ No newline at end of file
diff --git a/Engine/Source/Programs/UnrealGameSync/UnrealGameSyncCmd/UnrealGameSyncCmd.csproj b/Engine/Source/Programs/UnrealGameSync/UnrealGameSyncCmd/UnrealGameSyncCmd.csproj
--- a/Engine/Source/Programs/UnrealGameSync/UnrealGameSyncCmd/UnrealGameSyncCmd.csproj
+++ b/Engine/Source/Programs/UnrealGameSync/UnrealGameSyncCmd/UnrealGameSyncCmd.csproj
@@ -1,9 +1,6 @@
 <Project Sdk="Microsoft.NET.Sdk">
-  <Import Project="..\..\Shared\UnrealEngine.csproj.props" />
-
   <PropertyGroup>
     <OutputType>Exe</OutputType>
-    <Nullable>enable</Nullable>
     <AssemblyName>ugs</AssemblyName>
   </PropertyGroup>
 
diff --git a/Engine/Source/Programs/UnrealGameSync/UnrealGameSyncShared/UnrealGameSyncShared.csproj b/Engine/Source/Programs/UnrealGameSync/UnrealGameSyncShared/UnrealGameSyncShared.csproj
--- a/Engine/Source/Programs/UnrealGameSync/UnrealGameSyncShared/UnrealGameSyncShared.csproj
+++ b/Engine/Source/Programs/UnrealGameSync/UnrealGameSyncShared/UnrealGameSyncShared.csproj
@@ -1,8 +1,5 @@
 <Project Sdk="Microsoft.NET.Sdk">
-  <Import Project="..\..\Shared\UnrealEngine.csproj.props" />
-
   <PropertyGroup>
-    <Nullable>enable</Nullable>
     <RootNamespace>UnrealGameSync</RootNamespace>
   </PropertyGroup>
 
@@ -29,7 +26,6 @@
   
   <ItemGroup>
     <PackageReference Include="Microsoft.Extensions.Caching.Memory" Version="8.0.1" />
-    <PackageReference Include="System.IO.FileSystem.DriveInfo" Version="4.3.1" />
   </ItemGroup>
 
   <ItemGroup>
diff --git a/Engine/Source/Programs/UnrealGameSync/UnrealGameSyncShared/Utility.cs b/Engine/Source/Programs/UnrealGameSync/UnrealGameSyncShared/Utility.cs
--- a/Engine/Source/Programs/UnrealGameSync/UnrealGameSyncShared/Utility.cs
+++ b/Engine/Source/Programs/UnrealGameSync/UnrealGameSyncShared/Utility.cs
@@ -313,7 +313,7 @@ namespace UnrealGameSync
 				PerforceChangeDetails? cachedDetails = null;
 
 				List<int> changeBatch = new List<int>(BatchSize);
-				while (changeBatch.Count < BatchSize && await changeNumberEnumerator.MoveNextAsync(cancellationToken))
+				while (changeBatch.Count < BatchSize && await changeNumberEnumerator.MoveNextAsync())
 				{
 					string cacheKey = GetChangeDetailsCacheKey(changeNumberEnumerator.Current, hash);
 					if (s_changeCache.TryGetValue(cacheKey, out cachedDetails))
diff --git a/Engine/Source/Programs/UnrealGameSync/UnrealGameSyncShared/WorkspaceState.cs b/Engine/Source/Programs/UnrealGameSync/UnrealGameSyncShared/WorkspaceState.cs
--- a/Engine/Source/Programs/UnrealGameSync/UnrealGameSyncShared/WorkspaceState.cs
+++ b/Engine/Source/Programs/UnrealGameSync/UnrealGameSyncShared/WorkspaceState.cs
@@ -330,7 +330,14 @@ namespace UnrealGameSync
 			}
 
 			byte[] buffer = new byte[stream.Length];
-			stream.Read(buffer);
+			try
+			{
+				stream.ReadExactly(buffer);
+			}
+			catch (EndOfStreamException)
+			{
+				return _createNew();
+			}
 			return Utility.TryDeserializeJson<WorkspaceState>(buffer) ?? _createNew();
 		}
 	}
