make tools command independent from a workspace
improved tools info gathering

#Synced-CL 51609771

diff --git a/Engine/Source/Programs/UnrealGameSync/UnrealGameSyncCmd/Commands/ToolsCommand.cs b/Engine/Source/Programs/UnrealGameSync/UnrealGameSyncCmd/Commands/ToolsCommand.cs
--- a/Engine/Source/Programs/UnrealGameSync/UnrealGameSyncCmd/Commands/ToolsCommand.cs
+++ b/Engine/Source/Programs/UnrealGameSync/UnrealGameSyncCmd/Commands/ToolsCommand.cs
@@ -32,29 +32,28 @@ namespace UnrealGameSyncCmd.Commands
 		public bool Update { get; set; }
 	}
 
-	internal class ToolsCommand : Command
+	internal class ToolsCommand : ICommand
 	{
-		public override async Task ExecuteAsync(CommandContext context)
+		public static async Task ExecuteAsync(ILogger logger, CommandLineArguments args, UserSettings? globalSettings)
 		{
-			ILogger logger = context.Logger!;
 			if (!RuntimeInformation.IsOSPlatform(OSPlatform.Windows))
 			{
 				logger.LogInformation("This command is only available for Windows.");
 				return;
 			}
 
-			context.Arguments!.TryGetPositionalArgument(out string? toolId);
+			args.TryGetPositionalArgument(out string? toolId);
 
-			ToolsCommandCommandOptions options = context.Arguments.ApplyTo<ToolsCommandCommandOptions>(logger);
-			context.Arguments.CheckAllArgumentsUsed();
+			ToolsCommandCommandOptions options = args.ApplyTo<ToolsCommandCommandOptions>(logger);
+			args.CheckAllArgumentsUsed();
 
-			if (context.GlobalSettings == null)
+			if (globalSettings == null)
 			{
 				logger.LogError("Could not retrieve user settings.");
 				return;
 			}
 
-			IPerforceConnection perforceClient = context.PerforceClient!;
+			using IPerforceConnection perforce = await PerforceConnection.CreateAsync(logger);
 
 			DirectoryReference dataFolder = DirectoryReference.Combine(DirectoryReference.GetSpecialFolder(Environment.SpecialFolder.LocalApplicationData)!, "UnrealGameSync");
 			DirectoryReference.CreateDirectory(dataFolder);
@@ -80,7 +79,7 @@ namespace UnrealGameSyncCmd.Commands
 			ServiceProvider serviceProvider = services.BuildServiceProvider();
 
 			using ToolUpdateMonitor toolUpdateMonitor =
-				new ToolUpdateMonitor(perforceClient.Settings, dataFolder, context.GlobalSettings, logger, serviceProvider);
+				new ToolUpdateMonitor(perforce.Settings, dataFolder, globalSettings, logger, serviceProvider);
 
 			// get the list of tools available
 			logger.LogInformation("Retrieving tools information, please wait.");
@@ -166,11 +165,11 @@ namespace UnrealGameSyncCmd.Commands
 					}
 				}
 
-				if (!newEnabledTools.SequenceEqual(context.GlobalSettings.EnabledTools))
+				if (!newEnabledTools.SequenceEqual(globalSettings.EnabledTools))
 				{
-					context.GlobalSettings.EnabledTools.Clear();
-					context.GlobalSettings.EnabledTools.UnionWith(newEnabledTools);
-					context.GlobalSettings.Save(logger);
+					globalSettings.EnabledTools.Clear();
+					globalSettings.EnabledTools.UnionWith(newEnabledTools);
+					globalSettings.Save(logger);
 
 					logger.LogInformation("Updating tools");
 					await toolUpdateMonitor.UpdateToolsAsync();
diff --git a/Engine/Source/Programs/UnrealGameSync/UnrealGameSyncCmd/Program.cs b/Engine/Source/Programs/UnrealGameSync/UnrealGameSyncCmd/Program.cs
--- a/Engine/Source/Programs/UnrealGameSync/UnrealGameSyncCmd/Program.cs
+++ b/Engine/Source/Programs/UnrealGameSync/UnrealGameSyncCmd/Program.cs
@@ -189,6 +189,13 @@ namespace UnrealGameSyncCmd
 							break;
 						}
 
+					//////////////////////////////////////
+					case ToolsCommand _:
+						{
+							await ToolsCommand.ExecuteAsync(logger, args, globalSettings);
+							break;
+						}
+
 					//////////////////////////////////////
 					case UpgradeCommand _:
 						{
diff --git a/Engine/Source/Programs/UnrealGameSync/UnrealGameSyncShared/ToolUpdateMonitor.cs b/Engine/Source/Programs/UnrealGameSync/UnrealGameSyncShared/ToolUpdateMonitor.cs
--- a/Engine/Source/Programs/UnrealGameSync/UnrealGameSyncShared/ToolUpdateMonitor.cs
+++ b/Engine/Source/Programs/UnrealGameSync/UnrealGameSyncShared/ToolUpdateMonitor.cs
@@ -1,6 +1,7 @@
 ﻿// Copyright Epic Games, Inc. All Rights Reserved.
 
 using System;
+using System.Collections.Concurrent;
 using System.Collections.Generic;
 using System.Diagnostics;
 using System.IO;
@@ -91,7 +92,7 @@ namespace UnrealGameSync
 
 		bool _readLegacyConfig;
 
-		Dictionary<string, ToolInfo> _perforceTools = new Dictionary<string, ToolInfo>(StringComparer.Ordinal);
+		IReadOnlyDictionary<string, ToolInfo> _perforceTools = new Dictionary<string, ToolInfo>(StringComparer.Ordinal);
 		int _perforceToolsChange = -1;
 
 		IReadOnlyList<ToolInfo> _tools = Array.Empty<ToolInfo>();
@@ -348,33 +349,31 @@ namespace UnrealGameSync
 			List<ChangesRecord> changes = await perforce.GetChangesAsync(ChangesOptions.None, 1, ChangeStatus.Submitted, $"{DeploymentSettings.Instance.ToolsDepotPath}/...", cancellationToken);
 			if (changes.Count > 0 && changes[0].Number != _perforceToolsChange)
 			{
-				Dictionary<string, ToolInfo> newPerforceTools = new Dictionary<string, ToolInfo>(StringComparer.Ordinal);
+				ConcurrentDictionary<string, ToolInfo> newPerforceTools = new ConcurrentDictionary<string, ToolInfo>(StringComparer.Ordinal);
 
 				List<FilesRecord> fileRecords = await perforce.FilesAsync(FilesOptions.ExcludeDeleted, $"{DeploymentSettings.Instance.ToolsDepotPath}/...", cancellationToken);
 
-				foreach (FilesRecord fileRecord in fileRecords)
+				await Parallel.ForEachAsync(fileRecords, cancellationToken, async (fileRecord, ct) =>
 				{
 					if (fileRecord.DepotFile != null && fileRecord.DepotFile.EndsWith(".ini", StringComparison.OrdinalIgnoreCase))
 					{
 						string zipFile = fileRecord.DepotFile.Substring(0, fileRecord.DepotFile.Length - 4) + ".zip";
 						FilesRecord? zipRecord = fileRecords.FirstOrDefault(x => String.Equals(x.DepotFile, zipFile, StringComparison.OrdinalIgnoreCase));
-
 						if (zipRecord != null)
 						{
 							string toolRevision = $"{zipFile}@{Math.Max(fileRecord.Change, zipRecord.Change)}";
-
 							ToolInfo? toolInfo;
 							if (!_perforceTools.TryGetValue(toolRevision, out toolInfo))
 							{
-								toolInfo = await ReadToolDefinitionAsync(perforce, $"{fileRecord.DepotFile}@{fileRecord.Change}", toolRevision, cancellationToken);
+								toolInfo = await ReadToolDefinitionAsync(perforce, $"{fileRecord.DepotFile}@{fileRecord.Change}", toolRevision, ct);
 							}
 							if (toolInfo != null)
 							{
-								newPerforceTools.Add(toolRevision, toolInfo);
+								newPerforceTools.TryAdd(toolRevision, toolInfo);
 							}
 						}
 					}
-				}
+				});
 
 				_perforceTools = newPerforceTools;
 				_perforceToolsChange = changes[0].Number;
