added an option to the ugs cli filter command to set|unset SyncBaseContent

#Synced-CL 51609771

diff --git a/Engine/Source/Programs/UnrealGameSync/UnrealGameSyncCmd/Commands/FilterCommand.cs b/Engine/Source/Programs/UnrealGameSync/UnrealGameSyncCmd/Commands/FilterCommand.cs
--- a/Engine/Source/Programs/UnrealGameSync/UnrealGameSyncCmd/Commands/FilterCommand.cs
+++ b/Engine/Source/Programs/UnrealGameSync/UnrealGameSyncCmd/Commands/FilterCommand.cs
@@ -21,6 +21,10 @@ namespace UnrealGameSyncCmd.Commands
 		[CommandLine("-Reset")]
 		public bool Reset { get; set; } = false;
 
+		[CommandLine("-SetSyncBaseContent", Value = "true")]
+		[CommandLine("-UnsetSyncBaseContent", Value = "false")]
+		public bool? SyncBaseContent { get; set; } = null;
+
 		[CommandLine("-Include=")]
 		public List<string> Include { get; set; } = new List<string>();
 
@@ -83,6 +87,12 @@ namespace UnrealGameSyncCmd.Commands
 			FilterCommandOptions options = context.Arguments!.ApplyTo<FilterCommandOptions>(logger);
 			context.Arguments.CheckAllArgumentsUsed(logger);
 
+			if (options.SyncBaseContent != null)
+			{
+				workspaceSettings.SyncBaseContent = options.SyncBaseContent.Value;
+				workspaceSettings.Save(logger);
+			}
+
 			if (options.Global)
 			{
 				ApplyCommandOptions(logger, context, globalFilter, options, syncCategories.Values);
