using project override value to get badges when set

#Synced-CL 51609771

diff --git a/Engine/Source/Programs/UnrealGameSync/UnrealGameSync/Controls/WorkspaceControl.cs b/Engine/Source/Programs/UnrealGameSync/UnrealGameSync/Controls/WorkspaceControl.cs
--- a/Engine/Source/Programs/UnrealGameSync/UnrealGameSync/Controls/WorkspaceControl.cs
+++ b/Engine/Source/Programs/UnrealGameSync/UnrealGameSync/Controls/WorkspaceControl.cs
@@ -508,7 +508,23 @@ namespace UnrealGameSync
 			_perforceMonitor.OnStreamChange += StreamChanged;
 			_perforceMonitor.OnLoginExpired += LoginExpired;
 
-			_eventMonitor = new EventMonitor(_apiUrl, PerforceUtils.GetClientOrDepotDirectoryName(SelectedProjectIdentifier), openProjectInfo.PerforceSettings.UserName, _serviceProvider);
+			string projectOverride = String.Empty;
+
+			ConfigSection? projectConfigSection = openProjectInfo.LatestProjectConfigFile.FindSection(SelectedProjectIdentifier);
+			if (projectConfigSection != null)
+			{
+				projectOverride = projectConfigSection.GetValue("CloudArchivesProjectOverride", String.Empty);
+			}
+
+			if (String.IsNullOrEmpty(projectOverride))
+			{
+				_eventMonitor = new EventMonitor(_apiUrl, PerforceUtils.GetClientOrDepotDirectoryName(SelectedProjectIdentifier), openProjectInfo.PerforceSettings.UserName, _serviceProvider);
+			}
+			else
+			{
+				_eventMonitor = new EventMonitor(_apiUrl, PerforceUtils.GetClientOrDepotDirectoryName(projectOverride), openProjectInfo.PerforceSettings.UserName, _serviceProvider);
+			}
+
 			_eventMonitor.OnUpdatesReady += UpdateReviewsCallback;
 
 			UpdateColumnSettings(true);
