display message from CloudArchiveOverride instead of stream message when set

#Synced-CL 51609771

diff --git a/Engine/Source/Programs/UnrealGameSync/UnrealGameSync/Controls/WorkspaceControl.cs b/Engine/Source/Programs/UnrealGameSync/UnrealGameSync/Controls/WorkspaceControl.cs
--- a/Engine/Source/Programs/UnrealGameSync/UnrealGameSync/Controls/WorkspaceControl.cs
+++ b/Engine/Source/Programs/UnrealGameSync/UnrealGameSync/Controls/WorkspaceControl.cs
@@ -4371,47 +4371,43 @@ namespace UnrealGameSync
 			ConfigFile projectConfigFile = _perforceMonitor.LatestProjectConfigFile;
 			if (projectConfigFile != null)
 			{
-				string? message;
-				if (TryGetProjectSetting(projectConfigFile, "Message", out message))
-				{
-					alert = CreateStatusLineFromMarkdown(message);
-				}
-
-				string? statusPanelColor;
-				if (TryGetProjectSetting(projectConfigFile, "StatusPanelColor", out statusPanelColor))
-				{
-					tintColor = System.Drawing.ColorTranslator.FromHtml(statusPanelColor);
-				}
 
-				/*string[]? infoMessages;
-				if (TryGetProjectSetting(projectConfigFile, "InfoMessage", out infoMessages))
+				string? message = String.Empty;
+				if (TryGetProjectSetting(projectConfigFile, "CloudArchivesProjectOverride", out string? projectOverride))
 				{
-					const int maxInfoMessages = 7;
-					if (infoMessages.Length > maxInfoMessages)
-					{
-						_logger.LogWarning("Maximum number of InfoMessage entries exceeded, found {MessageCount} messages but will only display {MaxMessages}", infoMessages.Length, maxInfoMessages);
-					}
-
-					foreach (string infoMessage in infoMessages.Take(Math.Min(infoMessages.Length, maxInfoMessages)))
+					if (String.IsNullOrWhiteSpace(projectOverride) == false)
 					{
-						ConfigObject infoMessageObj = new ConfigObject(infoMessage);
-
-						string? infoText = infoMessageObj.GetValue("Message");
-
-						if (String.IsNullOrEmpty(infoText) == false)
+						if (ConfigUtils.TryGetProjectSetting(projectConfigFile, projectOverride, "Message", out string? value))
 						{
-							StatusLine infoLine = CreateStatusLineFromMarkdown(infoText, applyDpiScalingToIcons: false);
-
-							string? infoTooltip = infoMessageObj.GetValue("Tooltip");
-							if (String.IsNullOrEmpty(infoTooltip) == false)
+							if(!String.IsNullOrWhiteSpace(value))
 							{
-								infoLine.ToolTip = infoTooltip;
+								message = value;
 							}
+						}
+					}
+				}
 
-							infoLines.Add(infoLine);
+				if (String.IsNullOrWhiteSpace(message))
+				{
+					if (TryGetProjectSetting(projectConfigFile, "Message", out string? value))
+					{
+						if (!String.IsNullOrWhiteSpace(value))
+						{
+							message = value;
 						}
 					}
-				}*/
+				}
+
+				if (!String.IsNullOrWhiteSpace(message))
+				{
+					alert = CreateStatusLineFromMarkdown(message);
+				}
+
+				string? statusPanelColor;
+				if (TryGetProjectSetting(projectConfigFile, "StatusPanelColor", out statusPanelColor))
+				{
+					tintColor = System.Drawing.ColorTranslator.FromHtml(statusPanelColor);
+				}
 			}
 
 			using (Graphics graphics = CreateGraphics())
