forgotten use case, when the project's archives is empty we need to add the archive to the list on -set

#Synced-CL 51609771
#Synced-CL 51609771

diff --git a/Engine/Source/Programs/UnrealGameSync/UnrealGameSyncCmd/Commands/ArchiveOverrideCommand.cs b/Engine/Source/Programs/UnrealGameSync/UnrealGameSyncCmd/Commands/ArchiveOverrideCommand.cs
--- a/Engine/Source/Programs/UnrealGameSync/UnrealGameSyncCmd/Commands/ArchiveOverrideCommand.cs
+++ b/Engine/Source/Programs/UnrealGameSync/UnrealGameSyncCmd/Commands/ArchiveOverrideCommand.cs
@@ -95,14 +95,19 @@ namespace UnrealGameSyncCmd.Commands
 					return;
 				}
 
-				List<BaseArchiveChannel> enabled = archives.Where(a => a.Type.Equals(options.ArchiveType, StringComparison.OrdinalIgnoreCase) && a.Name.Equals(options.ArchiveName, StringComparison.OrdinalIgnoreCase)).ToList();
-				if (enabled.Any())
+				List<BaseArchiveChannel> matchingArchives = archives.Where(a => a.Type.Equals(options.ArchiveType, StringComparison.OrdinalIgnoreCase) && a.Name.Equals(options.ArchiveName, StringComparison.OrdinalIgnoreCase)).ToList();
+				if (matchingArchives.Any())
 				{
 					ResetArchiveSettings(logger, context.UserProjectSettings!);
-					foreach (BaseArchiveChannel archive in enabled)
+					foreach (BaseArchiveChannel archive in matchingArchives)
 					{
 						ArchiveSettings? archiveSettings = context.UserProjectSettings!.Archives.FirstOrDefault(x => x.Type == archive.Type);
-						if (archiveSettings != null)
+						if (archiveSettings == null)
+						{
+							archiveSettings = new ArchiveSettings(true, archive.Type, new string[] { archive.Name });
+							context.UserProjectSettings.Archives.Add(archiveSettings);
+						}
+						else
 						{
 							archiveSettings.Enabled = true;
 							archiveSettings.Order.Remove(archive.Name);
