additional paths to sync should be considered as sync filters

#Synced-CL 51609771

diff --git a/Engine/Source/Programs/UnrealGameSync/UnrealGameSyncShared/WorkspaceUpdate.cs b/Engine/Source/Programs/UnrealGameSync/UnrealGameSyncShared/WorkspaceUpdate.cs
--- a/Engine/Source/Programs/UnrealGameSync/UnrealGameSyncShared/WorkspaceUpdate.cs
+++ b/Engine/Source/Programs/UnrealGameSync/UnrealGameSyncShared/WorkspaceUpdate.cs
@@ -1948,22 +1948,17 @@ namespace UnrealGameSync
 			}
 
 			// add any additional paths
-			syncPaths.AddRange(additionalPathsToSync);
+			foreach (string syncPath in additionalPathsToSync)
+			{
+				applySyncPath(syncPaths, syncPath);
+			}
 
 			// Apply the sync filter to that list. We only want inclusive rules in the output list, but we can manually apply exclusions to previous entries.
 			if (syncFilter != null)
 			{
 				foreach (string syncPath in syncFilter)
 				{
-					string trimSyncPath = syncPath.Trim();
-					if (trimSyncPath.StartsWith("/", StringComparison.Ordinal))
-					{
-						syncPaths.Add(trimSyncPath);
-					}
-					else if (trimSyncPath.StartsWith("-/", StringComparison.Ordinal) && trimSyncPath.EndsWith("...", StringComparison.Ordinal))
-					{
-						syncPaths.RemoveAll(x => x.StartsWith(trimSyncPath.Substring(1, trimSyncPath.Length - 4), StringComparison.Ordinal));
-					}
+					applySyncPath(syncPaths, syncPath);
 				}
 			}
 
@@ -2019,6 +2014,19 @@ namespace UnrealGameSync
 			}
 		}
 
+		static void applySyncPath(List<string> syncPaths, string syncPath)
+		{
+			string trimSyncPath = syncPath.Trim();
+			if (trimSyncPath.StartsWith("/", StringComparison.Ordinal))
+			{
+				syncPaths.Add(trimSyncPath);
+			}
+			else if (trimSyncPath.StartsWith("-/", StringComparison.Ordinal) && trimSyncPath.EndsWith("...", StringComparison.Ordinal))
+			{
+				syncPaths.RemoveAll(x => x.StartsWith(trimSyncPath.Substring(1, trimSyncPath.Length - 4), StringComparison.Ordinal));
+			}
+		}
+
 		static Task<(WorkspaceUpdateResult, string)> SyncFileRevisions(IPerforceConnection perforce, string prefix, WorkspaceUpdateContext context, List<string> syncCommands, HashSet<string> remainingDepotPaths, ProgressValue progress, ILogger logger, CancellationToken cancellationToken)
 		{
 			Queue<List<string>> syncCommandLists = new Queue<List<string>>();
