revert AdditionalPaths to sync from the filtered view, we need to introduce a new concept of AdditionalPathsToSyncAndFilter for this purpose.

#Synced-CL 51609771

diff --git a/Engine/Source/Programs/UnrealGameSync/UnrealGameSyncShared/WorkspaceUpdate.cs b/Engine/Source/Programs/UnrealGameSync/UnrealGameSyncShared/WorkspaceUpdate.cs
--- a/Engine/Source/Programs/UnrealGameSync/UnrealGameSyncShared/WorkspaceUpdate.cs
+++ b/Engine/Source/Programs/UnrealGameSync/UnrealGameSyncShared/WorkspaceUpdate.cs
@@ -1947,12 +1947,6 @@ namespace UnrealGameSync
 				syncPaths.AddRange(baseContentPaths);
 			}
 
-			// add any additional paths
-			foreach (string syncPath in additionalPathsToSync)
-			{
-				applySyncPath(syncPaths, syncPath);
-			}
-
 			// Apply the sync filter to that list. We only want inclusive rules in the output list, but we can manually apply exclusions to previous entries.
 			if (syncFilter != null)
 			{
