UnrealGameSync - use the Nuget version of Microsoft.VisualStudio.Setup.Configuration.Interop rather than the checked-in binary. The Nuget package has an inbuilt targets file that embeds the interop types.

#Synced-CL 52010980

diff --git a/Engine/Source/Programs/UnrealGameSync/UnrealGameSync/UnrealGameSync.csproj b/Engine/Source/Programs/UnrealGameSync/UnrealGameSync/UnrealGameSync.csproj
--- a/Engine/Source/Programs/UnrealGameSync/UnrealGameSync/UnrealGameSync.csproj
+++ b/Engine/Source/Programs/UnrealGameSync/UnrealGameSync/UnrealGameSync.csproj
@@ -26,6 +26,7 @@
   </PropertyGroup>
   <ItemGroup>
     <PackageReference Include="EnvDTE" Version="17.11.40262" />
+    <PackageReference Include="Microsoft.VisualStudio.Setup.Configuration.Interop" Version="3.14.2075" />
     <PackageReference Include="Sentry" Version="3.15.0" />
     <PackageReference Include="Serilog.Expressions" Version="3.2.1" />
     <PackageReference Include="Serilog.Extensions.Hosting" Version="4.2.0" />
@@ -52,12 +53,6 @@
     </None>
   </ItemGroup>
 
-  <ItemGroup>
-    <Reference Include="Microsoft.VisualStudio.Setup.Configuration.Interop">
-      <HintPath>..\..\..\..\Binaries\ThirdParty\VisualStudio\Microsoft.VisualStudio.Setup.Configuration.Interop.dll</HintPath>
-      <EmbedInteropTypes>true</EmbedInteropTypes>
-    </Reference>
-  </ItemGroup>
   <ItemGroup>
     <Compile Update="Properties\Settings.Designer.cs">
       <DesignTimeSharedInput>True</DesignTimeSharedInput>
