supressing analysis issue for now

#Synced-CL 51609771

diff --git a/Engine/Source/Programs/UnrealGameSync/UnrealGameSyncShared/WorkspaceUpdate.cs b/Engine/Source/Programs/UnrealGameSync/UnrealGameSyncShared/WorkspaceUpdate.cs
--- a/Engine/Source/Programs/UnrealGameSync/UnrealGameSyncShared/WorkspaceUpdate.cs
+++ b/Engine/Source/Programs/UnrealGameSync/UnrealGameSyncShared/WorkspaceUpdate.cs
@@ -1858,7 +1858,9 @@ namespace UnrealGameSync
 			return syncPaths.Select(x => project.ClientRootPath + x).ToList();
 		}
 
+#pragma warning disable IDE0060 // Remove unused parameter
 		public static List<string> GetRelativeSyncPaths(ProjectInfo project, WorkspaceUpdateOptions options, IReadOnlyList<string>? syncFilter, ISet<string> additionalPathsToSync, ISet<string> baseContentPaths)
+#pragma warning restore IDE0060 // Remove unused parameter
 		{
 			bool syncAllProjects = (options & WorkspaceUpdateOptions.SyncAllProjects) != 0;
 			bool bSyncLocalProjects = (options & WorkspaceUpdateOptions.SyncLocalProjects) != 0;
