Added message to prompt users to clean workspace after switching stream.

#Synced-CL 49593614

diff --git a/Engine/Source/Programs/UnrealGameSync/UnrealGameSync/Controls/StatusPanel.cs b/Engine/Source/Programs/UnrealGameSync/UnrealGameSync/Controls/StatusPanel.cs
--- a/Engine/Source/Programs/UnrealGameSync/UnrealGameSync/Controls/StatusPanel.cs
+++ b/Engine/Source/Programs/UnrealGameSync/UnrealGameSync/Controls/StatusPanel.cs
@@ -455,6 +455,7 @@ namespace UnrealGameSync
 	{
 		const float LineSpacing = 1.35f;
 		const int InfoMargin = 15;
+		const int TopMargin = 10;
 
 		InvertibleImage _projectLogo = new InvertibleImage(Properties.Resources.DefaultProjectLogo, false);
 		Rectangle _projectLogoBounds;
@@ -731,7 +732,7 @@ namespace UnrealGameSync
 			float totalLineHeight = _lines.Sum(x => x.LineHeight);
 
 			// Space out all the lines
-			float lineY = (bodyHeight - totalLineHeight * (int)(Font.Height * LineSpacing)) / 2;
+			float lineY = ((bodyHeight - totalLineHeight * (int)(Font.Height * LineSpacing)) / 2) + TopMargin;
 			foreach (StatusLine line in _lines)
 			{
 				lineY += (int)(Font.Height * LineSpacing * line.LineHeight * 0.5f);
diff --git a/Engine/Source/Programs/UnrealGameSync/UnrealGameSync/Controls/WorkspaceControl.cs b/Engine/Source/Programs/UnrealGameSync/UnrealGameSync/Controls/WorkspaceControl.cs
--- a/Engine/Source/Programs/UnrealGameSync/UnrealGameSync/Controls/WorkspaceControl.cs
+++ b/Engine/Source/Programs/UnrealGameSync/UnrealGameSync/Controls/WorkspaceControl.cs
@@ -506,6 +506,7 @@ namespace UnrealGameSync
 			_perforceMonitor.OnUpdate += UpdateBuildListCallback;
 			_perforceMonitor.OnUpdateMetadata += UpdateBuildMetadataCallback;
 			_perforceMonitor.OnStreamChange += StreamChanged;
+			_perforceMonitor.OnStreamViewChange += StreamViewChanged;
 			_perforceMonitor.OnLoginExpired += LoginExpired;
 
 			string projectOverride = String.Empty;
@@ -978,6 +979,23 @@ namespace UnrealGameSync
 		}
 
 		private void StreamChanged()
+		{
+			RefreshAfterStreamChange();
+
+			// Flag that the stream has been switched so we can prompt to clean the workspace
+			// We only do this when the stream itself has switched not the view
+			_workspace?.ModifyState(x =>
+			{
+				x.StreamSwitched = true;
+			});
+		}
+
+		private void StreamViewChanged()
+		{
+			RefreshAfterStreamChange();
+		}
+
+		private void RefreshAfterStreamChange()
 		{
 			_owner.StreamChanged(this);
 
@@ -988,17 +1006,6 @@ namespace UnrealGameSync
 				x.CurrentChangeNumber = -1;
 				x.CurrentCodeChangeNumber = -1;
 			});
-
-			/*
-			StatusPanel.SuspendDisplay();
-
-			string PrevSelectedFileName = SelectedFileName;
-			if(TryCloseProject())
-			{
-				OpenProject(PrevSelectedFileName);
-			}
-
-			StatusPanel.ResumeDisplay();*/
 		}
 
 		private void LoginExpired()
@@ -4306,13 +4313,28 @@ namespace UnrealGameSync
 				programsLine.AddLink("More... \u25BE", FontStyle.Regular, (p, r) => { ShowActionsMenu(r); });
 				lines.Add(programsLine);
 
+				if (_workspace.State.StreamSwitched)
+				{
+					ApplicationTheme theme = ProgramApplicationContext.GetTheme();
+
+					lines.Add(new StatusLine() { LineHeight = 0.5f });
+
+					StatusLine streamSwitchLine = new StatusLine();
+					streamSwitchLine.AddIcon(Properties.Resources.StatusIcons, new Size(16, 16), 4);
+					streamSwitchLine.AddText("  Workspace may be broken after switching stream  ", theme.WarningColor);
+					streamSwitchLine.AddLink("Clean Workspace To Fix", FontStyle.Bold | FontStyle.Underline, () => { CleanWorkspace(); });
+
+					lines.Add(streamSwitchLine);
+				}
+
 				// Get the summary of the last sync
 				if (_workspace.State.LastSyncChangeNumber > 0)
 				{
 					string? summaryText;
 					if (_workspace.State.LastSyncChangeNumber == _workspace.CurrentChangeNumber && _workspace.State.LastSyncResult == WorkspaceUpdateResult.Success && _workspace.State.LastSyncTime.HasValue)
 					{
-						lines.Add(new StatusLine() { LineHeight = 0.5f });
+						float spacing = _workspace.State.StreamSwitched ? 0.3f : 0.5f;
+						lines.Add(new StatusLine() { LineHeight = spacing });
 
 						StatusLine successLine = new StatusLine();
 						successLine.AddIcon(Properties.Resources.StatusIcons, new Size(16, 16), 0);
@@ -4321,7 +4343,8 @@ namespace UnrealGameSync
 					}
 					else if (GetLastUpdateMessage(_workspace.State.LastSyncResult, _workspace.State.LastSyncResultMessage, out summaryText))
 					{
-						lines.Add(new StatusLine() { LineHeight = 0.5f });
+						float spacing = _workspace.State.StreamSwitched ? 0.3f : 0.5f;
+						lines.Add(new StatusLine() { LineHeight = spacing });
 
 						int summaryTextLength = summaryText.IndexOf('\n', StringComparison.Ordinal);
 						if (summaryTextLength == -1)
@@ -6570,6 +6593,11 @@ namespace UnrealGameSync
 		}
 
 		private void MoreToolsContextMenu_CleanWorkspace_Click(object sender, EventArgs e)
+		{
+			CleanWorkspace();
+		}
+
+		private void CleanWorkspace()
 		{
 			if (!WaitForProgramsToFinish())
 			{
@@ -6601,6 +6629,7 @@ namespace UnrealGameSync
 				_workspace.ModifyState(x =>
 				{
 					x.LastSyncEditorArchive = "0";
+					x.StreamSwitched = false;
 				});
 			}
 		}
diff --git a/Engine/Source/Programs/UnrealGameSync/UnrealGameSync/PerforceMonitor.cs b/Engine/Source/Programs/UnrealGameSync/UnrealGameSync/PerforceMonitor.cs
--- a/Engine/Source/Programs/UnrealGameSync/UnrealGameSync/PerforceMonitor.cs
+++ b/Engine/Source/Programs/UnrealGameSync/UnrealGameSync/PerforceMonitor.cs
@@ -94,6 +94,7 @@ namespace UnrealGameSync
 		public event Action? OnUpdate;
 		public event Action? OnUpdateMetadata;
 		public event Action? OnStreamChange;
+		public event Action? OnStreamViewChange;
 		public event Action? OnLoginExpired;
 
 		readonly object _lockObject = new object();
@@ -140,6 +141,7 @@ namespace UnrealGameSync
 			OnUpdate = null;
 			OnUpdateMetadata = null;
 			OnStreamChange = null;
+			OnStreamViewChange = null;
 			OnLoginExpired = null;
 
 			if (_workerTask != null)
@@ -302,7 +304,7 @@ namespace UnrealGameSync
 
 									if (isVirtual && cache!.Streams.Any())
 									{
-										_synchronizationContext.Post(_ => OnStreamChange?.Invoke(), null);
+										_synchronizationContext.Post(_ => OnStreamViewChange?.Invoke(), null);
 									}
 
 									PerforceMonitorCache newCache = new PerforceMonitorCache()
diff --git a/Engine/Source/Programs/UnrealGameSync/UnrealGameSyncShared/ApplicationTheme.cs b/Engine/Source/Programs/UnrealGameSync/UnrealGameSyncShared/ApplicationTheme.cs
--- a/Engine/Source/Programs/UnrealGameSync/UnrealGameSyncShared/ApplicationTheme.cs
+++ b/Engine/Source/Programs/UnrealGameSync/UnrealGameSyncShared/ApplicationTheme.cs
@@ -65,7 +65,7 @@ namespace UnrealGameSync
 		[Theme]
 		public Color SuccessColor { get; init; } = Color.FromArgb(28, 180, 64);
 		[Theme]
-		public Color WarningColor { get; init; } = Color.FromArgb(128, 128, 0);
+		public Color WarningColor { get; init; } = Color.FromArgb(SystemColors.HotTrack.B, SystemColors.HotTrack.G, SystemColors.HotTrack.R);
 		[Theme]
 		public Color ErrorColor { get; init; } = Color.FromArgb(189, 54, 47);
 		[Theme]
@@ -134,7 +134,7 @@ namespace UnrealGameSync
 				LinkColorHover = Color.FromArgb(255, 170, 84),
 				LinkColorActive = Color.FromArgb(127, 85, 84),
 				SuccessColor = Color.FromArgb(74, 227, 110),
-				WarningColor = Color.FromArgb(255, 255, 128),
+				WarningColor = Color.FromArgb(255, 170, 84),
 				ErrorColor = Color.FromArgb(208, 74, 67),
 				
 				StandardNotificationBGTint = Color.FromArgb(15, 15, 36),
diff --git a/Engine/Source/Programs/UnrealGameSync/UnrealGameSyncShared/WorkspaceState.cs b/Engine/Source/Programs/UnrealGameSync/UnrealGameSyncShared/WorkspaceState.cs
--- a/Engine/Source/Programs/UnrealGameSync/UnrealGameSyncShared/WorkspaceState.cs
+++ b/Engine/Source/Programs/UnrealGameSync/UnrealGameSyncShared/WorkspaceState.cs
@@ -27,6 +27,8 @@ namespace UnrealGameSync
 		public string ProjectPath { get; set; } = String.Empty;
 		public string? StreamName { get; set; }
 
+		public bool StreamSwitched { get; set; } = false; // Whether the stream has been switched since the last sync
+
 		public string ProjectIdentifier { get; set; } = String.Empty; // Should be fully reset if this changes
 		public bool IsEnterpriseProject { get; set; }
 
@@ -98,6 +100,8 @@ namespace UnrealGameSync
 			AdditionalChangeNumbers.Clear();
 			AdditionalChangeNumbers.AddRange(other.AdditionalChangeNumbers);
 
+			StreamSwitched = other.StreamSwitched;
+
 			LastSyncChangeNumber = other.LastSyncChangeNumber;
 			LastSyncResult = other.LastSyncResult;
 			LastSyncResultMessage = other.LastSyncResultMessage;
@@ -149,6 +153,7 @@ namespace UnrealGameSync
 		public string BranchPath => _inner.BranchPath;
 		public string ProjectPath => _inner.ProjectPath;
 		public string? StreamName => _inner.StreamName;
+		public bool StreamSwitched => _inner.StreamSwitched;
 		public string ProjectIdentifier => _inner.ProjectIdentifier;
 		public bool IsEnterpriseProject => _inner.IsEnterpriseProject;
 
