UnrealGameSyncAutomation: Update to TargetFramework net10.0

#rnx
#Synced-CL 52406569
#Synced-CL 52406569

diff --git a/Engine/Source/Programs/UnrealGameSync/UnrealGameSyncAutomation/UnrealGameSyncAutomation.csproj b/Engine/Source/Programs/UnrealGameSync/UnrealGameSyncAutomation/UnrealGameSyncAutomation.csproj
--- a/Engine/Source/Programs/UnrealGameSync/UnrealGameSyncAutomation/UnrealGameSyncAutomation.csproj
+++ b/Engine/Source/Programs/UnrealGameSync/UnrealGameSyncAutomation/UnrealGameSyncAutomation.csproj
@@ -2,7 +2,7 @@
 <Project Sdk="Microsoft.NET.Sdk">
   <PropertyGroup>
     <OutputType>WinExe</OutputType>
-    <TargetFramework>net8.0-windows</TargetFramework>
+    <TargetFramework>net10.0-windows</TargetFramework>
     <UseWindowsForms>true</UseWindowsForms>
     <GenerateAssemblyInfo>false</GenerateAssemblyInfo>
     <RootNamespace>UnrealGameSyncAutomation</RootNamespace>
