[NuGet] Update EnvDTE to 17.14.40260, Microsoft.VisualStudio.Debugger.InteropA to 17.14.40260, Microsoft.VisualStudio.DiagnosticsHub.BenchmarkDotNetDiagnosers to 18.6.37125.3, BenchmarkDotNet to 0.15.8

#rnx
#Synced-CL 52554515

diff --git a/Engine/Source/Programs/Shared/EpicGames.Benchmarks/EpicGames.Benchmarks.csproj b/Engine/Source/Programs/Shared/EpicGames.Benchmarks/EpicGames.Benchmarks.csproj
--- a/Engine/Source/Programs/Shared/EpicGames.Benchmarks/EpicGames.Benchmarks.csproj
+++ b/Engine/Source/Programs/Shared/EpicGames.Benchmarks/EpicGames.Benchmarks.csproj
@@ -5,8 +5,8 @@
   </PropertyGroup>
 
   <ItemGroup>
-    <PackageReference Include="BenchmarkDotNet" Version="0.15.2" />
-    <PackageReference Include="Microsoft.VisualStudio.DiagnosticsHub.BenchmarkDotNetDiagnosers" Version="18.0.36313.1" />
+    <PackageReference Include="BenchmarkDotNet" Version="0.15.8" />
+    <PackageReference Include="Microsoft.VisualStudio.DiagnosticsHub.BenchmarkDotNetDiagnosers" Version="18.6.37125.3" />
   </ItemGroup>
 
   <ItemGroup>
diff --git a/Engine/Source/Programs/UnrealGameSync/UnrealGameSync/UnrealGameSync.csproj b/Engine/Source/Programs/UnrealGameSync/UnrealGameSync/UnrealGameSync.csproj
--- a/Engine/Source/Programs/UnrealGameSync/UnrealGameSync/UnrealGameSync.csproj
+++ b/Engine/Source/Programs/UnrealGameSync/UnrealGameSync/UnrealGameSync.csproj
@@ -25,7 +25,7 @@
     <DefineConstants>$(DefineConstants);WITH_AUTOUPDATE</DefineConstants>
   </PropertyGroup>
   <ItemGroup>
-    <PackageReference Include="EnvDTE" Version="17.11.40262" />
+    <PackageReference Include="EnvDTE" Version="17.14.40260" />
     <PackageReference Include="Microsoft.VisualStudio.Setup.Configuration.Interop" Version="3.14.2075" />
     <PackageReference Include="Sentry" Version="3.15.0" />
     <PackageReference Include="Serilog.Expressions" Version="3.2.1" />
