[UGS] Preserve CurrentChangeNumber on sync failure to keep fast sync working on retry

#Synced-CL 52551117

diff --git a/Engine/Source/Programs/UnrealGameSync/UnrealGameSyncShared/WorkspaceUpdate.cs b/Engine/Source/Programs/UnrealGameSync/UnrealGameSyncShared/WorkspaceUpdate.cs
--- a/Engine/Source/Programs/UnrealGameSync/UnrealGameSyncShared/WorkspaceUpdate.cs
+++ b/Engine/Source/Programs/UnrealGameSync/UnrealGameSyncShared/WorkspaceUpdate.cs
@@ -1056,6 +1056,11 @@ namespace UnrealGameSync
 
 					syncTree.GetOptimizedSyncCommands(project.ClientRootPath, Context.ChangeNumber, batchBuilder);
 
+					// Save the current changelist before clearing, so we can restore it if the sync
+					// transfers files successfully but fails in a post-sync step (e.g. files to resolve).
+					// This preserves the ability to use fast sync on retry instead of forcing a full sync.
+					int previousChangeNumber = state.CurrentChangeNumber;
+
 					// Clear the current sync changelist, in case we cancel
 					if (!Context.Options.HasFlag(WorkspaceUpdateOptions.SyncSingleChange))
 					{
@@ -1077,6 +1082,14 @@ namespace UnrealGameSync
 						(WorkspaceUpdateResult, string) syncResult = await SyncFileRevisions(perforce, "Syncing files...", Context, batchBuilder.Batches, remainingDepotPaths, Progress, logger, cancellationToken);
 						if (syncResult.Item1 != WorkspaceUpdateResult.Success)
 						{
+							// For retryable results (e.g. FilesToClobber), restore the previous changelist
+							// so fast sync remains usable on the next attempt instead of forcing a full sync.
+							if (syncResult.Item1 == WorkspaceUpdateResult.FilesToClobber
+								&& !Context.Options.HasFlag(WorkspaceUpdateOptions.SyncSingleChange)
+								&& previousChangeNumber > 0)
+							{
+								state = stateMgr.Modify(x => x.CurrentChangeNumber = previousChangeNumber);
+							}
 							transferStopwatch.AddData(new { SyncResult = syncResult.Item2.ToString(), CompletedFilesFiles = syncDepotPaths.Count - remainingDepotPaths.Count });
 							return syncResult;
 						}
@@ -1087,6 +1100,14 @@ namespace UnrealGameSync
 						times.Add(new Tuple<string, TimeSpan>("Sync Files Transfer", stop));
 					}
 
+					// File transfer succeeded. Restore the previous changelist so that if a post-sync
+					// step fails (e.g. files to resolve, version file update), fast sync remains usable
+					// on the next retry. The success path below will overwrite this with Context.ChangeNumber.
+					if (!Context.Options.HasFlag(WorkspaceUpdateOptions.SyncSingleChange) && previousChangeNumber > 0)
+					{
+						state = stateMgr.Modify(x => x.CurrentChangeNumber = previousChangeNumber);
+					}
+
 					int versionChangeNumber = -1;
 					if (Context.Options.HasFlag(WorkspaceUpdateOptions.Sync) && !Context.Options.HasFlag(WorkspaceUpdateOptions.UpdateFilter))
 					{
