Updated cancelling the ClobberWindow in UGS to clear the sync state to prevent the user thinking they are synced to the previous CL when some files may have been updated.

#rnx
#Synced-CL 49593614

diff --git a/Engine/Source/Programs/UnrealGameSync/UnrealGameSync/Controls/WorkspaceControl.cs b/Engine/Source/Programs/UnrealGameSync/UnrealGameSync/Controls/WorkspaceControl.cs
--- a/Engine/Source/Programs/UnrealGameSync/UnrealGameSync/Controls/WorkspaceControl.cs
+++ b/Engine/Source/Programs/UnrealGameSync/UnrealGameSync/Controls/WorkspaceControl.cs
@@ -1436,6 +1436,10 @@ namespace UnrealGameSync
 						StartWorkspaceUpdate(context, _updateCallback);
 						return;
 					}
+					else
+					{
+						_workspace?.ModifyState(x => x.CurrentChangeNumber = -1);
+					}
 				}
 			}
 			else if (result == WorkspaceUpdateResult.FailedToCompileWithCleanWorkspace)
