Add a way of ignoring the clean warning after re-targetting a workspace to a dfferent stream in UGS.

#rnx
#Synced-CL 49593614

#Synced-CL 49593614

diff --git a/Engine/Source/Programs/UnrealGameSync/UnrealGameSync/Controls/WorkspaceControl.cs b/Engine/Source/Programs/UnrealGameSync/UnrealGameSync/Controls/WorkspaceControl.cs
--- a/Engine/Source/Programs/UnrealGameSync/UnrealGameSync/Controls/WorkspaceControl.cs
+++ b/Engine/Source/Programs/UnrealGameSync/UnrealGameSync/Controls/WorkspaceControl.cs
@@ -4327,6 +4327,14 @@ namespace UnrealGameSync
 					streamSwitchLine.AddIcon(Properties.Resources.StatusIcons, new Size(16, 16), 4);
 					streamSwitchLine.AddText("  Workspace may be broken after switching stream  ", theme.WarningColor);
 					streamSwitchLine.AddLink("Clean Workspace To Fix", FontStyle.Bold | FontStyle.Underline, () => { CleanWorkspace(); });
+					streamSwitchLine.AddText("  |  ");
+					streamSwitchLine.AddLink("Ignore", FontStyle.Regular, () => 
+					{ 
+						_workspace?.ModifyState(x => 
+						{ 
+							x.StreamSwitched = false; 
+						}); 
+					});
 
 					lines.Add(streamSwitchLine);
 				}
