Disabling UGS clean prompt after changing stream whilst we address issues with cleaning directly after a stream swap.

#rnx
#Synced-CL 49593614

diff --git a/Engine/Source/Programs/UnrealGameSync/UnrealGameSync/Controls/WorkspaceControl.cs b/Engine/Source/Programs/UnrealGameSync/UnrealGameSync/Controls/WorkspaceControl.cs
--- a/Engine/Source/Programs/UnrealGameSync/UnrealGameSync/Controls/WorkspaceControl.cs
+++ b/Engine/Source/Programs/UnrealGameSync/UnrealGameSync/Controls/WorkspaceControl.cs
@@ -984,10 +984,10 @@ namespace UnrealGameSync
 
 			// Flag that the stream has been switched so we can prompt to clean the workspace
 			// We only do this when the stream itself has switched not the view
-			_workspace?.ModifyState(x =>
-			{
-				x.StreamSwitched = true;
-			});
+			//_workspace?.ModifyState(x =>
+			//{
+			//	x.StreamSwitched = true;
+			//});
 		}
 
 		private void StreamViewChanged()
